Add functionality for serializing arbitrary object graphs to a binary archive using attribute annotations rather than explicit serialization methods. Generates dynamic IL methods at runtime to minimize boxing/unboxing and dispatch overhead.

#rnx

diff --git a/Engine/Source/Programs/Shared/EpicGames.Core/BinaryArchive.cs b/Engine/Source/Programs/Shared/EpicGames.Core/BinaryArchive.cs
new file mode 100644
--- /dev/null
+++ b/Engine/Source/Programs/Shared/EpicGames.Core/BinaryArchive.cs
@@ -0,0 +1,455 @@
+using System;
+using System.Collections.Concurrent;
+using System.Collections.Generic;
+using System.Linq.Expressions;
+using System.Reflection;
+using System.Reflection.Emit;
+using System.Text;
+
+#nullable enable
+
+namespace EpicGames.Core
+{
+	/// <summary>
+	/// Common functionality for dealing with generic type serialization in binary archives
+	/// </summary>
+	public static class BinaryArchive
+	{
+		class Instantiator<T> where T : class, new()
+		{
+			public static T Instance = new T();
+		}
+
+		sealed class TypeSerializer
+		{
+			public MethodInfo ReadMethodInfo { get; }
+			public Func<BinaryArchiveReader, object> ReadMethod { get; }
+			public MethodInfo WriteMethodInfo { get; }
+			public Action<BinaryArchiveWriter, object> WriteMethod { get; }
+
+			public TypeSerializer(MethodInfo ReadMethodInfo, MethodInfo WriteMethodInfo)
+			{
+				this.ReadMethodInfo = ReadMethodInfo;
+				this.WriteMethodInfo = WriteMethodInfo;
+
+				Type Type = ReadMethodInfo.ReturnType;
+				this.ReadMethod = CreateBoxedReadMethod(Type, ReadMethodInfo);
+				this.WriteMethod = CreateBoxedWriteMethod(Type, WriteMethodInfo);
+			}
+
+			static Func<BinaryArchiveReader, object> CreateBoxedReadMethod(Type Type, MethodInfo MethodInfo)
+			{
+				DynamicMethod ReaderMethod = new DynamicMethod($"Boxed_{MethodInfo.Name}", typeof(object), new[] { typeof(BinaryArchiveReader) });
+
+				ILGenerator Generator = ReaderMethod.GetILGenerator(64);
+				Generator.Emit(OpCodes.Ldarg_0);
+				Generator.EmitCall(OpCodes.Call, MethodInfo, null);
+				if (!Type.IsClass)
+				{
+					Generator.Emit(OpCodes.Box);
+				}
+				Generator.Emit(OpCodes.Ret);
+
+				return (Func<BinaryArchiveReader, object>)ReaderMethod.CreateDelegate(typeof(Func<BinaryArchiveReader, object>));
+			}
+
+			static Action<BinaryArchiveWriter, object> CreateBoxedWriteMethod(Type Type, MethodInfo MethodInfo)
+			{
+				DynamicMethod WriterMethod = new DynamicMethod($"Boxed_{MethodInfo.Name}", typeof(void), new[] { typeof(BinaryArchiveWriter), typeof(object) });
+
+				ILGenerator Generator = WriterMethod.GetILGenerator(64);
+				Generator.Emit(OpCodes.Ldarg_0);
+				Generator.Emit(OpCodes.Ldarg_1);
+				Generator.Emit(OpCodes.Unbox_Any, Type);
+				Generator.EmitCall(OpCodes.Call, MethodInfo, null);
+				Generator.Emit(OpCodes.Ret);
+
+				return (Action<BinaryArchiveWriter, object>)WriterMethod.CreateDelegate(typeof(Action<BinaryArchiveWriter, object>));
+			}
+
+			public static TypeSerializer Create<T>(Expression<Func<BinaryArchiveReader, T>> ReaderExpr, Expression<Action<BinaryArchiveWriter, T>> WriterExpr)
+			{
+				MethodInfo ReadMethod = ((MethodCallExpression)ReaderExpr.Body).Method;
+				MethodInfo WriteMethod = ((MethodCallExpression)WriterExpr.Body).Method;
+				return new TypeSerializer(ReadMethod, WriteMethod);
+			}
+		}
+
+		static ConcurrentDictionary<Type, TypeSerializer> TypeToSerializerInfo = new ConcurrentDictionary<Type, TypeSerializer>()
+		{
+			[typeof(byte)] = TypeSerializer.Create(Reader => Reader.ReadByte(), (Writer, Value) => Writer.WriteByte(Value)),
+			[typeof(sbyte)] = TypeSerializer.Create(Reader => Reader.ReadSignedByte(), (Writer, Value) => Writer.WriteSignedByte(Value)),
+			[typeof(short)] = TypeSerializer.Create(Reader => Reader.ReadShort(), (Writer, Value) => Writer.WriteShort(Value)),
+			[typeof(ushort)] = TypeSerializer.Create(Reader => Reader.ReadUnsignedShort(), (Writer, Value) => Writer.WriteUnsignedShort(Value)),
+			[typeof(int)] = TypeSerializer.Create(Reader => Reader.ReadInt(), (Writer, Value) => Writer.WriteInt(Value)),
+			[typeof(uint)] = TypeSerializer.Create(Reader => Reader.ReadUnsignedInt(), (Writer, Value) => Writer.WriteUnsignedInt(Value)),
+			[typeof(long)] = TypeSerializer.Create(Reader => Reader.ReadLong(), (Writer, Value) => Writer.WriteLong(Value)),
+			[typeof(ulong)] = TypeSerializer.Create(Reader => Reader.ReadUnsignedLong(), (Writer, Value) => Writer.WriteUnsignedLong(Value)),
+			[typeof(string)] = TypeSerializer.Create(Reader => Reader.ReadString(), (Writer, Value) => Writer.WriteString(Value))
+		};
+
+		static ConcurrentDictionary<Type, ContentHash> TypeToDigest = new ConcurrentDictionary<Type, ContentHash>();
+		static ConcurrentDictionary<ContentHash, Type> DigestToType = new ConcurrentDictionary<ContentHash, Type>();
+
+		static MethodInfo GetMethodInfo(Expression<Action> Expr)
+		{
+			return ((MethodCallExpression)Expr.Body).Method;
+		}
+
+		static MethodInfo GetMethodInfo<T>(Expression<Action<T>> Expr)
+		{
+			return ((MethodCallExpression)Expr.Body).Method;
+		}
+
+		static MethodInfo GetGenericMethodInfo<T, T1>(Expression<Action<T, T1>> Expr)
+		{
+			return ((MethodCallExpression)Expr.Body).Method.GetGenericMethodDefinition();
+		}
+
+		static MethodInfo ReadBoolMethodInfo = GetMethodInfo<BinaryArchiveReader>(x => x.ReadBool());
+		static MethodInfo ReadIntMethodInfo = GetMethodInfo<BinaryArchiveReader>(x => x.ReadInt());
+		static MethodInfo ReadArrayMethodInfo = GetGenericMethodInfo<BinaryArchiveReader, Func<int>>((x, y) => x.ReadArray(y));
+		static MethodInfo ReadPolymorphicObjectMethodInfo = GetMethodInfo(() => ReadObject(null!));
+		static MethodInfo WriteBoolMethodInfo = GetMethodInfo<BinaryArchiveWriter>(x => x.WriteBool(false));
+		static MethodInfo WriteIntMethodInfo = GetMethodInfo<BinaryArchiveWriter>(x => x.WriteInt(0));
+		static MethodInfo WriteArrayMethodInfo = GetGenericMethodInfo<BinaryArchiveWriter, Action<int>>((x, y) => x.WriteArray(null!, y));
+		static MethodInfo WritePolymorphicObjectMethodInfo = GetMethodInfo(() => WritePolymorphicObject(null!, null!));
+
+		/// <summary>
+		/// Reads an arbitrary value from the given archive.
+		/// </summary>
+		/// <typeparam name="T">Type to read</typeparam>
+		/// <param name="Reader">Reader to serialize from</param>
+		/// <returns>The given value</returns>
+		public static T Read<T>(this BinaryArchiveReader Reader)
+		{
+			Type Type = typeof(T);
+			if (Type.IsClass && !Type.IsSealed)
+			{
+				return (T)ReadObject(Reader);
+			}
+			else
+			{
+				return (T)FindOrAddSerializerInfo(typeof(T)).ReadMethod(Reader)!;
+			}
+		}
+
+		/// <summary>
+		/// Writes an arbitrary type to the given archive. May be an object or value type.
+		/// </summary>
+		/// <typeparam name="T">The type to write</typeparam>
+		/// <param name="Writer">Writer to serialize to</param>
+		/// <param name="Value">Value to write</param>
+		public static void Write<T>(this BinaryArchiveWriter Writer, T Value)
+		{
+			Type Type = typeof(T);
+			if (Type.IsClass && !Type.IsSealed)
+			{
+				Writer.WriteObjectReference(Value, () => WriteNewPolymorphicObject(Writer, Value!));
+			}
+			else
+			{
+				FindOrAddSerializerInfo(Type).WriteMethod(Writer, Value!);
+			}
+		}
+
+		/// <summary>
+		/// Registers the given type. Allows serializing/deserializing it through calls to ReadType/WriteType.
+		/// </summary>
+		/// <param name="Type">Type to register</param>
+		public static void RegisterType(Type Type)
+		{
+			ContentHash Digest = ContentHash.MD5($"{Type.Assembly.FullName}\n{Type.FullName}");
+			TypeToDigest.TryAdd(Type, Digest);
+			DigestToType.TryAdd(Digest, Type);
+		}
+
+		/// <summary>
+		/// Registers all types in the given assembly with the <see cref="BinarySerializableAttribute"/> attribute for serialization
+		/// </summary>
+		/// <param name="Assembly">Assembly to search in</param>
+		public static void RegisterTypes(Assembly Assembly)
+		{
+			foreach (Type Type in Assembly.GetTypes())
+			{
+				if (Type.GetCustomAttribute<BinarySerializableAttribute>() != null)
+				{
+					RegisterType(Type);
+				}
+			}
+		}
+
+		/// <summary>
+		/// Reads a tyoe from the archive
+		/// </summary>
+		/// <param name="Reader">Reader to deserializer the type from</param>
+		/// <returns>The matching type</returns>
+		public static Type ReadType(this BinaryArchiveReader Reader)
+		{
+			return Reader.ReadObjectReference(() => DigestToType[Reader.ReadContentHash()]);
+		}
+
+		/// <summary>
+		/// Writes a type to an archive
+		/// </summary>
+		/// <param name="Writer">Writer to serialize to</param>
+		/// <param name="Type">The type to serialize</param>
+		public static void WriteType(this BinaryArchiveWriter Writer, Type Type)
+		{
+			Writer.WriteObjectReference(Type, () => Writer.WriteContentHash(TypeToDigest[Type]));
+		}
+
+		static object ReadNewPolymorphicObject(BinaryArchiveReader Reader)
+		{
+			Type FinalType = Reader.ReadType();
+			TypeSerializer Info = FindOrAddSerializerInfo(FinalType);
+			return Info.ReadMethod(Reader)!;
+		}
+
+		public static object ReadObject(this BinaryArchiveReader Reader)
+		{
+			return Reader.ReadObjectReference(() => ReadNewPolymorphicObject(Reader));
+		}
+
+		static void WriteNewPolymorphicObject(BinaryArchiveWriter Writer, object Value)
+		{
+			Type ActualType = Value.GetType();
+			Writer.WriteType(ActualType);
+
+			TypeSerializer Info = FindOrAddSerializerInfo(ActualType);
+			Info.WriteMethod(Writer, Value);
+		}
+
+		static void WritePolymorphicObject(this BinaryArchiveWriter Writer, object? Value)
+		{
+			Writer.WriteObjectReference(Value, () => WriteNewPolymorphicObject(Writer, Value!));
+		}
+
+		/// <summary>
+		/// Finds an existing serializer for the given type, or creates one
+		/// </summary>
+		/// <param name="Type">Type to create a serializer for</param>
+		/// <param name="ConverterType">Type of the converter to use</param>
+		/// <returns>Instance of the type serializer</returns>
+		static TypeSerializer FindOrAddSerializerInfo(Type Type, Type? ConverterType = null)
+		{
+			Type SerializerKey = ConverterType ?? Type;
+
+			// Get the serializer info
+			TypeSerializer? SerializerInfo;
+			if (!TypeToSerializerInfo.TryGetValue(SerializerKey, out SerializerInfo))
+			{
+				lock (TypeToSerializerInfo)
+				{
+					SerializerInfo = CreateTypeSerializer(Type, ConverterType);
+					TypeToSerializerInfo[SerializerKey] = SerializerInfo;
+				}
+			}
+
+			return SerializerInfo;
+		}
+
+		/// <summary>
+		/// Creates a serializer for the given type
+		/// </summary>
+		/// <param name="Type">Type to create a serializer from</param>
+		/// <param name="ConverterType">Converter for the type</param>
+		/// <returns>New instance of a type serializer</returns>
+		static TypeSerializer CreateTypeSerializer(Type Type, Type? ConverterType = null)
+		{
+			// If there's a specific converter defined, generate serialization methods using that
+			if (ConverterType != null)
+			{
+				// Make sure the converter is derived from IBinaryConverter
+				Type InterfaceType = typeof(IBinaryConverter<>).MakeGenericType(Type);
+				if (!InterfaceType.IsAssignableFrom(ConverterType))
+				{
+					throw new NotImplementedException($"Converter does not implement IBinaryArchiveConverter<{Type.Name}>");
+				}
+
+				// Instantiate the converter, and store it in a static variable
+				Type InstantiatorType = typeof(Instantiator<>).MakeGenericType(ConverterType);
+				FieldInfo ConverterField = InstantiatorType.GetField(nameof(Instantiator<object>.Instance))!;
+
+				// Get the mapping between interface methods and instance methods
+				InterfaceMapping InterfaceMapping = ConverterType.GetInterfaceMap(InterfaceType);
+
+				// Create a reader method
+				DynamicMethod ReaderMethod = new DynamicMethod($"BinaryArchiveReader_Dynamic_{Type.Name}_With_{ConverterType.Name}", Type, new[] { typeof(BinaryArchiveReader) });
+				int ReaderIdx = Array.FindIndex(InterfaceMapping.InterfaceMethods, x => x.Name.Equals(nameof(IBinaryConverter<object>.Read), StringComparison.Ordinal));
+				GenerateConverterReaderForwardingMethod(ReaderMethod, ConverterField, InterfaceMapping.TargetMethods[ReaderIdx]);
+
+				// Create a writer method
+				DynamicMethod WriterMethod = new DynamicMethod($"BinaryArchiveWriter_Dynamic_{Type.Name}_With_{ConverterType.Name}", typeof(void), new[] { typeof(BinaryArchiveWriter), typeof(object) });
+				int WriterIdx = Array.FindIndex(InterfaceMapping.InterfaceMethods, x => x.Name.Equals(nameof(IBinaryConverter<object>.Write), StringComparison.Ordinal));
+				GenerateConverterWriterForwardingMethod(WriterMethod, ConverterField, InterfaceMapping.TargetMethods[WriterIdx]);
+
+				return new TypeSerializer(ReaderMethod, WriterMethod);
+			}
+
+			// Get the default converter type
+			Type? DefaultConverterType;
+			if(BinaryConverter.TryGetConverterType(Type, out DefaultConverterType))
+			{
+				return FindOrAddSerializerInfo(Type, DefaultConverterType);
+			}
+
+			// Check if it's an array
+			if (Type.IsArray)
+			{
+				Type ElementType = Type.GetElementType()!;
+
+				MethodInfo ReaderMethod = ReadArrayMethodInfo.MakeGenericMethod(ElementType);
+				MethodInfo WriterMethod = WriteArrayMethodInfo.MakeGenericMethod(ElementType);
+
+				return new TypeSerializer(ReaderMethod, WriterMethod);
+			}
+
+			// Check if it's a class
+			if (Type.IsClass)
+			{
+				// Get all the class properties
+				List<PropertyInfo> Properties = new List<PropertyInfo>();
+				foreach (PropertyInfo Property in Type.GetProperties(BindingFlags.Public | BindingFlags.NonPublic | BindingFlags.Instance))
+				{
+					MethodInfo? GetMethod = Property.GetGetMethod();
+					MethodInfo? SetMethod = Property.GetSetMethod();
+					if (GetMethod != null && SetMethod != null && Property.GetCustomAttribute<BinaryIgnoreAttribute>() == null)
+					{
+						Properties.Add(Property);
+					}
+				}
+
+				// Find the type constructor
+				ConstructorInfo? ConstructorInfo = Type.GetConstructor(BindingFlags.Instance | BindingFlags.Public | BindingFlags.NonPublic, null, Array.Empty<Type>(), null);
+				if (ConstructorInfo == null)
+				{
+					throw new NotImplementedException($"Type '{Type.Name}' does not have a parameterless constructor");
+				}
+
+				// Create the methods
+				DynamicMethod WriterMethod = new DynamicMethod($"BinaryArchiveWriter_Dynamic_{Type.Name}", typeof(void), new[] { typeof(BinaryArchiveWriter), Type });
+				GenerateClassWriterMethod(WriterMethod, Properties);
+
+				DynamicMethod ReaderMethod = new DynamicMethod($"BinaryArchiveReader_Dynamic_{Type.Name}", Type, new[] { typeof(BinaryArchiveReader) });
+				GenerateClassReaderMethod(ReaderMethod, ConstructorInfo, Properties);
+
+				return new TypeSerializer(ReaderMethod, WriterMethod);
+			}
+
+			throw new NotImplementedException($"Unable to create a serializer for {Type.Name}");
+		}
+
+		static void GenerateClassWriterMethod(DynamicMethod WriterMethod, List<PropertyInfo> Properties)
+		{
+			// Get the IL generator
+			ILGenerator Generator = WriterMethod.GetILGenerator(256 + (Properties.Count * 24));
+
+			// Check if the argument is null, and write 'false' if it is, then early out.
+			Label NonNullInstance = Generator.DefineLabel();
+			Generator.Emit(OpCodes.Ldarg_1);                    // [0] = Arg1 (Instance)
+			Generator.Emit(OpCodes.Brtrue, NonNullInstance);    // [EMPTY] Instance != null -> NonNullInstance
+
+			Generator.Emit(OpCodes.Ldarg_0);                    // [0] = Arg0 (Writer)
+			Generator.Emit(OpCodes.Ldc_I4_0);                   // [1] = false
+			Generator.EmitCall(OpCodes.Call, WriteBoolMethodInfo, null);  // [EMPTY] WriteBool(Write, false)
+
+			Generator.Emit(OpCodes.Ret);                        // [EMPTY] Return
+
+			// Otherwise write true.
+			Generator.MarkLabel(NonNullInstance);
+
+			Generator.Emit(OpCodes.Ldarg_0);                    // [0] = Arg0 (Writer)
+			Generator.Emit(OpCodes.Ldc_I4_1);                   // [1] = true
+			Generator.EmitCall(OpCodes.Call, WriteBoolMethodInfo, null);  // [EMPTY] WriteBool(Writer, true)
+
+			// Write all the properties
+			foreach (PropertyInfo Property in Properties)
+			{
+				BinaryConverterAttribute? Converter = Property.GetCustomAttribute<BinaryConverterAttribute>();
+
+				MethodInfo? WritePropertyMethod;
+				if (Property.PropertyType.IsClass && !Property.PropertyType.IsSealed && Converter == null)
+				{
+					WritePropertyMethod = WritePolymorphicObjectMethodInfo;
+				}
+				else
+				{
+					WritePropertyMethod = FindOrAddSerializerInfo(Property.PropertyType, Converter?.Type).WriteMethodInfo;
+				}
+
+				Generator.Emit(OpCodes.Ldarg_0);											// [0] = Arg0 (BinaryArchiveWriter)
+				Generator.Emit(OpCodes.Ldarg_1);											// [1] = Arg1 (Instance)
+				Generator.EmitCall(OpCodes.Call, Property.GetGetMethod(true)!, null);		// [1] = GetMethod(Instance)
+				Generator.EmitCall(OpCodes.Call, WritePropertyMethod, null);		        // [EMPTY] Write(BinaryArchiveWriter, GetMethod(Instance));
+			}
+
+			// Return
+			Generator.Emit(OpCodes.Ret);
+		}
+
+		static void GenerateClassReaderMethod(DynamicMethod ReaderMethod, ConstructorInfo ConstructorInfo, List<PropertyInfo> Properties)
+		{
+			// Get the IL generator
+			ILGenerator Generator = ReaderMethod.GetILGenerator(256 + (Properties.Count * 24));
+
+			// Check if it was a null instance
+			Generator.Emit(OpCodes.Ldarg_0);                    // [0] = Arg1 (Reader)
+			Generator.EmitCall(OpCodes.Call, ReadBoolMethodInfo, null);   // [0] = Reader.ReadBool()
+
+			Label NonNullInstance = Generator.DefineLabel();
+			Generator.Emit(OpCodes.Brtrue, NonNullInstance);    // Bool != null -> NonNullInstance
+			Generator.Emit(OpCodes.Ldnull);                     // [0] = null
+			Generator.Emit(OpCodes.Ret);                        // return
+
+			// Construct an instance
+			Generator.MarkLabel(NonNullInstance);
+			Generator.Emit(OpCodes.Newobj, ConstructorInfo);    // [0] = new Type()
+
+			// Read all the properties
+			foreach (PropertyInfo Property in Properties)
+			{
+				BinaryConverterAttribute? Converter = Property.GetCustomAttribute<BinaryConverterAttribute>();
+
+				MethodInfo? ReadPropertyMethod;
+				if (Property.PropertyType.IsClass && !Property.PropertyType.IsSealed && Converter == null)
+				{
+					ReadPropertyMethod = ReadPolymorphicObjectMethodInfo;
+				}
+				else
+				{
+					ReadPropertyMethod = FindOrAddSerializerInfo(Property.PropertyType, Converter?.Type).ReadMethodInfo;
+				}
+
+				Generator.Emit(OpCodes.Dup);                                                // [1] = new Type()
+
+				Generator.Emit(OpCodes.Ldarg_0);                                            // [2] = Arg1 (Reader)
+				Generator.Emit(OpCodes.Call, ReadPropertyMethod);	                        // [2] = Reader.ReadMethod() or ReadMethod(Reader)
+
+				Generator.Emit(OpCodes.Call, Property.GetSetMethod(true)!);                 // SetMethod(new Type(), ReadMethod(Reader))
+			}
+
+			// Return the instance
+			Generator.Emit(OpCodes.Ret);                        // [0] = new Type()
+		}
+
+		static void GenerateConverterReaderForwardingMethod(DynamicMethod ReaderMethod, FieldInfo ConverterField, MethodInfo TargetMethod)
+		{
+			ILGenerator Generator = ReaderMethod.GetILGenerator(64);
+			Generator.Emit(OpCodes.Ldsfld, ConverterField);
+			Generator.Emit(OpCodes.Ldarg_0);
+			Generator.Emit(OpCodes.Call, TargetMethod);
+			Generator.Emit(OpCodes.Ret);
+		}
+
+		static void GenerateConverterWriterForwardingMethod(DynamicMethod WriterMethod, FieldInfo ConverterField, MethodInfo TargetMethod)
+		{
+			ILGenerator Generator = WriterMethod.GetILGenerator(64);
+			Generator.Emit(OpCodes.Ldsfld, ConverterField);
+			Generator.Emit(OpCodes.Ldarg_0);
+			Generator.Emit(OpCodes.Ldarg_1);
+			Generator.Emit(OpCodes.Call, TargetMethod);
+			Generator.Emit(OpCodes.Ret);
+		}
+	}
+}
diff --git a/Engine/Source/Programs/Shared/EpicGames.Core/BinaryArchiveAttributes.cs b/Engine/Source/Programs/Shared/EpicGames.Core/BinaryArchiveAttributes.cs
new file mode 100644
--- /dev/null
+++ b/Engine/Source/Programs/Shared/EpicGames.Core/BinaryArchiveAttributes.cs
@@ -0,0 +1,49 @@
+using System;
+using System.Collections.Generic;
+using System.Text;
+
+namespace EpicGames.Core
+{
+	/// <summary>
+	/// Instructs the serializer to ignore this property during serialization
+	/// </summary>
+	[AttributeUsage(AttributeTargets.Property)]
+	public class BinaryIgnoreAttribute : Attribute
+	{
+		public BinaryIgnoreAttribute()
+		{
+		}
+	}
+
+	/// <summary>
+	/// Marks this class as supporting binary serialization. Used to automatically register types via BinaryArchive.RegisterTypes().
+	/// </summary>
+	[AttributeUsage(AttributeTargets.Class)]
+	public class BinarySerializableAttribute : Attribute
+	{
+		public BinarySerializableAttribute()
+		{
+		}
+	}
+
+	/// <summary>
+	/// Instructs the serializer to use a specific converter for this property or class
+	/// </summary>
+	[AttributeUsage(AttributeTargets.Property | AttributeTargets.Class)]
+	public class BinaryConverterAttribute : Attribute
+	{
+		/// <summary>
+		/// The serializer type
+		/// </summary>
+		public Type Type { get; }
+
+		/// <summary>
+		/// Constructor
+		/// </summary>
+		/// <param name="Type">The serializer type</param>
+		public BinaryConverterAttribute(Type Type)
+		{
+			this.Type = Type;
+		}
+	}
+}
diff --git a/Engine/Source/Programs/Shared/EpicGames.Core/BinaryArchiveWriter.cs b/Engine/Source/Programs/Shared/EpicGames.Core/BinaryArchiveWriter.cs
--- a/Engine/Source/Programs/Shared/EpicGames.Core/BinaryArchiveWriter.cs
+++ b/Engine/Source/Programs/Shared/EpicGames.Core/BinaryArchiveWriter.cs
@@ -490,10 +490,9 @@ namespace EpicGames.Core
 		/// <summary>
 		/// Writes an object to the output. If the specific instance has already been written, preserves the reference to that.
 		/// </summary>
-		/// <typeparam name="T">Type of the object to serialize</typeparam>
 		/// <param name="Object">The object to serialize</param>
 		/// <param name="WriteObject">Delegate used to write the object</param>
-		public void WriteObjectReference<T>(T Object, Action WriteObject) where T : class
+		public void WriteObjectReference(object Object, Action WriteObject)
 		{
 			if(Object == null)
 			{
@@ -514,5 +513,15 @@ namespace EpicGames.Core
 				}
 			}
 		}
+
+		/// <summary>
+		/// Writes an object to the output. If the specific instance has already been written, preserves the reference to that.
+		/// </summary>
+		/// <param name="Object">The object to serialize</param>
+		/// <param name="WriteObject">Delegate used to write the object</param>
+		public void WriteObjectReference<T>(T Object, Action WriteObject) where T : class
+		{
+			WriteObjectReference((object)Object, WriteObject);
+		}
 	}
 }
diff --git a/Engine/Source/Programs/Shared/EpicGames.Core/BinaryConverter.cs b/Engine/Source/Programs/Shared/EpicGames.Core/BinaryConverter.cs
new file mode 100644
--- /dev/null
+++ b/Engine/Source/Programs/Shared/EpicGames.Core/BinaryConverter.cs
@@ -0,0 +1,103 @@
+using EpicGames.Core;
+using System;
+using System.Collections.Concurrent;
+using System.Collections.Generic;
+using System.Reflection;
+using System.Text;
+
+#nullable enable
+
+namespace EpicGames.Core
+{
+	/// <summary>
+	/// Base interface for a binary converter. Implementations must derive from <see cref="IBinaryConverter{TValue}"/> instead
+	/// </summary>
+	public interface IBinaryConverter
+	{
+		/// <summary>
+		/// Converter version number
+		/// </summary>
+		int Version { get; }
+	}
+
+	/// <summary>
+	/// Interface for converting serializing object to and from a binary format
+	/// </summary>
+	/// <typeparam name="TValue">The value to serialize</typeparam>
+	public interface IBinaryConverter<TValue> : IBinaryConverter
+	{
+		/// <summary>
+		/// Reads a value from the archive
+		/// </summary>
+		/// <param name="Reader">The archive reader</param>
+		/// <returns>New instance of the value</returns>
+		TValue Read(BinaryArchiveReader Reader);
+
+		/// <summary>
+		/// Writes a value to the archive
+		/// </summary>
+		/// <param name="Writer">The archive writer</param>
+		/// <param name="Value">The value to write</param>
+		void Write(BinaryArchiveWriter Writer, TValue Value);
+	}
+
+	/// <summary>
+	/// Registration of IBinaryConverter instances
+	/// </summary>
+	public static class BinaryConverter
+	{
+		/// <summary>
+		/// Map from type to the converter type
+		/// </summary>
+		static ConcurrentDictionary<Type, Type> TypeToConverterType = new ConcurrentDictionary<Type, Type>();
+
+		/// <summary>
+		/// Explicitly register the converter for a type. If Type is a generic type, the converter should also be a generic type with the same type arguments.
+		/// </summary>
+		/// <param name="Type">Type to register a converter for</param>
+		/// <param name="ConverterType">The converter type</param>
+		public static void RegisterConverter(Type Type, Type ConverterType)
+		{
+			if (!TypeToConverterType.TryAdd(Type, ConverterType))
+			{
+				throw new Exception($"Type '{Type.Name}' already has a registered converter ({TypeToConverterType[Type].Name})");
+			}
+		}
+
+		/// <summary>
+		/// Attempts to get the converter for a particular type
+		/// </summary>
+		/// <param name="Type">The type to use</param>
+		/// <param name="ConverterType">The converter type</param>
+		/// <returns>True if a converter was found</returns>
+		public static bool TryGetConverterType(Type Type, out Type? ConverterType)
+		{
+			Type? CustomConverterType;
+			if (TypeToConverterType.TryGetValue(Type, out CustomConverterType))
+			{
+				ConverterType = CustomConverterType;
+				return true;
+			}
+
+			BinaryConverterAttribute? ConverterAttribute = Type.GetCustomAttribute<BinaryConverterAttribute>();
+			if (ConverterAttribute != null)
+			{
+				ConverterType = ConverterAttribute.Type;
+				return true;
+			}
+
+			if (Type.IsGenericType)
+			{
+				BinaryConverterAttribute? GenericConverterAttribute = Type.GetGenericTypeDefinition().GetCustomAttribute<BinaryConverterAttribute>();
+				if (GenericConverterAttribute != null)
+				{
+					ConverterType = GenericConverterAttribute.Type.MakeGenericType(Type.GetGenericArguments());
+					return true;
+				}
+			}
+
+			ConverterType = null;
+			return false;
+		}
+	}
+}
