Merging functionality from Horde fork of EpicGames.Core, and first pass on fixing up nullable annotations.

#rnx

diff --git a/Engine/Source/Programs/Shared/EpicGames.Core/AssemblyUtils.cs b/Engine/Source/Programs/Shared/EpicGames.Core/AssemblyUtils.cs
--- a/Engine/Source/Programs/Shared/EpicGames.Core/AssemblyUtils.cs
+++ b/Engine/Source/Programs/Shared/EpicGames.Core/AssemblyUtils.cs
@@ -7,6 +7,8 @@ using System.IO;
 using System.Linq;
 using System.Reflection;
 
+#nullable disable
+
 namespace EpicGames.Core
 {
 	public static class AssemblyUtils
diff --git a/Engine/Source/Programs/Shared/EpicGames.Core/AsyncEvent.cs b/Engine/Source/Programs/Shared/EpicGames.Core/AsyncEvent.cs
new file mode 100644
--- /dev/null
+++ b/Engine/Source/Programs/Shared/EpicGames.Core/AsyncEvent.cs
@@ -0,0 +1,44 @@
+// Copyright Epic Games, Inc. All Rights Reserved.
+
+using System.Threading;
+using System.Threading.Tasks;
+
+namespace EpicGames.Core
+{
+	/// <summary>
+	/// Wrapper to allow awaiting an event being signalled, similar to an AutoResetEvent.
+	/// </summary>
+	public class AsyncEvent
+	{
+		/// <summary>
+		/// Completion source to wait on
+		/// </summary>
+		TaskCompletionSource<bool> Source = new TaskCompletionSource<bool>();
+
+		/// <summary>
+		/// Signal the event
+		/// </summary>
+		public void Set()
+		{
+			TaskCompletionSource<bool> PrevSource = Interlocked.Exchange(ref Source, new TaskCompletionSource<bool>());
+			PrevSource.SetResult(true);
+		}
+
+		/// <summary>
+		/// Determines if this event is set
+		/// </summary>
+		/// <returns>True if the event is set</returns>
+		public bool IsSet()
+		{
+			return Source.Task.IsCompleted;
+		}
+
+		/// <summary>
+		/// Waits for this event to be set
+		/// </summary>
+		public Task Task
+		{
+			get { return Source.Task; }
+		}
+	}
+}
diff --git a/Engine/Source/Programs/Shared/EpicGames.Core/BinaryArchive.cs b/Engine/Source/Programs/Shared/EpicGames.Core/BinaryArchive.cs
--- a/Engine/Source/Programs/Shared/EpicGames.Core/BinaryArchive.cs
+++ b/Engine/Source/Programs/Shared/EpicGames.Core/BinaryArchive.cs
@@ -186,7 +186,7 @@ namespace EpicGames.Core
 		/// <returns>The matching type</returns>
 		public static Type ReadType(this BinaryArchiveReader Reader)
 		{
-			return Reader.ReadObjectReference(() => DigestToType[Reader.ReadContentHash()]);
+			return Reader.ReadObjectReference(() => DigestToType[Reader.ReadContentHash()!]);
 		}
 
 		/// <summary>
diff --git a/Engine/Source/Programs/Shared/EpicGames.Core/BinaryArchiveReader.cs b/Engine/Source/Programs/Shared/EpicGames.Core/BinaryArchiveReader.cs
--- a/Engine/Source/Programs/Shared/EpicGames.Core/BinaryArchiveReader.cs
+++ b/Engine/Source/Programs/Shared/EpicGames.Core/BinaryArchiveReader.cs
@@ -7,6 +7,8 @@ using System.Linq;
 using System.Text;
 using System.Threading.Tasks;
 
+#nullable disable
+
 namespace EpicGames.Core
 {
 	/// <summary>
diff --git a/Engine/Source/Programs/Shared/EpicGames.Core/BinaryArchiveWriter.cs b/Engine/Source/Programs/Shared/EpicGames.Core/BinaryArchiveWriter.cs
--- a/Engine/Source/Programs/Shared/EpicGames.Core/BinaryArchiveWriter.cs
+++ b/Engine/Source/Programs/Shared/EpicGames.Core/BinaryArchiveWriter.cs
@@ -8,6 +8,8 @@ using System.Runtime.CompilerServices;
 using System.Text;
 using System.Threading.Tasks;
 
+#nullable disable
+
 namespace EpicGames.Core
 {
 	/// <summary>
diff --git a/Engine/Source/Programs/Shared/EpicGames.Core/BinaryReaderExtensions.cs b/Engine/Source/Programs/Shared/EpicGames.Core/BinaryReaderExtensions.cs
--- a/Engine/Source/Programs/Shared/EpicGames.Core/BinaryReaderExtensions.cs
+++ b/Engine/Source/Programs/Shared/EpicGames.Core/BinaryReaderExtensions.cs
@@ -7,6 +7,8 @@ using System.Linq;
 using System.Text;
 using System.Threading.Tasks;
 
+#nullable disable
+
 namespace EpicGames.Core
 {
 	public static class BinaryReaderExtensions
diff --git a/Engine/Source/Programs/Shared/EpicGames.Core/BinaryWriterExtensions.cs b/Engine/Source/Programs/Shared/EpicGames.Core/BinaryWriterExtensions.cs
--- a/Engine/Source/Programs/Shared/EpicGames.Core/BinaryWriterExtensions.cs
+++ b/Engine/Source/Programs/Shared/EpicGames.Core/BinaryWriterExtensions.cs
@@ -7,6 +7,8 @@ using System.Linq;
 using System.Text;
 using System.Threading.Tasks;
 
+#nullable disable
+
 namespace EpicGames.Core
 {
 	/// <summary>
diff --git a/Engine/Source/Programs/Shared/EpicGames.Core/CancellationTask.cs b/Engine/Source/Programs/Shared/EpicGames.Core/CancellationTask.cs
new file mode 100644
--- /dev/null
+++ b/Engine/Source/Programs/Shared/EpicGames.Core/CancellationTask.cs
@@ -0,0 +1,52 @@
+// Copyright Epic Games, Inc. All Rights Reserved.
+
+using System;
+using System.Collections.Generic;
+using System.Text;
+using System.Threading;
+using System.Threading.Tasks;
+
+namespace EpicGames.Core
+{
+	/// <summary>
+	/// Wrapper class to create a waitable task out of a cancellation token
+	/// </summary>
+	public class CancellationTask : IDisposable
+	{
+		/// <summary>
+		/// Completion source for the task
+		/// </summary>
+		TaskCompletionSource<bool> CompletionSource;
+
+		/// <summary>
+		/// Registration handle with the cancellation token
+		/// </summary>
+		CancellationTokenRegistration Registration;
+
+		/// <summary>
+		/// Constructor
+		/// </summary>
+		/// <param name="Token">The cancellation token to register with</param>
+		public CancellationTask(CancellationToken Token)
+		{
+			CompletionSource = new TaskCompletionSource<bool>();
+			Registration = Token.Register(() => CompletionSource.TrySetResult(true));
+		}
+
+		/// <summary>
+		/// The task that can be waited on
+		/// </summary>
+		public Task Task
+		{
+			get { return CompletionSource.Task; }
+		}
+
+		/// <summary>
+		/// Dispose of any allocated resources
+		/// </summary>
+		public void Dispose()
+		{
+			Registration.Dispose();
+		}
+	}
+}
diff --git a/Engine/Source/Programs/Shared/EpicGames.Core/CommandLineArguments.cs b/Engine/Source/Programs/Shared/EpicGames.Core/CommandLineArguments.cs
--- a/Engine/Source/Programs/Shared/EpicGames.Core/CommandLineArguments.cs
+++ b/Engine/Source/Programs/Shared/EpicGames.Core/CommandLineArguments.cs
@@ -10,6 +10,8 @@ using System.Reflection;
 using System.Text;
 using System.Threading.Tasks;
 
+#nullable disable
+
 namespace EpicGames.Core
 {
 	/// <summary>
diff --git a/Engine/Source/Programs/Shared/EpicGames.Core/CommandLineAttribute.cs b/Engine/Source/Programs/Shared/EpicGames.Core/CommandLineAttribute.cs
--- a/Engine/Source/Programs/Shared/EpicGames.Core/CommandLineAttribute.cs
+++ b/Engine/Source/Programs/Shared/EpicGames.Core/CommandLineAttribute.cs
@@ -18,12 +18,12 @@ namespace EpicGames.Core
 		/// <summary>
 		/// Prefix for the option, with a leading '-' and trailing '=' character if a value is expected.
 		/// </summary>
-		public string Prefix;
+		public string? Prefix;
 
 		/// <summary>
 		/// Specifies a fixed value for this argument. Specifying an alternate value is not permitted.
 		/// </summary>
-		public string Value = null;
+		public string? Value = null;
 
 		/// <summary>
 		/// Whether this argument is required
@@ -39,7 +39,7 @@ namespace EpicGames.Core
 		/// Constructor
 		/// </summary>
 		/// <param name="Prefix">Prefix for this argument</param>
-		public CommandLineAttribute(string Prefix = null)
+		public CommandLineAttribute(string? Prefix = null)
 		{
 			this.Prefix = Prefix;
 
diff --git a/Engine/Source/Programs/Shared/EpicGames.Core/CommandLineExtensions.cs b/Engine/Source/Programs/Shared/EpicGames.Core/CommandLineExtensions.cs
new file mode 100644
--- /dev/null
+++ b/Engine/Source/Programs/Shared/EpicGames.Core/CommandLineExtensions.cs
@@ -0,0 +1,137 @@
+// Copyright Epic Games, Inc. All Rights Reserved.
+
+using System;
+using System.Collections.Generic;
+using System.Text;
+
+namespace HordeAgent.Utility
+{
+	/// <summary>
+	/// Extensions to StringBuilder for building command lines
+	/// </summary>
+	public static class CommandLineExtensions
+	{
+		/// <summary>
+		/// Determines if the given argument needs to be escaped
+		/// </summary>
+		/// <param name="Argument">The argument to check</param>
+		/// <returns>True if the argument needs to be escaped</returns>
+		static bool NeedsEscaping(string Argument)
+		{
+			return Argument.Contains(' ') || Argument.Contains('\"');
+		}
+
+		/// <summary>
+		/// Appends command line argument with a prefixed space. The argument may contain spaces or quotes.
+		/// </summary>
+		/// <param name="Builder">The command line to append to</param>
+		/// <param name="Argument">The argument to append</param>
+		public static void AppendArgument(this StringBuilder Builder, string Argument)
+		{
+			if (Builder.Length > 0)
+			{
+				Builder.Append(' ');
+			}
+
+			int EqualsIdx = Argument.IndexOf('=');
+			if (EqualsIdx != -1)
+			{
+				string Name = Argument.Substring(0, EqualsIdx + 1);
+				if (!NeedsEscaping(Name))
+				{
+					Builder.Append(Name);
+					Argument = Argument.Substring(EqualsIdx + 1);
+				}
+			}
+
+			AppendCommandLineArgumentWithoutSpace(Builder, Argument);
+		}
+
+		/// <summary>
+		/// Appends command line argument with a prefixed space. The argument may contain spaces or quotes.
+		/// </summary>
+		/// <param name="Builder">The command line to append to</param>
+		/// <param name="Name">Name of the argument (eg. -Foo=)</param>
+		/// <param name="Value">Value of the argument</param>
+		public static void AppendArgument(this StringBuilder Builder, string Name, string Value)
+		{
+			if (Builder.Length > 0)
+			{
+				Builder.Append(' ');
+			}
+
+			if (NeedsEscaping(Name))
+			{
+				AppendCommandLineArgumentWithoutSpace(Builder, Name + Value);
+			}
+			else
+			{
+				Builder.Append(Name);
+				AppendCommandLineArgumentWithoutSpace(Builder, Value);
+			}
+		}
+
+		/// <summary>
+		/// Appends an escaped command line argument. The argument may contain spaces or quotes, and is escaped according to the rules in
+		/// https://docs.microsoft.com/en-us/windows/win32/api/shellapi/nf-shellapi-commandlinetoargvw.
+		/// </summary>
+		/// <param name="Builder">The builder to append to</param>
+		/// <param name="Argument">The argument to escape</param>
+		public static void AppendCommandLineArgumentWithoutSpace(this StringBuilder Builder, string Argument)
+		{
+			if (!NeedsEscaping(Argument))
+			{
+				// No escaping necessary if the argument doesn't contain any special characters
+				Builder.Append(Argument);
+			}
+			else
+			{
+				// Escape the whole string following the rules on the CommandLineToArgV MSDN page. 
+				Builder.Append('\"');
+				for (int Idx = 0; Idx < Argument.Length; Idx++)
+				{
+					char Character = Argument[Idx];
+					if (Character == '\"')
+					{
+						// Escape a single quotation mark
+						Builder.Append("\\\"");
+					}
+					else if (Character == '\\')
+					{
+						// Special handling for slashes which may be followed by a quotation mark, as dictated by CommandLineToArgV
+						int StartIdx = Idx;
+						for (; ; )
+						{
+							int NextIdx = Idx + 1;
+							if (NextIdx == Argument.Length)
+							{
+								// Will have a trailing quotation mark toggling 'in quotes' mode (2n)
+								Builder.Append('\\', (NextIdx - StartIdx) * 2);
+								break;
+							}
+							else if (Argument[NextIdx] == '\"')
+							{
+								// Needs to have a trailing quotation mark, so need to escape each backslash plus the quotation mark (2n+1)
+								Builder.Append('\\', (NextIdx - StartIdx) * 2 + 1);
+								break;
+							}
+							else if (Argument[NextIdx] != '\\')
+							{
+								// No trailing quote; can just pass through verbatim
+								Builder.Append('\\', (NextIdx - StartIdx));
+								break;
+							}
+							Idx = NextIdx;
+						}
+					}
+					else
+					{
+						// Regular character
+						Builder.Append(Character);
+					}
+				}
+				Builder.Append('\"');
+			}
+		}
+	}
+}
diff --git a/Engine/Source/Programs/Shared/EpicGames.Core/ContentHash.cs b/Engine/Source/Programs/Shared/EpicGames.Core/ContentHash.cs
--- a/Engine/Source/Programs/Shared/EpicGames.Core/ContentHash.cs
+++ b/Engine/Source/Programs/Shared/EpicGames.Core/ContentHash.cs
@@ -2,6 +2,7 @@
 
 using System;
 using System.Collections.Generic;
+using System.Diagnostics.CodeAnalysis;
 using System.IO;
 using System.Linq;
 using System.Security.Cryptography;
@@ -41,7 +42,7 @@ namespace EpicGames.Core
 		/// </summary>
 		/// <param name="Other">The object to compare against</param>
 		/// <returns>True if the hashes are equal, false otherwise</returns>
-		public override bool Equals(object Other)
+		public override bool Equals(object? Other)
 		{
 			return Equals(Other as ContentHash);
 		}
@@ -51,9 +52,9 @@ namespace EpicGames.Core
 		/// </summary>
 		/// <param name="Other">The hash to compare against</param>
 		/// <returns>True if the hashes are equal, false otherwise</returns>
-		public bool Equals(ContentHash Other)
+		public bool Equals(ContentHash? Other)
 		{
-			if((object)Other == null)
+			if((object?)Other == null)
 			{
 				return false;
 			}
@@ -221,7 +222,7 @@ namespace EpicGames.Core
 		/// <returns>Value of the hash</returns>
 		public static ContentHash Parse(string Text)
 		{
-			ContentHash Hash;
+			ContentHash? Hash;
 			if(!TryParse(Text, out Hash))
 			{
 				throw new ArgumentException(String.Format("'{0}' is not a valid content hash", Text));
@@ -234,9 +235,9 @@ namespace EpicGames.Core
 		/// </summary>
 		/// <param name="Text">Text to parse</param>
 		/// <returns>Value of the hash</returns>
-		public static bool TryParse(string Text, out ContentHash Hash)
+		public static bool TryParse(string Text, [NotNullWhen(true)] out ContentHash? Hash)
 		{
-			byte[] Bytes;
+			byte[]? Bytes;
 			if(StringUtils.TryParseHexString(Text, out Bytes))
 			{
 				Hash = new ContentHash(Bytes);
@@ -283,9 +284,9 @@ namespace EpicGames.Core
 		/// </summary>
 		/// <param name="Writer">The writer to output data to</param>
 		/// <param name="Hash">The hash to write</param>
-		public static void WriteContentHash(this BinaryArchiveWriter Writer, ContentHash Hash)
+		public static void WriteContentHash(this BinaryArchiveWriter Writer, ContentHash? Hash)
 		{
-			if(Hash == null)
+			if(ReferenceEquals(Hash, null))
 			{
 				Writer.WriteByteArray(null);
 			}
@@ -300,9 +301,9 @@ namespace EpicGames.Core
 		/// </summary>
 		/// <param name="Reader">Reader to serialize data from</param>
 		/// <returns>New hash instance</returns>
-		public static ContentHash ReadContentHash(this BinaryArchiveReader Reader)
+		public static ContentHash? ReadContentHash(this BinaryArchiveReader Reader)
 		{
-			byte[] Data = Reader.ReadByteArray();
+			byte[]? Data = Reader.ReadByteArray();
 			if(Data == null)
 			{
 				return null;
diff --git a/Engine/Source/Programs/Shared/EpicGames.Core/CsProjectInfo.cs b/Engine/Source/Programs/Shared/EpicGames.Core/CsProjectInfo.cs
--- a/Engine/Source/Programs/Shared/EpicGames.Core/CsProjectInfo.cs
+++ b/Engine/Source/Programs/Shared/EpicGames.Core/CsProjectInfo.cs
@@ -9,6 +9,8 @@ using System.Text;
 using System.Threading.Tasks;
 using System.Xml;
 
+#nullable disable
+
 namespace EpicGames.Core
 {
 	/// <summary>
diff --git a/Engine/Source/Programs/Shared/EpicGames.Core/DirectoryReference.cs b/Engine/Source/Programs/Shared/EpicGames.Core/DirectoryReference.cs
--- a/Engine/Source/Programs/Shared/EpicGames.Core/DirectoryReference.cs
+++ b/Engine/Source/Programs/Shared/EpicGames.Core/DirectoryReference.cs
@@ -7,6 +7,8 @@ using System.Linq;
 using System.Text;
 using System.Threading.Tasks;
 
+#nullable disable
+
 namespace EpicGames.Core
 {
 	/// <summary>
diff --git a/Engine/Source/Programs/Shared/EpicGames.Core/EpicGames.Core.csproj b/Engine/Source/Programs/Shared/EpicGames.Core/EpicGames.Core.csproj
--- a/Engine/Source/Programs/Shared/EpicGames.Core/EpicGames.Core.csproj
+++ b/Engine/Source/Programs/Shared/EpicGames.Core/EpicGames.Core.csproj
@@ -7,6 +7,7 @@
     <GenerateAssemblyInfo>false</GenerateAssemblyInfo>
     <GenerateTargetFrameworkAttribute>false</GenerateTargetFrameworkAttribute>
     <Configurations>Debug;Release;Development</Configurations>
+    <Nullable>enable</Nullable>
   </PropertyGroup>
   <PropertyGroup Condition=" '$(Configuration)|$(Platform)' == 'Debug|AnyCPU' ">
     <DocumentationFile></DocumentationFile>
diff --git a/Engine/Source/Programs/Shared/EpicGames.Core/ExceptionUtils.cs b/Engine/Source/Programs/Shared/EpicGames.Core/ExceptionUtils.cs
--- a/Engine/Source/Programs/Shared/EpicGames.Core/ExceptionUtils.cs
+++ b/Engine/Source/Programs/Shared/EpicGames.Core/ExceptionUtils.cs
@@ -26,7 +26,7 @@ namespace EpicGames.Core
 		/// <param name="Message">Message to append to the context list</param>
 		public static void AddContext(Exception Ex, string Message)
 		{
-			List<string> Messages = Ex.Data[ContextEntryName] as List<string>;
+			List<string>? Messages = Ex.Data[ContextEntryName] as List<string>;
 			if (Messages == null)
 			{
 				Messages = new List<string>();
@@ -54,7 +54,7 @@ namespace EpicGames.Core
 		/// <returns>Sequence of context lines</returns>
 		public static IEnumerable<string> GetContext(Exception Ex)
 		{
-			List<string> Messages = Ex.Data[ContextEntryName] as List<string>;
+			List<string>? Messages = Ex.Data[ContextEntryName] as List<string>;
 			if (Messages != null)
 			{
 				foreach (string Message in Messages)
@@ -74,7 +74,7 @@ namespace EpicGames.Core
 			StringBuilder ErrorMessage = new StringBuilder();
 			if (Ex is AggregateException)
 			{
-				Exception InnerException = ((AggregateException)Ex).InnerException;
+				Exception InnerException = ((AggregateException)Ex).InnerException!;
 				ErrorMessage.Append(InnerException.ToString());
 				foreach (string Line in GetContext(InnerException))
 				{
@@ -100,7 +100,7 @@ namespace EpicGames.Core
 		public static string FormatExceptionDetails(Exception Ex)
 		{
 			List<Exception> ExceptionStack = new List<Exception>();
-			for (Exception CurrentEx = Ex; CurrentEx != null; CurrentEx = CurrentEx.InnerException)
+			for (Exception? CurrentEx = Ex; CurrentEx != null; CurrentEx = CurrentEx.InnerException)
 			{
 				ExceptionStack.Add(CurrentEx);
 			}
@@ -113,9 +113,9 @@ namespace EpicGames.Core
 
 				if (CurrentEx.Data.Count > 0)
 				{
-					foreach (object Key in CurrentEx.Data.Keys)
+					foreach (object? Key in CurrentEx.Data.Keys)
 					{
-						object Value = CurrentEx.Data[Key];
+						object Value = CurrentEx.Data[Key!]!;
 
 						string ValueString;
 						if(Value is List<string>)
@@ -124,7 +124,7 @@ namespace EpicGames.Core
 						}
 						else
 						{
-							ValueString = Value.ToString();
+							ValueString = Value.ToString()!;
 						}
 
 						Message.AppendFormat("   data: {0} = {1}", Key, ValueString);
diff --git a/Engine/Source/Programs/Shared/EpicGames.Core/FileFilter.cs b/Engine/Source/Programs/Shared/EpicGames.Core/FileFilter.cs
--- a/Engine/Source/Programs/Shared/EpicGames.Core/FileFilter.cs
+++ b/Engine/Source/Programs/Shared/EpicGames.Core/FileFilter.cs
@@ -329,7 +329,7 @@ namespace EpicGames.Core
 			LastNode.Type = Type;
 
 			// Update the maximums along that path
-			for (FileFilterNode UpdateNode = LastNode; UpdateNode != null; UpdateNode = UpdateNode.Parent)
+			for (FileFilterNode? UpdateNode = LastNode; UpdateNode != null; UpdateNode = UpdateNode.Parent)
 			{
 				if (Type == FileFilterType.Include)
 				{
diff --git a/Engine/Source/Programs/Shared/EpicGames.Core/FileFilterNode.cs b/Engine/Source/Programs/Shared/EpicGames.Core/FileFilterNode.cs
--- a/Engine/Source/Programs/Shared/EpicGames.Core/FileFilterNode.cs
+++ b/Engine/Source/Programs/Shared/EpicGames.Core/FileFilterNode.cs
@@ -16,7 +16,7 @@ namespace EpicGames.Core
 		/// <summary>
 		/// Node which this is parented to. Null for the root node.
 		/// </summary>
-		public readonly FileFilterNode Parent;
+		public readonly FileFilterNode? Parent;
 
 		/// <summary>
 		/// Pattern to match for this node. May contain * or ? wildcards.
@@ -51,7 +51,7 @@ namespace EpicGames.Core
 		/// <summary>
 		/// Default constructor.
 		/// </summary>
-		public FileFilterNode(FileFilterNode InParent, string InPattern)
+		public FileFilterNode(FileFilterNode? InParent, string InPattern)
 		{
 			Parent = InParent;
 			Pattern = InPattern;
diff --git a/Engine/Source/Programs/Shared/EpicGames.Core/FilePattern.cs b/Engine/Source/Programs/Shared/EpicGames.Core/FilePattern.cs
--- a/Engine/Source/Programs/Shared/EpicGames.Core/FilePattern.cs
+++ b/Engine/Source/Programs/Shared/EpicGames.Core/FilePattern.cs
@@ -283,7 +283,7 @@ namespace EpicGames.Core
 		/// <param name="TargetPattern">Matching output pattern</param>
 		/// <param name="Filter">Filter to apply to source files</param>
 		/// <param name="TargetFileToSourceFile">Dictionary to receive a mapping from target file to source file. An exception is thrown if multiple source files map to one target file, or a source file is also used as a target file.</param>
-		public static Dictionary<FileReference, FileReference> CreateMapping(HashSet<FileReference> Files, ref FilePattern SourcePattern, ref FilePattern TargetPattern)
+		public static Dictionary<FileReference, FileReference> CreateMapping(HashSet<FileReference>? Files, ref FilePattern SourcePattern, ref FilePattern TargetPattern)
 		{
 			// If the source pattern ends in a directory separator, or a set of input files are specified and it doesn't contain wildcards, treat it as a full directory match
 			if(SourcePattern.EndsWithDirectorySeparator())
@@ -361,7 +361,7 @@ namespace EpicGames.Core
 				// Add them to the output map
 				for(int Idx = 0; Idx < TargetFiles.Length; Idx++)
 				{
-					FileReference ExistingSourceFile;
+					FileReference? ExistingSourceFile;
 					if(TargetFileToSourceFile.TryGetValue(TargetFiles[Idx], out ExistingSourceFile) && ExistingSourceFile != SourceFiles[Idx])
 					{
 						throw new FilePatternException("Output file '{0}' is mapped from '{1}' and '{2}'", TargetFiles[Idx], ExistingSourceFile, SourceFiles[Idx]);
diff --git a/Engine/Source/Programs/Shared/EpicGames.Core/FileReference.cs b/Engine/Source/Programs/Shared/EpicGames.Core/FileReference.cs
--- a/Engine/Source/Programs/Shared/EpicGames.Core/FileReference.cs
+++ b/Engine/Source/Programs/Shared/EpicGames.Core/FileReference.cs
@@ -7,6 +7,8 @@ using System.Linq;
 using System.Text;
 using System.Threading.Tasks;
 
+#nullable disable
+
 namespace EpicGames.Core
 {
 	/// <summary>
diff --git a/Engine/Source/Programs/Shared/EpicGames.Core/FileUtils.cs b/Engine/Source/Programs/Shared/EpicGames.Core/FileUtils.cs
--- a/Engine/Source/Programs/Shared/EpicGames.Core/FileUtils.cs
+++ b/Engine/Source/Programs/Shared/EpicGames.Core/FileUtils.cs
@@ -1,11 +1,14 @@
 // Copyright Epic Games, Inc. All Rights Reserved.
 
+using Microsoft.Win32.SafeHandles;
 using System;
 using System.Collections.Generic;
 using System.ComponentModel;
+using System.Diagnostics;
 using System.IO;
 using System.Linq;
 using System.Runtime.InteropServices;
+using System.Runtime.InteropServices.ComTypes;
 using System.Text;
 using System.Threading.Tasks;
 
@@ -45,6 +48,11 @@ namespace EpicGames.Core
 		/// </summary>
 		public static IEqualityComparer<string> PlatformPathComparer = Environment.OSVersion.Platform == PlatformID.Unix ? StringComparer.Ordinal : StringComparer.OrdinalIgnoreCase;
 
+		/// <summary>
+		/// Utf8 string comparer that should be used for native path comparisons
+		/// </summary>
+		public static IEqualityComparer<ReadOnlyUtf8String> PlatformPathComparerUtf8 = Environment.OSVersion.Platform == PlatformID.Unix ? ReadOnlyUtf8StringComparer.Ordinal : ReadOnlyUtf8StringComparer.OrdinalIgnoreCase;
+
 		/// <summary>
 		/// Finds the on-disk case of a a file
 		/// </summary>
@@ -71,7 +79,11 @@ namespace EpicGames.Core
 		{
 			if (!DirectoryReference.Exists(Directory))
 			{
-				CreateDirectoryTree(Directory.ParentDirectory);
+				DirectoryReference? ParentDirectory = Directory.ParentDirectory;
+				if (ParentDirectory != null)
+				{
+					CreateDirectoryTree(ParentDirectory);
+				}
 				DirectoryReference.CreateDirectory(Directory);
 			}
 		}
@@ -82,7 +94,14 @@ namespace EpicGames.Core
 		/// <param name="FileName">Name of the file to delete</param>
 		public static void ForceDeleteFile(string FileName)
 		{
-			ForceDeleteFile(new FileInfo(FileName));
+			if (Environment.OSVersion.Platform == PlatformID.Win32NT)
+			{
+				ForceDeleteFileWin32(FileName);
+			}
+			else
+			{
+				ForceDeleteFile(new FileInfo(FileName));
+			}
 		}
 
 		/// <summary>
@@ -91,17 +110,25 @@ namespace EpicGames.Core
 		/// <param name="File">The file to delete</param>
 		public static void ForceDeleteFile(FileInfo File)
 		{
-			try
+			if (Environment.OSVersion.Platform == PlatformID.Win32NT)
 			{
-				if (File.Exists)
-				{
-					File.Attributes = FileAttributes.Normal;
-					File.Delete();
-				}
+				ForceDeleteFileWin32(File.FullName);
+				File.Refresh();
 			}
-			catch (Exception Ex)
+			else
 			{
-				throw new WrappedFileOrDirectoryException(Ex, String.Format("Unable to delete '{0}': {1}", File.FullName, Ex.Message.TrimEnd()));
+				try
+				{
+					if (File.Exists)
+					{
+						File.Attributes = FileAttributes.Normal;
+						File.Delete();
+					}
+				}
+				catch (Exception Ex)
+				{
+					throw new WrappedFileOrDirectoryException(Ex, String.Format("Unable to delete '{0}': {1}", File.FullName, Ex.Message.TrimEnd()));
+				}
 			}
 		}
 
@@ -111,7 +138,14 @@ namespace EpicGames.Core
 		/// <param name="Location">The file to delete</param>
 		public static void ForceDeleteFile(FileReference Location)
 		{
-			ForceDeleteFile(new FileInfo(Location.FullName));
+			if (Environment.OSVersion.Platform == PlatformID.Win32NT)
+			{
+				ForceDeleteFileWin32(Location.FullName);
+			}
+			else
+			{
+				ForceDeleteFile(new FileInfo(Location.FullName));
+			}
 		}
 
 		/// <summary>
@@ -196,19 +230,22 @@ namespace EpicGames.Core
 			}
 			else
 			{
-				foreach (FileInfo File in Directory.EnumerateFiles())
-				{
-					ForceDeleteFile(File);
-				}
-				foreach (DirectoryInfo SubDirectory in Directory.EnumerateDirectories())
+				if (Directory.Exists)
 				{
-					if (SubDirectory.Attributes.HasFlag(FileAttributes.ReparsePoint))
+					foreach (FileInfo File in Directory.EnumerateFiles())
 					{
-						ForceDeleteDirectoryInternal(SubDirectory);
+						ForceDeleteFile(File);
 					}
-					else
+					foreach (DirectoryInfo SubDirectory in Directory.EnumerateDirectories())
 					{
-						ForceDeleteDirectory(SubDirectory);
+						if (SubDirectory.Attributes.HasFlag(FileAttributes.ReparsePoint))
+						{
+							ForceDeleteDirectoryInternal(SubDirectory);
+						}
+						else
+						{
+							ForceDeleteDirectory(SubDirectory);
+						}
 					}
 				}
 			}
@@ -277,7 +314,7 @@ namespace EpicGames.Core
 				}
 
 				// Throw the original exception
-				throw new WrappedFileOrDirectoryException(Ex, String.Format("{0}: Unable to move {1} to {2}", Ex.GetType(), SourceLocation, TargetLocation));
+				throw new WrappedFileOrDirectoryException(Ex, String.Format("Unable to move {0} to {1}", SourceLocation, TargetLocation));
 			}
 		}
 
@@ -329,6 +366,8 @@ namespace EpicGames.Core
 
 		static readonly IntPtr INVALID_HANDLE_VALUE = new IntPtr(-1);
 
+		const int ERROR_PATH_NOT_FOUND = 3;
+
 		private static void ForceDeleteLongDirectoryContentsWin32(string DirName)
 		{
 			WIN32_FIND_DATA FindData = new WIN32_FIND_DATA();
@@ -354,10 +393,7 @@ namespace EpicGames.Core
 							{
 								SetFileAttributesW(FullName, FILE_ATTRIBUTE_NORMAL);
 							}
-							if (!DeleteFileW(FullName))
-							{
-								throw new WrappedFileOrDirectoryException(new Win32Exception(), "Unable to delete " + FullName);
-							}
+							ForceDeleteFileWin32(FullName);
 						}
 
 						if (!FindNextFileW(hFind, ref FindData))
@@ -379,8 +415,280 @@ namespace EpicGames.Core
 
 			if (!RemoveDirectory(DirName))
 			{
-				throw new WrappedFileOrDirectoryException(new Win32Exception(), "Unable to delete " + DirName);
+				int ErrorCode = Marshal.GetLastWin32Error();
+				if (ErrorCode != ERROR_PATH_NOT_FOUND)
+				{
+					throw new WrappedFileOrDirectoryException(new Win32Exception(ErrorCode), "Unable to delete " + DirName);
+				}
+			}
+		}
+
+		const int ERROR_FILE_NOT_FOUND = 2;
+		const int ERROR_ACCESS_DENIED = 5;
+
+		[DllImport("kernel32.dll", CharSet = CharSet.Unicode, SetLastError = true)]
+		internal static extern int GetFileAttributesW(string lpFileName);
+
+		public static void ForceDeleteFileWin32(string FileName)
+		{
+			// Try to delete the file normally
+			if (DeleteFileW(FileName))
+			{
+				return;
+			}
+
+			// Capture the exception for failing to delete the file
+			Win32Exception Ex = new Win32Exception();
+
+			// Check the file exists and is not readonly
+			int Attributes = GetFileAttributesW(FileName);
+			if (Attributes == -1)
+			{
+				int ErrorCode = Marshal.GetLastWin32Error();
+				if (ErrorCode == ERROR_PATH_NOT_FOUND || ErrorCode == ERROR_FILE_NOT_FOUND)
+				{
+					return;
+				}
+			}
+			else
+			{
+				if ((Attributes & (int)FileAttributes.ReadOnly) != 0)
+				{
+					if (SetFileAttributesW(FileName, (int)FileAttributes.Normal) && DeleteFileW(FileName))
+					{
+						return;
+					}
+				}
+			}
+
+			// Get a useful error message about why the delete failed
+			StringBuilder Message = new StringBuilder($"Unable to delete {FileName} - {Ex.Message}");
+			if (Ex.NativeErrorCode == ERROR_ACCESS_DENIED)
+			{
+				List<FileLockInfo_Win32>? LockInfoList;
+				try
+				{
+					LockInfoList = GetFileLockInfo_Win32(FileName);
+				}
+				catch
+				{
+					LockInfoList = null;
+				}
+
+				if (LockInfoList != null && LockInfoList.Count > 0)
+				{
+					Message.Append("\nProcesses with open handles to file:");
+					foreach (FileLockInfo_Win32 LockInfo in LockInfoList)
+					{
+						Message.Append($"\n  {LockInfo}");
+					}
+				}
+			}
+			throw new WrappedFileOrDirectoryException(Ex, Message.ToString());
+		}
+
+		#endregion
+
+		#region Win32 Restart Manager API
+
+		[StructLayout(LayoutKind.Sequential)]
+		struct RM_UNIQUE_PROCESS
+		{
+			public int dwProcessId;
+			public FILETIME ProcessStartTime;
+		}
+
+		const int RmRebootReasonNone = 0;
+		static readonly int RM_SESSION_KEY_LEN = Marshal.SizeOf<Guid>();
+		static readonly int CCH_RM_SESSION_KEY = RM_SESSION_KEY_LEN * 2;
+		const int CCH_RM_MAX_APP_NAME = 255;
+		const int CCH_RM_MAX_SVC_NAME = 63;
+
+		enum RM_APP_TYPE
+		{
+			RmUnknownApp = 0,
+			RmMainWindow = 1,
+			RmOtherWindow = 2,
+			RmService = 3,
+			RmExplorer = 4,
+			RmConsole = 5,
+			RmCritical = 1000
+		}
+
+		[StructLayout(LayoutKind.Sequential, CharSet = CharSet.Unicode)]
+		struct RM_PROCESS_INFO
+		{
+			public RM_UNIQUE_PROCESS Process;
+
+			[MarshalAs(UnmanagedType.ByValTStr, SizeConst = CCH_RM_MAX_APP_NAME + 1)]
+			public string strAppName;
+
+			[MarshalAs(UnmanagedType.ByValTStr, SizeConst = CCH_RM_MAX_SVC_NAME + 1)]
+			public string strServiceShortName;
+
+			public RM_APP_TYPE ApplicationType;
+			public uint AppStatus;
+			public uint TSSessionId;
+
+			[MarshalAs(UnmanagedType.Bool)]
+			public bool bRestartable;
+		}
+
+		[DllImport("rstrtmgr.dll", CharSet = CharSet.Unicode)]
+		static extern int RmRegisterResources(uint pSessionHandle, uint nFiles, string[] rgsFilenames, uint nApplications, [In] RM_UNIQUE_PROCESS[]? rgApplications, uint nServices, string[]? rgsServiceNames);
+
+		[DllImport("rstrtmgr.dll", CharSet = CharSet.Auto)]
+		static extern int RmStartSession(out uint pSessionHandle, int dwSessionFlags, StringBuilder strSessionKey);
+
+		[DllImport("rstrtmgr.dll")]
+		static extern int RmEndSession(uint pSessionHandle);
+
+		[DllImport("rstrtmgr.dll")]
+		static extern int RmGetList(uint dwSessionHandle,
+									out uint pnProcInfoNeeded,
+									ref uint pnProcInfo,
+									[In, Out] RM_PROCESS_INFO[] rgAffectedApps,
+									ref uint lpdwRebootReasons);
+
+		const int PROCESS_QUERY_LIMITED_INFORMATION = 0x1000;
+
+		[DllImport("kernel32.dll", SetLastError = true)]
+		static extern SafeProcessHandle OpenProcess(
+			 int processAccess,
+			 bool bInheritHandle,
+			 int processId
+		);
+
+		[DllImport("kernel32.dll", SetLastError = true)]
+		[return: MarshalAs(UnmanagedType.Bool)]
+		static extern bool GetProcessTimes(SafeProcessHandle hProcess, out FILETIME
+		   lpCreationTime, out FILETIME lpExitTime, out FILETIME lpKernelTime,
+		   out FILETIME lpUserTime);
+
+		[DllImport("kernel32.dll", SetLastError = true)]
+		static extern bool QueryFullProcessImageName([In]SafeProcessHandle hProcess, [In]int dwFlags, [Out]StringBuilder lpExeName, ref int lpdwSize);
+
+		/// <summary>
+		/// Information about a locked file
+		/// </summary>
+		public class FileLockInfo_Win32
+		{
+			/// <summary>
+			/// Process id
+			/// </summary>
+			public int ProcessId;
+
+			/// <summary>
+			/// Path to the process holding the lock
+			/// </summary>
+			public string? FileName;
+
+			/// <summary>
+			/// Name of the application
+			/// </summary>
+			public string AppName;
+
+			/// <summary>
+			/// Time at which the process started
+			/// </summary>
+			public DateTime StartTime;
+
+			/// <summary>
+			/// Constructor
+			/// </summary>
+			/// <param name="ProcessId">Process id</param>
+			/// <param name="FileName">Path to the process holding the lock</param>
+			/// <param name="AppName">Name of the application</param>
+			/// <param name="StartTime">Time at which the process started</param>
+			public FileLockInfo_Win32(int ProcessId, string? FileName, string AppName, DateTime StartTime)
+			{
+				this.ProcessId = ProcessId;
+				this.FileName = FileName;
+				this.AppName = AppName;
+				this.StartTime = StartTime;
 			}
+
+			/// <inheritdoc/>
+			public override string ToString()
+			{
+				return $"{ProcessId}: {FileName ?? AppName} (started {StartTime})";
+			}
+		}
+
+		/// <summary>
+		/// Gets a list of processes that have a handle to the given file open
+		/// </summary>
+		/// <param name="FileName">File to check</param>
+		/// <returns>List of processes with a lock open</returns>
+		public static List<FileLockInfo_Win32> GetFileLockInfo_Win32(string FileName)
+		{
+			uint SessionHandle = 0;
+			try
+			{
+				StringBuilder SessionKey = new StringBuilder(CCH_RM_SESSION_KEY + 1);
+
+				int Result = RmStartSession(out SessionHandle, 0, SessionKey);
+				if (Result != 0)
+				{
+					throw new Win32Exception(Result, "Unable to open restart manager session");
+				}
+
+				Result = RmRegisterResources(SessionHandle, 1, new string[] { FileName }, 0, null, 0, null);
+				if (Result != 0)
+				{
+					throw new Win32Exception(Result, "Unable to register resource with restart manager");
+				}
+
+				uint nProcInfoNeeded = 0;
+				uint nProcInfo = 10;
+				uint Reason = 0;
+				RM_PROCESS_INFO[] ProcessInfoArray = new RM_PROCESS_INFO[nProcInfo];
+				Result = RmGetList(SessionHandle, out nProcInfoNeeded, ref nProcInfo, ProcessInfoArray, ref Reason);
+				if (Result != 0)
+				{
+					throw new Win32Exception(Result, "Unable to query processes with file handle open");
+				}
+
+				List<FileLockInfo_Win32> FileLocks = new List<FileLockInfo_Win32>();
+				for (int Idx = 0; Idx < nProcInfo; Idx++)
+				{
+					RM_PROCESS_INFO ProcessInfo = ProcessInfoArray[Idx];
+					long StartTimeTicks = FileTimeToTicks(ProcessInfo.Process.ProcessStartTime);
+
+					string? ImageName = null;
+					using (SafeProcessHandle hProcess = OpenProcess(PROCESS_QUERY_LIMITED_INFORMATION, false, ProcessInfo.Process.dwProcessId))
+					{
+						if (hProcess != null)
+						{
+							FILETIME CreateTime, ExitTime, KernelTime, UserTime;
+							if (GetProcessTimes(hProcess, out CreateTime, out ExitTime, out KernelTime, out UserTime) && FileTimeToTicks(CreateTime) == StartTimeTicks)
+							{
+								int Capacity = 260;
+								StringBuilder ImageNameBuilder = new StringBuilder(Capacity);
+								if (QueryFullProcessImageName(hProcess, 0, ImageNameBuilder, ref Capacity))
+								{
+									ImageName = ImageNameBuilder.ToString(0, Capacity);
+								}
+							}
+						}
+					}
+
+					FileLocks.Add(new FileLockInfo_Win32(ProcessInfo.Process.dwProcessId, ImageName, ProcessInfo.strAppName, DateTime.FromFileTime(StartTimeTicks)));
+				}
+				return FileLocks;
+			}
+			finally
+			{
+				if (SessionHandle != 0)
+				{
+					RmEndSession(SessionHandle);
+				}
+			}
+		}
+
+		private static long FileTimeToTicks(FILETIME FileTime)
+		{
+			return (long)(uint)FileTime.dwLowDateTime | ((long)(uint)FileTime.dwHighDateTime << 32);
 		}
 
 		#endregion
diff --git a/Engine/Source/Programs/Shared/EpicGames.Core/HttpClientExtensions.cs b/Engine/Source/Programs/Shared/EpicGames.Core/HttpClientExtensions.cs
new file mode 100644
--- /dev/null
+++ b/Engine/Source/Programs/Shared/EpicGames.Core/HttpClientExtensions.cs
@@ -0,0 +1,128 @@
+// Copyright Epic Games, Inc. All Rights Reserved.
+
+using System;
+using System.Collections.Generic;
+using System.IO;
+using System.Net;
+using System.Net.Http;
+using System.Text;
+using System.Text.Json;
+using System.Text.Json.Serialization;
+using System.Threading;
+using System.Threading.Tasks;
+using System.Web;
+
+namespace EpicGames.Core
+{
+	/// <summary>
+	/// Extension methods for consuming REST APIs via JSON objects
+	/// </summary>
+	public static class HttpClientExtensions
+	{
+		/// <summary>
+		/// Gets a resource from an HTTP endpoint and parses it as a JSON object
+		/// </summary>
+		/// <typeparam name="TResponse">The object type to return</typeparam>
+		/// <param name="Client">The http client instance</param>
+		/// <param name="Url">The url to retrieve</param>
+		/// <param name="CancellationToken">Cancels the request</param>
+		/// <returns>New instance of the object</returns>
+		public static async Task<TResponse> GetAsync<TResponse>(this HttpClient Client, string Url, CancellationToken CancellationToken)
+		{
+			using (HttpResponseMessage Response = await Client.GetAsync(Url, CancellationToken))
+			{
+				Response.EnsureSuccessStatusCode();
+				return await ParseJsonContent<TResponse>(Response);
+			}
+		}
+
+		/// <summary>
+		/// Posts an object to an HTTP endpoint as a JSON object
+		/// </summary>
+		/// <typeparam name="TRequest">The object type to post</typeparam>
+		/// <param name="Client">The http client instance</param>
+		/// <param name="Url">The url to post to</param>
+		/// <param name="Object">The object to post</param>
+		/// <param name="CancellationToken">Cancels the request</param>
+		/// <returns>Response message</returns>
+		public static async Task<HttpResponseMessage> PostAsync<TRequest>(this HttpClient Client, string Url, TRequest Object, CancellationToken CancellationToken)
+		{
+			return await Client.PostAsync(Url, ToJsonContent(Object), CancellationToken);
+		}
+
+		/// <summary>
+		/// Posts an object to an HTTP endpoint as a JSON object, and parses the response object
+		/// </summary>
+		/// <typeparam name="TResponse">The object type to return</typeparam>
+		/// <typeparam name="TRequest">The object type to post</typeparam>
+		/// <param name="Client">The http client instance</param>
+		/// <param name="Url">The url to post to</param>
+		/// <param name="Request">The object to post</param>
+		/// <param name="CancellationToken">Cancels the request</param>
+		/// <returns>The response parsed into the requested type</returns>
+		public static async Task<TResponse> PostAsync<TResponse, TRequest>(this HttpClient Client, string Url, TRequest Request, CancellationToken CancellationToken)
+		{
+			using (HttpResponseMessage Response = await PostAsync(Client, Url, Request, CancellationToken))
+			{
+				Response.EnsureSuccessStatusCode();
+				return await ParseJsonContent<TResponse>(Response);
+			}
+		}
+
+		/// <summary>
+		/// Puts an object to an HTTP endpoint as a JSON object
+		/// </summary>
+		/// <typeparam name="TRequest">The object type to post</typeparam>
+		/// <param name="Client">The http client instance</param>
+		/// <param name="Url">The url to post to</param>
+		/// <param name="Object">The object to post</param>
+		/// <param name="CancellationToken">Cancels the request</param>
+		/// <returns>Response message</returns>
+		public static async Task<HttpResponseMessage> PutAsync<TRequest>(this HttpClient Client, string Url, TRequest Object, CancellationToken CancellationToken)
+		{
+			return await Client.PutAsync(Url, ToJsonContent(Object), CancellationToken);
+		}
+
+		/// <summary>
+		/// Puts an object to an HTTP endpoint as a JSON object, and parses the response object
+		/// </summary>
+		/// <typeparam name="TResponse">The object type to return</typeparam>
+		/// <typeparam name="TRequest">The object type to post</typeparam>
+		/// <param name="Client">The http client instance</param>
+		/// <param name="Url">The url to post to</param>
+		/// <param name="Request">The object to post</param>
+		/// <param name="CancellationToken">Cancels the request</param>
+		/// <returns>The response parsed into the requested type</returns>
+		public static async Task<TResponse> PutAsync<TResponse, TRequest>(this HttpClient Client, string Url, TRequest Request, CancellationToken CancellationToken)
+		{
+			using (HttpResponseMessage Response = await PutAsync(Client, Url, Request, CancellationToken))
+			{
+				Response.EnsureSuccessStatusCode();
+				return await ParseJsonContent<TResponse>(Response);
+			}
+		}
+
+		/// <summary>
+		/// Converts an object to a JSON http content object
+		/// </summary>
+		/// <typeparam name="T">Type of the object to parse</typeparam>
+		/// <param name="Object">The object instance</param>
+		/// <returns>Http content object</returns>
+		private static HttpContent ToJsonContent<T>(T Object)
+		{
+			return new StringContent(JsonSerializer.Serialize<T>(Object), Encoding.UTF8, "application/json");
+		}
+
+		/// <summary>
+		/// Parses a HTTP response as a JSON object
+		/// </summary>
+		/// <typeparam name="T">Type of the object to parse</typeparam>
+		/// <param name="Message">The message received</param>
+		/// <returns>Parsed object instance</returns>
+		private static async Task<T> ParseJsonContent<T>(HttpResponseMessage Message)
+		{
+			byte[] Bytes = await Message.Content.ReadAsByteArrayAsync();
+			return JsonSerializer.Deserialize<T>(Bytes);
+		}
+	}
+}
diff --git a/Engine/Source/Programs/Shared/EpicGames.Core/JsonKnownTypes.cs b/Engine/Source/Programs/Shared/EpicGames.Core/JsonKnownTypes.cs
new file mode 100644
--- /dev/null
+++ b/Engine/Source/Programs/Shared/EpicGames.Core/JsonKnownTypes.cs
@@ -0,0 +1,232 @@
+// Copyright Epic Games, Inc. All Rights Reserved.
+
+using System;
+using System.Buffers;
+using System.Collections.Generic;
+using System.Diagnostics.CodeAnalysis;
+using System.Linq;
+using System.Text.Json;
+using System.Text.Json.Serialization;
+using System.Threading.Tasks;
+
+namespace EpicGames.Core
+{
+	/// <summary>
+	/// Sets a name used to discriminate between classes derived from a base class
+	/// </summary>
+	[AttributeUsage(AttributeTargets.Class)]
+	public class JsonDiscriminatorAttribute : Attribute
+	{
+		/// <summary>
+		/// Name to use to discriminate between different classes
+		/// </summary>
+		public string Name { get; }
+
+		/// <summary>
+		/// Constructor
+		/// </summary>
+		/// <param name="Name">Name to use to discriminate between different classes</param>
+		public JsonDiscriminatorAttribute(string Name)
+		{
+			this.Name = Name;
+		}
+	}
+
+	/// <summary>
+	/// Marks a class as the base class of a hierarchy, allowing any known subclasses to be serialized for it.
+	/// </summary>
+	[AttributeUsage(AttributeTargets.Class)]
+	public class JsonKnownTypesAttribute : Attribute
+	{
+		/// <summary>
+		/// Array of derived classes
+		/// </summary>
+		public Type[] Types { get; }
+
+		/// <summary>
+		/// Constructor
+		/// </summary>
+		/// <param name="Types">Array of derived classes</param>
+		public JsonKnownTypesAttribute(params Type[] Types)
+		{
+			this.Types = Types;
+		}
+	}
+
+	/// <summary>
+	/// Serializer for polymorphic objects
+	/// </summary>
+	public class JsonKnownTypesConverter<T> : JsonConverter<T>
+	{
+		/// <summary>
+		/// Name of the 'Type' field used to store discriminators in the serialized objects
+		/// </summary>
+		const string TypePropertyName = "Type";
+
+		/// <summary>
+		/// Mapping from discriminator to class type
+		/// </summary>
+		Dictionary<string, Type> DiscriminatorToType;
+
+		/// <summary>
+		/// Mapping from class type to discriminator
+		/// </summary>
+		Dictionary<Type, JsonEncodedText> TypeToDiscriminator;
+
+		/// <summary>
+		/// Constructor
+		/// </summary>
+		/// <param name="KnownTypes">Set of derived class types</param>
+		public JsonKnownTypesConverter(IEnumerable<Type> KnownTypes)
+		{
+			DiscriminatorToType = new Dictionary<string, Type>(StringComparer.OrdinalIgnoreCase);
+			foreach (Type KnownType in KnownTypes)
+			{
+				if (!typeof(T).IsAssignableFrom(KnownType))
+				{
+					throw new InvalidOperationException($"{KnownType} is not derived from {typeof(T)}");
+				}
+				foreach (JsonDiscriminatorAttribute DiscriminatorAttribute in KnownType.GetCustomAttributes(typeof(JsonDiscriminatorAttribute), true))
+				{
+					DiscriminatorToType.Add(DiscriminatorAttribute.Name, KnownType);
+				}
+			}
+
+			TypeToDiscriminator = DiscriminatorToType.ToDictionary(x => x.Value, x => JsonEncodedText.Encode(x.Key));
+		}
+
+		/// <summary>
+		/// Determines whether the given type can be converted
+		/// </summary>
+		/// <param name="TypeToConvert">Type to convert</param>
+		/// <returns>True if the type can be converted</returns>
+		public override bool CanConvert(Type TypeToConvert)
+		{
+			return TypeToConvert == typeof(T) || TypeToDiscriminator.ContainsKey(TypeToConvert);
+		}
+
+		/// <summary>
+		/// Reads an object from the given input stream
+		/// </summary>
+		/// <param name="Reader">UTF8 reader</param>
+		/// <param name="TypeToConvert">Type to be read from the stream</param>
+		/// <param name="Options">Options for serialization</param>
+		/// <returns>Instance of the serialized object</returns>
+		public override T Read(ref Utf8JsonReader Reader, Type TypeToConvert, JsonSerializerOptions Options)
+		{
+			using (JsonDocument Document = JsonDocument.ParseValue(ref Reader))
+			{
+				JsonElement Element;
+				if (!TryGetPropertyRespectingCase(Document.RootElement, TypePropertyName, Options, out Element))
+				{
+					throw new JsonException($"Missing '{TypePropertyName}' field for serializing {TypeToConvert.Name}");
+				}
+
+				Type? TargetType;
+				if (!DiscriminatorToType.TryGetValue(Element.GetString(), out TargetType))
+				{
+					throw new JsonException($"Invalid '{TypePropertyName}' field ('{Element.GetString()}') for serializing {TypeToConvert.Name}");
+				}
+
+				string Text = Document.RootElement.GetRawText();
+				T Result = (T)JsonSerializer.Deserialize(Text, TargetType, Options);
+				return Result;
+			}
+		}
+
+		/// <summary>
+		/// Finds a property within an object, ignoring case
+		/// </summary>
+		/// <param name="Object">The object to search</param>
+		/// <param name="PropertyName">Name of the property to search</param>
+		/// <param name="Options">Options for serialization. The <see cref="JsonSerializerOptions.PropertyNameCaseInsensitive"/> property defines whether a case insensitive scan will be performed.</param>
+		/// <param name="OutElement">The property value, if successful</param>
+		/// <returns>True if the property was found, false otherwise</returns>
+		static bool TryGetPropertyRespectingCase(JsonElement Object, string PropertyName, JsonSerializerOptions Options, out JsonElement OutElement)
+		{
+			JsonElement NewElement;
+			if (Object.TryGetProperty(PropertyName, out NewElement))
+			{
+				OutElement = NewElement;
+				return true;
+			}
+
+			if (Options.PropertyNameCaseInsensitive)
+			{
+				foreach (JsonProperty Property in Object.EnumerateObject())
+				{
+					if (String.Equals(Property.Name, PropertyName, StringComparison.OrdinalIgnoreCase))
+					{
+						OutElement = Property.Value;
+						return true;
+					}
+				}
+			}
+
+			OutElement = new JsonElement();
+			return false;
+		}
+
+		/// <summary>
+		/// Writes an object to the given output stream
+		/// </summary>
+		/// <param name="Writer">UTF8 writer</param>
+		/// <param name="Value">Object to be serialized to the stream</param>
+		/// <param name="Options">Options for serialization</param>
+		public override void Write(Utf8JsonWriter Writer, T Value, JsonSerializerOptions Options)
+		{
+			if (Value == null)
+			{
+				Writer.WriteNullValue();
+			}
+			else
+			{
+				Type ValueType = Value.GetType();
+
+				Writer.WriteStartObject();
+				Writer.WriteString(Options.PropertyNamingPolicy?.ConvertName(TypePropertyName) ?? TypePropertyName, TypeToDiscriminator[ValueType]);
+
+				byte[] Bytes = JsonSerializer.SerializeToUtf8Bytes(Value, ValueType, Options);
+				using (JsonDocument Document = JsonDocument.Parse(Bytes))
+				{
+					foreach (JsonProperty Property in Document.RootElement.EnumerateObject())
+					{
+						Property.WriteTo(Writer);
+					}
+				}
+
+				Writer.WriteEndObject();
+			}
+		}
+	}
+
+	/// <summary>
+	/// Factory for JsonKnownTypes converters
+	/// </summary>
+	public class JsonKnownTypesConverterFactory : JsonConverterFactory
+	{
+		/// <summary>
+		/// Determines whether it's possible to create a converter for the given type
+		/// </summary>
+		/// <param name="TypeToConvert">The type being converted</param>
+		/// <returns>True if the type can be converted</returns>
+		public override bool CanConvert(Type TypeToConvert)
+		{
+			object[] Attributes = TypeToConvert.GetCustomAttributes(typeof(JsonKnownTypesAttribute), false);
+			return Attributes.Length > 0;
+		}
+
+		/// <summary>
+		/// Creates a converter for the given type
+		/// </summary>
+		/// <param name="TypeToConvert">The type being converted</param>
+		/// <param name="Options">Options for serialization</param>
+		/// <returns>New converter for the requested type</returns>
+		public override JsonConverter CreateConverter(Type TypeToConvert, JsonSerializerOptions Options)
+		{
+			object[] Attributes = TypeToConvert.GetCustomAttributes(typeof(JsonKnownTypesAttribute), false);
+			Type GenericType = typeof(JsonKnownTypesConverter<>).MakeGenericType(TypeToConvert);
+			return (JsonConverter)Activator.CreateInstance(GenericType, Attributes.SelectMany(x => ((JsonKnownTypesAttribute)x).Types))!;
+		}
+	}
+}
diff --git a/Engine/Source/Programs/Shared/EpicGames.Core/JsonObject.cs b/Engine/Source/Programs/Shared/EpicGames.Core/JsonObject.cs
--- a/Engine/Source/Programs/Shared/EpicGames.Core/JsonObject.cs
+++ b/Engine/Source/Programs/Shared/EpicGames.Core/JsonObject.cs
@@ -7,6 +7,8 @@ using System.Linq;
 using System.Text;
 using System.Text.Json;
 
+#nullable disable
+
 namespace EpicGames.Core
 {
 	/// <summary>
diff --git a/Engine/Source/Programs/Shared/EpicGames.Core/JsonWriter.cs b/Engine/Source/Programs/Shared/EpicGames.Core/JsonWriter.cs
--- a/Engine/Source/Programs/Shared/EpicGames.Core/JsonWriter.cs
+++ b/Engine/Source/Programs/Shared/EpicGames.Core/JsonWriter.cs
@@ -80,7 +80,7 @@ namespace EpicGames.Core
 			if(!bLeaveOpen && Writer != null)
 			{
 				Writer.Dispose();
-				Writer = null;
+				Writer = null!;
 			}
 		}
 
@@ -228,7 +228,7 @@ namespace EpicGames.Core
 		/// <param name="Values">Values for the field</param>
 		public void WriteEnumArrayField<T>(string Name, IEnumerable<T> Values) where T : struct
 		{
-			WriteStringArrayField(Name, Values.Select(x => x.ToString()));
+			WriteStringArrayField(Name, Values.Select(x => x.ToString()!));
 		}
 
 		/// <summary>
@@ -264,7 +264,7 @@ namespace EpicGames.Core
 		/// </summary>
 		/// <param name="Name">Name of the field</param>
 		/// <param name="Value">Value for the field</param>
-		public void WriteValue(string Name, string Value)
+		public void WriteValue(string Name, string? Value)
 		{
 			WriteCommaNewline();
 
@@ -313,7 +313,7 @@ namespace EpicGames.Core
 		/// <param name="Value">Value for the field</param>
 		public void WriteEnumValue<T>(string Name, T Value) where T : struct
 		{
-			WriteValue(Name, Value.ToString());
+			WriteValue(Name, Value.ToString()!);
 		}
 
 		void WriteCommaNewline()
@@ -338,7 +338,7 @@ namespace EpicGames.Core
 			bRequiresComma = true;
 		}
 
-		void WriteEscapedString(string Value)
+		void WriteEscapedString(string? Value)
 		{
 			// Escape any characters which may not appear in a JSON string (see http://www.json.org).
 			Writer.Write("\"");
diff --git a/Engine/Source/Programs/Shared/EpicGames.Core/ListExtensions.cs b/Engine/Source/Programs/Shared/EpicGames.Core/ListExtensions.cs
new file mode 100644
--- /dev/null
+++ b/Engine/Source/Programs/Shared/EpicGames.Core/ListExtensions.cs
@@ -0,0 +1,43 @@
+// Copyright Epic Games, Inc. All Rights Reserved.
+
+using System;
+using System.Collections;
+using System.Collections.Generic;
+using System.Diagnostics.CodeAnalysis;
+using System.Globalization;
+using System.Text;
+
+namespace EpicGames.Core
+{
+	/// <summary>
+	/// Extension methods for lists
+	/// </summary>s
+	public static class ListExtensions
+	{
+		/// <summary>
+		/// Sorts a list by a particular field
+		/// </summary>
+		/// <typeparam name="TElement">List element</typeparam>
+		/// <typeparam name="TField">Field type to sort by</typeparam>
+		/// <param name="List">List to sort</param>
+		/// <param name="Selector">Selects a field from the element</param>
+		public static void SortBy<TElement, TField>(this List<TElement> List, Func<TElement, TField> Selector)
+		{
+			IComparer<TField> DefaultComparer = Comparer<TField>.Default;
+			SortBy(List, Selector, DefaultComparer);
+		}
+
+		/// <summary>
+		/// Sorts a list by a particular field
+		/// </summary>
+		/// <typeparam name="TElement">List element</typeparam>
+		/// <typeparam name="TField">Field type to sort by</typeparam>
+		/// <param name="List">List to sort</param>
+		/// <param name="Selector">Selects a field from the element</param>
+		/// <param name="Comparer">Comparer for fields</param>
+		public static void SortBy<TElement, TField>(this List<TElement> List, Func<TElement, TField> Selector, IComparer<TField> Comparer)
+		{
+			List.Sort((X, Y) => Comparer.Compare(Selector(X), Selector(Y)));
+		}
+	}
+}
diff --git a/Engine/Source/Programs/Shared/EpicGames.Core/Log.cs b/Engine/Source/Programs/Shared/EpicGames.Core/Log.cs
--- a/Engine/Source/Programs/Shared/EpicGames.Core/Log.cs
+++ b/Engine/Source/Programs/Shared/EpicGames.Core/Log.cs
@@ -9,6 +9,8 @@ using System.Reflection;
 using System.Runtime.CompilerServices;
 using System.Text;
 
+#nullable disable
+
 namespace EpicGames.Core
 {
 	/// <summary>
diff --git a/Engine/Source/Programs/Shared/EpicGames.Core/LogIndentScope.cs b/Engine/Source/Programs/Shared/EpicGames.Core/LogIndentScope.cs
--- a/Engine/Source/Programs/Shared/EpicGames.Core/LogIndentScope.cs
+++ b/Engine/Source/Programs/Shared/EpicGames.Core/LogIndentScope.cs
@@ -16,7 +16,7 @@ namespace EpicGames.Core
 		/// <summary>
 		/// The previous indent
 		/// </summary>
-		string PrevIndent;
+		string? PrevIndent;
 
 		/// <summary>
 		/// Constructor
diff --git a/Engine/Source/Programs/Shared/EpicGames.Core/LogStatusScope.cs b/Engine/Source/Programs/Shared/EpicGames.Core/LogStatusScope.cs
--- a/Engine/Source/Programs/Shared/EpicGames.Core/LogStatusScope.cs
+++ b/Engine/Source/Programs/Shared/EpicGames.Core/LogStatusScope.cs
@@ -91,7 +91,7 @@ namespace EpicGames.Core
 			if(Message != null)
 			{
 				Log.PopStatus();
-				Message = null;
+				Message = null!;
 			}
 		}
 	}
diff --git a/Engine/Source/Programs/Shared/EpicGames.Core/ManagedProcess.cs b/Engine/Source/Programs/Shared/EpicGames.Core/ManagedProcess.cs
--- a/Engine/Source/Programs/Shared/EpicGames.Core/ManagedProcess.cs
+++ b/Engine/Source/Programs/Shared/EpicGames.Core/ManagedProcess.cs
@@ -12,6 +12,8 @@ using System.Text;
 using System.Threading;
 using System.Threading.Tasks;
 
+#nullable disable
+
 namespace EpicGames.Core
 {
 	/// <summary>
diff --git a/Engine/Source/Programs/Shared/EpicGames.Core/ParallelTask.cs b/Engine/Source/Programs/Shared/EpicGames.Core/ParallelTask.cs
new file mode 100644
--- /dev/null
+++ b/Engine/Source/Programs/Shared/EpicGames.Core/ParallelTask.cs
@@ -0,0 +1,53 @@
+// Copyright Epic Games, Inc. All Rights Reserved.
+
+using System;
+using System.Collections.Generic;
+using System.Linq;
+using System.Text;
+using System.Threading.Tasks;
+using System.Threading.Tasks.Dataflow;
+
+namespace EpicGames.Core
+{
+	/// <summary>
+	/// Async equivalents for parallel methods
+	/// </summary>
+	public static class ParallelTask
+	{
+		/// <summary>
+		/// Execute a large number of tasks in parallel over a collection. Parallel.ForEach() method isn't generally compatible with asynchronous programming, because any
+		/// exceptions are thrown on the 
+		/// </summary>
+		/// <param name="FromInclusive">The starting index</param>
+		/// <param name="ToExclusive">The last index, exclusive</param>
+		/// <param name="Action">Action to perform for each item in the collection</param>
+		/// <returns>Async task</returns>
+		public static Task ForAsync(int FromInclusive, int ToExclusive, Action<int> Action)
+		{
+			return ForEachAsync(Enumerable.Range(FromInclusive, ToExclusive - FromInclusive), Action);
+		}
+
+		/// <summary>
+		/// Execute a large number of tasks in parallel over a collection. Parallel.ForEach() method isn't generally compatible with asynchronous programming, because any
+		/// exceptions are thrown on the 
+		/// </summary>
+		/// <typeparam name="T">The collection type</typeparam>
+		/// <param name="Collection">The collection to iterate over</param>
+		/// <param name="Action">Action to perform for each item in the collection</param>
+		/// <returns>Async task</returns>
+		public static Task ForEachAsync<T>(IEnumerable<T> Collection, Action<T> Action)
+		{
+			ExecutionDataflowBlockOptions Options = new ExecutionDataflowBlockOptions();
+			Options.MaxDegreeOfParallelism = DataflowBlockOptions.Unbounded;
+
+			ActionBlock<T> Actions = new ActionBlock<T>(Action, Options);
+			foreach (T Item in Collection)
+			{
+				Actions.Post(Item);
+			}
+			Actions.Complete();
+
+			return Actions.Completion;
+		}
+	}
+}
diff --git a/Engine/Source/Programs/Shared/EpicGames.Core/ReadOnlyListExtensions.cs b/Engine/Source/Programs/Shared/EpicGames.Core/ReadOnlyListExtensions.cs
new file mode 100644
--- /dev/null
+++ b/Engine/Source/Programs/Shared/EpicGames.Core/ReadOnlyListExtensions.cs
@@ -0,0 +1,129 @@
+// Copyright Epic Games, Inc. All Rights Reserved.
+
+using System;
+using System.Collections;
+using System.Collections.Generic;
+using System.Linq;
+using System.Threading.Tasks;
+
+namespace EpicGames.Core
+{
+	/// <summary>
+	/// Extension methods for functionality on <see cref="List{T}"/> but missing from <see cref="IReadOnlyList{T}"/>
+	/// </summary>
+	public static class ReadOnlyListExtensions
+	{
+		/// <summary>
+		/// Performs a binary search on the given list
+		/// </summary>
+		/// <typeparam name="T">The element type for the list</typeparam>
+		/// <param name="List">The list type</param>
+		/// <param name="Item">Item to search for</param>
+		/// <returns>As List.BinarySearch</returns>
+		public static int BinarySearch<T>(this IReadOnlyList<T> List, T Item)
+		{
+			return BinarySearch(List, x => x, Item, Comparer<T>.Default);
+		}
+
+		/// <summary>
+		/// Performs a binary search on the given list
+		/// </summary>
+		/// <typeparam name="T">The element type for the list</typeparam>
+		/// <param name="List">The list type</param>
+		/// <param name="Item">Item to search for</param>
+		/// <param name="Comparer">Comparer for elements in the list</param>
+		/// <returns>As List.BinarySearch</returns>
+		public static int BinarySearch<T>(this IReadOnlyList<T> List, T Item, IComparer<T> Comparer)
+		{
+			return BinarySearch(List, x => x, Item, Comparer);
+		}
+
+		/// <summary>
+		/// Binary searches a list based on a projection
+		/// </summary>
+		/// <typeparam name="TItem">The item in the list</typeparam>
+		/// <typeparam name="TField">The field to search on</typeparam>
+		/// <param name="List">The list to search</param>
+		/// <param name="Projection">The projection to apply to each item in the list</param>
+		/// <param name="Item">The item to find</param>
+		/// <returns>As <see cref="List{T}.BinarySearch(T)"/></returns>
+		public static int BinarySearch<TItem, TField>(this IReadOnlyList<TItem> List, Func<TItem, TField> Projection, TField Item)
+		{
+			return BinarySearch(List, Projection, Item, Comparer<TField>.Default);
+		}
+
+		/// <summary>
+		/// Binary searches a list based on a projection
+		/// </summary>
+		/// <typeparam name="TItem">The item in the list</typeparam>
+		/// <typeparam name="TField">The field to search on</typeparam>
+		/// <param name="List">The list to search</param>
+		/// <param name="Projection">The projection to apply to each item in the list</param>
+		/// <param name="Item">The item to find</param>
+		/// <param name="Comparer">Comparer for field elements</param>
+		/// <returns>As <see cref="List{T}.BinarySearch(T)"/></returns>
+		public static int BinarySearch<TItem, TField>(this IReadOnlyList<TItem> List, Func<TItem, TField> Projection, TField Item, IComparer<TField> Comparer)
+		{
+			int LowerBound = 0;
+			int UpperBound = List.Count - 1;
+			while (LowerBound <= UpperBound)
+			{
+				int Idx = LowerBound + (UpperBound - LowerBound) / 2;
+
+				int Comparison = Comparer.Compare(Projection(List[Idx]), Item);
+				if (Comparison == 0)
+				{
+					return Idx;
+				}
+				else if (Comparison < 0)
+				{
+					LowerBound = Idx + 1;
+				}
+				else
+				{
+					UpperBound = Idx - 1;
+				}
+			}
+			return ~LowerBound;
+		}
+
+		/// <summary>
+		/// Converts a read only list to a different type
+		/// </summary>
+		/// <typeparam name="TInput">The input element type</typeparam>
+		/// <typeparam name="TOutput">The output element type</typeparam>
+		/// <param name="Input">Input list</param>
+		/// <param name="Convert">Conversion function</param>
+		/// <returns>New list of items</returns>
+		public static List<TOutput> ConvertAll<TInput, TOutput>(this IReadOnlyList<TInput> Input, Func<TInput, TOutput> Convert)
+		{
+			List<TOutput> Outputs = new List<TOutput>(Input.Count);
+			for (int Idx = 0; Idx < Input.Count; Idx++)
+			{
+				Outputs.Add(Convert(Input[Idx]));
+			}
+			return Outputs;
+		}
+
+		/// <summary>
+		/// Finds the index of the first element matching a predicate
+		/// </summary>
+		/// <typeparam name="T"></typeparam>
+		/// <param name="List">List to search</param>
+		/// <param name="Predicate">Predicate for the list</param>
+		/// <returns>Index of the element</returns>
+		public static int FindIndex<T>(this IReadOnlyList<T> List, Predicate<T> Predicate)
+		{
+			int FoundIndex = -1;
+			for(int Idx = 0; Idx < List.Count; Idx++)
+			{
+				if (Predicate(List[Idx]))
+				{
+					FoundIndex = Idx;
+					break;
+				}
+			}
+			return FoundIndex;
+		}
+	}
+}
diff --git a/Engine/Source/Programs/Shared/EpicGames.Core/ReadOnlyUtf8String.cs b/Engine/Source/Programs/Shared/EpicGames.Core/ReadOnlyUtf8String.cs
new file mode 100644
--- /dev/null
+++ b/Engine/Source/Programs/Shared/EpicGames.Core/ReadOnlyUtf8String.cs
@@ -0,0 +1,289 @@
+// Copyright Epic Games, Inc. All Rights Reserved.
+
+using System;
+using System.Collections.Generic;
+using System.Runtime.CompilerServices;
+using System.Text;
+
+namespace EpicGames.Core
+{
+	/// <summary>
+	/// Represents a memory region which can be treated as a utf-8 string.
+	/// </summary>
+	public struct ReadOnlyUtf8String : IEquatable<ReadOnlyUtf8String>
+	{
+		/// <summary>
+		/// An empty string
+		/// </summary>
+		public static readonly ReadOnlyUtf8String Empty = new ReadOnlyUtf8String();
+
+		/// <summary>
+		/// The data represented by this string
+		/// </summary>
+		public ReadOnlyMemory<byte> Memory { get; }
+
+		/// <summary>
+		/// Returns read only span for this string
+		/// </summary>
+		public ReadOnlySpan<byte> Span
+		{
+			get { return Memory.Span; } 
+		}
+
+		/// <summary>
+		/// Determines if this string is empty
+		/// </summary>
+		public bool IsEmpty
+		{
+			get { return Memory.IsEmpty; }
+		}
+
+		/// <summary>
+		/// Returns the length of this string
+		/// </summary>
+		public int Length
+		{
+			get { return Memory.Length; }
+		}
+
+		/// <summary>
+		/// Allows indexing individual bytes of the data
+		/// </summary>
+		/// <param name="Index">Byte index</param>
+		/// <returns>Byte at the given index</returns>
+		public byte this[int Index]
+		{
+			get { return Span[Index]; }
+		}
+
+		/// <summary>
+		/// Constructor
+		/// </summary>
+		/// <param name="Text">Text to construct from</param>
+		public ReadOnlyUtf8String(string Text)
+		{
+			this.Memory = Encoding.UTF8.GetBytes(Text);
+		}
+
+		/// <summary>
+		/// Constructor
+		/// </summary>
+		/// <param name="Memory">The data to construct from</param>
+		public ReadOnlyUtf8String(ReadOnlyMemory<byte> Memory)
+		{
+			this.Memory = Memory;
+		}
+
+		/// <summary>
+		/// Constructor
+		/// </summary>
+		/// <param name="Buffer">The buffer to construct from</param>
+		/// <param name="Offset">Offset within the buffer</param>
+		/// <param name="Length">Length of the string within the buffer</param>
+		public ReadOnlyUtf8String(byte[] Buffer, int Offset, int Length)
+		{
+			this.Memory = new ReadOnlyMemory<byte>(Buffer, Offset, Length);
+		}
+
+		/// <summary>
+		/// Tests two strings for equality
+		/// </summary>
+		/// <param name="A">The first string to compare</param>
+		/// <param name="B">The second string to compare</param>
+		/// <returns>True if the strings are equal</returns>
+		public static bool operator ==(ReadOnlyUtf8String A, ReadOnlyUtf8String B)
+		{
+			return A.Equals(B);
+		}
+
+		/// <summary>
+		/// Tests two strings for inequality
+		/// </summary>
+		/// <param name="A">The first string to compare</param>
+		/// <param name="B">The second string to compare</param>
+		/// <returns>True if the strings are not equal</returns>
+		public static bool operator !=(ReadOnlyUtf8String A, ReadOnlyUtf8String B)
+		{
+			return !A.Equals(B);
+		}
+
+		/// <summary>
+		/// Tests whether two strings are equal
+		/// </summary>
+		/// <param name="Other">The other string to compare to</param>
+		/// <returns>True if the strings are equal</returns>
+		public bool Equals(ReadOnlyUtf8String Other)
+		{
+			return Span.SequenceEqual(Other.Span);
+		}
+
+		/// <summary>
+		/// Tests if this string starts with another string
+		/// </summary>
+		/// <param name="Other">The string to check against</param>
+		/// <returns>True if this string starts with the other string</returns>
+		public bool StartsWith(ReadOnlyUtf8String Other)
+		{
+			return Span.StartsWith(Other.Span);
+		}
+
+		/// <summary>
+		/// Tests if this string ends with another string
+		/// </summary>
+		/// <param name="Other">The string to check against</param>
+		/// <returns>True if this string ends with the other string</returns>
+		public bool EndsWith(ReadOnlyUtf8String Other)
+		{
+			return Span.EndsWith(Other.Span);
+		}
+
+		/// <summary>
+		/// Slices this string at the given start position
+		/// </summary>
+		/// <param name="Start">Start position</param>
+		/// <returns>String starting at the given position</returns>
+		public ReadOnlyUtf8String Slice(int Start)
+		{
+			return new ReadOnlyUtf8String(Memory.Slice(Start));
+		}
+
+		/// <summary>
+		/// Slices this string at the given start position and length
+		/// </summary>
+		/// <param name="Start">Start position</param>
+		/// <returns>String starting at the given position</returns>
+		public ReadOnlyUtf8String Slice(int Start, int Count)
+		{
+			return new ReadOnlyUtf8String(Memory.Slice(Start, Count));
+		}
+
+		/// <summary>
+		/// Tests if this string is equal to the other object
+		/// </summary>
+		/// <param name="obj">Object to compare to</param>
+		/// <returns>True if the objects are equivalent</returns>
+		public override bool Equals(object? obj)
+		{
+			ReadOnlyUtf8String? Other = obj as ReadOnlyUtf8String?;
+			return Other != null && Equals(Other.Value);
+		}
+
+		/// <summary>
+		/// Returns the hash code of this string
+		/// </summary>
+		/// <returns>Hash code for the string</returns>
+		public override int GetHashCode()
+		{
+			int Hash = 5381;
+			for(int Idx = 0; Idx < Memory.Length; Idx++)
+			{
+				Hash += (Hash << 5) + Span[Idx];
+			}
+			return Hash;
+		}
+
+		/// <summary>
+		/// Gets the string represented by this data
+		/// </summary>
+		/// <returns>The utf-8 string</returns>
+		public override string ToString()
+		{
+			return Encoding.UTF8.GetString(Span);
+		}
+
+		/// <summary>
+		/// Converts a string to a utf-8 string
+		/// </summary>
+		/// <param name="Text">Text to convert</param>
+		public static implicit operator ReadOnlyUtf8String(string Text)
+		{
+			return new ReadOnlyUtf8String(new ReadOnlyMemory<byte>(Encoding.UTF8.GetBytes(Text)));
+		}
+	}
+
+	/// <summary>
+	/// Comparison classes for utf8 strings
+	/// </summary>
+	public abstract class ReadOnlyUtf8StringComparer : IEqualityComparer<ReadOnlyUtf8String>
+	{
+		/// <summary>
+		/// Ordinal comparer for utf8 strings
+		/// </summary>
+		class OrdinalComparer : ReadOnlyUtf8StringComparer
+		{
+			/// <inheritdoc/>
+			public override bool Equals(ReadOnlyUtf8String StrA, ReadOnlyUtf8String StrB)
+			{
+				return StrA.Equals(StrB);
+			}
+
+			/// <inheritdoc/>
+			public override int GetHashCode(ReadOnlyUtf8String String)
+			{
+				return String.GetHashCode();
+			}
+		}
+
+		/// <summary>
+		/// Comparison between ReadOnlyUtf8String objects that ignores case for ASCII characters
+		/// </summary>
+		class OrdinalIgnoreCaseComparer : ReadOnlyUtf8StringComparer 
+		{
+			/// <inheritdoc/>
+			public override bool Equals(ReadOnlyUtf8String StrA, ReadOnlyUtf8String StrB)
+			{
+				if (StrA.Length != StrB.Length)
+				{
+					return false;
+				}
+
+				for (int Idx = 0; Idx < StrA.Length; Idx++)
+				{
+					if (StrA[Idx] != StrB[Idx] && ToUpper(StrA[Idx]) != ToUpper(StrB[Idx]))
+					{
+						return false;
+					}
+				}
+
+				return true;
+			}
+
+			/// <inheritdoc/>
+			public override int GetHashCode(ReadOnlyUtf8String String)
+			{
+				HashCode HashCode = new HashCode();
+				for (int Idx = 0; Idx < String.Length; Idx++)
+				{
+					HashCode.Add(ToUpper(String[Idx]));
+				}
+				return HashCode.ToHashCode();
+			}
+
+			/// <summary>
+			/// Convert a character to uppercase
+			/// </summary>
+			/// <param name="Character">Character to convert</param>
+			/// <returns>The uppercase version of the character</returns>
+			static byte ToUpper(byte Character)
+			{
+				return (Character >= 'a' && Character <= 'z') ? (byte)(Character - 'a' + 'A') : Character;
+			}
+		}
+
+		/// <summary>
+		/// Static instance of the ordinal utf8 ordinal comparer
+		/// </summary>
+		public static ReadOnlyUtf8StringComparer Ordinal = new OrdinalComparer();
+
+		/// <summary>
+		/// Static instance of the case-insensitive utf8 ordinal string comparer
+		/// </summary>
+		public static ReadOnlyUtf8StringComparer OrdinalIgnoreCase = new OrdinalIgnoreCaseComparer();
+
+		/// <inheritdoc/>
+		public abstract bool Equals(ReadOnlyUtf8String StrA, ReadOnlyUtf8String StrB);
+
+		/// <inheritdoc/>
+		public abstract int GetHashCode(ReadOnlyUtf8String String);
+	}
+}
diff --git a/Engine/Source/Programs/Shared/EpicGames.Core/StringUtils.cs b/Engine/Source/Programs/Shared/EpicGames.Core/StringUtils.cs
--- a/Engine/Source/Programs/Shared/EpicGames.Core/StringUtils.cs
+++ b/Engine/Source/Programs/Shared/EpicGames.Core/StringUtils.cs
@@ -2,6 +2,7 @@
 
 using System;
 using System.Collections.Generic;
+using System.Diagnostics.CodeAnalysis;
 using System.Linq;
 using System.Text;
 using System.Threading.Tasks;
@@ -10,6 +11,35 @@ namespace EpicGames.Core
 {
 	public static class StringUtils
 	{
+		/// <summary>
+		/// Array mapping from ascii index to hexadecimal digits.
+		/// </summary>
+		static sbyte[] HexDigits;
+
+		/// <summary>
+		/// Static constructor. Initializes the HexDigits array.
+		/// </summary>
+		static StringUtils()
+		{
+			HexDigits = new sbyte[256];
+			for (int Idx = 0; Idx < 256; Idx++)
+			{
+				HexDigits[Idx] = -1;
+			}
+			for (int Idx = '0'; Idx <= '9'; Idx++)
+			{
+				HexDigits[Idx] = (sbyte)(Idx - '0');
+			}
+			for (int Idx = 'a'; Idx <= 'f'; Idx++)
+			{
+				HexDigits[Idx] = (sbyte)(10 + Idx - 'a');
+			}
+			for (int Idx = 'A'; Idx <= 'F'; Idx++)
+			{
+				HexDigits[Idx] = (sbyte)(10 + Idx - 'A');
+			}
+		}
+
 		/// <summary>
 		/// Indents a string by a given indent
 		/// </summary>
@@ -26,6 +56,57 @@ namespace EpicGames.Core
 			return Result;
 		}
 
+		/// <summary>
+		/// Expand all the property references (of the form $(PropertyName)) in a string.
+		/// </summary>
+		/// <param name="Text">The input string to expand properties in</param>
+		/// <param name="Properties">Dictionary of properties to expand</param>
+		/// <returns>The expanded string</returns>
+		public static string ExpandProperties(string Text, Dictionary<string, string> Properties)
+		{
+			return ExpandProperties(Text, Name => { Properties.TryGetValue(Name, out string? Value); return Value; });
+		}
+
+		/// <summary>
+		/// Expand all the property references (of the form $(PropertyName)) in a string.
+		/// </summary>
+		/// <param name="Text">The input string to expand properties in</param>
+		/// <param name="GetPropertyValue">Delegate to retrieve a property value</param>
+		/// <returns>The expanded string</returns>
+		public static string ExpandProperties(string Text, Func<string, string?> GetPropertyValue)
+		{
+			string Result = Text;
+			for (int Idx = Result.IndexOf("$("); Idx != -1; Idx = Result.IndexOf("$(", Idx))
+			{
+				// Find the end of the variable name
+				int EndIdx = Result.IndexOf(')', Idx + 2);
+				if (EndIdx == -1)
+				{
+					break;
+				}
+
+				// Extract the variable name from the string
+				string Name = Result.Substring(Idx + 2, EndIdx - (Idx + 2));
+
+				// Check if we've got a value for this variable
+				string? Value = GetPropertyValue(Name);
+				if (Value == null)
+				{
+					// Do not expand it; must be preprocessing the script.
+					Idx = EndIdx;
+				}
+				else
+				{
+					// Replace the variable, or skip past it
+					Result = Result.Substring(0, Idx) + Value + Result.Substring(EndIdx + 1);
+
+					// Make sure we skip over the expanded variable; we don't want to recurse on it.
+					Idx += Value.Length;
+				}
+			}
+			return Result;
+		}
+
 		/// <summary>
 		/// Takes a given sentence and wraps it on a word by word basis so that no line exceeds the
 		/// set maximum line length. Words longer than a line are broken up. Returns the sentence as
@@ -145,41 +226,58 @@ namespace EpicGames.Core
 		}
 
 		/// <summary>
-		/// Parses a hexadecimal digit
+		/// Formats a list of items
 		/// </summary>
-		/// <param name="Character">Character to parse</param>
-		/// <returns>Value of this digit</returns>
-		public static bool TryParseHexDigit(char Character, out int Value)
+		/// <param name="Items">Array of items</param>
+		/// <param name="MaxCount">Maximum number of items to include in the list</param>
+		/// <returns>Formatted list of items</returns>
+		public static string FormatList(string[] Items, int MaxCount)
 		{
-			if(Character >= '0' && Character <= '9')
+			if (Items.Length == 0)
 			{
-				Value = Character - '0';
-				return true;
+				return "unknown";
 			}
-			else if(Character >= 'a' && Character <= 'f')
+			else if (Items.Length == 1)
 			{
-				Value = 10 + (Character - 'a');
-				return true;
+				return Items[0];
 			}
-			else if(Character >= 'A' && Character <= 'F')
+			else if (Items.Length <= MaxCount)
 			{
-				Value = 10 + (Character - 'A');
-				return true;
+				return $"{String.Join(", ", Items.Take(Items.Length - 1))} and {Items.Last()}";
 			}
 			else
 			{
-				Value = 0;
-				return false;
+				return $"{String.Join(", ", Items.Take(MaxCount - 1))} and {Items.Length - (MaxCount - 1)} others";
 			}
 		}
 
+		/// <summary>
+		/// Parses a hexadecimal digit
+		/// </summary>
+		/// <param name="Character">Character to parse</param>
+		/// <returns>Value of this digit, or -1 if invalid</returns>
+		public static int GetHexDigit(byte Character)
+		{
+			return HexDigits[Character];
+		}
+
+		/// <summary>
+		/// Parses a hexadecimal digit
+		/// </summary>
+		/// <param name="Character">Character to parse</param>
+		/// <returns>Value of this digit, or -1 if invalid</returns>
+		public static int GetHexDigit(char Character)
+		{
+			return HexDigits[Math.Min((uint)Character, 127)];
+		}
+
 		/// <summary>
 		/// Parses a hexadecimal string into an array of bytes
 		/// </summary>
 		/// <returns>Array of bytes</returns>
 		public static byte[] ParseHexString(string Text)
 		{
-			byte[] Bytes;
+			byte[]? Bytes;
 			if(!TryParseHexString(Text, out Bytes))
 			{
 				throw new FormatException(String.Format("Invalid hex string: '{0}'", Text));
@@ -187,12 +285,26 @@ namespace EpicGames.Core
 			return Bytes;
 		}
 
+		/// <summary>
+		/// Parses a hexadecimal string into an array of bytes
+		/// </summary>
+		/// <returns>Array of bytes</returns>
+		public static byte[] ParseHexString(ReadOnlySpan<byte> Text)
+		{
+			byte[]? Bytes;
+			if (!TryParseHexString(Text, out Bytes))
+			{
+				throw new FormatException($"Invalid hex string: '{Encoding.UTF8.GetString(Text)}'");
+			}
+			return Bytes;
+		}
+
 		/// <summary>
 		/// Parses a hexadecimal string into an array of bytes
 		/// </summary>
 		/// <param name="Text">Text to parse</param>
 		/// <returns></returns>
-		public static bool TryParseHexString(string Text, out byte[] OutBytes)
+		public static bool TryParseHexString(string Text, [NotNullWhen(true)] out byte[]? OutBytes)
 		{
 			if((Text.Length & 1) != 0)
 			{
@@ -202,18 +314,56 @@ namespace EpicGames.Core
 			byte[] Bytes = new byte[Text.Length / 2];
 			for(int Idx = 0; Idx < Text.Length; Idx += 2)
 			{
-				int A, B;
-				if(!TryParseHexDigit(Text[Idx], out A) || !TryParseHexDigit(Text[Idx + 1], out B))
+				int Value = (GetHexDigit(Text[Idx]) << 4) | GetHexDigit(Text[Idx + 1]);
+				if(Value < 0)
+				{
+					OutBytes = null;
+					return false;
+				}
+				Bytes[Idx / 2] = (byte)Value;
+			}
+			OutBytes = Bytes;
+			return true;
+		}
+
+		/// <summary>
+		/// Parses a hexadecimal string into an array of bytes
+		/// </summary>
+		/// <param name="Text">Text to parse</param>
+		/// <returns></returns>
+		public static bool TryParseHexString(ReadOnlySpan<byte> Text, [NotNullWhen(true)] out byte[]? OutBytes)
+		{
+			if ((Text.Length & 1) != 0)
+			{
+				throw new FormatException("Length of hex string must be a multiple of two characters");
+			}
+
+			byte[] Bytes = new byte[Text.Length / 2];
+			for (int Idx = 0; Idx < Text.Length; Idx += 2)
+			{
+				int Value = ParseHexByte(Text, Idx);
+				if (Value < 0)
 				{
 					OutBytes = null;
 					return false;
 				}
-				Bytes[Idx / 2] = (byte)((A << 4) | B);
+				Bytes[Idx / 2] = (byte)Value;
 			}
 			OutBytes = Bytes;
 			return true;
 		}
 
+		/// <summary>
+		/// Parse a hex byte from the given offset into a span of utf8 characters
+		/// </summary>
+		/// <param name="Text">The text to parse</param>
+		/// <param name="Idx">Index within the text to parse</param>
+		/// <returns>The parsed value, or a negative value on error</returns>
+		public static int ParseHexByte(ReadOnlySpan<byte> Text, int Idx)
+		{
+			return ((int)HexDigits[Text[Idx]] << 4) | ((int)HexDigits[Text[Idx + 1]]);
+		}
+
 		/// <summary>
 		/// Formats an array of bytes as a hexadecimal string
 		/// </summary>
@@ -231,5 +381,22 @@ namespace EpicGames.Core
 			}
 			return new string(Characters);
 		}
+
+		/// <summary>
+		/// Quotes a string as a command line argument
+		/// </summary>
+		/// <param name="String">The string to quote</param>
+		/// <returns>The quoted argument if it contains any spaces, otherwise the original string</returns>
+		public static string QuoteArgument(this string String)
+		{
+			if (String.Contains(' '))
+			{
+				return $"\"{String}\"";
+			}
+			else
+			{
+				return String;
+			}
+		}
 	}
 }
diff --git a/Engine/Source/Programs/Shared/EpicGames.Core/StringView.cs b/Engine/Source/Programs/Shared/EpicGames.Core/StringView.cs
new file mode 100644
--- /dev/null
+++ b/Engine/Source/Programs/Shared/EpicGames.Core/StringView.cs
@@ -0,0 +1,173 @@
+// Copyright Epic Games, Inc. All Rights Reserved.
+
+using System;
+using System.Collections.Generic;
+using System.Diagnostics.CodeAnalysis;
+using System.Text;
+
+namespace EpicGames.Core
+{
+	/// <summary>
+	/// View of a character string. Allows comparing/manipulating substrings without unnecessary memory allocatinos.
+	/// </summary>
+	public readonly struct StringView : IEquatable<StringView>
+	{
+		/// <summary>
+		/// Memory containing the characters
+		/// </summary>
+		public ReadOnlyMemory<char> Memory { get; }
+
+		/// <summary>
+		/// Span for the sequence of characters
+		/// </summary>
+		public ReadOnlySpan<char> Span 
+		{ 
+			get { return Memory.Span; } 
+		}
+
+		/// <summary>
+		/// Constructor
+		/// </summary>
+		/// <param name="Memory">The memory containing the characters</param>
+		public StringView(ReadOnlyMemory<char> Memory)
+		{
+			this.Memory = Memory;
+		}
+
+		/// <summary>
+		/// Constructor
+		/// </summary>
+		/// <param name="Text">String to construct from</param>
+		public StringView(string Text)
+		{
+			Memory = Text.AsMemory();
+		}
+
+		/// <summary>
+		/// Constructs a view onto a substring of the given string
+		/// </summary>
+		/// <param name="Text">String to construct from</param>
+		/// <param name="Index">Offset within the string for this view</param>
+		public StringView(string Text, int Index)
+		{
+			Memory = Text.AsMemory(Index);
+		}
+
+		/// <summary>
+		/// Constructs a view onto a substring of the given string
+		/// </summary>
+		/// <param name="Text">String to construct from</param>
+		/// <param name="Index">Offset within the string for this view</param>
+		/// <param name="Count">Number of characters to include</param>
+		public StringView(string Text, int Index, int Count)
+		{
+			Memory = Text.AsMemory(Index, Count);
+		}
+
+		/// <summary>
+		/// Equality comparer
+		/// </summary>
+		/// <param name="X">First string to compare</param>
+		/// <param name="Y">Second string to compare</param>
+		/// <returns>True if the strings are equal</returns>
+		public static bool operator ==(StringView X, StringView Y)
+		{
+			return X.Equals(Y);
+		}
+
+		/// <summary>
+		/// Inequality comparer
+		/// </summary>
+		/// <param name="X">First string to compare</param>
+		/// <param name="Y">Second string to compare</param>
+		/// <returns>True if the strings are equal</returns>
+		public static bool operator !=(StringView X, StringView Y)
+		{
+			return !X.Equals(Y);
+		}
+
+		/// <inheritdoc/>
+		public override bool Equals(object? Obj)
+		{
+			return Obj is StringView && Equals((StringView)Obj);
+		}
+
+		/// <inheritdoc/>
+		public override int GetHashCode()
+		{
+			return String.GetHashCode(Memory.Span);
+		}
+
+		/// <inheritdoc/>
+		public override string ToString()
+		{
+			return new string(Memory.Span);
+		}
+
+		/// <inheritdoc/>
+		public bool Equals(StringView Other)
+		{
+			return Equals(Other, StringComparison.CurrentCulture);
+		}
+
+		/// <inheritdoc/>
+		public bool Equals(StringView? Other, StringComparison comparisonType)
+		{
+			return Other.HasValue && Memory.Span.Equals(Other.Value.Memory.Span, comparisonType);
+		}
+
+		/// <summary>
+		/// Implicit conversion operator from a regular string
+		/// </summary>
+		/// <param name="Text">The string to construct from</param>
+		public static implicit operator StringView(string Text) => new StringView(Text);
+	}
+
+	/// <summary>
+	/// Comparer for StringView objects
+	/// </summary>
+	public class StringViewComparer : IComparer<StringView>, IEqualityComparer<StringView>
+	{
+		/// <summary>
+		/// Static instance of an ordinal StringView comparer
+		/// </summary>
+		public static StringViewComparer Ordinal = new StringViewComparer(StringComparison.Ordinal);
+
+		/// <summary>
+		/// Static instance of an ordinal StringView comparer which ignores case
+		/// </summary>
+		public static StringViewComparer OrdinalIgnoreCase = new StringViewComparer(StringComparison.OrdinalIgnoreCase);
+
+		/// <summary>
+		/// The comparison type
+		/// </summary>
+		public StringComparison ComparisonType { get; }
+
+		/// <summary>
+		/// Constructor
+		/// </summary>
+		/// <param name="ComparisonType">Type of comparison to perform</param>
+		public StringViewComparer(StringComparison ComparisonType)
+		{
+			this.ComparisonType = ComparisonType;
+		}
+
+		/// <inheritdoc/>
+		public bool Equals(StringView X, StringView Y)
+		{
+			return X.Span.Equals(Y.Span, ComparisonType);
+		}
+
+		/// <inheritdoc/>
+		public int GetHashCode(StringView obj)
+		{
+			return String.GetHashCode(obj.Span);
+		}
+
+		/// <inheritdoc/>
+		public int Compare(StringView X, StringView Y)
+		{
+			return X.Span.CompareTo(Y.Span, ComparisonType);
+		}
+	}
+}
diff --git a/Engine/Source/Programs/Shared/EpicGames.Core/ThreadPoolWorkQueue.cs b/Engine/Source/Programs/Shared/EpicGames.Core/ThreadPoolWorkQueue.cs
--- a/Engine/Source/Programs/Shared/EpicGames.Core/ThreadPoolWorkQueue.cs
+++ b/Engine/Source/Programs/Shared/EpicGames.Core/ThreadPoolWorkQueue.cs
@@ -51,7 +51,7 @@ namespace EpicGames.Core
 				Wait();
 
 				EmptyEvent.Dispose();
-				EmptyEvent = null;
+				EmptyEvent = null!;
 			}
 		}
 
@@ -89,11 +89,11 @@ namespace EpicGames.Core
 		/// Internal method to execute an action
 		/// </summary>
 		/// <param name="ActionToExecute">The action to execute</param>
-		void Execute(object ActionToExecute)
+		void Execute(object? ActionToExecute)
 		{
 			try
 			{
-				((Action)ActionToExecute)();
+				((Action)ActionToExecute!)();
 			}
 			catch(Exception Ex)
 			{
diff --git a/Engine/Source/Programs/Shared/EpicGames.Core/TraceSpan.cs b/Engine/Source/Programs/Shared/EpicGames.Core/TraceSpan.cs
--- a/Engine/Source/Programs/Shared/EpicGames.Core/TraceSpan.cs
+++ b/Engine/Source/Programs/Shared/EpicGames.Core/TraceSpan.cs
@@ -36,13 +36,13 @@ namespace EpicGames.Core
 		class TraceSpanImpl : ITraceSpan
 		{
 			public string Name;
-			public string Resource;
-			public string Service;
+			public string? Resource;
+			public string? Service;
 			public DateTimeOffset StartTime;
 			public DateTimeOffset? FinishTime;
 			public Dictionary<string, string> Metadata = new Dictionary<string, string>();
 
-			public TraceSpanImpl(string Name, string Resource, string Service)
+			public TraceSpanImpl(string Name, string? Resource, string? Service)
 			{
 				this.Name = Name;
 				this.Resource = Resource;
@@ -72,7 +72,7 @@ namespace EpicGames.Core
 		/// <summary>
 		/// Creates a scope using the current provider
 		/// </summary>
-		public static ITraceSpan Create(string Name, string Resource = null, string Service = null)
+		public static ITraceSpan Create(string Name, string? Resource = null, string? Service = null)
 		{
 			TraceSpanImpl Span = new TraceSpanImpl(Name, Resource, Service);
 			Spans.Add(Span);
@@ -84,7 +84,7 @@ namespace EpicGames.Core
 		/// </summary>
 		public static void Flush()
 		{
-			string TelemetryDir = Environment.GetEnvironmentVariable("UE_TELEMETRY_DIR");
+			string? TelemetryDir = Environment.GetEnvironmentVariable("UE_TELEMETRY_DIR");
 			if (TelemetryDir != null)
 			{
 				FileReference File;
@@ -93,7 +93,7 @@ namespace EpicGames.Core
 					DirectoryReference TelemetryDirRef = new DirectoryReference(TelemetryDir);
 					DirectoryReference.CreateDirectory(TelemetryDirRef);
 
-					string FileName = String.Format("{0}.{1}.json", Path.GetFileName(Assembly.GetEntryAssembly().Location), Process.Id, Process.StartTime.Ticks);
+					string FileName = String.Format("{0}.{1}.json", Path.GetFileName(Assembly.GetEntryAssembly()!.Location), Process.Id, Process.StartTime.Ticks);
 					File = FileReference.Combine(TelemetryDirRef, FileName);
 				}
 
diff --git a/Engine/Source/Programs/UnrealBuildTool/Preprocessor/SourceFile.cs b/Engine/Source/Programs/UnrealBuildTool/Preprocessor/SourceFile.cs
--- a/Engine/Source/Programs/UnrealBuildTool/Preprocessor/SourceFile.cs
+++ b/Engine/Source/Programs/UnrealBuildTool/Preprocessor/SourceFile.cs
@@ -109,7 +109,7 @@ namespace UnrealBuildTool
 		public SourceFile(BinaryArchiveReader Reader)
 		{
 			File = Reader.ReadFileItem();
-			Hash = Reader.ReadContentHash();
+			Hash = Reader.ReadContentHash()!;
 			LastWriteTimeUtc = Reader.ReadLong();
 			Markup = Reader.ReadArray(() => new SourceFileMarkup(Reader));
 			HeaderGuardMacro = Reader.ReadIdentifier();
diff --git a/Engine/Source/Programs/UnrealBuildTool/System/CommandLine.cs b/Engine/Source/Programs/UnrealBuildTool/System/CommandLine.cs
--- a/Engine/Source/Programs/UnrealBuildTool/System/CommandLine.cs
+++ b/Engine/Source/Programs/UnrealBuildTool/System/CommandLine.cs
@@ -86,7 +86,7 @@ namespace UnrealBuildTool
 					IEnumerable<CommandLineAttribute> Attributes = FieldInfo.GetCustomAttributes<CommandLineAttribute>();
 					foreach(CommandLineAttribute Attribute in Attributes)
 					{
-						string Prefix = Attribute.Prefix;
+						string? Prefix = Attribute.Prefix;
 						if(Prefix == null)
 						{
 							if(NonNullableType(FieldInfo.FieldType) == typeof(bool))
