Merging additional changes from Horde fork of EpicGames.Core.

#rnx

diff --git a/Engine/Source/Programs/Shared/EpicGames.Core/AsyncUtils.cs b/Engine/Source/Programs/Shared/EpicGames.Core/AsyncUtils.cs
new file mode 100644
--- /dev/null
+++ b/Engine/Source/Programs/Shared/EpicGames.Core/AsyncUtils.cs
@@ -0,0 +1,66 @@
+// Copyright Epic Games, Inc. All Rights Reserved.
+
+using System;
+using System.Collections.Generic;
+using System.Text;
+using System.Threading.Tasks;
+
+namespace EpicGames.Core
+{
+	/// <summary>
+	/// Utility functions for manipulating async tasks
+	/// </summary>
+	public static class AsyncUtils
+	{
+		/// <summary>
+		/// Removes all the complete tasks from a list, allowing each to throw exceptions as necessary
+		/// </summary>
+		/// <param name="Tasks">List of tasks to remove tasks from</param>
+		public static void RemoveCompleteTasks(this List<Task> Tasks)
+		{
+			int OutIdx = 0;
+			for (int Idx = 0; Idx < Tasks.Count; Idx++)
+			{
+				if (Tasks[Idx].IsCompleted)
+				{
+					Tasks[Idx].Wait();
+				}
+				else
+				{
+					if (Idx != OutIdx)
+					{
+						Tasks[OutIdx] = Tasks[Idx];
+					}
+					OutIdx++;
+				}
+			}
+			Tasks.RemoveRange(OutIdx, Tasks.Count - OutIdx);
+		}
+
+		/// <summary>
+		/// Removes all the complete tasks from a list, allowing each to throw exceptions as necessary
+		/// </summary>
+		/// <param name="Tasks">List of tasks to remove tasks from</param>
+		/// <returns>Return values from the completed tasks</returns>
+		public static List<T> RemoveCompleteTasks<T>(this List<Task<T>> Tasks)
+		{
+			List<T> Results = new List<T>();
+
+			int OutIdx = 0;
+			for (int Idx = 0; Idx < Tasks.Count; Idx++)
+			{
+				if (Tasks[Idx].IsCompleted)
+				{
+					Results.Add(Tasks[Idx].Result);
+				}
+				else if (Idx != OutIdx)
+				{
+					Tasks[OutIdx++] = Tasks[Idx];
+				}
+			}
+			Tasks.RemoveRange(OutIdx, Tasks.Count - OutIdx);
+
+			return Results;
+		}
+	}
+}
diff --git a/Engine/Source/Programs/Shared/EpicGames.Core/EnumerableExtensions.cs b/Engine/Source/Programs/Shared/EpicGames.Core/EnumerableExtensions.cs
new file mode 100644
--- /dev/null
+++ b/Engine/Source/Programs/Shared/EpicGames.Core/EnumerableExtensions.cs
@@ -0,0 +1,39 @@
+// Copyright Epic Games, Inc. All Rights Reserved.
+
+using System;
+using System.Collections.Generic;
+using System.Text;
+
+namespace EpicGames.Core
+{
+	/// <summary>
+	/// Extension methods for <see cref="IEnumerable{T}"/>
+	/// </summary>
+	public static class EnumerableExtensions
+	{
+		/// <summary>
+		/// Split the sequence into batches of at most the given size
+		/// </summary>
+		/// <typeparam name="TElement">The element type</typeparam>
+		/// <param name="Sequence">Sequence to split into batches</param>
+		/// <param name="BatchSize">Maximum size of each batch</param>
+		/// <returns>Sequence of batches</returns>
+		public static IEnumerable<IReadOnlyList<TElement>> Batch<TElement>(this IEnumerable<TElement> Sequence, int BatchSize)
+		{
+			List<TElement> Elements = new List<TElement>(BatchSize);
+			foreach (TElement Element in Sequence)
+			{
+				Elements.Add(Element);
+				if (Elements.Count == BatchSize)
+				{
+					yield return Elements;
+					Elements.Clear();
+				}
+			}
+			if (Elements.Count > 0)
+			{
+				yield return Elements;
+			}
+		}
+	}
+}
diff --git a/Engine/Source/Programs/Shared/EpicGames.Core/ManagedProcess.cs b/Engine/Source/Programs/Shared/EpicGames.Core/ManagedProcess.cs
--- a/Engine/Source/Programs/Shared/EpicGames.Core/ManagedProcess.cs
+++ b/Engine/Source/Programs/Shared/EpicGames.Core/ManagedProcess.cs
@@ -5,6 +5,7 @@ using System;
 using System.Collections.Generic;
 using System.ComponentModel;
 using System.Diagnostics;
+using System.Diagnostics.CodeAnalysis;
 using System.IO;
 using System.Linq;
 using System.Runtime.InteropServices;
@@ -12,8 +13,6 @@ using System.Text;
 using System.Threading;
 using System.Threading.Tasks;
 
-#nullable disable
-
 namespace EpicGames.Core
 {
 	/// <summary>
@@ -72,7 +71,7 @@ namespace EpicGames.Core
 		/// Handle to the native job object that this process is added to. This handle is closed by the Dispose() method (and will automatically be closed by the OS on process exit),
 		/// resulting in the child process being killed.
 		/// </summary>
-		internal SafeFileHandle JobHandle
+		internal SafeFileHandle? JobHandle
 		{
 			get;
 			private set;
@@ -160,9 +159,9 @@ namespace EpicGames.Core
 		class STARTUPINFO
 		{
 			public int cb;
-			public string lpReserved;
-			public string lpDesktop;
-			public string lpTitle;
+			public string? lpReserved;
+			public string? lpDesktop;
+			public string? lpTitle;
 			public uint dwX;
 			public uint dwY;
 			public uint dwXSize;
@@ -174,13 +173,13 @@ namespace EpicGames.Core
 			public short wShowWindow;
 			public short cbReserved2;
 			public IntPtr lpReserved2;
-			public SafeHandle hStdInput;
-			public SafeHandle hStdOutput;
-			public SafeHandle hStdError;
+			public SafeHandle? hStdInput;
+			public SafeHandle? hStdOutput;
+			public SafeHandle? hStdError;
 		}
 
 		[DllImport("kernel32.dll", SetLastError=true)]
-		static extern int AssignProcessToJobObject(SafeFileHandle hJob, IntPtr hProcess);
+		static extern int AssignProcessToJobObject(SafeFileHandle? hJob, IntPtr hProcess);
 
 		[DllImport("kernel32.dll", SetLastError=true)]
 		static extern int CloseHandle(IntPtr hObject);
@@ -209,7 +208,7 @@ namespace EpicGames.Core
 		}
 
         [DllImport("kernel32.dll", SetLastError=true)]
-        static extern int CreateProcess(/*[MarshalAs(UnmanagedType.LPTStr)]*/ string lpApplicationName, StringBuilder lpCommandLine, IntPtr lpProcessAttributes, IntPtr lpThreadAttributes, bool bInheritHandles, ProcessCreationFlags dwCreationFlags, IntPtr lpEnvironment, /*[MarshalAs(UnmanagedType.LPTStr)]*/ string lpCurrentDirectory, STARTUPINFO lpStartupInfo, PROCESS_INFORMATION lpProcessInformation);
+        static extern int CreateProcess(/*[MarshalAs(UnmanagedType.LPTStr)]*/ string? lpApplicationName, StringBuilder lpCommandLine, IntPtr lpProcessAttributes, IntPtr lpThreadAttributes, bool bInheritHandles, ProcessCreationFlags dwCreationFlags, IntPtr lpEnvironment, /*[MarshalAs(UnmanagedType.LPTStr)]*/ string? lpCurrentDirectory, STARTUPINFO lpStartupInfo, PROCESS_INFORMATION lpProcessInformation);
 
 		[DllImport("kernel32.dll", SetLastError=true)]
 		static extern int ResumeThread(IntPtr hThread);
@@ -249,32 +248,32 @@ namespace EpicGames.Core
 		/// <summary>
 		/// Handle for the child process.
 		/// </summary>
-		SafeFileHandle ProcessHandle;
+		SafeFileHandle? ProcessHandle;
 
 		/// <summary>
 		/// The write end of the child process' stdin pipe.
 		/// </summary>
-		SafeFileHandle StdInWrite;
+		SafeFileHandle? StdInWrite;
 
 		/// <summary>
 		/// The read end of the child process' stdout pipe.
 		/// </summary>
-		SafeFileHandle StdOutRead;
+		SafeFileHandle? StdOutRead;
 
 		/// <summary>
 		/// Output stream for the child process.
 		/// </summary>
-		FileStream InnerStream;
+		FileStream? InnerStream;
 
 		/// <summary>
 		/// Reader for the process' output stream.
 		/// </summary>
-		StreamReader ReadStream;
+		StreamReader? ReadStream;
 
 		/// <summary>
 		/// Standard process implementation for non-Windows platforms.
 		/// </summary>
-		Process FrameworkProcess;
+		Process? FrameworkProcess;
 
 		/// <summary>
 		/// Static lock object. This is used to synchronize the creation of child processes - in particular, the inheritance of stdout/stderr write pipes. If processes
@@ -292,7 +291,7 @@ namespace EpicGames.Core
 		/// <param name="Environment">Environment variables for the new process. May be null, in which case the current process' environment is inherited</param>
 		/// <param name="Input">Text to be passed via stdin to the new process. May be null.</param>
 		/// <param name="Priority">Priority for the child process</param>
-		public ManagedProcess(ManagedProcessGroup Group, string FileName, string CommandLine, string WorkingDirectory, IReadOnlyDictionary<string, string> Environment, byte[] Input, ProcessPriorityClass Priority)
+		public ManagedProcess(ManagedProcessGroup? Group, string FileName, string CommandLine, string? WorkingDirectory, IReadOnlyDictionary<string, string>? Environment, byte[]? Input, ProcessPriorityClass Priority)
 		{
 			// Create the child process
 			// NOTE: Child process must be created in a separate method to avoid stomping exception callstacks (https://stackoverflow.com/a/2494150)
@@ -324,7 +323,7 @@ namespace EpicGames.Core
 		/// <param name="Environment">Environment variables for the new process. May be null, in which case the current process' environment is inherited</param>
 		/// <param name="Input">Text to be passed via stdin to the new process. May be null.</param>
 		/// <param name="Priority">Priority for the child process</param>
-		private void CreateManagedProcessWin32(ManagedProcessGroup Group, string FileName, string CommandLine, string WorkingDirectory, IReadOnlyDictionary<string, string> Environment, byte[] Input, ProcessPriorityClass Priority)
+		private void CreateManagedProcessWin32(ManagedProcessGroup? Group, string FileName, string CommandLine, string? WorkingDirectory, IReadOnlyDictionary<string, string>? Environment, byte[]? Input, ProcessPriorityClass Priority)
 		{
 			IntPtr EnvironmentBlock = IntPtr.Zero;
 			try
@@ -380,9 +379,9 @@ namespace EpicGames.Core
 					// inherited handles are closed.
 					lock (LockObject)
 					{
-						SafeFileHandle StdInRead = null;
-						SafeFileHandle StdOutWrite = null;
-						SafeWaitHandle StdErrWrite = null;
+						SafeFileHandle? StdInRead = null;
+						SafeFileHandle? StdOutWrite = null;
+						SafeWaitHandle? StdErrWrite = null;
 						try
 						{
 							// Create stdin and stdout pipes for the child process. We'll close the handles for the child process' ends after it's been created.
@@ -390,11 +389,11 @@ namespace EpicGames.Core
 							SecurityAttributes.nLength = Marshal.SizeOf(SecurityAttributes);
 							SecurityAttributes.bInheritHandle = 1;
 
-							if (CreatePipe(out StdInRead, out StdInWrite, SecurityAttributes, 0) == 0 || SetHandleInformation(StdInWrite, HANDLE_FLAG_INHERIT, 0) == 0)
+							if (CreatePipe(out StdInRead, out StdInWrite, SecurityAttributes, 4 * 1024) == 0 || SetHandleInformation(StdInWrite, HANDLE_FLAG_INHERIT, 0) == 0)
 							{
 								throw new Win32ExceptionWithCode("Unable to create stdin pipe");
 							}
-							if (CreatePipe(out StdOutRead, out StdOutWrite, SecurityAttributes, 0) == 0 || SetHandleInformation(StdOutRead, HANDLE_FLAG_INHERIT, 0) == 0)
+							if (CreatePipe(out StdOutRead, out StdOutWrite, SecurityAttributes, 1024 * 1024) == 0 || SetHandleInformation(StdOutRead, HANDLE_FLAG_INHERIT, 0) == 0)
 							{
 								throw new Win32ExceptionWithCode("Unable to create stdout pipe");
 							}
@@ -425,7 +424,6 @@ namespace EpicGames.Core
 									throw new Win32ExceptionWithCode("Unable to create process");
 								}
 
-								Log.TraceWarning("Unable to create process {0} (access denied); waiting {1}ms to retry...", FileName, RetryDelay[AttemptIdx]);
 								Thread.Sleep(RetryDelay[AttemptIdx]);
 							}
 						}
@@ -520,7 +518,7 @@ namespace EpicGames.Core
 		/// <param name="Environment">Environment variables for the new process. May be null, in which case the current process' environment is inherited</param>
 		/// <param name="Input">Text to be passed via stdin to the new process. May be null.</param>
 		/// <param name="Priority">Priority for the child process</param>
-		private void CreateManagedProcessPortable(ManagedProcessGroup Group, string FileName, string CommandLine, string WorkingDirectory, IReadOnlyDictionary<string, string> Environment, byte[] Input, ProcessPriorityClass Priority)
+		private void CreateManagedProcessPortable(ManagedProcessGroup? Group, string FileName, string CommandLine, string? WorkingDirectory, IReadOnlyDictionary<string, string>? Environment, byte[]? Input, ProcessPriorityClass Priority)
 		{
 			// Fallback for Mono platforms
 			FrameworkProcess = new Process();
@@ -611,7 +609,7 @@ namespace EpicGames.Core
 			Task<int> ReadTask;
 			if(FrameworkProcess == null)
 			{
-				ReadTask = InnerStream.ReadAsync(Buffer, Offset, Count);
+				ReadTask = InnerStream!.ReadAsync(Buffer, Offset, Count);
 			}
 			else
 			{
@@ -624,6 +622,43 @@ namespace EpicGames.Core
 			return ReadTask.Result;
 		}
 
+		/// <summary>
+		/// Reads data from the process output
+		/// </summary>
+		/// <param name="Buffer">The buffer to receive the data</param>
+		/// <param name="Offset">Offset within the buffer to write to</param>
+		/// <param name="Count">Maximum number of bytes to read</param>
+		/// <returns>Number of bytes read</returns>
+		public async Task<int> ReadAsync(byte[] Buffer, int Offset, int Count, CancellationToken CancellationToken)
+		{
+			if (FrameworkProcess == null)
+			{
+				return await InnerStream!.ReadAsync(Buffer, Offset, Count, CancellationToken);
+			}
+			else
+			{
+				return await FrameworkProcess.StandardOutput.BaseStream.ReadAsync(Buffer, Offset, Count, CancellationToken);
+			}
+		}
+
+		/// <summary>
+		/// Copy the process output to the given stream
+		/// </summary>
+		/// <param name="OutputStream">The output stream</param>
+		/// <param name="CancellationToken">Cancellation token</param>
+		/// <returns></returns>
+		public Task CopyTo(Stream OutputStream, CancellationToken CancellationToken)
+		{
+			if(FrameworkProcess == null)
+			{
+				return InnerStream!.CopyToAsync(OutputStream, CancellationToken);
+			}
+			else
+			{
+				return FrameworkProcess.StandardOutput.BaseStream.CopyToAsync(OutputStream, CancellationToken);
+			}
+		}
+
 		/// <summary>
 		/// Read all the output from the process. Does not return until the process terminates.
 		/// </summary>
@@ -680,21 +715,37 @@ namespace EpicGames.Core
 			return OutputLines;
 		}
 
+		/// <summary>
+		/// Reads a single line asynchronously
+		/// </summary>
+		/// <returns>New line</returns>
+		public async Task<string?> ReadLineAsync()
+		{
+			if (FrameworkProcess == null)
+			{
+				return await ReadStream!.ReadLineAsync();
+			}
+			else
+			{
+				return await FrameworkProcess.StandardOutput.ReadLineAsync();
+			}
+		}
+
 		/// <summary>
 		/// Block until the process outputs a line of text, or terminates.
 		/// </summary>
 		/// <param name="Line">Variable to receive the output line</param>
 		/// <param name="Token">Cancellation token which can be used to abort the read operation.</param>
 		/// <returns>True if a line was read, false if the process terminated without writing another line, or the cancellation token was signaled.</returns>
-		public bool TryReadLine(out string Line, CancellationToken? Token = null)
+		public bool TryReadLine([NotNullWhen(true)] out string? Line, CancellationToken? Token = null)
 		{
 			try
 			{
 				// Busy wait for the ReadLine call to finish, so we can get interrupted by thread abort exceptions
-				Task<string> ReadLineTask;
+				Task<string?> ReadLineTask;
 				if(FrameworkProcess == null)
 				{
-					ReadLineTask = ReadStream.ReadLineAsync();
+					ReadLineTask = ReadStream!.ReadLineAsync();
 				}
 				else
 				{
@@ -720,6 +771,22 @@ namespace EpicGames.Core
 			}
 		}
 
+		/// <summary>
+		/// Waits for the process to exit
+		/// </summary>
+		/// <returns>Exit code of the process</returns>
+		public void WaitForExit()
+		{
+			if (FrameworkProcess == null)
+			{
+				WaitForSingleObject(ProcessHandle!, INFINITE);
+			}
+			else
+			{
+				FrameworkProcess.WaitForExit();
+			}
+		}
+
 		/// <summary>
 		/// The exit code of the process. Throws an exception if the process has not terminated.
 		/// </summary>
@@ -730,7 +797,7 @@ namespace EpicGames.Core
 				if(FrameworkProcess == null)
 				{
 					int Value;
-					if(GetExitCodeProcess(ProcessHandle, out Value) == 0)
+					if(GetExitCodeProcess(ProcessHandle!, out Value) == 0)
 					{
 						throw new Win32Exception();
 					}
diff --git a/Engine/Source/Programs/Shared/EpicGames.Core/ThreadPoolWorkQueue.cs b/Engine/Source/Programs/Shared/EpicGames.Core/ThreadPoolWorkQueue.cs
--- a/Engine/Source/Programs/Shared/EpicGames.Core/ThreadPoolWorkQueue.cs
+++ b/Engine/Source/Programs/Shared/EpicGames.Core/ThreadPoolWorkQueue.cs
@@ -101,7 +101,6 @@ namespace EpicGames.Core
 				{
 					Exceptions.Add(Ex);
 				}
-				throw;
 			}
 			finally
 			{
diff --git a/Engine/Source/Programs/UnrealBuildTool/Platform/Windows/VCPreludeCache.cs b/Engine/Source/Programs/UnrealBuildTool/Platform/Windows/VCPreludeCache.cs
--- a/Engine/Source/Programs/UnrealBuildTool/Platform/Windows/VCPreludeCache.cs
+++ b/Engine/Source/Programs/UnrealBuildTool/Platform/Windows/VCPreludeCache.cs
@@ -156,7 +156,7 @@ namespace UnrealBuildTool
 						Writer.WriteLine("// Arguments: {0}", CommandLine);
 						Writer.WriteLine();
 
-						string Line;
+						string? Line;
 						while (Process.TryReadLine(out Line))
 						{
 							Line = Line.Trim();
