Change UAT/UBT log system to be backed by a Microsoft.Extensions.Logging.ILogger instance, to allow us to start transitioning towards structured logging for all C# tools.

Current ILogger implementation contains the same functionality as before, with message formatting options passed through to ILogger.Log via a LogMessage instance.

#rnx

diff --git a/Engine/Source/Programs/Shared/EpicGames.Core/EpicGames.Core.csproj b/Engine/Source/Programs/Shared/EpicGames.Core/EpicGames.Core.csproj
--- a/Engine/Source/Programs/Shared/EpicGames.Core/EpicGames.Core.csproj
+++ b/Engine/Source/Programs/Shared/EpicGames.Core/EpicGames.Core.csproj
@@ -26,6 +26,7 @@
     </Compile>
   </ItemGroup>
   <ItemGroup>
+    <PackageReference Include="Microsoft.Extensions.Logging.Abstractions" Version="5.0.0" />
     <PackageReference Include="System.Memory" Version="4.5.4" />
   </ItemGroup>
 </Project>
\ No newline at end of file
diff --git a/Engine/Source/Programs/Shared/EpicGames.Core/Log.cs b/Engine/Source/Programs/Shared/EpicGames.Core/Log.cs
--- a/Engine/Source/Programs/Shared/EpicGames.Core/Log.cs
+++ b/Engine/Source/Programs/Shared/EpicGames.Core/Log.cs
@@ -1,5 +1,6 @@
 // Copyright Epic Games, Inc. All Rights Reserved.
 
+using Microsoft.Extensions.Logging;
 using System;
 using System.Collections.Generic;
 using System.Diagnostics;
@@ -21,37 +22,37 @@
 		/// <summary>
 		/// The log event is a fatal error
 		/// </summary>
-		Fatal = 1,
+		Fatal = LogLevel.Critical,
 
 		/// <summary>
 		/// The log event is an error
 		/// </summary>
-		Error = 2,
+		Error = LogLevel.Error,
 
 		/// <summary>
 		/// The log event is a warning
 		/// </summary>
-		Warning = 4,
+		Warning = LogLevel.Warning,
 
 		/// <summary>
 		/// Output the log event to the console
 		/// </summary>
-		Console = 8,
+		Console = LogLevel.Information,
 
 		/// <summary>
 		/// Output the event to the on-disk log
 		/// </summary>
-		Log = 16,
+		Log = LogLevel.Debug,
 
 		/// <summary>
 		/// The log event should only be displayed if verbose logging is enabled
 		/// </summary>
-		Verbose = 32,
+		Verbose = LogLevel.Trace,
 
 		/// <summary>
 		/// The log event should only be displayed if very verbose logging is enabled
 		/// </summary>
-		VeryVerbose = 64
+		VeryVerbose = LogLevel.Trace
 	}
 
 	/// <summary>
@@ -77,5 +78,45 @@
 	}
 
 	/// <summary>
+	/// Contents of a legacy log message
+	/// </summary>
+	public struct LogMessage
+	{
+		/// <summary>
+		/// Number of stack frames to skip
+		/// </summary>
+		public int StackFramesToSkip;
+
+		/// <summary>
+		/// Whether to only output this message once
+		/// </summary>
+		public bool bWriteOnce;
+
+		/// <summary>
+		/// Options for formatting the message
+		/// </summary>
+		public LogFormatOptions FormatOptions;
+
+		/// <summary>
+		/// Format string
+		/// </summary>
+		public string Format;
+
+		/// <summary>
+		/// Argument list
+		/// </summary>
+		public object[] Args;
+
+		/// <summary>
+		/// Convert the message to a single string
+		/// </summary>
+		/// <returns></returns>
+		public override string ToString()
+		{
+			return String.Format(Format, Args);
+		}
+	}
+
+	/// <summary>
 	/// UAT/UBT Custom log system.
 	/// 
@@ -91,38 +132,22 @@
 	static public class Log
 	{
 		/// <summary>
-		/// Temporary status message displayed on the console.
+		/// Singleton instance of the default logger
 		/// </summary>
-		[DebuggerDisplay("{HeadingText}")]
-		class StatusMessage
-		{
-			/// <summary>
-			/// The heading for this status message.
-			/// </summary>
-			public string HeadingText;
-
-			/// <summary>
-			/// The current status text.
-			/// </summary>
-			public string CurrentText;
-
-			/// <summary>
-			/// Whether the heading has been written to the console. Before the first time that lines are output to the log in the midst of a status scope, the heading will be written on a line of its own first.
-			/// </summary>
-			public bool bHasFlushedHeadingText;
-		}
+		private static DefaultLogger DefaultLogger { get; } = new DefaultLogger();
 
 		/// <summary>
-		/// Object used for synchronization
+		/// The public logging interface
 		/// </summary>
-		private static object SyncObject = new object();
+		public static ILogger Logger { get; set; } = DefaultLogger;
 
 		/// <summary>
 		/// When true, verbose logging is enabled.
 		/// </summary>
 		public static LogEventType OutputLevel
 		{
-			get; set;
+			get => DefaultLogger.OutputLevel;
+			set => DefaultLogger.OutputLevel = value;
 		}
 
 		/// <summary>
@@ -130,7 +155,8 @@
 		/// </summary>
 		public static bool IncludeTimestamps
 		{
-			get; set;
+			get => DefaultLogger.IncludeTimestamps;
+			set => DefaultLogger.IncludeTimestamps = value;
 		}
 
 		/// <summary>
@@ -138,15 +164,17 @@
 		/// </summary>
 		public static bool IncludeSeverityPrefix
 		{
-			get; set;
+			get => DefaultLogger.IncludeSeverityPrefix;
+			set => DefaultLogger.IncludeSeverityPrefix = value;
 		}
-		
+
 		/// <summary>
 		/// When true, warnings and errors will have a prefix suitable for display by MSBuild (avoiding error messages showing as (EXEC : Error : ")
 		/// </summary>
 		public static bool IncludeProgramNameWithSeverityPrefix
 		{
-			get; set;
+			get => DefaultLogger.IncludeProgramNameWithSeverityPrefix;
+			set => DefaultLogger.IncludeProgramNameWithSeverityPrefix = value;
 		}
 
 		/// <summary>
@@ -154,86 +182,40 @@
 		/// </summary>
 		public static bool IncludeCallingMethod
 		{
-			get; set;
+			get => DefaultLogger.IncludeCallingMethod;
+			set => DefaultLogger.IncludeCallingMethod = value;
 		}
-		
+
 		/// <summary>
 		/// When true, console output will have the calling mehod prepended to the output as MethodName:
 		/// </summary>
 		public static bool IncludeCallingMethodForConsole
 		{
-			get; set;
+			get => DefaultLogger.IncludeCallingMethodForConsole;
+			set => DefaultLogger.IncludeCallingMethodForConsole = value;
 		}
-		
+
 		/// <summary>
 		/// When true, will detect warnings and errors and set the console output color to yellow and red.
 		/// </summary>
 		public static bool ColorConsoleOutput
 		{
-			get; set;
+			get => DefaultLogger.ColorConsoleOutput;
+			set => DefaultLogger.ColorConsoleOutput = value;
 		}
 
 		/// <summary>
 		/// Path to the log file being written to. May be null.
 		/// </summary>
-		public static FileReference OutputFile
-		{
-			get; private set;
-		}
+		public static FileReference OutputFile => DefaultLogger.OutputFile;
 
 		/// <summary>
-		/// Whether console output is redirected. This prevents writing status updates that rely on moving the cursor.
-		/// </summary>
-		private static bool AllowStatusUpdates
-		{
-			get { return !Console.IsOutputRedirected; }
-		}
-
-		/// <summary>
-		/// When configured, this tracks time since initialization to prepend a timestamp to each log.
-		/// </summary>
-		private static Stopwatch Timer = Stopwatch.StartNew();
-
-		/// <summary>
-		/// A collection of strings that have been already written once
-		/// </summary>
-		private static HashSet<string> WriteOnceSet = new HashSet<string>();
-
-		/// <summary>
-		/// Stack of status scope information.
-		/// </summary>
-		private static Stack<StatusMessage> StatusMessageStack = new Stack<StatusMessage>();
-
-		/// <summary>
-		/// The currently visible status text
-		/// </summary>
-		private static string StatusText = "";
-
-		/// <summary>
-		/// Last time a status message was pushed to the stack
-		/// </summary>
-		private static Stopwatch StatusTimer = new Stopwatch();
-
-		/// <summary>
 		/// Indent added to every output line
 		/// </summary>
 		public static string Indent
 		{
-			get; set;
+			get => DefaultLogger.Indent;
+			set => DefaultLogger.Indent = value; 
 		}
 
 		/// <summary>
-		/// Static initializer
-		/// </summary>
-		static Log()
-		{
-			Indent = "";
-			OutputLevel = LogEventType.Log;
-			IncludeSeverityPrefix = true;
-			IncludeProgramNameWithSeverityPrefix = false;
-			IncludeCallingMethod = true;
-			IncludeCallingMethodForConsole = false;
-			ColorConsoleOutput = true;
-		}
-
-		/// <summary>
@@ -243,21 +225,7 @@
 		/// <returns>The created trace listener</returns>
 		public static TextWriterTraceListener AddFileWriter(string Name, FileReference OutputFile)
 		{
-			try
-			{
-				Log.OutputFile = OutputFile;
-				DirectoryReference.CreateDirectory(OutputFile.Directory);
-				TextWriterTraceListener LogTraceListener = new TextWriterTraceListener(new StreamWriter(OutputFile.FullName), Name);
-				lock (SyncObject)
-				{
-					Trace.Listeners.Add(LogTraceListener);
-				}
-				return LogTraceListener;
-			}
-			catch (Exception Ex)
-			{
-				throw new Exception(String.Format("Unable to open log file for writing ({0})", OutputFile), Ex);
-			}
+			return DefaultLogger.AddFileWriter(Name, OutputFile);
 		}
 
 		/// <summary>
@@ -266,13 +234,7 @@
 		/// <param name="TraceListener">The <see cref="TraceListener"/> to add.</param>
 		public static void AddTraceListener(TraceListener TraceListener)
 		{
-			lock (SyncObject)
-			{
-				if (!Trace.Listeners.Contains(TraceListener))
-				{
-					Trace.Listeners.Add(TraceListener);
-				}
-			}
+			DefaultLogger.AddTraceListener(TraceListener);
 		}
 
 		/// <summary>
@@ -281,13 +243,7 @@
 		/// <param name="TraceListener">The <see cref="TraceListener"/> to remove.</param>
 		public static void RemoveTraceListener(TraceListener TraceListener)
 		{
-			lock (SyncObject)
-			{
-				if (Trace.Listeners.Contains(TraceListener))
-				{
-					Trace.Listeners.Remove(TraceListener);
-				}
-			}
+			DefaultLogger.RemoveTraceListener(TraceListener);
 		}
 
 		/// <summary>
@@ -296,104 +252,10 @@
 		/// <returns>True if a TextWriterTraceListener has been added</returns>
 		public static bool HasFileWriter()
 		{
-			foreach(TraceListener Listener in Trace.Listeners)
-			{
-				if(Listener is TextWriterTraceListener)
-				{
-					return true;
-				}
-			}
-			return false;
+			return DefaultLogger.HasFileWriter();
 		}
 
 		/// <summary>
-		/// Gets the name of the Method N levels deep in the stack frame. Used to trap what method actually made the logging call.
-		/// Only used when bLogSources is true.
-		/// </summary>
-		/// <param name="StackFramesToSkip"></param>
-		/// <returns>ClassName.MethodName</returns>
-		[MethodImplAttribute(MethodImplOptions.NoInlining)]
-		private static string GetSource(int StackFramesToSkip)
-		{
-			StackFrame Frame = new StackFrame(2 + StackFramesToSkip);
-			System.Reflection.MethodBase Method = Frame.GetMethod();
-			return String.Format("{0}.{1}", Method.DeclaringType.Name, Method.Name);
-		}
-
-		/// <summary>
-		/// Converts a LogEventType into a log prefix. Only used when bLogSeverity is true.
-		/// </summary>
-		/// <param name="Severity"></param>
-		/// <returns></returns>
-		private static string GetSeverityPrefix(LogEventType Severity)
-		{
-			switch (Severity)
-			{
-				case LogEventType.Fatal:
-					return "FATAL ERROR: ";
-				case LogEventType.Error:
-					return "ERROR: ";
-				case LogEventType.Warning:
-					return "WARNING: ";
-				case LogEventType.Console:
-					return "";
-				case LogEventType.Verbose:
-					return "VERBOSE: ";
-				case LogEventType.VeryVerbose:
-					return "VVERBOSE: ";
-				default:
-					return "";
-			}
-		}
-
-		/// <summary>
-		/// Formats message for logging. Enforces the configured options.
-		/// </summary>
-		/// <param name="StackFramesToSkip">Number of frames to skip to get to the originator of the log request.</param>
-		/// <param name="Verbosity">Message verbosity level</param>
-		/// <param name="Options">Options for formatting this string</param>
-		/// <param name="bForConsole">Whether the message is intended for console output</param>
-		/// <param name="Format">Message text format string</param>
-		/// <param name="Args">Message text parameters</param>
-		/// <returns>Formatted message</returns>
-		[MethodImplAttribute(MethodImplOptions.NoInlining)]
-		private static List<string> FormatMessage(int StackFramesToSkip, LogEventType Verbosity, LogFormatOptions Options, bool bForConsole, string Format, params object[] Args)
-		{
-			string TimePrefix = IncludeTimestamps? String.Format("[{0:hh\\:mm\\:ss\\.fff}] ", Timer.Elapsed) : "";
-			string SourcePrefix = (bForConsole ? IncludeCallingMethodForConsole : IncludeCallingMethod) ? string.Format("{0}: ", GetSource(StackFramesToSkip)) : "";
-			string SeverityPrefix = (IncludeSeverityPrefix && ((Options & LogFormatOptions.NoSeverityPrefix) == 0)) ? GetSeverityPrefix(Verbosity) : "";
-
-			// Include the executable name when running inside MSBuild. If unspecified, MSBuild re-formats them with an "EXEC :" prefix.
-			if(SeverityPrefix.Length > 0 && IncludeProgramNameWithSeverityPrefix)
-			{
-				SeverityPrefix = String.Format("{0}: {1}", Path.GetFileNameWithoutExtension(Assembly.GetEntryAssembly().Location), SeverityPrefix);
-			}
-
-			// If there are no extra args, don't try to format the string, in case it has any format control characters in it (our LOCTEXT strings tend to).
-			string[] Lines = ((Args.Length > 0) ? String.Format(Format, Args) : Format).TrimEnd(' ', '\t', '\r', '\n').Split('\n');
-
-			List<string> FormattedLines = new List<string>();
-			FormattedLines.Add(String.Format("{0}{1}{2}{3}{4}", TimePrefix, SourcePrefix, Indent, SeverityPrefix, Lines[0].TrimEnd('\r')));
-
-			if (Lines.Length > 1)
-			{
-				int PaddingLength = 0;
-				while(PaddingLength < Lines[0].Length && Char.IsWhiteSpace(Lines[0][PaddingLength]))
-				{
-					PaddingLength++;
-				}
-
-				string Padding = new string(' ', SeverityPrefix.Length) + Lines[0].Substring(0, PaddingLength);
-				for (int Idx = 1; Idx < Lines.Length; Idx++)
-				{
-					FormattedLines.Add(String.Format("{0}{1}{2}{3}{4}", TimePrefix, SourcePrefix, Indent, Padding, Lines[Idx].TrimEnd('\r')));
-				}
-			}
-
-			return FormattedLines;
-		}
-
-		/// <summary>
 		/// Writes a formatted message to the console. All other functions should boil down to calling this method.
 		/// </summary>
 		/// <param name="StackFramesToSkip">Number of frames to skip to get to the originator of the log request.</param>
@@ -405,82 +267,13 @@
 		[MethodImplAttribute(MethodImplOptions.NoInlining)]
 		private static void WriteLinePrivate(int StackFramesToSkip, bool bWriteOnce, LogEventType Verbosity, LogFormatOptions FormatOptions, string Format, params object[] Args)
 		{
-			// if we want this message only written one time, check if it was already written out
-			if (bWriteOnce)
-			{
-				string Formatted = string.Format(Format, Args);
-				if (WriteOnceSet.Contains(Formatted))
-				{
-					return;
-				}
-
-				WriteOnceSet.Add(Formatted);
-			}
-
-			if (Verbosity <= OutputLevel)
-			{
-				lock (SyncObject)
-				{
-					// Output to all the other trace listeners
-					List<string> Lines = FormatMessage(StackFramesToSkip + 1, Verbosity, FormatOptions, false, Format, Args);
-					foreach (TraceListener Listener in Trace.Listeners)
-					{
-						foreach(string Line in Lines)
-						{
-							Listener.WriteLine(Line);
-						}
-						Listener.Flush();
-					}
-
-					// Handle the console output separately; we format things differently
-					if ((Verbosity != LogEventType.Log || OutputLevel >= LogEventType.Verbose) && (FormatOptions & LogFormatOptions.NoConsoleOutput) == 0)
-					{
-						FlushStatusHeading();
-
-						bool bResetConsoleColor = false;
-						if (ColorConsoleOutput)
-						{
-							if(Verbosity == LogEventType.Warning)
-							{
-								Console.ForegroundColor = ConsoleColor.Yellow;
-								bResetConsoleColor = true;
-							}
-							if(Verbosity <= LogEventType.Error)
-							{
-								Console.ForegroundColor = ConsoleColor.Red;
-								bResetConsoleColor = true;
-							}
-						}
-						try
-						{
-							List<string> ConsoleLines = FormatMessage(StackFramesToSkip + 1, Verbosity, FormatOptions, true, Format, Args);
-							foreach (string ConsoleLine in ConsoleLines)
-							{
-								Console.WriteLine(ConsoleLine);
-							}
-						}
-						catch (IOException)
-						{
-							// Potential file access/sharing issue on std out
-							// This can occur on some versions of mono (e.g. macOS 6.12.0) if writing to a full pipe
-							// during IPC when the reader isn't consuming it quick enough
-						}
-						finally
-						{
-							// make sure we always put the console color back.
-							if(bResetConsoleColor)
-							{
-								Console.ResetColor();
-							}
-						}
-
-						if(StatusMessageStack.Count > 0 && AllowStatusUpdates)
-						{
-							SetStatusText(StatusMessageStack.Peek().CurrentText);
-						}
-					}
-				}
-			}
+			LogMessage Message = new LogMessage();
+			Message.StackFramesToSkip = StackFramesToSkip;
+			Message.bWriteOnce = bWriteOnce;
+			Message.FormatOptions = FormatOptions;
+			Message.Format = Format;
+			Message.Args = Args;
+			Logger.Log((LogLevel)Verbosity, new EventId(), Message, null, (State, Ex) => State.ToString());
 		}
 
 		/// <summary>
@@ -804,27 +597,5 @@
 		}
 
 		/// <summary>
-		/// Flushes the current status text before writing out a new log line or status message
-		/// </summary>
-		static void FlushStatusHeading()
-		{
-			if(StatusMessageStack.Count > 0)
-			{
-				StatusMessage CurrentStatus = StatusMessageStack.Peek();
-				if(CurrentStatus.HeadingText.Length > 0 && !CurrentStatus.bHasFlushedHeadingText && AllowStatusUpdates)
-				{
-					SetStatusText(CurrentStatus.HeadingText);
-					Console.WriteLine();
-					StatusText = "";
-					CurrentStatus.bHasFlushedHeadingText = true;
-				}
-				else
-				{
-					SetStatusText("");
-				}
-			}
-		}
-
-		/// <summary>
 		/// Enter a scope with the given status message. The message will be written to the console without a newline, allowing it to be updated through subsequent calls to UpdateStatus().
 		/// The message will be written to the log immediately. If another line is written while in a status scope, the initial status message is flushed to the console first.
@@ -834,23 +605,7 @@
 		[MethodImplAttribute(MethodImplOptions.NoInlining)]
 		public static void PushStatus(string Message)
 		{
-			lock(SyncObject)
-			{
-				FlushStatusHeading();
-
-				StatusMessage NewStatusMessage = new StatusMessage();
-				NewStatusMessage.HeadingText = Message;
-				NewStatusMessage.CurrentText = Message;
-				StatusMessageStack.Push(NewStatusMessage);
-
-				StatusTimer.Restart();
-
-				if(Message.Length > 0)
-				{
-					WriteLine(LogEventType.Log, LogFormatOptions.NoConsoleOutput, "{0}", Message);
-					SetStatusText(Message);
-				}
-			}
+			DefaultLogger.PushStatus(Message);
 		}
 
 		/// <summary>
@@ -861,17 +616,7 @@
 		[MethodImplAttribute(MethodImplOptions.NoInlining)]
 		public static void UpdateStatus(string Message)
 		{
-			lock(SyncObject)
-			{
-				StatusMessage CurrentStatusMessage = StatusMessageStack.Peek();
-				CurrentStatusMessage.CurrentText = Message;
-
-				if(AllowStatusUpdates || StatusTimer.Elapsed.TotalSeconds > 10.0)
-				{
-					SetStatusText(Message);
-					StatusTimer.Restart();
-				}
-			}
+			DefaultLogger.UpdateStatus(Message);
 		}
 
 		/// <summary>
@@ -882,12 +627,514 @@
 		[MethodImplAttribute(MethodImplOptions.NoInlining)]
 		public static void PopStatus()
 		{
-			lock(SyncObject)
+			DefaultLogger.PopStatus();
+		}
+	}
+
+	/// <summary>
+	/// Default log output device
+	/// </summary>
+	class DefaultLogger : ILogger
+	{
+		/// <summary>
+		/// Temporary status message displayed on the console.
+		/// </summary>
+		[DebuggerDisplay("{HeadingText}")]
+		class StatusMessage
+		{
+			/// <summary>
+			/// The heading for this status message.
+			/// </summary>
+			public string HeadingText;
+
+			/// <summary>
+			/// The current status text.
+			/// </summary>
+			public string CurrentText;
+
+			/// <summary>
+			/// Whether the heading has been written to the console. Before the first time that lines are output to the log in the midst of a status scope, the heading will be written on a line of its own first.
+			/// </summary>
+			public bool bHasFlushedHeadingText;
+		}
+
+		/// <summary>
+		/// Object used for synchronization
+		/// </summary>
+		private object SyncObject = new object();
+
+		/// <summary>
+		/// When true, verbose logging is enabled.
+		/// </summary>
+		public LogEventType OutputLevel
+		{
+			get; set;
+		}
+
+		/// <summary>
+		/// Whether to include timestamps on each line of log output
+		/// </summary>
+		public bool IncludeTimestamps
+		{
+			get; set;
+		}
+
+		/// <summary>
+		/// When true, warnings and errors will have a WARNING: or ERROR: prexifx, respectively.
+		/// </summary>
+		public bool IncludeSeverityPrefix
+		{
+			get; set;
+		}
+
+		/// <summary>
+		/// When true, warnings and errors will have a prefix suitable for display by MSBuild (avoiding error messages showing as (EXEC : Error : ")
+		/// </summary>
+		public bool IncludeProgramNameWithSeverityPrefix
+		{
+			get; set;
+		}
+
+		/// <summary>
+		/// When true, logs will have the calling mehod prepended to the output as MethodName:
+		/// </summary>
+		public bool IncludeCallingMethod
+		{
+			get; set;
+		}
+
+		/// <summary>
+		/// When true, console output will have the calling mehod prepended to the output as MethodName:
+		/// </summary>
+		public bool IncludeCallingMethodForConsole
+		{
+			get; set;
+		}
+
+		/// <summary>
+		/// When true, will detect warnings and errors and set the console output color to yellow and red.
+		/// </summary>
+		public bool ColorConsoleOutput
+		{
+			get; set;
+		}
+
+		/// <summary>
+		/// Path to the log file being written to. May be null.
+		/// </summary>
+		public FileReference OutputFile
+		{
+			get; private set;
+		}
+
+		/// <summary>
+		/// Whether console output is redirected. This prevents writing status updates that rely on moving the cursor.
+		/// </summary>
+		private bool AllowStatusUpdates
+		{
+			get { return !Console.IsOutputRedirected; }
+		}
+
+		/// <summary>
+		/// When configured, this tracks time since initialization to prepend a timestamp to each log.
+		/// </summary>
+		private Stopwatch Timer = Stopwatch.StartNew();
+
+		/// <summary>
+		/// A collection of strings that have been already written once
+		/// </summary>
+		private HashSet<string> WriteOnceSet = new HashSet<string>();
+
+		/// <summary>
+		/// Stack of status scope information.
+		/// </summary>
+		private Stack<StatusMessage> StatusMessageStack = new Stack<StatusMessage>();
+
+		/// <summary>
+		/// The currently visible status text
+		/// </summary>
+		private string StatusText = "";
+
+		/// <summary>
+		/// Last time a status message was pushed to the stack
+		/// </summary>
+		private Stopwatch StatusTimer = new Stopwatch();
+
+		/// <summary>
+		/// Indent added to every output line
+		/// </summary>
+		public string Indent
+		{
+			get; set;
+		}
+
+		/// <summary>
+		/// Constructor
+		/// </summary>
+		public DefaultLogger()
+		{
+			Indent = "";
+			OutputLevel = LogEventType.Log;
+			IncludeSeverityPrefix = true;
+			IncludeProgramNameWithSeverityPrefix = false;
+			IncludeCallingMethod = true;
+			IncludeCallingMethodForConsole = false;
+			ColorConsoleOutput = true;
+		}
+
+		/// <summary>
+		/// Adds a trace listener that writes to a log file
+		/// </summary>
+		/// <param name="OutputFile">The file to write to</param>
+		/// <returns>The created trace listener</returns>
+		public TextWriterTraceListener AddFileWriter(string Name, FileReference OutputFile)
+		{
+			try
+			{
+				this.OutputFile = OutputFile;
+				DirectoryReference.CreateDirectory(OutputFile.Directory);
+				TextWriterTraceListener LogTraceListener = new TextWriterTraceListener(new StreamWriter(OutputFile.FullName), Name);
+				lock (SyncObject)
+				{
+					Trace.Listeners.Add(LogTraceListener);
+				}
+				return LogTraceListener;
+			}
+			catch (Exception Ex)
+			{
+				throw new Exception(String.Format("Unable to open log file for writing ({0})", OutputFile), Ex);
+			}
+		}
+
+		/// <summary>
+		/// Adds a <see cref="TraceListener"/> to the collection in a safe manner.
+		/// </summary>
+		/// <param name="TraceListener">The <see cref="TraceListener"/> to add.</param>
+		public void AddTraceListener(TraceListener TraceListener)
+		{
+			lock (SyncObject)
+			{
+				if (!Trace.Listeners.Contains(TraceListener))
+				{
+					Trace.Listeners.Add(TraceListener);
+				}
+			}
+		}
+
+		/// <summary>
+		/// Removes a <see cref="TraceListener"/> from the collection in a safe manner.
+		/// </summary>
+		/// <param name="TraceListener">The <see cref="TraceListener"/> to remove.</param>
+		public void RemoveTraceListener(TraceListener TraceListener)
+		{
+			lock (SyncObject)
+			{
+				if (Trace.Listeners.Contains(TraceListener))
+				{
+					Trace.Listeners.Remove(TraceListener);
+				}
+			}
+		}
+
+		/// <summary>
+		/// Determines if a TextWriterTraceListener has been added to the list of trace listeners
+		/// </summary>
+		/// <returns>True if a TextWriterTraceListener has been added</returns>
+		public static bool HasFileWriter()
+		{
+			foreach (TraceListener Listener in Trace.Listeners)
+			{
+				if (Listener is TextWriterTraceListener)
+				{
+					return true;
+				}
+			}
+			return false;
+		}
+
+		/// <summary>
+		/// Gets the name of the Method N levels deep in the stack frame. Used to trap what method actually made the logging call.
+		/// Only used when bLogSources is true.
+		/// </summary>
+		/// <param name="StackFramesToSkip"></param>
+		/// <returns>ClassName.MethodName</returns>
+		[MethodImplAttribute(MethodImplOptions.NoInlining)]
+		private static string GetSource(int StackFramesToSkip)
+		{
+			StackFrame Frame = new StackFrame(2 + StackFramesToSkip);
+			System.Reflection.MethodBase Method = Frame.GetMethod();
+			return String.Format("{0}.{1}", Method.DeclaringType.Name, Method.Name);
+		}
+
+		/// <summary>
+		/// Converts a LogEventType into a log prefix. Only used when bLogSeverity is true.
+		/// </summary>
+		/// <param name="Severity"></param>
+		/// <returns></returns>
+		private static string GetSeverityPrefix(LogEventType Severity)
+		{
+			switch (Severity)
+			{
+				case LogEventType.Fatal:
+					return "FATAL ERROR: ";
+				case LogEventType.Error:
+					return "ERROR: ";
+				case LogEventType.Warning:
+					return "WARNING: ";
+				case LogEventType.Console:
+					return "";
+				case LogEventType.Verbose:
+					return "VERBOSE: ";
+				default:
+					return "";
+			}
+		}
+
+		/// <summary>
+		/// Formats message for logging. Enforces the configured options.
+		/// </summary>
+		/// <param name="StackFramesToSkip">Number of frames to skip to get to the originator of the log request.</param>
+		/// <param name="Verbosity">Message verbosity level</param>
+		/// <param name="Options">Options for formatting this string</param>
+		/// <param name="bForConsole">Whether the message is intended for console output</param>
+		/// <param name="Format">Message text format string</param>
+		/// <param name="Args">Message text parameters</param>
+		/// <returns>Formatted message</returns>
+		[MethodImplAttribute(MethodImplOptions.NoInlining)]
+		private List<string> FormatMessage(int StackFramesToSkip, LogEventType Verbosity, LogFormatOptions Options, bool bForConsole, string Format, params object[] Args)
+		{
+			string TimePrefix = IncludeTimestamps ? String.Format("[{0:hh\\:mm\\:ss\\.fff}] ", Timer.Elapsed) : "";
+			string SourcePrefix = (bForConsole ? IncludeCallingMethodForConsole : IncludeCallingMethod) ? string.Format("{0}: ", GetSource(StackFramesToSkip)) : "";
+			string SeverityPrefix = (IncludeSeverityPrefix && ((Options & LogFormatOptions.NoSeverityPrefix) == 0)) ? GetSeverityPrefix(Verbosity) : "";
+
+			// Include the executable name when running inside MSBuild. If unspecified, MSBuild re-formats them with an "EXEC :" prefix.
+			if (SeverityPrefix.Length > 0 && IncludeProgramNameWithSeverityPrefix)
+			{
+				SeverityPrefix = String.Format("{0}: {1}", Path.GetFileNameWithoutExtension(Assembly.GetEntryAssembly().Location), SeverityPrefix);
+			}
+
+			// If there are no extra args, don't try to format the string, in case it has any format control characters in it (our LOCTEXT strings tend to).
+			string[] Lines = ((Args.Length > 0) ? String.Format(Format, Args) : Format).TrimEnd(' ', '\t', '\r', '\n').Split('\n');
+
+			List<string> FormattedLines = new List<string>();
+			FormattedLines.Add(String.Format("{0}{1}{2}{3}{4}", TimePrefix, SourcePrefix, Indent, SeverityPrefix, Lines[0].TrimEnd('\r')));
+
+			if (Lines.Length > 1)
+			{
+				int PaddingLength = 0;
+				while (PaddingLength < Lines[0].Length && Char.IsWhiteSpace(Lines[0][PaddingLength]))
+				{
+					PaddingLength++;
+				}
+
+				string Padding = new string(' ', SeverityPrefix.Length) + Lines[0].Substring(0, PaddingLength);
+				for (int Idx = 1; Idx < Lines.Length; Idx++)
+				{
+					FormattedLines.Add(String.Format("{0}{1}{2}{3}{4}", TimePrefix, SourcePrefix, Indent, Padding, Lines[Idx].TrimEnd('\r')));
+				}
+			}
+
+			return FormattedLines;
+		}
+
+		public IDisposable BeginScope<TState>(TState State)
+		{
+			throw new NotImplementedException();
+		}
+
+		public bool IsEnabled(LogLevel LogLevel)
+		{
+			throw new NotImplementedException();
+		}
+
+		public void Log<TState>(LogLevel LogLevel, EventId EventId, TState State, Exception Exception, Func<TState, Exception, string> Formatter)
+		{
+			LogMessage? Message = State as LogMessage?;
+			if (Message != null)
+			{
+				WriteLinePrivate(Message.Value.StackFramesToSkip, Message.Value.bWriteOnce, (LogEventType)LogLevel, Message.Value.FormatOptions, Message.Value.Format, Message.Value.Args);
+			}
+			else
+			{
+				WriteLinePrivate(1, false, (LogEventType)LogLevel, LogFormatOptions.None, Formatter(State, Exception));
+			}
+		}
+
+		/// <summary>
+		/// Writes a formatted message to the console. All other functions should boil down to calling this method.
+		/// </summary>
+		/// <param name="StackFramesToSkip">Number of frames to skip to get to the originator of the log request.</param>
+		/// <param name="bWriteOnce">If true, this message will be written only once</param>
+		/// <param name="Verbosity">Message verbosity level. We only meaningfully use values up to Verbose</param>
+		/// <param name="FormatOptions">Options for formatting messages</param>
+		/// <param name="Format">Message format string.</param>
+		/// <param name="Args">Optional arguments</param>
+		[MethodImplAttribute(MethodImplOptions.NoInlining)]
+		private void WriteLinePrivate(int StackFramesToSkip, bool bWriteOnce, LogEventType Verbosity, LogFormatOptions FormatOptions, string Format, params object[] Args)
+		{
+			// if we want this message only written one time, check if it was already written out
+			if (bWriteOnce)
+			{
+				string Formatted = string.Format(Format, Args);
+				if (WriteOnceSet.Contains(Formatted))
+				{
+					return;
+				}
+
+				WriteOnceSet.Add(Formatted);
+			}
+
+			if (Verbosity > OutputLevel)
+			{
+				lock (SyncObject)
+				{
+					// Output to all the other trace listeners
+					List<string> Lines = FormatMessage(StackFramesToSkip + 1, Verbosity, FormatOptions, false, Format, Args);
+					foreach (TraceListener Listener in Trace.Listeners)
+					{
+						foreach (string Line in Lines)
+						{
+							Listener.WriteLine(Line);
+						}
+						Listener.Flush();
+					}
+
+					// Handle the console output separately; we format things differently
+					if (Verbosity > LogEventType.Log && (FormatOptions & LogFormatOptions.NoConsoleOutput) == 0)
+					{
+						FlushStatusHeading();
+
+						bool bResetConsoleColor = false;
+						if (ColorConsoleOutput)
+						{
+							if (Verbosity == LogEventType.Warning)
+							{
+								Console.ForegroundColor = ConsoleColor.Yellow;
+								bResetConsoleColor = true;
+							}
+							if (Verbosity >= LogEventType.Error)
+							{
+								Console.ForegroundColor = ConsoleColor.Red;
+								bResetConsoleColor = true;
+							}
+						}
+						try
+						{
+							List<string> ConsoleLines = FormatMessage(StackFramesToSkip + 1, Verbosity, FormatOptions, true, Format, Args);
+							foreach (string ConsoleLine in ConsoleLines)
+							{
+								Console.WriteLine(ConsoleLine);
+							}
+						}
+						catch (IOException)
+						{
+							// Potential file access/sharing issue on std out
+							// This can occur on some versions of mono (e.g. macOS 6.12.0) if writing to a full pipe
+							// during IPC when the reader isn't consuming it quick enough
+						}
+						finally
+						{
+							// make sure we always put the console color back.
+							if (bResetConsoleColor)
+							{
+								Console.ResetColor();
+							}
+						}
+
+						if (StatusMessageStack.Count > 0 && AllowStatusUpdates)
+						{
+							SetStatusText(StatusMessageStack.Peek().CurrentText);
+						}
+					}
+				}
+			}
+		}
+
+		/// <summary>
+		/// Flushes the current status text before writing out a new log line or status message
+		/// </summary>
+		void FlushStatusHeading()
+		{
+			if (StatusMessageStack.Count > 0)
+			{
+				StatusMessage CurrentStatus = StatusMessageStack.Peek();
+				if (CurrentStatus.HeadingText.Length > 0 && !CurrentStatus.bHasFlushedHeadingText && AllowStatusUpdates)
+				{
+					SetStatusText(CurrentStatus.HeadingText);
+					Console.WriteLine();
+					StatusText = "";
+					CurrentStatus.bHasFlushedHeadingText = true;
+				}
+				else
+				{
+					SetStatusText("");
+				}
+			}
+		}
+
+		/// <summary>
+		/// Enter a scope with the given status message. The message will be written to the console without a newline, allowing it to be updated through subsequent calls to UpdateStatus().
+		/// The message will be written to the log immediately. If another line is written while in a status scope, the initial status message is flushed to the console first.
+		/// </summary>
+		/// <param name="Message">The status message</param>
+		[Conditional("TRACE")]
+		[MethodImplAttribute(MethodImplOptions.NoInlining)]
+		public void PushStatus(string Message)
+		{
+			lock (SyncObject)
+			{
+				FlushStatusHeading();
+
+				StatusMessage NewStatusMessage = new StatusMessage();
+				NewStatusMessage.HeadingText = Message;
+				NewStatusMessage.CurrentText = Message;
+				StatusMessageStack.Push(NewStatusMessage);
+
+				StatusTimer.Restart();
+
+				if (Message.Length > 0)
+				{
+					WriteLinePrivate(1, false, LogEventType.Log, LogFormatOptions.NoConsoleOutput, "{0}", Message);
+					SetStatusText(Message);
+				}
+			}
+		}
+
+		/// <summary>
+		/// Updates the current status message. This will overwrite the previous status line.
+		/// </summary>
+		/// <param name="Message">The status message</param>
+		[Conditional("TRACE")]
+		[MethodImplAttribute(MethodImplOptions.NoInlining)]
+		public void UpdateStatus(string Message)
+		{
+			lock (SyncObject)
 			{
 				StatusMessage CurrentStatusMessage = StatusMessageStack.Peek();
+				CurrentStatusMessage.CurrentText = Message;
+
+				if (AllowStatusUpdates || StatusTimer.Elapsed.TotalSeconds > 10.0)
+				{
+					SetStatusText(Message);
+					StatusTimer.Restart();
+				}
+			}
+		}
+
+		/// <summary>
+		/// Updates the Pops the top status message from the stack. The mess
+		/// </summary>
+		/// <param name="Message"></param>
+		[Conditional("TRACE")]
+		[MethodImplAttribute(MethodImplOptions.NoInlining)]
+		public void PopStatus()
+		{
+			lock (SyncObject)
+			{
+				StatusMessage CurrentStatusMessage = StatusMessageStack.Peek();
 				SetStatusText(CurrentStatusMessage.CurrentText);
 
-				if(StatusText.Length > 0)
+				if (StatusText.Length > 0)
 				{
 					Console.WriteLine();
 					StatusText = "";
@@ -901,22 +1148,22 @@
 		/// Update the status text. For internal use only; does not modify the StatusMessageStack objects.
 		/// </summary>
 		/// <param name="NewStatusText">New status text to display</param>
-		private static void SetStatusText(string NewStatusText)
+		private void SetStatusText(string NewStatusText)
 		{
-			if(NewStatusText.Length > 0)
+			if (NewStatusText.Length > 0)
 			{
 				NewStatusText = Indent + NewStatusText;
 			}
 
-			if(StatusText != NewStatusText)
+			if (StatusText != NewStatusText)
 			{
 				int NumCommonChars = 0;
-				while(NumCommonChars < StatusText.Length && NumCommonChars < NewStatusText.Length && StatusText[NumCommonChars] == NewStatusText[NumCommonChars])
+				while (NumCommonChars < StatusText.Length && NumCommonChars < NewStatusText.Length && StatusText[NumCommonChars] == NewStatusText[NumCommonChars])
 				{
 					NumCommonChars++;
 				}
 
-				if(!AllowStatusUpdates && NumCommonChars < StatusText.Length)
+				if (!AllowStatusUpdates && NumCommonChars < StatusText.Length)
 				{
 					// Prevent writing backspace characters if the console doesn't support it
 					Console.WriteLine();
@@ -927,7 +1174,7 @@
 				StringBuilder Text = new StringBuilder();
 				Text.Append('\b', StatusText.Length - NumCommonChars);
 				Text.Append(NewStatusText, NumCommonChars, NewStatusText.Length - NumCommonChars);
-				if(NewStatusText.Length < StatusText.Length)
+				if (NewStatusText.Length < StatusText.Length)
 				{
 					int NumChars = StatusText.Length - NewStatusText.Length;
 					Text.Append(' ', NumChars);
