Make LogIndentScope independent of the logger being used. Any loggers that want to use this state for additional formatting can do so with an enricher.

#rnx

diff --git a/Engine/Source/Programs/AutomationTool/AutomationUtils/ProcessUtils.cs b/Engine/Source/Programs/AutomationTool/AutomationUtils/ProcessUtils.cs
--- a/Engine/Source/Programs/AutomationTool/AutomationUtils/ProcessUtils.cs
+++ b/Engine/Source/Programs/AutomationTool/AutomationUtils/ProcessUtils.cs
@@ -840,27 +840,20 @@ namespace AutomationTool
                 LogWithVerbosity(SpewVerbosity,"Running: " + App + " " + (String.IsNullOrEmpty(CommandLine) ? "" : CommandLine));
             }
 
-			string PrevIndent = null;
-			if(Options.HasFlag(ERunOptions.AllowSpew))
-			{
-				PrevIndent = EpicGames.Core.Log.Indent;
-				EpicGames.Core.Log.Indent += "  ";
-			}
-
 			IProcessResult Result = ProcessManager.CreateProcess(App, Options.HasFlag(ERunOptions.AllowSpew), !Options.HasFlag(ERunOptions.NoStdOutCapture), Env, SpewVerbosity: SpewVerbosity, SpewFilterCallback: SpewFilterCallback);
-			try
+			using (LogIndentScope Scope = Options.HasFlag(ERunOptions.AllowSpew) ? new LogIndentScope("  ") : null)
 			{
 				Process Proc = Result.ProcessObject;
 
 				bool bRedirectStdOut = (Options & ERunOptions.NoStdOutRedirect) != ERunOptions.NoStdOutRedirect;
 				Proc.StartInfo.FileName = App;
-				
+
 				// Process Arguments follow windows conventions in .NET Core
 				// Which means single quotes ' are not considered quotes.
 				// see [REDACTED_PROFILE_URL]/runtime/issues/29857
 				// also see UE-102580
 				Proc.StartInfo.Arguments = String.IsNullOrEmpty(CommandLine) ? "" : CommandLine.Replace('\'', '\"');
-				
+
 				Proc.StartInfo.UseShellExecute = false;
 				if (bRedirectStdOut)
 				{
@@ -898,13 +891,6 @@ namespace AutomationTool
 					Result.ExitCode = -1;
 				}
 			}
-			finally
-			{
-				if(PrevIndent != null)
-				{
-					EpicGames.Core.Log.Indent = PrevIndent;
-				}
-			}
 
 			if (!Options.HasFlag(ERunOptions.NoWaitForExit))
 			{
diff --git a/Engine/Source/Programs/Shared/EpicGames.Core/Log.cs b/Engine/Source/Programs/Shared/EpicGames.Core/Log.cs
--- a/Engine/Source/Programs/Shared/EpicGames.Core/Log.cs
+++ b/Engine/Source/Programs/Shared/EpicGames.Core/Log.cs
@@ -209,15 +209,6 @@ namespace EpicGames.Core
 		/// </summary>
 		public static FileReference OutputFile => DefaultLogger.OutputFile;
 
-		/// <summary>
-		/// Indent added to every output line
-		/// </summary>
-		public static string Indent
-		{
-			get => DefaultLogger.Indent;
-			set => DefaultLogger.Indent = value; 
-		}
-
 		/// <summary>
 		/// Adds a trace listener that writes to a log file
 		/// </summary>
@@ -760,20 +751,11 @@ namespace EpicGames.Core
 		/// </summary>
 		private Stopwatch StatusTimer = new Stopwatch();
 
-		/// <summary>
-		/// Indent added to every output line
-		/// </summary>
-		public string Indent
-		{
-			get; set;
-		}
-
 		/// <summary>
 		/// Constructor
 		/// </summary>
 		public DefaultLogger()
 		{
-			Indent = "";
 			OutputLevel = LogEventType.Log;
 			IncludeSeverityPrefix = true;
 			IncludeProgramNameWithSeverityPrefix = false;
@@ -916,6 +898,8 @@ namespace EpicGames.Core
 			// If there are no extra args, don't try to format the string, in case it has any format control characters in it (our LOCTEXT strings tend to).
 			string[] Lines = ((Args.Length > 0) ? String.Format(Format, Args) : Format).TrimEnd(' ', '\t', '\r', '\n').Split('\n');
 
+			string Indent = LogIndent.Current;
+
 			List<string> FormattedLines = new List<string>();
 			FormattedLines.Add(String.Format("{0}{1}{2}{3}{4}", TimePrefix, SourcePrefix, Indent, SeverityPrefix, Lines[0].TrimEnd('\r')));
 
@@ -1152,7 +1136,7 @@ namespace EpicGames.Core
 		{
 			if (NewStatusText.Length > 0)
 			{
-				NewStatusText = Indent + NewStatusText;
+				NewStatusText = LogIndent.Current + NewStatusText;
 			}
 
 			if (StatusText != NewStatusText)
diff --git a/Engine/Source/Programs/Shared/EpicGames.Core/LogIndent.cs b/Engine/Source/Programs/Shared/EpicGames.Core/LogIndent.cs
new file mode 100644
--- /dev/null
+++ b/Engine/Source/Programs/Shared/EpicGames.Core/LogIndent.cs
@@ -0,0 +1,62 @@
+// Copyright Epic Games, Inc. All Rights Reserved.
+
+using System;
+using System.Collections.Generic;
+using System.Text;
+using System.Threading;
+
+namespace EpicGames.Core
+{
+	/// <summary>
+	/// Static class for tracking a stack of log indent prefixes
+	/// </summary>
+	public static class LogIndent
+	{
+		/// <summary>
+		/// Tracks the current state
+		/// </summary>
+		class State
+		{
+			public State? PrevState { get; }
+			public string Indent { get; }
+
+			public State(State? PrevState, string Indent)
+			{
+				this.PrevState = PrevState;
+				this.Indent = Indent;
+			}
+		}
+
+		/// <summary>
+		/// The current state value
+		/// </summary>
+		static AsyncLocal<State?> CurrentState = new AsyncLocal<State?>();
+
+		/// <summary>
+		/// Gets the current indent string
+		/// </summary>
+		public static string Current => CurrentState.Value?.Indent ?? String.Empty;
+
+		/// <summary>
+		/// Push a new indent onto the stack
+		/// </summary>
+		/// <param name="Indent">The indent to add</param>
+		public static void Push(string Indent)
+		{
+			State? PrevState = CurrentState.Value;
+			if (PrevState != null)
+			{
+				Indent = PrevState.Indent + Indent;
+			}
+			CurrentState.Value = new State(PrevState, Indent);
+		}
+
+		/// <summary>
+		/// Pops an indent off the stack
+		/// </summary>
+		public static void Pop()
+		{
+			CurrentState.Value = CurrentState.Value!.PrevState;
+		}
+	}
+}
diff --git a/Engine/Source/Programs/Shared/EpicGames.Core/LogIndentScope.cs b/Engine/Source/Programs/Shared/EpicGames.Core/LogIndentScope.cs
--- a/Engine/Source/Programs/Shared/EpicGames.Core/LogIndentScope.cs
+++ b/Engine/Source/Programs/Shared/EpicGames.Core/LogIndentScope.cs
@@ -4,6 +4,7 @@ using System;
 using System.Collections.Generic;
 using System.Linq;
 using System.Text;
+using System.Threading;
 using System.Threading.Tasks;
 
 namespace EpicGames.Core
@@ -14,9 +15,9 @@ namespace EpicGames.Core
 	public class LogIndentScope : IDisposable
 	{
 		/// <summary>
-		/// The previous indent
+		/// Whether the object has been disposed
 		/// </summary>
-		string? PrevIndent;
+		bool Disposed;
 
 		/// <summary>
 		/// Constructor
@@ -24,8 +25,7 @@ namespace EpicGames.Core
 		/// <param name="Indent">Indent to append to the existing indent</param>
 		public LogIndentScope(string Indent)
 		{
-			PrevIndent = Log.Indent;
-			Log.Indent += Indent;
+			LogIndent.Push(Indent);
 		}
 
 		/// <summary>
@@ -33,10 +33,10 @@ namespace EpicGames.Core
 		/// </summary>
 		public void Dispose()
 		{
-			if (PrevIndent != null)
+			if (!Disposed)
 			{
-				Log.Indent = PrevIndent;
-				PrevIndent = null;
+				LogIndent.Pop();
+				Disposed = true;
 			}
 		}
 	}
