Added missing copyright boilerplate

#rnx

diff --git a/Engine/Source/Programs/Shared/EpicGames.Core/BinaryArchive.cs b/Engine/Source/Programs/Shared/EpicGames.Core/BinaryArchive.cs
--- a/Engine/Source/Programs/Shared/EpicGames.Core/BinaryArchive.cs
+++ b/Engine/Source/Programs/Shared/EpicGames.Core/BinaryArchive.cs
@@ -1,3 +1,5 @@
+// Copyright Epic Games, Inc. All Rights Reserved.
+
 using System;
 using System.Collections.Concurrent;
 using System.Collections.Generic;
diff --git a/Engine/Source/Programs/Shared/EpicGames.Core/BinaryArchiveAttributes.cs b/Engine/Source/Programs/Shared/EpicGames.Core/BinaryArchiveAttributes.cs
--- a/Engine/Source/Programs/Shared/EpicGames.Core/BinaryArchiveAttributes.cs
+++ b/Engine/Source/Programs/Shared/EpicGames.Core/BinaryArchiveAttributes.cs
@@ -1,3 +1,5 @@
+// Copyright Epic Games, Inc. All Rights Reserved.
+
 using System;
 using System.Collections.Generic;
 using System.Text;
diff --git a/Engine/Source/Programs/Shared/EpicGames.Core/BinaryConverter.cs b/Engine/Source/Programs/Shared/EpicGames.Core/BinaryConverter.cs
--- a/Engine/Source/Programs/Shared/EpicGames.Core/BinaryConverter.cs
+++ b/Engine/Source/Programs/Shared/EpicGames.Core/BinaryConverter.cs
@@ -1,3 +1,5 @@
+// Copyright Epic Games, Inc. All Rights Reserved.
+
 using EpicGames.Core;
 using System;
 using System.Collections.Concurrent;
