Merge changes from Horde fork of EpicGames.Perforce (various performance improvements, bug fixes, nullable and async aware)

#rnx

diff --git a/Engine/Source/Programs/Shared/EpicGames.Perforce/EpicGames.Perforce.csproj b/Engine/Source/Programs/Shared/EpicGames.Perforce/EpicGames.Perforce.csproj
--- a/Engine/Source/Programs/Shared/EpicGames.Perforce/EpicGames.Perforce.csproj
+++ b/Engine/Source/Programs/Shared/EpicGames.Perforce/EpicGames.Perforce.csproj
@@ -1,4 +1,5 @@
 <Project Sdk="Microsoft.NET.Sdk">
+
   <PropertyGroup>
     <TargetFramework>netcoreapp3.1</TargetFramework>
     <Configuration Condition=" '$(Configuration)' == '' ">Development</Configuration>
@@ -7,6 +8,7 @@
     <GenerateAssemblyInfo>false</GenerateAssemblyInfo>
     <GenerateTargetFrameworkAttribute>false</GenerateTargetFrameworkAttribute>
     <Configurations>Debug;Release;Development</Configurations>
+    <nullable>enable</nullable>
   </PropertyGroup>
   <PropertyGroup Condition=" '$(Configuration)|$(Platform)' == 'Debug|AnyCPU' ">
     <DocumentationFile></DocumentationFile>
diff --git a/Engine/Source/Programs/Shared/EpicGames.Perforce/PerforceAttributes.cs b/Engine/Source/Programs/Shared/EpicGames.Perforce/PerforceAttributes.cs
--- a/Engine/Source/Programs/Shared/EpicGames.Perforce/PerforceAttributes.cs
+++ b/Engine/Source/Programs/Shared/EpicGames.Perforce/PerforceAttributes.cs
@@ -1,4 +1,4 @@
-// Copyright Epic Games, Inc. All Rights Reserved.
+// Copyright 1998-2019 Epic Games, Inc. All Rights Reserved.
 
 using System;
 using System.Collections.Generic;
diff --git a/Engine/Source/Programs/Shared/EpicGames.Perforce/PerforceChildProcess.cs b/Engine/Source/Programs/Shared/EpicGames.Perforce/PerforceChildProcess.cs
--- a/Engine/Source/Programs/Shared/EpicGames.Perforce/PerforceChildProcess.cs
+++ b/Engine/Source/Programs/Shared/EpicGames.Perforce/PerforceChildProcess.cs
@@ -1,280 +1,977 @@
-// Copyright Epic Games, Inc. All Rights Reserved.
+// Copyright 1998-2019 Epic Games, Inc. All Rights Reserved.
 
+using Datadog.Trace;
 using EpicGames.Core;
+using Microsoft.Extensions.Logging;
 using System;
+using System.Buffers.Binary;
+using System.Buffers.Text;
+using System.Collections.Concurrent;
 using System.Collections.Generic;
 using System.Diagnostics;
+using System.Diagnostics.CodeAnalysis;
 using System.IO;
 using System.Linq;
+using System.Reflection;
+using System.Security.Cryptography;
 using System.Text;
+using System.Text.RegularExpressions;
+using System.Threading;
 using System.Threading.Tasks;
 
 namespace EpicGames.Perforce
 {
-	enum PerforceRecordType
+	/// <summary>
+	/// Wraps a call to a p4.exe child process, and allows reading data from it
+	/// </summary>
+	class PerforceChildProcess : IDisposable
 	{
-		Invalid,
-		Stat,
-		Text,
-		Binary,
-	}
+		/// <summary>
+		/// String constants for records
+		/// </summary>
+		static class ReadOnlyUtf8StringConstants
+		{
+			public static ReadOnlyUtf8String Code = "code";
+			public static ReadOnlyUtf8String Stat = "stat";
+			public static ReadOnlyUtf8String Info = "info";
+			public static ReadOnlyUtf8String Error = "error";
+		}
 
-	enum PerforceFieldType
-	{
-		Record,
-		EndRecord,
-		String,
-		Integer,
-		EndStream
-	}
+		/// <summary>
+		/// Standard prefix for a returned record: record indicator, string, 4 bytes, 'code', string, [value]
+		/// </summary>
+		static readonly byte[] RecordPrefix = { (byte)'{', (byte)'s', 4, 0, 0, 0, (byte)'c', (byte)'o', (byte)'d', (byte)'e', (byte)'s' };
 
-	class PerforceChildProcess : IDisposable
-	{
+		/// <summary>
+		/// The process group
+		/// </summary>
 		ManagedProcessGroup ChildProcessGroup;
+
+		/// <summary>
+		/// The child process instance
+		/// </summary>
 		ManagedProcess ChildProcess;
+
+		/// <summary>
+		/// Scope object for tracing
+		/// </summary>
+		Scope Scope;
+
+		/// <summary>
+		/// Buffer for data from the child process
+		/// </summary>
 		byte[] Buffer;
+
+		/// <summary>
+		/// Current read position within the buffer
+		/// </summary>
 		int BufferPos;
+
+		/// <summary>
+		/// End of the valid portion of the buffer (exclusive)
+		/// </summary>
 		int BufferEnd;
+
+		/// <summary>
+		/// Whether the remaining contents of the buffer represent the end of the file
+		/// </summary>
 		bool bEndOfStream;
 
-		public PerforceChildProcess(byte[] InputData, string Format, params object[] Args)
+		/// <summary>
+		/// Constructor
+		/// </summary>
+		/// <param name="InputData">Input data to pass to the child process</param>
+		/// <param name="Arguments">Command line arguments</param>
+		/// <param name="Logger">Logging device</param>
+		public PerforceChildProcess(byte[]? InputData, string Arguments, ILogger Logger)
+		{
+			string PerforceFileName = GetExecutable();
+
+			string FullArgumentList = "-G " + Arguments;
+			Logger.LogDebug("Running {0} {1}", PerforceFileName, FullArgumentList);
+
+			Scope = Tracer.Instance.StartActive($"{GetCommandName(Arguments)}", serviceName: "perforce");
+			Scope.Span.SetTag("arguments", FullArgumentList);
+
+			ChildProcessGroup = new ManagedProcessGroup();
+			ChildProcess = new ManagedProcess(ChildProcessGroup, PerforceFileName, FullArgumentList, null, null, InputData, ProcessPriorityClass.Normal);
+
+			Buffer = new byte[64 * 1024];
+		}
+
+		/// <summary>
+		/// Gets the command name from the argument list
+		/// </summary>
+		/// <param name="CommandLine">The command line arguments</param>
+		/// <returns>The command name</returns>
+		public static string GetCommandName(string CommandLine)
+		{
+			string[] Arguments = CommandLineArguments.Split(CommandLine);
+			for(int Idx = 0; Idx < Arguments.Length; Idx++)
+			{
+				string Argument = Arguments[Idx];
+				if(Argument == "-p" || Argument == "-u" || Argument == "-c")
+				{
+					Idx++;
+				}
+				else if(Argument.Length > 0 && Char.IsLetter(Argument[0]))
+				{
+					return Argument;
+				}
+			}
+			return "(unknown)";
+		}
+
+		/// <summary>
+		/// Gets the path to the P4.EXE executable
+		/// </summary>
+		/// <returns>Path to the executable</returns>
+		public static string GetExecutable()
 		{
 			string PerforceFileName;
-			if(Environment.OSVersion.Platform == PlatformID.Win32NT || Environment.OSVersion.Platform == PlatformID.Win32S || Environment.OSVersion.Platform == PlatformID.Win32Windows)
+			if (Environment.OSVersion.Platform == PlatformID.Win32NT || Environment.OSVersion.Platform == PlatformID.Win32S || Environment.OSVersion.Platform == PlatformID.Win32Windows)
 			{
 				PerforceFileName = "p4.exe";
 			}
 			else
 			{
-				PerforceFileName = File.Exists("/usr/local/bin/p4")? "/usr/local/bin/p4" : "/usr/bin/p4";
+				PerforceFileName = File.Exists("/usr/local/bin/p4") ? "/usr/local/bin/p4" : "/usr/bin/p4";
 			}
-
-			string FullArgumentList = "-G " + String.Format(Format, Args);
-			Log.TraceLog("Running {0} {1}", PerforceFileName, FullArgumentList);
-
-			ChildProcessGroup = new ManagedProcessGroup();
-			ChildProcess = new ManagedProcess(ChildProcessGroup, PerforceFileName, FullArgumentList, null, null, InputData, ProcessPriorityClass.Normal);
-
-			Buffer = new byte[2048];
+			return PerforceFileName;
 		}
 
+		/// <summary>
+		/// Dispose of resources owned by this object
+		/// </summary>
 		public void Dispose()
 		{
-			if(ChildProcess != null)
+			if (ChildProcess != null)
 			{
 				ChildProcess.Dispose();
-				ChildProcess = null;
+				ChildProcess = null!;
 			}
-			if(ChildProcessGroup != null)
+			if (ChildProcessGroup != null)
 			{
 				ChildProcessGroup.Dispose();
-				ChildProcessGroup = null;
+				ChildProcessGroup = null!;
+			}
+			if (Scope != null)
+			{
+				Scope.Dispose();
+				Scope = null!;
 			}
 		}
 
-		public bool IsEndOfStream()
+		/// <summary>
+		/// Reads all output from the child process as a string
+		/// </summary>
+		/// <param name="CancellationToken">Cancellation token to abort the read</param>
+		/// <returns>Exit code and output from the process</returns>
+		public async Task<Tuple<bool, string>> TryReadToEndAsync(CancellationToken CancellationToken)
 		{
-			if(!bEndOfStream)
+			MemoryStream Stream = new MemoryStream();
+			await ChildProcess!.CopyTo(Stream, CancellationToken);
+			ChildProcess.WaitForExit();
+			string Data = Encoding.Default.GetString(Stream.ToArray());
+			return Tuple.Create(ChildProcess.ExitCode == 0, Data);
+		}
+
+		/// <summary>
+		/// Formats the current contents of the buffer to a string
+		/// </summary>
+		/// <param name="Dat">The next byte that was read</param>
+		/// <returns>String representation of the buffer</returns>
+		private static string FormatDataAsString(Span<byte> Data)
+		{
+			StringBuilder Result = new StringBuilder();
+			if (Data.Length > 0)
 			{
-				bEndOfStream = (BufferPos == BufferEnd && !TryUpdateBuffer(1));
+				for (int Idx = 0; Idx < Data.Length && Idx < 1024;)
+				{
+					Result.Append("\n   ");
+
+					// Output to the end of the line
+					for (; Idx < Data.Length && Idx < 1024 && Data[Idx] != '\r' && Data[Idx] != '\n'; Idx++)
+					{
+						if (Data[Idx] == '\t')
+						{
+							Result.Append("\t");
+						}
+						else if (Data[Idx] == '\\')
+						{
+							Result.Append("\\");
+						}
+						else if (Data[Idx] >= 0x20 && Data[Idx] <= 0x7f)
+						{
+							Result.Append((char)Data[Idx]);
+						}
+						else
+						{
+							Result.AppendFormat("\\x{0:x2}", Data[Idx]);
+						}
+					}
+
+					// Skip the newline characters
+					if (Idx < Data.Length && Data[Idx] == '\r')
+					{
+						Idx++;
+					}
+					if (Idx < Data.Length && Data[Idx] == '\n')
+					{
+						Idx++;
+					}
+				}
 			}
-			return bEndOfStream;
+			return Result.ToString();
+		}
+
+		/// <summary>
+		/// Formats the current contents of the buffer to a string
+		/// </summary>
+		/// <param name="NextByte">The next byte that was read</param>
+		/// <returns>String representation of the buffer</returns>
+		private string FormatDataAsHexDump(int BufferIdx)
+		{
+			return FormatDataAsHexDump(Buffer.AsSpan(BufferIdx, Math.Min(BufferEnd - BufferIdx, 1024)));
 		}
 
-		private string FormatUnexpectedData(byte NextByte)
+		/// <summary>
+		/// Formats the current contents of the buffer to a string
+		/// </summary>
+		/// <param name="Data">The data to format</param>
+		/// <returns>String representation of the buffer</returns>
+		private static string FormatDataAsHexDump(Span<byte> Data)
 		{
-			// Read all the data to the buffer
-			List<byte> Data = new List<byte>();
-			for(int Idx = 0; Idx < 1024; Idx++)
+			// Format the result
+			StringBuilder Result = new StringBuilder();
+
+			const int RowLength = 16;
+			for (int BaseIdx = 0; BaseIdx < Data.Length; BaseIdx += RowLength)
 			{
-				Data.Add(NextByte);
-				if(IsEndOfStream())
+				Result.Append("\n    ");
+				for (int Offset = 0; Offset < RowLength; Offset++)
 				{
-					break;
+					int Idx = BaseIdx + Offset;
+					if (Idx >= Data.Length)
+					{
+						Result.Append("   ");
+					}
+					else
+					{
+						Result.AppendFormat("{0:x2} ", Data[Idx]);
+					}
+				}
+				Result.Append("   ");
+				for (int Offset = 0; Offset < RowLength; Offset++)
+				{
+					int Idx = BaseIdx + Offset;
+					if (Idx >= Data.Length)
+					{
+						break;
+					}
+					else if (Data[Idx] < 0x20 || Data[Idx] >= 0x7f)
+					{
+						Result.Append('.');
+					}
+					else
+					{
+						Result.Append((char)Data[Idx]);
+					}
 				}
-				NextByte = ReadByte();
 			}
+			return Result.ToString();
+		}
 
-			// Check if it's printable. 
-			bool bIsPrintable = true;
-			for(int Idx = 0; Idx < Data.Count; Idx++)
+		/// <summary>
+		/// Try to read plain text from the output stream. Fails if the line starts with a binary record marker ('{').
+		/// </summary>
+		/// <param name="CancellationToken">Cancellation token</param>
+		/// <returns>The line that was read.</returns>
+		public async Task<string?> TryReadLineAsync(CancellationToken CancellationToken)
+		{
+			for (; ; )
 			{
-				if(Data[Idx] < 0x09 || (Data[Idx] >= 0x0e && Data[Idx] <= 0x1f) || Data[Idx] == 0x1f)
+				// Move data to the start of the buffer
+				if (BufferPos > 0)
 				{
-					bIsPrintable = false;
+					Array.Copy(Buffer, BufferPos, Buffer, 0, BufferEnd - BufferPos);
+					BufferEnd -= BufferPos;
+					BufferPos = 0;
+				}
+
+				// Append data to the end of the buffer
+				if (BufferEnd < Buffer.Length)
+				{
+					int BytesRead = await ChildProcess!.ReadAsync(Buffer, BufferEnd, Buffer.Length - BufferEnd, CancellationToken);
+					if (BytesRead == 0)
+					{
+						return null;
+					}
+					BufferEnd += BytesRead;
+				}
+
+				// Move past any lines of plain text output
+				if (BufferEnd > 0)
+				{
+					// Start of a data record
+					if(Buffer[0] == '{')
+					{
+						return null;
+					}
+
+					// Try to reach the end of the line
+					for(int Idx = 0; Idx < BufferEnd; Idx++)
+					{
+						if(Buffer[Idx] == '\n')
+						{
+							BufferPos = Idx + 1;
+							return Encoding.UTF8.GetString(Buffer, 0, BufferPos);
+						}
+					}
 				}
 			}
+		}
 
-			// Format the result
-			StringBuilder Result = new StringBuilder();
-			if(Data.Count > 0)
+		/// <summary>
+		/// Read a list of responses from the child process
+		/// </summary>
+		/// <param name="StatRecordInfo">The type of stat record to parse</param>
+		/// <param name="CancellationToken">Cancellation token for the read</param>
+		/// <returns>Async task</returns>
+		public async Task<List<PerforceResponse>> ReadResponsesAsync(CachedRecordInfo? StatRecordInfo, CancellationToken CancellationToken)
+		{
+			List<PerforceResponse> Responses = new List<PerforceResponse>();
+			while (BufferPos != BufferEnd || !bEndOfStream)
 			{
-				if(bIsPrintable)
+				// Update the buffer contents
+				if (BufferPos > 0)
 				{
-					Result.Append("\n    ");
-					for(int Idx = 0; Idx < Data.Count; Idx++)
+					Array.Copy(Buffer, BufferPos, Buffer, 0, BufferEnd - BufferPos);
+					BufferEnd -= BufferPos;
+					BufferPos = 0;
+				}
+				else if(BufferEnd == Buffer.Length)
+				{
+					Array.Resize(ref Buffer, Buffer.Length + 65536);
+				}
+
+				// Start a read to add more data
+				Task<int>? ReadTask;
+				if (BufferEnd < Buffer.Length && !bEndOfStream)
+				{
+					ReadTask = Task.Run(() => ChildProcess!.ReadAsync(Buffer, BufferEnd, Buffer.Length - BufferEnd, CancellationToken));
+				}
+				else
+				{
+					ReadTask = null;
+				}
+
+				if (BufferEnd > BufferPos)
+				{
+					if (Responses.Count == 0 && BufferPos == 0 && Buffer[0] != '{')
 					{
-						if(Data[Idx] == '\n')
+						// Check for the whole message not being a marshalled python object, and produce a better response in that scenario
+						if (bEndOfStream || BufferEnd == Buffer.Length)
 						{
-							Result.Append("\n    ");
+							throw new PerforceException("Unexpected response from server:{0}", FormatDataAsString(Buffer.AsSpan(0, BufferEnd)));
 						}
-						else if(Data[Idx] != '\r')
+					}
+					else
+					{
+						// Try to process all the existing data
+						for (; ; )
 						{
-							Result.Append((char)Data[Idx]);
+							int InitialBufferPos = BufferPos;
+
+							PerforceResponse? Response;
+							if (!TryReadResponse(StatRecordInfo, out Response))
+							{
+								BufferPos = InitialBufferPos;
+								break;
+							}
+
+							Responses.Add(Response);
+						}
+
+						// If the stream is complete but we couldn't parse a response from the server, treat it as an error
+						if (bEndOfStream && BufferPos < BufferEnd)
+						{
+							throw new PerforceException("Unexpected response from server:{0}", FormatDataAsString(Buffer.AsSpan(BufferPos, BufferEnd - BufferPos)));
 						}
 					}
 				}
+
+				// Wait for the read to complete
+				if(ReadTask != null)
+				{
+					int ReadSize = await ReadTask;
+					if(ReadSize == 0)
+					{
+						bEndOfStream = true;
+					}
+					BufferEnd += ReadSize;
+				}
+			}
+			return Responses;
+		}
+
+		/// <summary>
+		/// Read a list of responses from the child process
+		/// </summary>
+		/// <param name="HandleRecord">Delegate to invoke for each record read</param>
+		/// <param name="CancellationToken">Cancellation token for the read</param>
+		/// <returns>Async task</returns>
+		public async Task ReadRecordsAsync(Action<PerforceRecord> HandleRecord, CancellationToken CancellationToken)
+		{
+			PerforceRecord Record = new PerforceRecord();
+			Record.Rows = new List<KeyValuePair<ReadOnlyUtf8String, PerforceValue>>();
+
+			for(; ;)
+			{
+				// Update the buffer contents
+				if (BufferPos > 0)
+				{
+					byte[] PrevBuffer = Buffer;
+					Buffer = new byte[(BufferEnd - BufferPos) + 65536];
+					Array.Copy(PrevBuffer, BufferPos, Buffer, 0, BufferEnd - BufferPos);
+					BufferEnd -= BufferPos;
+					BufferPos = 0;
+				}
+				else if (BufferEnd == Buffer.Length)
+				{
+					Array.Resize(ref Buffer, Buffer.Length + 65536);
+				}
+
+				// Start a read to add more data
+				Task<int>? ReadTask;
+				if (BufferEnd < Buffer.Length && !bEndOfStream)
+				{
+					ReadTask = Task.Run(() => ChildProcess!.ReadAsync(Buffer, BufferEnd, Buffer.Length - BufferEnd, CancellationToken));
+				}
 				else
 				{
-					for(int Idx = 0; Idx < Data.Count; Idx++)
+					ReadTask = null;
+				}
+
+				// Try to process all the existing data
+				for (; ; )
+				{
+					int InitialBufferPos = BufferPos;
+
+					if (!ReadRecord(Record.Rows))
 					{
-						Result.AppendFormat("{0}{1:x2}", ((Idx & 31) == 0)? "\n    " : " ", Data[Idx]);
+						BufferPos = InitialBufferPos;
+						break;
 					}
+
+					HandleRecord(Record);
+				}
+
+				// If we're at the end of the stream, check we parsed everything
+				if (bEndOfStream)
+				{
+					if (BufferPos == BufferEnd)
+					{
+						break;
+					}
+					else
+					{
+						throw new PerforceException("Unable to parse last chunk of data from response stream." + FormatDataAsHexDump(BufferPos));
+					}
+				}
+
+				// Wait for the read to complete
+				if (ReadTask != null)
+				{
+					int ReadSize = await ReadTask;
+					if (ReadSize == 0)
+					{
+						bEndOfStream = true;
+					}
+					BufferEnd += ReadSize;
 				}
 			}
-			return Result.ToString();
 		}
 
-		public bool TryReadRecord(List<KeyValuePair<string, object>> Record)
+		/// <summary>
+		/// Reads from the buffer into a record object
+		/// </summary>
+		/// <param name="Rows">List of rows to read into</param>
+		/// <returns>True if a record could be read; false if more data is required</returns>
+		bool ReadRecord(List<KeyValuePair<ReadOnlyUtf8String, PerforceValue>> Rows)
 		{
-			// Check if we've reached the end of the stream. This is the only condition where we return false.
-			if(IsEndOfStream())
+			Rows.Clear();
+
+			// Check we can read the initial record marker
+			if (BufferPos >= BufferEnd)
 			{
 				return false;
 			}
-
-			// Check that a dictionary follows
-			byte Temp = ReadByte();
-			if(Temp != '{')
+			if (Buffer[BufferPos] != '{')
 			{
-				throw new PerforceException("Unexpected data while parsing marshaled output - expected '{{', got: {0}", FormatUnexpectedData(Temp));
+				throw new PerforceException("Invalid record start");
 			}
+			BufferPos++;
 
-			// Read all the fields in the record
-			Record.Clear();
-			for(;;)
+			// Capture the start of the string
+			int StartPos = BufferPos;
+			for (; ; )
 			{
-				// Read the next field type. Perforce only outputs string records. A '0' character indicates the end of the dictionary.
-				byte KeyFieldType = ReadByte();
-				if(KeyFieldType == '0')
+				// Check that we've got a string field
+				if (BufferPos >= BufferEnd)
+				{
+					return false;
+				}
+
+				// If this is the end of the record, break out
+				byte KeyType = Buffer[BufferPos++];
+				if (KeyType == '0')
 				{
 					break;
 				}
-				else if(KeyFieldType != 's')
+				else if (KeyType != 's')
+				{
+					throw new PerforceException("Unexpected key field type while parsing marshalled output ({0}) - expected 's', got: {1}", (int)KeyType, FormatDataAsHexDump(BufferPos - 1));
+				}
+
+				// Read the tag
+				ReadOnlyUtf8String Key;
+				if (!TryReadString(out Key))
 				{
-					throw new PerforceException("Unexpected key field type while parsing marshalled output ({0}) - expected 's', got: {1}", (int)KeyFieldType, FormatUnexpectedData(KeyFieldType));
+					return false;
 				}
 
-				// Read the key
-				string Key = ReadString();
+				// Remember the start of the value
+				int ValueOffset = BufferPos;
 
-				// Read the value type.
-				byte ValueFieldType = ReadByte();
-				if(ValueFieldType == 'i')
+				// Read the value type
+				byte ValueType;
+				if (!TryReadByte(out ValueType))
 				{
-					// An integer
-					Record.Add(new KeyValuePair<string, object>(Key, ReadInt32()));
+					return false;
 				}
-				else if(ValueFieldType == 's')
+
+				// Parse the appropriate value
+				PerforceValue Value;
+				if (ValueType == 's')
 				{
-					// A string (or variable length data)
-					byte[] Value = ReadBytes();
-					if(Record.Count > 0 && Record[0].Key == "code" && ((string)Record[0].Value == "binary" || (string)Record[0].Value == "text"))
+					ReadOnlyUtf8String String;
+					if (!TryReadString(out String))
 					{
-						Record.Add(new KeyValuePair<string, object>(Key, Value));
+						return false;
 					}
-					else
+					Value = new PerforceValue(Buffer.AsMemory(ValueOffset, BufferPos - ValueOffset));
+				}
+				else if (ValueType == 'i')
+				{
+					int Integer;
+					if(!TryReadInt(out Integer))
 					{
-						Record.Add(new KeyValuePair<string, object>(Key, Encoding.UTF8.GetString(Value)));
+						return false;
 					}
+					Value = new PerforceValue(Buffer.AsMemory(ValueOffset, BufferPos - ValueOffset));
 				}
 				else
 				{
-					throw new PerforceException("Unexpected value field type while parsing marshalled output ({0}) - expected 's', got: {1}", (int)ValueFieldType, FormatUnexpectedData(ValueFieldType));
+					throw new PerforceException("Unrecognized value type {0}", ValueType);
 				}
+
+				// Construct the response object with the record
+				Rows.Add(KeyValuePair.Create(Key, Value));
 			}
 			return true;
 		}
 
-		public byte ReadByte()
+		/// <summary>
+		/// Reads a response object from the buffer
+		/// </summary>
+		/// <param name="StatRecordInfo">The type of record expected to parse from the response</param>
+		/// <param name="Response">Receives the response object on success</param>
+		/// <returns>True if a response was read, false if the buffer needs more data</returns>
+		bool TryReadResponse(CachedRecordInfo? StatRecordInfo, [NotNullWhen(true)] out PerforceResponse? Response)
 		{
-			UpdateBuffer(1);
+			if (BufferPos + RecordPrefix.Length + 4 > BufferEnd)
+			{
+				Response = null;
+				return false;
+			}
+
+			Span<byte> Prefix = Buffer.AsSpan(BufferPos, RecordPrefix.Length);
+			if (!Prefix.SequenceEqual(RecordPrefix))
+			{
+				throw new PerforceException("Expected 'code' field at the start of record");
+			}
+			BufferPos += Prefix.Length;
+
+			ReadOnlyUtf8String Code;
+			if (!TryReadString(out Code))
+			{
+				Response = null;
+				return false;
+			}
+
+			// Dispatch it to the appropriate handler
+			object? Record;
+			if (Code == ReadOnlyUtf8StringConstants.Stat && StatRecordInfo != null)
+			{
+				if(!TryReadTypedRecord(ReadOnlyUtf8String.Empty, StatRecordInfo, out Record))
+				{
+					Response = null;
+					return false;
+				}
+			}
+			else if (Code == ReadOnlyUtf8StringConstants.Info)
+			{
+				if(!TryReadTypedRecord(ReadOnlyUtf8String.Empty, PerforceReflection.InfoRecordInfo, out Record))
+				{
+					Response = null;
+					return false;
+				}
+			}
+			else if (Code == ReadOnlyUtf8StringConstants.Error)
+			{
+				if (!TryReadTypedRecord(ReadOnlyUtf8String.Empty, PerforceReflection.ErrorRecordInfo, out Record))
+				{
+					Response = null;
+					return false;
+				}
+			}
+			else
+			{
+				throw new PerforceException("Unknown return code for record: {0}", Code);
+			}
 
-			byte NextByte = Buffer[BufferPos++];
-			return NextByte;
+			// Skip over the record terminator
+			if(BufferPos >= BufferEnd || Buffer[BufferPos] != '0')
+			{
+				throw new PerforceException("Unexpected record terminator");
+			}
+			BufferPos++;
+
+			// Create the response
+			Response = new PerforceResponse(Record);
+			return true;
 		}
 
-		public int ReadInt32()
+		/// <summary>
+		/// Parse an individual record from the server.
+		/// </summary>
+		/// <param name="RequiredSuffix">The required suffix for any subobject arrays.</param>
+		/// <param name="RecordInfo">Reflection information for the type being serialized into.</param>
+		/// <returns>The parsed object.</returns>
+		bool TryReadTypedRecord(ReadOnlyUtf8String RequiredSuffix, CachedRecordInfo RecordInfo, [NotNullWhen(true)] out object? Record)
 		{
-			UpdateBuffer(4);
+			// Create a bitmask for all the required tags
+			ulong RequiredTagsBitMask = 0;
 
-			int Value = Buffer[BufferPos + 0] | (Buffer[BufferPos + 1] << 8) | (Buffer[BufferPos + 2] << 16) | (Buffer[BufferPos + 3] << 24);
-			BufferPos += 4;
-			return Value;
+			// Create the new record
+			object? NewRecord = RecordInfo.CreateInstance();
+			if (NewRecord == null)
+			{
+				throw new InvalidDataException($"Unable to construct record of type {RecordInfo.Type}");
+			}
+
+			// Get the record info, and parse it into the object
+			for(; ;)
+			{
+				// Check that we've got a string field
+				if(BufferPos >= BufferEnd)
+				{
+					Record = null;
+					return false;
+				}
+
+				// If this is the end of the record, break out
+				byte KeyType = Buffer[BufferPos];
+				if (KeyType == '0')
+				{
+					break;
+				}
+				else if (KeyType != 's')
+				{
+					throw new PerforceException("Unexpected key field type while parsing marshalled output ({0}) - expected 's', got: {1}", (int)KeyType, FormatDataAsHexDump(BufferPos));
+				}
+
+				// Capture the initial buffer position, in case we have to roll back
+				int StartBufferPos = BufferPos;
+				BufferPos++;
+
+				// Read the tag
+				ReadOnlyUtf8String Tag;
+				if(!TryReadString(out Tag))
+				{
+					Record = null;
+					return false;
+				}
+
+				// Find the start of the array suffix
+				int SuffixIdx = Tag.Length;
+				while (SuffixIdx > 0 && (Tag[SuffixIdx - 1] == (byte)',' || (Tag[SuffixIdx - 1] >= '0' && Tag[SuffixIdx - 1] <= '9')))
+				{
+					SuffixIdx--;
+				}
+
+				// Separate the key into tag and suffix
+				ReadOnlyUtf8String Suffix = Tag.Slice(SuffixIdx);
+				Tag = Tag.Slice(0, SuffixIdx);
+
+				// Try to find the matching field
+				CachedTagInfo? TagInfo;
+				if (RecordInfo.NameToInfo.TryGetValue(Tag, out TagInfo))
+				{
+					RequiredTagsBitMask |= TagInfo.RequiredTagBitMask;
+				}
+
+				// Check whether it's a subobject or part of the current object.
+				if (Suffix == RequiredSuffix)
+				{
+					if(!TryReadValue(NewRecord, TagInfo))
+					{
+						Record = null;
+						return false;
+					}
+				}
+				else if (Suffix.StartsWith(RequiredSuffix) && (RequiredSuffix.Length == 0 || Suffix[RequiredSuffix.Length] == ','))
+				{
+					// Part of a subobject. If this record doesn't have any listed subobject type, skip the field and continue.
+					if (TagInfo != null)
+					{
+						// Get the list field
+						System.Collections.IList? List = (System.Collections.IList?)TagInfo.Field.GetValue(NewRecord);
+						if (List == null)
+						{
+							throw new PerforceException($"Empty list for {TagInfo.Field.Name}");
+						}
+
+						// Check the suffix matches the index of the next element
+						if (!IsCorrectIndex(Suffix, RequiredSuffix, List.Count))
+						{
+							throw new PerforceException("Subobject element received out of order: got {0}", Suffix);
+						}
+
+						// Add it to the list
+						if (!TryReadValue(NewRecord, TagInfo))
+						{
+							Record = null;
+							return false;
+						}
+					}
+					else if (RecordInfo.SubElementField != null)
+					{
+						// Move back to the start of this tag
+						BufferPos = StartBufferPos;
+
+						// Get the list field
+						System.Collections.IList? List = (System.Collections.IList?)RecordInfo.SubElementField.GetValue(NewRecord);
+						if (List == null)
+						{
+							throw new PerforceException($"Invalid field for {RecordInfo.SubElementField.Name}");
+						}
+
+						// Check the suffix matches the index of the next element
+						if (!IsCorrectIndex(Suffix, RequiredSuffix, List.Count))
+						{
+							throw new PerforceException("Subobject element received out of order: got {0}", Suffix);
+						}
+
+						// Parse the subobject and add it to the list
+						object? SubRecord;
+						if (!TryReadTypedRecord(Suffix, RecordInfo.SubElementRecordInfo!, out SubRecord))
+						{
+							Record = null;
+							return false;
+						}
+						List.Add(SubRecord);
+					}
+					else
+					{
+						// Just discard the value
+						if (!TryReadValue(NewRecord, TagInfo))
+						{
+							Record = null;
+							return false;
+						}
+					}
+				}
+				else
+				{
+					// Roll back
+					BufferPos = StartBufferPos;
+					break;
+				}
+			}
+
+			// Make sure we've got all the required tags we need
+			if (RequiredTagsBitMask != RecordInfo.RequiredTagsBitMask)
+			{
+				string MissingTagNames = String.Join(", ", RecordInfo.NameToInfo.Where(x => (RequiredTagsBitMask | x.Value.RequiredTagBitMask) != RequiredTagsBitMask).Select(x => x.Key));
+				throw new PerforceException("Missing '{0}' tag when parsing '{1}'", MissingTagNames, RecordInfo.Type.Name);
+			}
+
+			// Construct the response object with the record
+			Record = NewRecord;
+			return true;
 		}
 
-		public string ReadString()
+		/// <summary>
+		/// Reads a value from the input buffer
+		/// </summary>
+		/// <param name="NewRecord">The new record</param>
+		/// <param name="TagInfo">The current tag</param>
+		/// <returns></returns>
+		bool TryReadValue(object NewRecord, CachedTagInfo? TagInfo)
 		{
-			byte[] Bytes = ReadBytes();
-			return Encoding.UTF8.GetString(Bytes);
+			// Read the value type
+			byte ValueType;
+			if (!TryReadByte(out ValueType))
+			{
+				return false;
+			}
+
+			// Parse the appropriate value
+			if (ValueType == 's')
+			{
+				ReadOnlyUtf8String String;
+				if (!TryReadString(out String))
+				{
+					return false;
+				}
+				if (TagInfo != null)
+				{
+					TagInfo.SetFromString(NewRecord, String);
+				}
+			}
+			else if (ValueType == 'i')
+			{
+				int Integer;
+				if (!TryReadInt(out Integer))
+				{
+					return false;
+				}
+				if (TagInfo != null)
+				{
+					TagInfo.SetFromInteger(NewRecord, Integer);
+				}
+			}
+			else
+			{
+				throw new PerforceException("Unrecognized value type {0}", ValueType);
+			}
+
+			return true;
 		}
 
-		public byte[] ReadBytes()
+		/// <summary>
+		/// Attempts to read a single byte from the buffer
+		/// </summary>
+		/// <param name="Value">Receives the byte that was read</param>
+		/// <returns>True if a byte was read from the buffer, false if there was not enough data</returns>
+		bool TryReadByte(out byte Value)
 		{
-			int MaxNumBytes = ReadInt32();
+			if (BufferPos >= BufferEnd)
+			{
+				Value = 0;
+				return false;
+			}
 
-			byte[] Buffer = new byte[MaxNumBytes];
-			CopyBytes(MaxNumBytes, new MemoryStream(Buffer, true));
-			return Buffer;
+			Value = Buffer[BufferPos];
+			BufferPos++;
+			return true;
 		}
 
-		public void CopyBytes(int NumBytes, Stream OutputStream)
+		/// <summary>
+		/// Attempts to read a single int from the buffer
+		/// </summary>
+		/// <param name="Value">Receives the value that was read</param>
+		/// <returns>True if an int was read from the buffer, false if there was not enough data</returns>
+		bool TryReadInt(out int Value)
 		{
-			while(NumBytes > 0)
+			if(BufferPos + 4 > BufferEnd)
 			{
-				UpdateBuffer(1);
-				int NumBytesThisLoop = Math.Min(NumBytes, BufferEnd - BufferPos);
-				OutputStream.Write(Buffer, BufferPos, NumBytesThisLoop);
-				BufferPos += NumBytesThisLoop;
-				NumBytes -= NumBytesThisLoop;
+				Value = 0;
+				return false;
 			}
+
+			Value = Buffer[BufferPos + 0] | (Buffer[BufferPos + 1] << 8) | (Buffer[BufferPos + 2] << 16) | (Buffer[BufferPos + 3] << 24);
+			BufferPos += 4;
+			return true;
 		}
 
-		private void UpdateBuffer(int MinSize)
+		/// <summary>
+		/// Attempts to read a string with type indicator from the buffer
+		/// </summary>
+		/// <param name="String">Receives the value that was read</param>
+		/// <returns>True if a string was read from the buffer, false if there was not enough data</returns>
+		bool TryReadStringWithType(out ReadOnlyUtf8String String)
 		{
-			if(!TryUpdateBuffer(MinSize))
+			byte ValueType;
+			if(!TryReadByte(out ValueType))
+			{
+				String = new ReadOnlyUtf8String();
+				return false;
+			}
+			if (ValueType != 's')
 			{
-				throw new PerforceException("Unexpected end of stream");
+				throw new PerforceException("Expected string value, got '{0}'", ValueType);
 			}
+			return TryReadString(out String);
 		}
 
-		private bool TryUpdateBuffer(int MinSize)
+		/// <summary>
+		/// Attempts to read a string from the buffer
+		/// </summary>
+		/// <param name="String">Receives the value that was read</param>
+		/// <returns>True if a string was read from the buffer, false if there was not enough data</returns>
+		bool TryReadString(out ReadOnlyUtf8String String)
 		{
-			while(BufferPos + MinSize > BufferEnd)
+			int Length;
+			if(!TryReadInt(out Length))
 			{
-				if(BufferPos > 0)
-				{
-					Array.Copy(Buffer, BufferPos, Buffer, 0, BufferEnd - BufferPos);
-					BufferEnd -= BufferPos;
-					BufferPos = 0;
-				}
-
-				int ReadSize = ChildProcess.Read(Buffer, BufferEnd, Buffer.Length - BufferEnd);
-				if(ReadSize == 0)
-				{
-					if(BufferPos == BufferEnd)
-					{
-						return false;
-					}
-					else
-					{
-						throw new PerforceException("Partially read record");
-					}
-				}
+				String = new ReadOnlyUtf8String();
+				return false;
+			}
 
-				BufferEnd += ReadSize;
+			if(BufferPos + Length > BufferEnd)
+			{
+				String = new ReadOnlyUtf8String();
+				return false;
 			}
+
+			String = new ReadOnlyUtf8String(Buffer, BufferPos, Length);
+			BufferPos += Length;
 			return true;
 		}
+
+		/// <summary>
+		/// Determines if the given text contains the expected prefix followed by an array index
+		/// </summary>
+		/// <param name="Text">The text to check</param>
+		/// <param name="Prefix">The required prefix</param>
+		/// <param name="Index">The required index</param>
+		/// <returns>True if the index is correct</returns>
+		static bool IsCorrectIndex(ReadOnlyUtf8String Text, ReadOnlyUtf8String Prefix, int Index)
+		{
+			if (Prefix.Length > 0)
+			{
+				return Text.StartsWith(Prefix) && Text.Length > Prefix.Length && Text[Prefix.Length] == (byte)',' && IsCorrectIndex(Text.Span.Slice(Prefix.Length + 1), Index);
+			}
+			else
+			{
+				return IsCorrectIndex(Text.Span, Index);
+			}
+		}
+
+		/// <summary>
+		/// Determines if the given text matches the expected array index
+		/// </summary>
+		/// <param name="Span">The text to check</param>
+		/// <param name="ExpectedIndex">The expected array index</param>
+		/// <returns>True if the span matches</returns>
+		static bool IsCorrectIndex(ReadOnlySpan<byte> Span, int ExpectedIndex)
+		{
+			int Index;
+			int BytesConsumed;
+			return Utf8Parser.TryParse(Span, out Index, out BytesConsumed) && BytesConsumed == Span.Length && Index == ExpectedIndex;
+		}
 	}
 }
diff --git a/Engine/Source/Programs/Shared/EpicGames.Perforce/PerforceClientConnection.cs b/Engine/Source/Programs/Shared/EpicGames.Perforce/PerforceClientConnection.cs
new file mode 100644
--- /dev/null
+++ b/Engine/Source/Programs/Shared/EpicGames.Perforce/PerforceClientConnection.cs
@@ -0,0 +1,53 @@
+using Microsoft.Extensions.Logging;
+using System;
+using System.Collections.Generic;
+using System.Text;
+
+namespace EpicGames.Perforce
+{
+	/// <summary>
+	/// An instance of PerforceConnection which is strongly typed as having a valid client name
+	/// </summary>
+	public class PerforceClientConnection : PerforceConnection
+	{
+		/// <summary>
+		/// Accessor for the client name.
+		/// </summary>
+		public new string ClientName
+		{
+			get { return base.ClientName!; }
+		}
+
+		/// <summary>
+		/// Constructor
+		/// </summary>
+		/// <param name="Perforce">The Perforce connection</param>
+		/// <param name="ClientName">The client name</param>
+		public PerforceClientConnection(PerforceConnection Perforce, string ClientName)
+			: base(Perforce.ServerAndPort, Perforce.UserName, ClientName, Perforce.Logger)
+		{
+		}
+
+		/// <summary>
+		/// Constructor
+		/// </summary>
+		/// <param name="ServerAndPort">The server and port address</param>
+		/// <param name="UserName">Username to connect with</param>
+		/// <param name="ClientName">The client name to use</param>
+		/// <param name="Logger">Logging device</param>
+		public PerforceClientConnection(string? ServerAndPort, string? UserName, string ClientName, ILogger Logger) 
+			: base(ServerAndPort, UserName, ClientName, Logger)
+		{
+		}
+
+		/// <summary>
+		/// Remove the client parameter for the connection
+		/// </summary>
+		/// <returns>Bare connection to the server</returns>
+		public PerforceConnection WithoutClient()
+		{
+			return new PerforceConnection(this) { ClientName = null };
+		}
+	}
+}
+
diff --git a/Engine/Source/Programs/Shared/EpicGames.Perforce/PerforceConnection.cs b/Engine/Source/Programs/Shared/EpicGames.Perforce/PerforceConnection.cs
--- a/Engine/Source/Programs/Shared/EpicGames.Perforce/PerforceConnection.cs
+++ b/Engine/Source/Programs/Shared/EpicGames.Perforce/PerforceConnection.cs
@@ -1,5 +1,7 @@
-// Copyright Epic Games, Inc. All Rights Reserved.
+// Copyright 1998-2019 Epic Games, Inc. All Rights Reserved.
 
+using EpicGames.Core;
+using Microsoft.Extensions.Logging;
 using System;
 using System.Collections.Concurrent;
 using System.Collections.Generic;
@@ -8,6 +10,8 @@ using System.IO;
 using System.Linq;
 using System.Reflection;
 using System.Text;
+using System.Text.RegularExpressions;
+using System.Threading;
 using System.Threading.Tasks;
 
 namespace EpicGames.Perforce
@@ -17,112 +21,94 @@ namespace EpicGames.Perforce
 	/// </summary>
 	public class PerforceConnection
 	{
-		/// <summary>
-		/// Constant for the default changelist, where valid.
-		/// </summary>
-		public const int DefaultChange = -2;
-
 		#region Plumbing
 
 		/// <summary>
-		/// Stores cached information about a field with a P4Tag attribute
+		/// The current server and port
 		/// </summary>
-		class CachedTagInfo
+		public string? ServerAndPort
 		{
-			/// <summary>
-			/// Name of the tag. Specified in the attribute or inferred from the field name.
-			/// </summary>
-			public string Name;
-
-			/// <summary>
-			/// Whether this tag is optional or not.
-			/// </summary>
-			public bool Optional;
-
-			/// <summary>
-			/// The field containing the value of this data.
-			/// </summary>
-			public FieldInfo Field;
-
-			/// <summary>
-			/// Index into the bitmask of required types
-			/// </summary>
-			public ulong RequiredTagBitMask;
+			get;
+			set;
 		}
 
 		/// <summary>
-		/// Stores cached information about a record
+		/// The current user name
 		/// </summary>
-		class CachedRecordInfo
+		public string? UserName
 		{
-			/// <summary>
-			/// Type of the record
-			/// </summary>
-			public Type Type;
-
-			/// <summary>
-			/// Map of tag names to their cached reflection information
-			/// </summary>
-			public Dictionary<string, CachedTagInfo> TagNameToInfo = new Dictionary<string, CachedTagInfo>();
-
-			/// <summary>
-			/// Bitmask of all the required tags. Formed by bitwise-or'ing the RequiredTagBitMask fields for each required CachedTagInfo.
-			/// </summary>
-			public ulong RequiredTagsBitMask;
-
-			/// <summary>
-			/// The type of records to create for subelements
-			/// </summary>
-			public Type SubElementType;
-
-			/// <summary>
-			/// The cached record info for the subelement type
-			/// </summary>
-			public CachedRecordInfo SubElementRecordInfo;
-
-			/// <summary>
-			/// Field containing subelements
-			/// </summary>
-			public FieldInfo SubElementField;
+			get;
+			set;
 		}
 
 		/// <summary>
-		/// Unix epoch; used for converting times back into C# datetime objects
+		/// The current client name
 		/// </summary>
-		static readonly DateTime UnixEpoch = new DateTime(1970, 1, 1, 0, 0, 0, DateTimeKind.Utc);
+		public string? ClientName
+		{
+			get;
+			set;
+		}
 
 		/// <summary>
-		/// Cached map of enum types to a lookup mapping from p4 strings to enum values.
+		/// Name of this application, reported to server through -zprog arguments
 		/// </summary>
-		static ConcurrentDictionary<Type, Dictionary<string, int>> EnumTypeToFlags = new ConcurrentDictionary<Type, Dictionary<string, int>>();
+		public string? AppName
+		{
+			get;
+			set;
+		}
 
 		/// <summary>
-		/// Cached set of record 
+		/// Version of this application, reported to server through -zversion arguments
 		/// </summary>
-		static ConcurrentDictionary<Type, CachedRecordInfo> RecordTypeToInfo = new ConcurrentDictionary<Type, CachedRecordInfo>();
+		public string? AppVersion
+		{
+			get;
+			set;
+		}
 
 		/// <summary>
-		/// Default type for info
+		/// Additional options to append to the command line
 		/// </summary>
-		static CachedRecordInfo InfoRecordInfo = GetCachedRecordInfo(typeof(PerforceInfo));
+		public List<string> GlobalOptions { get; } = new List<string>();
 
 		/// <summary>
-		/// Default type for errors
+		/// The logging interface
 		/// </summary>
-		static CachedRecordInfo ErrorRecordInfo = GetCachedRecordInfo(typeof(PerforceError));
+		public ILogger Logger
+		{
+			get;
+			set;
+		}
 
 		/// <summary>
-		/// Global options for each command
+		/// Constructor
 		/// </summary>
-		public readonly string GlobalOptions;
+		/// <param name="ServerAndPort">The server address and port</param>
+		/// <param name="UserName">The user name</param>
+		/// <param name="Logger">Interface for logging</param>
+		public PerforceConnection(string? ServerAndPort, string? UserName, ILogger Logger)
+			: this(ServerAndPort, UserName, null, Logger)
+		{
+		}
 
 		/// <summary>
 		/// Constructor
 		/// </summary>
-		/// <param name="GlobalOptions">Global options to pass to every Perforce command</param>
-		public PerforceConnection(string GlobalOptions)
+		/// <param name="ServerAndPort">The server address and port</param>
+		/// <param name="UserName">The user name</param>
+		/// <param name="ClientName">The client name</param>
+		/// <param name="Logger">Interface for logging</param>
+		public PerforceConnection(string? ServerAndPort, string? UserName, string? ClientName, ILogger Logger)
+			: this(ServerAndPort, UserName, ClientName, null, null, Logger)
 		{
-			this.GlobalOptions = GlobalOptions;
+			AssemblyName EntryAssemblyName = Assembly.GetEntryAssembly()!.GetName();
+			if (EntryAssemblyName.Name != null)
+			{
+				AppName = EntryAssemblyName.Name;
+				AppVersion = EntryAssemblyName.Version?.ToString() ?? String.Empty;
+			}
 		}
 
 		/// <summary>
@@ -131,41 +117,62 @@ namespace EpicGames.Perforce
 		/// <param name="ServerAndPort">The server address and port</param>
 		/// <param name="UserName">The user name</param>
 		/// <param name="ClientName">The client name</param>
-		public PerforceConnection(string ServerAndPort, string UserName, string ClientName)
+		/// <param name="Logger">Interface for logging</param>
+		public PerforceConnection(string? ServerAndPort, string? UserName, string? ClientName, string? AppName, string? AppVersion, ILogger Logger)
+		{
+			this.ServerAndPort = ServerAndPort;
+			this.UserName = UserName;
+			this.ClientName = ClientName;
+			this.AppName = AppName;
+			this.AppVersion = AppVersion;
+			this.Logger = Logger;
+		}
+
+		/// <summary>
+		/// Constructor 
+		/// </summary>
+		/// <param name="Other">Connection to copy settings from</param>
+		/// <param name="ClientName">The client name</param>
+		public PerforceConnection(PerforceConnection Other)
+			: this(Other.ServerAndPort, Other.UserName, Other.ClientName, Other.AppName, Other.AppVersion, Other.Logger)
 		{
-			List<string> Options = new List<string>();
-			if(ServerAndPort != null)
+			GlobalOptions.AddRange(Other.GlobalOptions);
+		}
+
+		/// <summary>
+		/// Gets all the arguments
+		/// </summary>
+		/// <returns></returns>
+		string GetFullCommandLine(string CommandLine)
+		{
+			StringBuilder Result = new StringBuilder();
+			if (ServerAndPort != null)
 			{
-				Options.Add(String.Format("-p {0}", ServerAndPort));
+				Result.AppendFormat("-p {0} ", ServerAndPort);
 			}
 			if (UserName != null)
 			{
-				Options.Add(String.Format("-u {0}", UserName));
+				Result.AppendFormat("-u {0} ", UserName);
 			}
 			if (ClientName != null)
 			{
-				Options.Add(String.Format("-c {0}", ClientName));
+				Result.AppendFormat("-c {0} ", ClientName);
+			}
+			if (AppName != null)
+			{
+				Result.AppendFormat("-zprog={0} ", AppName);
+			}
+			if (AppVersion != null)
+			{
+				Result.AppendFormat("-zversion={0} ", AppVersion);
 			}
-			this.GlobalOptions = String.Join(" ", Options);
-		}
 
-		/// <summary>
-		/// Execute a Perforce command and parse the output as marshalled Python objects. This is more robustly defined than the text-based tagged output
-		/// format, because it avoids ambiguity when returned fields can have newlines.
-		/// </summary>
-		/// <param name="CommandLine">Command line to execute Perforce with</param>
-		/// <param name="InputData">Input data to pass to the Perforce server. May be null.</param>
-		/// <param name="RecordHandler">Handler for each received record.</param>
-		public void Command(string CommandLine, byte[] InputData, Action<List<KeyValuePair<string, object>>> RecordHandler)
-		{
-			using(PerforceChildProcess Process = new PerforceChildProcess(InputData, "{0} {1}", GlobalOptions, CommandLine))
+			foreach (string GlobalOption in GlobalOptions)
 			{
-				List<KeyValuePair<string, object>> Record = new List<KeyValuePair<string, object>>();
-				while(Process.TryReadRecord(Record))
-				{
-					RecordHandler(Record);
-				}
+				Result.AppendFormat("{0} ", GlobalOption);
 			}
+			Result.Append(CommandLine);
+			return Result.ToString();
 		}
 
 		/// <summary>
@@ -179,7 +186,7 @@ namespace EpicGames.Perforce
 			using (BinaryWriter Writer = new BinaryWriter(Stream))
 			{
 				Writer.Write((byte)'{');
-				foreach(KeyValuePair<string, object> KeyValuePair in KeyValuePairs)
+				foreach (KeyValuePair<string, object> KeyValuePair in KeyValuePairs)
 				{
 					Writer.Write('s');
 					byte[] KeyBytes = Encoding.UTF8.GetBytes(KeyValuePair.Key);
@@ -206,49 +213,18 @@ namespace EpicGames.Perforce
 		/// <summary>
 		/// Execute a command and parse the response
 		/// </summary>
-		/// <param name="Arguments">Arguments for the command</param>
+		/// <param name="CommandLine">Arguments for the command</param>
 		/// <param name="InputData">Input data to pass to Perforce</param>
 		/// <param name="StatRecordType">The type of records to return for "stat" responses</param>
+		/// <param name="CancellationToken">Token used to cancel the operation</param>
 		/// <returns>List of objects returned by the server</returns>
-		public List<PerforceResponse> Command(string Arguments, byte[] InputData, Type StatRecordType)
+		public async Task<List<PerforceResponse>> CommandAsync(string CommandLine, byte[]? InputData, Type? StatRecordType, CancellationToken CancellationToken)
 		{
-			CachedRecordInfo StatRecordInfo = (StatRecordType == null)? null : GetCachedRecordInfo(StatRecordType);
-
-			List<PerforceResponse> Responses = new List<PerforceResponse>();
-			Action<List<KeyValuePair<string, object>>> Handler = (KeyValuePairs) =>
+			using (PerforceChildProcess Process = new PerforceChildProcess(InputData, GetFullCommandLine(CommandLine), Logger))
 			{
-				if(KeyValuePairs.Count == 0)
-				{
-					throw new PerforceException("Unexpected empty record returned by Perforce.");
-				}
-				if(KeyValuePairs[0].Key != "code")
-				{
-					throw new PerforceException("Expected first returned field to be 'code'");
-				}
-
-				string Code = KeyValuePairs[0].Value as string;
-
-				int Idx = 1;
-				if (Code == "stat" && StatRecordType != null)
-				{
-					Responses.Add(ParseResponse(KeyValuePairs, ref Idx, "", StatRecordInfo));
-				}
-				else if (Code == "info")
-				{
-					Responses.Add(ParseResponse(KeyValuePairs, ref Idx, "", InfoRecordInfo));
-				}
-				else if(Code == "error")
-				{
-					Responses.Add(ParseResponse(KeyValuePairs, ref Idx, "", ErrorRecordInfo));
-				}
-				else
-				{
-					throw new PerforceException("Unknown return code for record: {0}", KeyValuePairs[0].Value);
-				}
-			};
-			Command(Arguments, InputData, Handler);
-
-			return Responses;
+				CachedRecordInfo? StatRecordInfo = (StatRecordType == null) ? null : PerforceReflection.GetCachedRecordInfo(StatRecordType);
+				return await Process.ReadResponsesAsync(StatRecordInfo, CancellationToken);
+			}
 		}
 
 		/// <summary>
@@ -256,10 +232,11 @@ namespace EpicGames.Perforce
 		/// </summary>
 		/// <param name="Arguments">Arguments for the command</param>
 		/// <param name="InputData">Input data to pass to Perforce</param>
+		/// <param name="CancellationToken">Token used to cancel the operation</param>
 		/// <returns>List of objects returned by the server</returns>
-		public PerforceResponseList<T> Command<T>(string Arguments, byte[] InputData) where T : class
+		public async Task<PerforceResponseList<T>> CommandAsync<T>(string Arguments, byte[]? InputData, CancellationToken CancellationToken) where T : class
 		{
-			List<PerforceResponse> Responses = Command(Arguments, InputData, typeof(T));
+			List<PerforceResponse> Responses = await CommandAsync(Arguments, InputData, typeof(T), CancellationToken);
 
 			PerforceResponseList<T> TypedResponses = new PerforceResponseList<T>();
 			foreach (PerforceResponse Response in Responses)
@@ -270,349 +247,115 @@ namespace EpicGames.Perforce
 		}
 
 		/// <summary>
-		/// Attempts to execute the given command, returning the results from the server or the first PerforceResponse object.
+		/// Execute a command and parse the response
 		/// </summary>
-		/// <param name="Arguments">Arguments for the command.</param>
-		/// <param name="InputData">Input data for the command.</param>
-		/// <param name="StatRecordType">Type of element to return in the response</param>
-		/// <returns>Response from the server; either an object of type T or error.</returns>
-		public PerforceResponse SingleResponseCommand(string Arguments, byte[] InputData, Type StatRecordType)
+		/// <param name="CommonArguments">Arguments for the command</param>
+		/// <param name="BatchedArguments">Arguments to pass to the command in batches</param>
+		/// <param name="InputData">Input data to pass to Perforce</param>
+		/// <param name="CancellationToken">Token used to cancel the operation</param>
+		/// <returns>List of objects returned by the server</returns>
+		public async Task<PerforceResponseList<T>> BatchedCommandAsync<T>(string CommonArguments, string[] BatchedArguments, byte[]? InputData, CancellationToken CancellationToken) where T : class
 		{
-			List<PerforceResponse> Responses = Command(Arguments, InputData, StatRecordType);
-			if(Responses.Count != 1)
+			PerforceResponseList<T> Responses = new PerforceResponseList<T>();
+			for (int FileSpecIdx = 0; FileSpecIdx < BatchedArguments.Length;)
 			{
-				throw new PerforceException("Expected one result from 'p4 {0}', got {1}", Arguments, Responses.Count);
+				StringBuilder Arguments = new StringBuilder(CommonArguments);
+				for (; FileSpecIdx < BatchedArguments.Length && Arguments.Length < 4096; FileSpecIdx++)
+				{
+					Arguments.AppendFormat(" \"{0}\"", BatchedArguments[FileSpecIdx]);
+				}
+				Responses.AddRange(await CommandAsync<T>(Arguments.ToString(), null, CancellationToken));
 			}
-			return Responses[0];
-		}
-
-		/// <summary>
-		/// Attempts to execute the given command, returning the results from the server or the first PerforceResponse object.
-		/// </summary>
-		/// <typeparam name="T">Type of record to parse</typeparam>
-		/// <param name="Arguments">Arguments for the command.</param>
-		/// <param name="InputData">Input data for the command.</param>
-		/// <returns>Response from the server; either an object of type T or error.</returns>
-		public PerforceResponse<T> SingleResponseCommand<T>(string Arguments, byte[] InputData) where T : class
-		{
-			return new PerforceResponse<T>(SingleResponseCommand(Arguments, InputData, typeof(T)));
+			return Responses;
 		}
 
 		/// <summary>
-		/// Parse an individual record from the server.
+		/// Execute a command and parse the response
 		/// </summary>
-		/// <param name="KeyValuePairs">List of tagged values returned by the server.</param>
-		/// <param name="Idx">Index of the first tagged value to parse.</param>
-		/// <param name="RequiredSuffix">The required suffix for any subobject arrays.</param>
-		/// <param name="RecordInfo">Reflection information for the type being serialized into.</param>
-		/// <returns>The parsed object.</returns>
-		PerforceResponse ParseResponse(List<KeyValuePair<string, object>> KeyValuePairs, ref int Idx, string RequiredSuffix, CachedRecordInfo RecordInfo)
+		/// <param name="CommandLine">Arguments for the command</param>
+		/// <param name="InputData">Input data to pass to Perforce</param>
+		/// <param name="HandleRecord">Delegate used to handle each record</param>
+		/// <param name="CancellationToken">Token used to cancel the operation</param>
+		/// <returns>List of objects returned by the server</returns>
+		public async Task RecordCommandAsync(string CommandLine, byte[]? InputData, Action<PerforceRecord> HandleRecord, CancellationToken CancellationToken)
 		{
-			// Create a bitmask for all the required tags
-			ulong RequiredTagsBitMask = 0;
-
-			// Get the record info, and parse it into the object
-			object NewRecord = Activator.CreateInstance(RecordInfo.Type);
-			while(Idx < KeyValuePairs.Count)
+			using (PerforceChildProcess Process = new PerforceChildProcess(InputData, GetFullCommandLine(CommandLine), Logger))
 			{
-				// Split out the tag and value
-				string Tag = KeyValuePairs[Idx].Key;
-				string Value = KeyValuePairs[Idx].Value.ToString();
-
-				// Parse the suffix from the current key
-				int SuffixIdx = Tag.Length;
-				while(SuffixIdx > 0 && (Tag[SuffixIdx - 1] == ',' || (Tag[SuffixIdx - 1] >= '0' && Tag[SuffixIdx - 1] <= '9')))
-				{
-					SuffixIdx--;
-				}
-
-				// Split out the suffix
-				string Suffix = Tag.Substring(SuffixIdx);
-				Tag = Tag.Substring(0, SuffixIdx);
-
-				// Check whether it's a subobject or part of the current object.
-				if (Suffix == RequiredSuffix)
-				{
-					// Part of the current object
-					CachedTagInfo TagInfo;
-					if (RecordInfo.TagNameToInfo.TryGetValue(Tag, out TagInfo))
-					{
-						FieldInfo FieldInfo = TagInfo.Field;
-						if (FieldInfo.FieldType == typeof(DateTime))
-						{
-							DateTime Time;
-							if(!DateTime.TryParse(Value, out Time))
-							{
-								Time = UnixEpoch + TimeSpan.FromSeconds(long.Parse(Value));
-							}
-							FieldInfo.SetValue(NewRecord, Time);
-						}
-						else if (FieldInfo.FieldType == typeof(bool))
-						{
-							FieldInfo.SetValue(NewRecord, Value.Length == 0 || Value == "true");
-						}
-						else if(FieldInfo.FieldType == typeof(Nullable<bool>))
-						{
-							FieldInfo.SetValue(NewRecord, Value == "true");
-						}
-						else if (FieldInfo.FieldType == typeof(int))
-						{
-							if(Value == "new" || Value == "none")
-							{
-								FieldInfo.SetValue(NewRecord, -1);
-							}
-							else if(Value.StartsWith("#"))
-							{
-								FieldInfo.SetValue(NewRecord, (Value == "#none") ? 0 : int.Parse(Value.Substring(1)));
-							}
-							else if(Value == "default")
-							{
-								FieldInfo.SetValue(NewRecord, DefaultChange);
-							}
-							else
-							{
-								FieldInfo.SetValue(NewRecord, int.Parse(Value));
-							}
-						}
-						else if (FieldInfo.FieldType == typeof(long))
-						{
-							FieldInfo.SetValue(NewRecord, long.Parse(Value));
-						}
-						else if (FieldInfo.FieldType == typeof(string))
-						{
-							FieldInfo.SetValue(NewRecord, Value);
-						}
-						else if(FieldInfo.FieldType.IsEnum)
-						{
-							FieldInfo.SetValue(NewRecord, ParseEnum(FieldInfo.FieldType, Value));
-						}
-						else
-						{
-							throw new PerforceException("Unsupported type of {0}.{1} for tag '{0}'", RecordInfo.Type.Name, FieldInfo.FieldType.Name, Tag);
-						}
-						RequiredTagsBitMask |= TagInfo.RequiredTagBitMask;
-					}
-					Idx++;
-				}
-				else if (Suffix.StartsWith(RequiredSuffix) && (RequiredSuffix.Length == 0 || Suffix[RequiredSuffix.Length] == ','))
-				{
-					// Part of a subobject. If this record doesn't have any listed subobject type, skip the field and continue.
-					if (RecordInfo.SubElementField == null)
-					{
-						CachedTagInfo TagInfo;
-						if (RecordInfo.TagNameToInfo.TryGetValue(Tag, out TagInfo))
-						{
-							FieldInfo FieldInfo = TagInfo.Field;
-							if (FieldInfo.FieldType == typeof(List<string>))
-							{
-								((List<string>)FieldInfo.GetValue(NewRecord)).Add(Value);
-							}
-							else
-							{
-								throw new PerforceException("Unsupported type of {0}.{1} for tag '{0}'", RecordInfo.Type.Name, FieldInfo.FieldType.Name, Tag);
-							}
-							RequiredTagsBitMask |= TagInfo.RequiredTagBitMask;
-						}
-						Idx++;
-					}
-					else
-					{
-						// Get the expected suffix for the next item based on the number of elements already in the list
-						System.Collections.IList List = (System.Collections.IList)RecordInfo.SubElementField.GetValue(NewRecord);
-						string RequiredChildSuffix = (RequiredSuffix.Length == 0) ? String.Format("{0}", List.Count) : String.Format("{0},{1}", RequiredSuffix, List.Count);
-						if (Suffix != RequiredChildSuffix)
-						{
-							throw new PerforceException("Subobject element received out of order; expected {0}{1}, got {0}{2}", Tag, RequiredChildSuffix, Suffix);
-						}
-
-						// Parse the subobject and add it to the list
-						PerforceResponse Response = ParseResponse(KeyValuePairs, ref Idx, RequiredChildSuffix, RecordInfo.SubElementRecordInfo);
-						List.Add(Response.Data);
-					}
-				}
-				else
-				{
-					break;
-				}
-			}
-
-			// Make sure we've got all the required tags we need
-			if (RequiredTagsBitMask != RecordInfo.RequiredTagsBitMask)
-			{
-				string MissingTagNames = String.Join(", ", RecordInfo.TagNameToInfo.Where(x => (RequiredTagsBitMask | x.Value.RequiredTagBitMask) != RequiredTagsBitMask).Select(x => x.Key));
-				throw new PerforceException("Missing '{0}' tag when parsing '{1}'", MissingTagNames, RecordInfo.Type.Name);
+				await Process.ReadRecordsAsync(HandleRecord, CancellationToken);
 			}
-
-			return new PerforceResponse(NewRecord);
 		}
 
 		/// <summary>
-		/// Gets a mapping of flags to enum values for the given type
+		/// Attempts to execute the given command, returning the results from the server or the first PerforceResponse object.
 		/// </summary>
-		/// <param name="EnumType">The enum type to retrieve flags for</param>
-		/// <returns>Map of name to enum value</returns>
-		static Dictionary<string, int> GetCachedEnumFlags(Type EnumType)
+		/// <param name="Arguments">Arguments for the command.</param>
+		/// <param name="InputData">Input data for the command.</param>
+		/// <param name="StatRecordType">Type of element to return in the response</param>
+		/// <param name="CancellationToken">Token used to cancel the operation</param>
+		/// <returns>Response from the server; either an object of type T or error.</returns>
+		public async Task<PerforceResponse> SingleResponseCommandAsync(string Arguments, byte[]? InputData, Type? StatRecordType, CancellationToken CancellationToken)
 		{
-			Dictionary<string, int> NameToValue;
-			if (!EnumTypeToFlags.TryGetValue(EnumType, out NameToValue))
+			List<PerforceResponse> Responses = await CommandAsync(Arguments, InputData, StatRecordType, CancellationToken);
+			if (Responses.Count != 1)
 			{
-				NameToValue = new Dictionary<string, int>();
-
-				FieldInfo[] Fields = EnumType.GetFields(BindingFlags.Public | BindingFlags.Static);
-				foreach (FieldInfo Field in Fields)
-				{
-					PerforceEnumAttribute Attribute = Field.GetCustomAttribute<PerforceEnumAttribute>();
-					if (Attribute != null)
-					{
-						NameToValue.Add(Attribute.Name, (int)Field.GetValue(null));
-					}
-				}
-
-				if (!EnumTypeToFlags.TryAdd(EnumType, NameToValue))
-				{
-					NameToValue = EnumTypeToFlags[EnumType];
-				}
+				throw new PerforceException("Expected one result from 'p4 {0}', got {1}", Arguments, Responses.Count);
 			}
-			return NameToValue;
+			return Responses[0];
 		}
 
 		/// <summary>
-		/// Parses an enum value, using PerforceEnumAttribute markup for names.
+		/// Attempts to execute the given command, returning the results from the server or the first PerforceResponse object.
 		/// </summary>
-		/// <param name="EnumType">Type of the enum to parse.</param>
-		/// <param name="Value">Value of the enum.</param>
-		/// <returns>Text for the enum.</returns>
-		string GetEnumText(Type EnumType, object Value)
+		/// <typeparam name="T">Type of record to parse</typeparam>
+		/// <param name="Arguments">Arguments for the command.</param>
+		/// <param name="InputData">Input data for the command.</param>
+		/// <param name="CancellationToken">Token used to cancel the operation</param>
+		/// <returns>Response from the server; either an object of type T or error.</returns>
+		public async Task<PerforceResponse<T>> SingleResponseCommandAsync<T>(string Arguments, byte[]? InputData, CancellationToken CancellationToken) where T : class
 		{
-			int IntegerValue = (int)Value;
-
-			Dictionary<string, int> NameToValue = GetCachedEnumFlags(EnumType);
-			if(EnumType.GetCustomAttribute<FlagsAttribute>() != null)
-			{
-				List<string> Names = new List<string>();
-				foreach (KeyValuePair<string, int> Pair in NameToValue)
-				{
-					if ((IntegerValue & Pair.Value) != 0)
-					{
-						Names.Add(Pair.Key);
-					}
-				}
-				return String.Join(" ", Names);
-			}
-			else
-			{
-				string Name = null;
-				foreach (KeyValuePair<string, int> Pair in NameToValue)
-				{
-					if (IntegerValue == Pair.Value)
-					{
-						Name = Pair.Key;
-						break;
-					}
-				}
-				return Name;
-			}
+			return new PerforceResponse<T>(await SingleResponseCommandAsync(Arguments, InputData, typeof(T), CancellationToken));
 		}
 
+		#endregion
+
+		#region p4 add
+
 		/// <summary>
-		/// Parses an enum value, using PerforceEnumAttribute markup for names.
+		/// Adds files to a pending changelist.
 		/// </summary>
-		/// <param name="EnumType">Type of the enum to parse.</param>
-		/// <param name="Text">Text to parse.</param>
-		/// <returns>The parsed enum value. Unknown values will be ignored.</returns>
-		object ParseEnum(Type EnumType, string Text)
+		/// <param name="ChangeNumber">Changelist to add files to</param>
+		/// <param name="FileNames">Files to be added</param>
+		/// <param name="CancellationToken">Token used to cancel the operation</param>
+		/// <returns>Response from the server</returns>
+		public async Task<List<AddRecord>> AddAsync(int ChangeNumber, string[] FileNames, CancellationToken CancellationToken)
 		{
-			Dictionary<string, int> NameToValue = GetCachedEnumFlags(EnumType);
-			if(EnumType.GetCustomAttribute<FlagsAttribute>() != null)
-			{
-				int Result = 0;
-				foreach (string Item in Text.Split(' '))
-				{
-					int ItemValue;
-					if (NameToValue.TryGetValue(Item, out ItemValue))
-					{
-						Result |= ItemValue;
-					}
-				}
-				return Enum.ToObject(EnumType, Result);
-			}
-			else
-			{
-				int Result;
-				NameToValue.TryGetValue(Text, out Result);
-				return Enum.ToObject(EnumType, Result);
-			}
+			return (await TryAddAsync(ChangeNumber, null, AddOptions.None, FileNames, CancellationToken)).Data;
 		}
 
 		/// <summary>
-		/// Gets reflection data for the given record type
+		/// Adds files to a pending changelist.
 		/// </summary>
-		/// <param name="RecordType">The type to retrieve record info for</param>
-		/// <returns>The cached reflection information for the given type</returns>
-		static CachedRecordInfo GetCachedRecordInfo(Type RecordType)
+		/// <param name="ChangeNumber">Changelist to add files to</param>
+		/// <param name="FileType">Type for new files</param>
+		/// <param name="Options">Options for the command</param>
+		/// <param name="FileNames">Files to be added</param>
+		/// <param name="CancellationToken">Token used to cancel the operation</param>
+		/// <returns>Response from the server</returns>
+		public async Task<List<AddRecord>> AddAsync(int ChangeNumber, string? FileType, AddOptions Options, string[] FileNames, CancellationToken CancellationToken)
 		{
-			CachedRecordInfo Record;
-			if (!RecordTypeToInfo.TryGetValue(RecordType, out Record))
-			{
-				Record = new CachedRecordInfo();
-				Record.Type = RecordType;
-
-				// Get all the fields for this type
-				FieldInfo[] Fields = RecordType.GetFields(BindingFlags.Public | BindingFlags.NonPublic | BindingFlags.Instance);
-
-				// Build the map of all tags for this record
-				foreach (FieldInfo Field in Fields)
-				{
-					PerforceTagAttribute TagAttribute = Field.GetCustomAttribute<PerforceTagAttribute>();
-					if (TagAttribute != null)
-					{
-						CachedTagInfo Tag = new CachedTagInfo();
-						Tag.Name = TagAttribute.Name ?? Field.Name;
-						Tag.Optional = TagAttribute.Optional;
-						Tag.Field = Field;
-						if(!Tag.Optional)
-						{
-							Tag.RequiredTagBitMask = Record.RequiredTagsBitMask + 1;
-							if(Tag.RequiredTagBitMask == 0)
-							{
-								throw new PerforceException("Too many required tags in {0}; max is {1}", RecordType.Name, sizeof(ulong) * 8);
-							}
-							Record.RequiredTagsBitMask |= Tag.RequiredTagBitMask;
-						}
-						Record.TagNameToInfo.Add(Tag.Name, Tag);
-					}
-
-					PerforceRecordListAttribute SubElementAttribute = Field.GetCustomAttribute<PerforceRecordListAttribute>();
-					if(SubElementAttribute != null)
-					{
-						Record.SubElementField = Field;
-						Record.SubElementType = Field.FieldType.GenericTypeArguments[0];
-						Record.SubElementRecordInfo = GetCachedRecordInfo(Record.SubElementType);
-					}
-				}
-
-				// Try to save the record info, or get the version that's already in the cache
-				if(!RecordTypeToInfo.TryAdd(RecordType, Record))
-				{
-					Record = RecordTypeToInfo[RecordType];
-				}
-			}
-			return Record;
+			return (await TryAddAsync(ChangeNumber, FileType, Options, FileNames, CancellationToken)).Data;
 		}
 
-		#endregion
-
-		#region p4 add
-
 		/// <summary>
 		/// Adds files to a pending changelist.
 		/// </summary>
 		/// <param name="ChangeNumber">Changelist to add files to</param>
 		/// <param name="FileNames">Files to be added</param>
+		/// <param name="CancellationToken">Token used to cancel the operation</param>
 		/// <returns>Response from the server</returns>
-		public PerforceResponseList<AddRecord> Add(int ChangeNumber, params string[] FileNames)
+		public Task<PerforceResponseList<AddRecord>> TryAddAsync(int ChangeNumber, string[] FileNames, CancellationToken CancellationToken)
 		{
-			return Add(ChangeNumber, null, AddOptions.None, FileNames);
+			return TryAddAsync(ChangeNumber, null, AddOptions.None, FileNames, CancellationToken);
 		}
 
 		/// <summary>
@@ -622,40 +365,41 @@ namespace EpicGames.Perforce
 		/// <param name="FileType">Type for new files</param>
 		/// <param name="Options">Options for the command</param>
 		/// <param name="FileNames">Files to be added</param>
+		/// <param name="CancellationToken">Token used to cancel the operation</param>
 		/// <returns>Response from the server</returns>
-		public PerforceResponseList<AddRecord> Add(int ChangeNumber, string FileType, AddOptions Options, params string[] FileNames)
+		public Task<PerforceResponseList<AddRecord>> TryAddAsync(int ChangeNumber, string? FileType, AddOptions Options, string[] FileNames, CancellationToken CancellationToken)
 		{
 			StringBuilder Arguments = new StringBuilder("add");
-			if(ChangeNumber != -1)
+			if (ChangeNumber != -1)
 			{
 				Arguments.AppendFormat(" -c {0}", ChangeNumber);
 			}
-			if((Options & AddOptions.DowngradeToAdd) != 0)
+			if ((Options & AddOptions.DowngradeToAdd) != 0)
 			{
 				Arguments.Append(" -d");
 			}
-			if((Options & AddOptions.IncludeWildcards) != 0)
+			if ((Options & AddOptions.IncludeWildcards) != 0)
 			{
 				Arguments.Append(" -f");
 			}
-			if((Options & AddOptions.NoIgnore) != 0)
+			if ((Options & AddOptions.NoIgnore) != 0)
 			{
 				Arguments.Append(" -I");
 			}
-			if((Options & AddOptions.PreviewOnly) != 0)
+			if ((Options & AddOptions.PreviewOnly) != 0)
 			{
 				Arguments.Append(" -n");
 			}
-			if(FileType != null)
+			if (FileType != null)
 			{
 				Arguments.AppendFormat(" -t \"{0}\"", FileType);
 			}
-			foreach(string FileName in FileNames)
+			foreach (string FileName in FileNames)
 			{
 				Arguments.AppendFormat(" \"{0}\"", FileName);
 			}
 
-			return Command<AddRecord>(Arguments.ToString(), null);
+			return CommandAsync<AddRecord>(Arguments.ToString(), null, CancellationToken);
 		}
 
 		#endregion
@@ -666,27 +410,48 @@ namespace EpicGames.Perforce
 		/// Creates a changelist with the p4 change command. 
 		/// </summary>
 		/// <param name="Record">Information for the change to create. The number field should be left set to -1.</param>
+		/// <param name="CancellationToken">Token used to cancel the operation</param>
+		/// <returns>The changelist number, or an error.</returns>
+		public async Task<ChangeRecord> CreateChangeAsync(ChangeRecord Record, CancellationToken CancellationToken)
+		{
+			return (await TryCreateChangeAsync(Record, CancellationToken)).Data;
+		}
+
+		/// <summary>
+		/// Creates a changelist with the p4 change command. 
+		/// </summary>
+		/// <param name="Record">Information for the change to create. The number field should be left set to -1.</param>
+		/// <param name="CancellationToken">Token used to cancel the operation</param>
 		/// <returns>The changelist number, or an error.</returns>
-		public PerforceResponse<int> CreateChange(ChangeRecord Record)
+		public async Task<PerforceResponse<ChangeRecord>> TryCreateChangeAsync(ChangeRecord Record, CancellationToken CancellationToken)
 		{
-			if(Record.Number != -1)
+			if (Record.Number != -1)
 			{
 				throw new PerforceException("'Number' field should be set to -1 to create a new changelist.");
 			}
 
-			PerforceResponse Response = SingleResponseCommand("change -i", SerializeRecord(Record), null);
-			if (Response.Failed)
+			PerforceResponse Response = await SingleResponseCommandAsync("change -i", SerializeRecord(Record), null, CancellationToken);
+
+			PerforceError? Error = Response.Error;
+			if (Error != null)
+			{
+				return new PerforceResponse<ChangeRecord>(Error);
+			}
+
+			PerforceInfo? Info = Response.Info;
+			if (Info == null)
 			{
-				return new PerforceResponse<int>(Response.Error);
+				throw new PerforceException("Unexpected info response from change command: {0}", Response);
 			}
 
-			string[] Tokens = Response.Info.Data.Split(' ');
+			string[] Tokens = Info.Data.Split(' ');
 			if (Tokens.Length != 3)
 			{
 				throw new PerforceException("Unexpected info response from change command: {0}", Response);
 			}
 
-			return new PerforceResponse<int>(int.Parse(Tokens[1]));
+			Record.Number = int.Parse(Tokens[1]);
+			return new PerforceResponse<ChangeRecord>(Record);
 		}
 
 		/// <summary>
@@ -694,25 +459,38 @@ namespace EpicGames.Perforce
 		/// </summary>
 		/// <param name="Options">Options for this command</param>
 		/// <param name="Record">Information for the change to create. The number field should be left set to zero.</param>
+		/// <param name="CancellationToken">Token used to cancel the operation</param>
 		/// <returns>The changelist number, or an error.</returns>
-		public PerforceResponse UpdateChange(UpdateChangeOptions Options, ChangeRecord Record)
+		public async Task UpdateChangeAsync(UpdateChangeOptions Options, ChangeRecord Record, CancellationToken CancellationToken)
 		{
-			if(Record.Number == -1)
+			(await TryUpdateChangeAsync(Options, Record, CancellationToken)).EnsureSuccess();
+		}
+
+		/// <summary>
+		/// Updates an existing changelist.
+		/// </summary>
+		/// <param name="Options">Options for this command</param>
+		/// <param name="Record">Information for the change to create. The number field should be left set to zero.</param>
+		/// <param name="CancellationToken">Token used to cancel the operation</param>
+		/// <returns>The changelist number, or an error.</returns>
+		public Task<PerforceResponse> TryUpdateChangeAsync(UpdateChangeOptions Options, ChangeRecord Record, CancellationToken CancellationToken)
+		{
+			if (Record.Number == -1)
 			{
 				throw new PerforceException("'Number' field must be set to update a changelist.");
 			}
 
 			StringBuilder Arguments = new StringBuilder("change -i");
-			if((Options & UpdateChangeOptions.Force) != 0)
+			if ((Options & UpdateChangeOptions.Force) != 0)
 			{
 				Arguments.Append(" -f");
 			}
-			if((Options & UpdateChangeOptions.Submitted) != 0)
+			if ((Options & UpdateChangeOptions.Submitted) != 0)
 			{
 				Arguments.Append(" -u");
 			}
 
-			return SingleResponseCommand(Arguments.ToString(), SerializeRecord(Record), null);
+			return SingleResponseCommandAsync(Arguments.ToString(), SerializeRecord(Record), null, CancellationToken);
 		}
 
 		/// <summary>
@@ -720,21 +498,34 @@ namespace EpicGames.Perforce
 		/// </summary>
 		/// <param name="Options">Options for the command</param>
 		/// <param name="ChangeNumber">Changelist number to delete</param>
+		/// <param name="CancellationToken">Token used to cancel the operation</param>
 		/// <returns>Response from the server</returns>
-		public PerforceResponse DeleteChange(DeleteChangeOptions Options, int ChangeNumber)
+		public async Task DeleteChangeAsync(DeleteChangeOptions Options, int ChangeNumber, CancellationToken CancellationToken)
+		{
+			(await TryDeleteChangeAsync(Options, ChangeNumber, CancellationToken)).EnsureSuccess();
+		}
+
+		/// <summary>
+		/// Deletes a changelist (p4 change -d)
+		/// </summary>
+		/// <param name="Options">Options for the command</param>
+		/// <param name="ChangeNumber">Changelist number to delete</param>
+		/// <param name="CancellationToken">Token used to cancel the operation</param>
+		/// <returns>Response from the server</returns>
+		public Task<PerforceResponse> TryDeleteChangeAsync(DeleteChangeOptions Options, int ChangeNumber, CancellationToken CancellationToken)
 		{
 			StringBuilder Arguments = new StringBuilder("change -d");
-			if((Options & DeleteChangeOptions.Submitted) != 0)
+			if ((Options & DeleteChangeOptions.Submitted) != 0)
 			{
 				Arguments.Append(" -f");
 			}
-			if((Options & DeleteChangeOptions.BeforeRenumber) != 0)
+			if ((Options & DeleteChangeOptions.BeforeRenumber) != 0)
 			{
 				Arguments.Append(" -O");
 			}
 			Arguments.AppendFormat(" {0}", ChangeNumber);
 
-			return SingleResponseCommand(Arguments.ToString(), null, null);
+			return SingleResponseCommandAsync(Arguments.ToString(), null, null, CancellationToken);
 		}
 
 		/// <summary>
@@ -742,19 +533,37 @@ namespace EpicGames.Perforce
 		/// </summary>
 		/// <param name="Options">Options for the command</param>
 		/// <param name="ChangeNumber">Changelist number to retrieve. -1 is the default changelist for this workspace.</param>
+		/// <param name="CancellationToken">Token used to cancel the operation</param>
 		/// <returns>Response from the server</returns>
-		public PerforceResponse<ChangeRecord> GetChange(GetChangeOptions Options, int ChangeNumber)
+		public async Task<ChangeRecord> GetChangeAsync(GetChangeOptions Options, int ChangeNumber, CancellationToken CancellationToken)
+		{
+			return (await TryGetChange(Options, ChangeNumber, CancellationToken)).Data;
+		}
+
+		/// <summary>
+		/// Gets a changelist
+		/// </summary>
+		/// <param name="Options">Options for the command</param>
+		/// <param name="ChangeNumber">Changelist number to retrieve. -1 is the default changelist for this workspace.</param>
+		/// <param name="CancellationToken">Token used to cancel the operation</param>
+		/// <returns>Response from the server</returns>
+		public Task<PerforceResponse<ChangeRecord>> TryGetChange(GetChangeOptions Options, int ChangeNumber, CancellationToken CancellationToken)
 		{
 			StringBuilder Arguments = new StringBuilder("change -o");
-			if((Options & GetChangeOptions.BeforeRenumber) != 0)
+			if ((Options & GetChangeOptions.BeforeRenumber) != 0)
 			{
 				Arguments.Append(" -O");
 			}
 			Arguments.AppendFormat(" {0}", ChangeNumber);
 
-			return SingleResponseCommand<ChangeRecord>(Arguments.ToString(), null);
+			return SingleResponseCommandAsync<ChangeRecord>(Arguments.ToString(), null, CancellationToken);
 		}
 
+		/// <summary>
+		/// Serializes a change record to a byte array
+		/// </summary>
+		/// <param name="Input">The record to serialize</param>
+		/// <returns>Serialized record</returns>
 		byte[] SerializeRecord(ChangeRecord Input)
 		{
 			List<KeyValuePair<string, object>> NameToValue = new List<KeyValuePair<string, object>>();
@@ -796,10 +605,11 @@ namespace EpicGames.Perforce
 		/// <param name="MaxChanges">List only the highest numbered changes</param>
 		/// <param name="Status">Limit the list to the changelists with the given status (pending, submitted or shelved)</param>
 		/// <param name="FileSpecs">Paths to query changes for</param>
+		/// <param name="CancellationToken">Token used to cancel the operation</param>
 		/// <returns>List of responses from the server.</returns>
-		public PerforceResponseList<ChangesRecord> Changes(ChangesOptions Options, int MaxChanges, ChangeStatus Status, params string[] FileSpecs)
+		public async Task<List<ChangesRecord>> GetChangesAsync(ChangesOptions Options, int MaxChanges, ChangeStatus Status, string[] FileSpecs, CancellationToken CancellationToken)
 		{
-			return Changes(Options, null, MaxChanges, Status, null, FileSpecs);
+			return (await TryGetChangesAsync(Options, MaxChanges, Status, FileSpecs, CancellationToken)).Data;
 		}
 
 		/// <summary>
@@ -811,8 +621,39 @@ namespace EpicGames.Perforce
 		/// <param name="Status">Limit the list to the changelists with the given status (pending, submitted or shelved)</param>
 		/// <param name="UserName">List only changes made by the named user</param>
 		/// <param name="FileSpecs">Paths to query changes for</param>
+		/// <param name="CancellationToken">Token used to cancel the operation</param>
 		/// <returns>List of responses from the server.</returns>
-		public PerforceResponseList<ChangesRecord> Changes(ChangesOptions Options, string ClientName, int MaxChanges, ChangeStatus Status, string UserName, params string[] FileSpecs)
+		public async Task<List<ChangesRecord>> GetChangesAsync(ChangesOptions Options, string? ClientName, int MaxChanges, ChangeStatus Status, string? UserName, string[]? FileSpecs, CancellationToken CancellationToken)
+		{
+			return (await TryGetChangesAsync(Options, ClientName, MaxChanges, Status, UserName, FileSpecs, CancellationToken)).Data;
+		}
+
+		/// <summary>
+		/// Enumerates changes on the server
+		/// </summary>
+		/// <param name="Options">Options for the command</param>
+		/// <param name="MaxChanges">List only the highest numbered changes</param>
+		/// <param name="Status">Limit the list to the changelists with the given status (pending, submitted or shelved)</param>
+		/// <param name="FileSpecs">Paths to query changes for</param>
+		/// <param name="CancellationToken">Token used to cancel the operation</param>
+		/// <returns>List of responses from the server.</returns>
+		public Task<PerforceResponseList<ChangesRecord>> TryGetChangesAsync(ChangesOptions Options, int MaxChanges, ChangeStatus Status, string[] FileSpecs, CancellationToken CancellationToken)
+		{
+			return TryGetChangesAsync(Options, null, MaxChanges, Status, null, FileSpecs, CancellationToken);
+		}
+
+		/// <summary>
+		/// Enumerates changes on the server
+		/// </summary>
+		/// <param name="Options">Options for the command</param>
+		/// <param name="ClientName">List only changes made from the named client workspace.</param>
+		/// <param name="MaxChanges">List only the highest numbered changes</param>
+		/// <param name="Status">Limit the list to the changelists with the given status (pending, submitted or shelved)</param>
+		/// <param name="UserName">List only changes made by the named user</param>
+		/// <param name="FileSpecs">Paths to query changes for</param>
+		/// <param name="CancellationToken">Token used to cancel the operation</param>
+		/// <returns>List of responses from the server.</returns>
+		public Task<PerforceResponseList<ChangesRecord>> TryGetChangesAsync(ChangesOptions Options, string? ClientName, int MaxChanges, ChangeStatus Status, string? UserName, string[]? FileSpecs, CancellationToken CancellationToken)
 		{
 			StringBuilder Arguments = new StringBuilder("changes");
 			if ((Options & ChangesOptions.IncludeIntegrations) != 0)
@@ -835,27 +676,94 @@ namespace EpicGames.Perforce
 			{
 				Arguments.Append(" -f");
 			}
-			if(ClientName != null)
+			if (ClientName != null)
 			{
 				Arguments.AppendFormat(" -c \"{0}\"", ClientName);
 			}
-			if(MaxChanges != -1)
+			if (MaxChanges != -1)
 			{
 				Arguments.AppendFormat(" -m {0}", MaxChanges);
 			}
-			if(Status != ChangeStatus.All)
+			if (Status != ChangeStatus.All)
 			{
-				Arguments.AppendFormat(" -s {0}", GetEnumText(typeof(ChangeStatus), Status));
+				Arguments.AppendFormat(" -s {0}", PerforceReflection.GetEnumText(typeof(ChangeStatus), Status));
 			}
-			if(UserName != null)
+			if (UserName != null)
 			{
 				Arguments.AppendFormat(" -u {0}", UserName);
 			}
-			foreach(string FileSpec in FileSpecs)
+			if (FileSpecs != null)
+			{
+				foreach (string FileSpec in FileSpecs)
+				{
+					Arguments.AppendFormat(" \"{0}\"", FileSpec);
+				}
+			}
+			return CommandAsync<ChangesRecord>(Arguments.ToString(), null, CancellationToken);
+		}
+
+		#endregion
+
+		#region p4 clean
+
+		/// <summary>
+		/// Cleans the workspace
+		/// </summary>
+		/// <param name="Options">Options for the command</param>
+		/// <param name="FileSpecs">Files to sync</param>
+		/// <param name="CancellationToken">Token used to cancel the operation</param>
+		/// <returns>Response from the server</returns>
+		public async Task<List<CleanRecord>> CleanAsync(CleanOptions Options, string[] FileSpecs, CancellationToken CancellationToken)
+		{
+			return (await TryCleanAsync(Options, FileSpecs, CancellationToken)).Data;
+		}
+
+		/// <summary>
+		/// Cleans the workspace
+		/// </summary>
+		/// <param name="Options">Options for the command</param>
+		/// <param name="FileSpecs">Files to sync</param>
+		/// <param name="CancellationToken">Token used to cancel the operation</param>
+		/// <returns>Response from the server</returns>
+		public async Task<PerforceResponseList<CleanRecord>> TryCleanAsync(CleanOptions Options, string[] FileSpecs, CancellationToken CancellationToken)
+		{
+			StringBuilder Arguments = new StringBuilder("clean");
+			if ((Options & CleanOptions.Edited) != 0)
+			{
+				Arguments.Append(" -e");
+			}
+			if ((Options & CleanOptions.Added) != 0)
+			{
+				Arguments.Append(" -a");
+			}
+			if ((Options & CleanOptions.Deleted) != 0)
+			{
+				Arguments.Append(" -d");
+			}
+			if ((Options & CleanOptions.Preview) != 0)
+			{
+				Arguments.Append(" -n");
+			}
+			if ((Options & CleanOptions.NoIgnoreChecking) != 0)
+			{
+				Arguments.Append(" -I");
+			}
+			if ((Options & CleanOptions.LocalSyntax) != 0)
+			{
+				Arguments.Append(" -l");
+			}
+			if ((Options & CleanOptions.ModifiedTimes) != 0)
+			{
+				Arguments.Append(" -m");
+			}
+			foreach (string FileSpec in FileSpecs)
 			{
 				Arguments.AppendFormat(" \"{0}\"", FileSpec);
 			}
-			return Command<ChangesRecord>(Arguments.ToString(), null);
+
+			PerforceResponseList<CleanRecord> Records = await CommandAsync<CleanRecord>(Arguments.ToString(), null, CancellationToken);
+			Records.RemoveAll(x => x.Error != null && x.Error.Generic == PerforceGenericCode.Empty);
+			return Records;
 		}
 
 		#endregion
@@ -866,20 +774,56 @@ namespace EpicGames.Perforce
 		/// Creates a client
 		/// </summary>
 		/// <param name="Record">The client record</param>
+		/// <param name="CancellationToken">Token used to cancel the operation</param>
 		/// <returns>Response from the server</returns>
-		public PerforceResponse CreateClient(ClientRecord Record)
+		public async Task CreateClientAsync(ClientRecord Record, CancellationToken CancellationToken)
 		{
-			return UpdateClient(Record);
+			(await TryCreateClientAsync(Record, CancellationToken)).EnsureSuccess();
 		}
 
 		/// <summary>
 		/// Creates a client
 		/// </summary>
 		/// <param name="Record">The client record</param>
+		/// <param name="CancellationToken">Token used to cancel the operation</param>
+		/// <returns>Response from the server</returns>
+		public Task<PerforceResponse> TryCreateClientAsync(ClientRecord Record, CancellationToken CancellationToken)
+		{
+			return TryUpdateClientAsync(Record, CancellationToken);
+		}
+
+		/// <summary>
+		/// Creates a client
+		/// </summary>
+		/// <param name="Record">The client record</param>
+		/// <param name="CancellationToken">Token used to cancel the operation</param>
+		/// <returns>Response from the server</returns>
+		public async Task UpdateClientAsync(ClientRecord Record, CancellationToken CancellationToken)
+		{
+			(await TryUpdateClientAsync(Record, CancellationToken)).EnsureSuccess();
+		}
+
+		/// <summary>
+		/// Update a client
+		/// </summary>
+		/// <param name="Record">The client record</param>
+		/// <param name="CancellationToken">Token used to cancel the operation</param>
+		/// <returns>Response from the server</returns>
+		public Task<PerforceResponse> TryUpdateClientAsync(ClientRecord Record, CancellationToken CancellationToken)
+		{
+			return SingleResponseCommandAsync("client -i", SerializeRecord(Record), null, CancellationToken);
+		}
+
+		/// <summary>
+		/// Deletes a client
+		/// </summary>
+		/// <param name="Options">Options for this command</param>
+		/// <param name="ClientName">Name of the client</param>
+		/// <param name="CancellationToken">Token used to cancel the operation</param>
 		/// <returns>Response from the server</returns>
-		public PerforceResponse UpdateClient(ClientRecord Record)
+		public async Task DeleteClientAsync(DeleteClientOptions Options, string ClientName, CancellationToken CancellationToken)
 		{
-			return SingleResponseCommand("client -i", SerializeRecord(Record), null);
+			(await TryDeleteClientAsync(Options, ClientName, CancellationToken)).EnsureSuccess();
 		}
 
 		/// <summary>
@@ -887,20 +831,34 @@ namespace EpicGames.Perforce
 		/// </summary>
 		/// <param name="Options">Options for this command</param>
 		/// <param name="ClientName">Name of the client</param>
+		/// <param name="CancellationToken">Token used to cancel the operation</param>
 		/// <returns>Response from the server</returns>
-		public PerforceResponse DeleteClient(DeleteClientOptions Options, string ClientName)
+		public Task<PerforceResponse> TryDeleteClientAsync(DeleteClientOptions Options, string ClientName, CancellationToken CancellationToken)
 		{
 			StringBuilder Arguments = new StringBuilder("client -d");
-			if((Options & DeleteClientOptions.Force) != 0)
+			if ((Options & DeleteClientOptions.Force) != 0)
 			{
 				Arguments.Append(" -f");
 			}
-			if((Options & DeleteClientOptions.DeleteShelved) != 0)
+			if ((Options & DeleteClientOptions.DeleteShelved) != 0)
 			{
 				Arguments.Append(" -Fs");
 			}
-			Arguments.AppendFormat(" -d \"{0}\"", ClientName);
-			return SingleResponseCommand<ClientRecord>(Arguments.ToString(), null);
+			Arguments.AppendFormat(" \"{0}\"", ClientName);
+			return SingleResponseCommandAsync(Arguments.ToString(), null, null, CancellationToken);
+		}
+
+		/// <summary>
+		/// Changes the stream associated with a client
+		/// </summary>
+		/// <param name="ClientName">The client name</param>
+		/// <param name="StreamName">The new stream to be associated with the client</param>
+		/// <param name="Options">Options for this command</param>
+		/// <param name="CancellationToken">Token used to cancel the operation</param>
+		/// <returns>Response from the server</returns>
+		public async Task SwitchClientToStreamAsync(string ClientName, string StreamName, SwitchClientOptions Options, CancellationToken CancellationToken)
+		{
+			(await TrySwitchClientToStreamAsync(ClientName, StreamName, Options, CancellationToken)).EnsureSuccess();
 		}
 
 		/// <summary>
@@ -909,44 +867,113 @@ namespace EpicGames.Perforce
 		/// <param name="ClientName">The client name</param>
 		/// <param name="StreamName">The new stream to be associated with the client</param>
 		/// <param name="Options">Options for this command</param>
+		/// <param name="CancellationToken">Token used to cancel the operation</param>
 		/// <returns>Response from the server</returns>
-		public PerforceResponse SwitchClientToStream(string ClientName, string StreamName, SwitchClientOptions Options)
+		public Task<PerforceResponse> TrySwitchClientToStreamAsync(string ClientName, string StreamName, SwitchClientOptions Options, CancellationToken CancellationToken)
 		{
 			StringBuilder Arguments = new StringBuilder("client -s");
-			if((Options & SwitchClientOptions.IgnoreOpenFiles) != 0)
+			if ((Options & SwitchClientOptions.IgnoreOpenFiles) != 0)
+			{
+				Arguments.Append(" -f");
+			}
+			Arguments.AppendFormat(" -S \"{0}\"", StreamName);
+
+			return SingleResponseCommandAsync(Arguments.ToString(), null, null, CancellationToken);
+		}
+
+		/// <summary>
+		/// Changes a client to mirror a template
+		/// </summary>
+		/// <param name="ClientName">The client name</param>
+		/// <param name="TemplateName">The new stream to be associated with the client</param>
+		/// <param name="CancellationToken">Token used to cancel the operation</param>
+		/// <returns>Response from the server</returns>
+		public async Task SwitchClientToTemplateAsync(string ClientName, string TemplateName, CancellationToken CancellationToken)
+		{
+			(await TrySwitchClientToTemplateAsync(ClientName, TemplateName, CancellationToken)).EnsureSuccess();
+		}
+
+		/// <summary>
+		/// Changes a client to mirror a template
+		/// </summary>
+		/// <param name="ClientName">The client name</param>
+		/// <param name="TemplateName">The new stream to be associated with the client</param>
+		/// <param name="CancellationToken">Token used to cancel the operation</param>
+		/// <returns>Response from the server</returns>
+		public Task<PerforceResponse> TrySwitchClientToTemplateAsync(string ClientName, string TemplateName, CancellationToken CancellationToken)
+		{
+			string Arguments = String.Format("client -s -t \"{0}\"", TemplateName);
+			return SingleResponseCommandAsync(Arguments, null, null, CancellationToken);
+		}
+
+		/// <summary>
+		/// Queries the current client definition
+		/// </summary>
+		/// <param name="ClientName">Name of the client. Specify null for the current client.</param>
+		/// <param name="CancellationToken">Token used to cancel the operation</param>
+		/// <returns>Response from the server; either a client record or error code</returns>
+		public async Task<ClientRecord> GetClientAsync(string ClientName, CancellationToken CancellationToken)
+		{
+			return (await TryGetClientAsync(ClientName, CancellationToken)).Data;
+		}
+
+		/// <summary>
+		/// Queries the current client definition
+		/// </summary>
+		/// <param name="ClientName">Name of the client. Specify null for the current client.</param>
+		/// <param name="CancellationToken">Token used to cancel the operation</param>
+		/// <returns>Response from the server; either a client record or error code</returns>
+		public Task<PerforceResponse<ClientRecord>> TryGetClientAsync(string ClientName, CancellationToken CancellationToken)
+		{
+			StringBuilder Arguments = new StringBuilder("client -o");
+			if (ClientName != null)
 			{
-				Arguments.Append(" -f");
+				Arguments.AppendFormat(" \"{0}\"", ClientName);
 			}
-			Arguments.AppendFormat(" -S \"{0}\"", StreamName);
-
-			return SingleResponseCommand(Arguments.ToString(), null, null);
+			return SingleResponseCommandAsync<ClientRecord>(Arguments.ToString(), null, CancellationToken);
 		}
 
 		/// <summary>
-		/// Changes a client to mirror a template
+		/// Queries information about a stream
 		/// </summary>
-		/// <param name="ClientName">The client name</param>
-		/// <param name="TemplateName">The new stream to be associated with the client</param>
-		/// <returns>Response from the server</returns>
-		public PerforceResponse SwitchClientToTemplate(string ClientName, string TemplateName)
+		/// <param name="StreamName">Name of the stream to query</param>
+		/// <param name="bIncludeView">Whether to include the stream view in the output</param>
+		/// <param name="CancellationToken">Token used to cancel the operation</param>
+		/// <returns>Stream information record</returns>
+		public async Task<StreamRecord> GetStreamAsync(string StreamName, bool bIncludeView, CancellationToken CancellationToken)
 		{
-			string Arguments = String.Format("client -s -t \"{0}\"", TemplateName);
-			return SingleResponseCommand(Arguments, null, null);
+			return (await TryGetStreamAsync(StreamName, bIncludeView, CancellationToken)).Data;
 		}
 
 		/// <summary>
-		/// Queries the current client definition
+		/// Queries information about a stream
 		/// </summary>
-		/// <param name="ClientName">Name of the client. Specify null for the current client.</param>
-		/// <returns>Response from the server; either a client record or error code</returns>
-		public PerforceResponse<ClientRecord> GetClient(string ClientName)
+		/// <param name="StreamName">Name of the stream to query</param>
+		/// <param name="bIncludeView">Whether to include the stream view in the output</param>
+		/// <param name="CancellationToken">Token used to cancel the operation</param>
+		/// <returns>Stream information record</returns>
+		public Task<PerforceResponse<StreamRecord>> TryGetStreamAsync(string StreamName, bool bIncludeView, CancellationToken CancellationToken)
 		{
-			StringBuilder Arguments = new StringBuilder("client -o");
-			if(ClientName != null)
+			StringBuilder Arguments = new StringBuilder("stream -o");
+			if (bIncludeView)
 			{
-				Arguments.AppendFormat(" \"{0}\"", ClientName);
+				Arguments.Append(" -v");
 			}
-			return SingleResponseCommand<ClientRecord>(Arguments.ToString(), null);
+			Arguments.Append(StreamName);
+
+			return SingleResponseCommandAsync<StreamRecord>(Arguments.ToString(), null, CancellationToken);
+		}
+
+		/// <summary>
+		/// Queries the view for a stream
+		/// </summary>
+		/// <param name="StreamName">Name of the stream.</param>
+		/// <param name="ChangeNumber">Changelist at which to query the stream view</param>
+		/// <param name="CancellationToken">Token used to cancel the operation</param>
+		/// <returns>Response from the server; either a client record or error code</returns>
+		public async Task<ClientRecord> GetStreamViewAsync(string StreamName, int ChangeNumber, CancellationToken CancellationToken)
+		{
+			return (await TryGetStreamViewAsync(StreamName, ChangeNumber, CancellationToken)).Data;
 		}
 
 		/// <summary>
@@ -954,16 +981,17 @@ namespace EpicGames.Perforce
 		/// </summary>
 		/// <param name="StreamName">Name of the stream.</param>
 		/// <param name="ChangeNumber">Changelist at which to query the stream view</param>
+		/// <param name="CancellationToken">Token used to cancel the operation</param>
 		/// <returns>Response from the server; either a client record or error code</returns>
-		public PerforceResponse<ClientRecord> GetStreamView(string StreamName, int ChangeNumber)
+		public Task<PerforceResponse<ClientRecord>> TryGetStreamViewAsync(string StreamName, int ChangeNumber, CancellationToken CancellationToken)
 		{
 			StringBuilder Arguments = new StringBuilder("client -o");
 			Arguments.AppendFormat(" -S \"{0}\"", StreamName);
-			if(ChangeNumber != -1)
+			if (ChangeNumber != -1)
 			{
 				Arguments.AppendFormat(" -c {0}", ChangeNumber);
 			}
-			return SingleResponseCommand<ClientRecord>(Arguments.ToString(), null);
+			return SingleResponseCommandAsync<ClientRecord>(Arguments.ToString(), null, CancellationToken);
 		}
 
 		/// <summary>
@@ -994,23 +1022,23 @@ namespace EpicGames.Perforce
 			{
 				NameToValue.Add(new KeyValuePair<string, object>("Root", Input.Root));
 			}
-			if(Input.Options != ClientOptions.None)
+			if (Input.Options != ClientOptions.None)
 			{
-				NameToValue.Add(new KeyValuePair<string, object>("Options", GetEnumText(typeof(ClientOptions), Input.Options)));
+				NameToValue.Add(new KeyValuePair<string, object>("Options", PerforceReflection.GetEnumText(typeof(ClientOptions), Input.Options)));
 			}
-			if(Input.SubmitOptions != ClientSubmitOptions.Unspecified)
+			if (Input.SubmitOptions != ClientSubmitOptions.Unspecified)
 			{
-				NameToValue.Add(new KeyValuePair<string, object>("SubmitOptions", GetEnumText(typeof(ClientSubmitOptions), Input.SubmitOptions)));
+				NameToValue.Add(new KeyValuePair<string, object>("SubmitOptions", PerforceReflection.GetEnumText(typeof(ClientSubmitOptions), Input.SubmitOptions)));
 			}
-			if(Input.LineEnd != ClientLineEndings.Unspecified)
+			if (Input.LineEnd != ClientLineEndings.Unspecified)
 			{
-				NameToValue.Add(new KeyValuePair<string, object>("LineEnd", GetEnumText(typeof(ClientLineEndings), Input.LineEnd)));
+				NameToValue.Add(new KeyValuePair<string, object>("LineEnd", PerforceReflection.GetEnumText(typeof(ClientLineEndings), Input.LineEnd)));
 			}
-			if(Input.Type != null)
+			if (Input.Type != null)
 			{
 				NameToValue.Add(new KeyValuePair<string, object>("Type", Input.Type));
 			}
-			if(Input.Stream != null)
+			if (Input.Stream != null)
 			{
 				NameToValue.Add(new KeyValuePair<string, object>("Stream", Input.Stream));
 			}
@@ -1030,10 +1058,38 @@ namespace EpicGames.Perforce
 		/// </summary>
 		/// <param name="Options">Options for this command</param>
 		/// <param name="UserName">List only client workspaces owned by this user.</param>
+		/// <param name="CancellationToken">Token used to cancel the operation</param>
+		/// <returns>Response from the server; either a client record or error code</returns>
+		public async Task<List<ClientsRecord>> GetClientsAsync(ClientsOptions Options, string? UserName, CancellationToken CancellationToken)
+		{
+			return (await TryGetClientsAsync(Options, UserName, CancellationToken)).Data;
+		}
+
+		/// <summary>
+		/// Queries the current client definition
+		/// </summary>
+		/// <param name="Options">Options for this command</param>
+		/// <param name="UserName">List only client workspaces owned by this user.</param>
+		/// <param name="CancellationToken">Token used to cancel the operation</param>
+		/// <returns>Response from the server; either a client record or error code</returns>
+		public Task<PerforceResponseList<ClientsRecord>> TryGetClientsAsync(ClientsOptions Options, string? UserName, CancellationToken CancellationToken)
+		{
+			return TryGetClientsAsync(Options, null, -1, null, UserName, CancellationToken);
+		}
+
+		/// <summary>
+		/// Queries the current client definition
+		/// </summary>
+		/// <param name="Options">Options for this command</param>
+		/// <param name="Filter">List only client workspaces matching filter. Treated as case sensitive if <ref>ClientsOptions.CaseSensitiveFilter</ref> is set.</param>
+		/// <param name="MaxResults">Limit the number of results to return. -1 for all.</param>
+		/// <param name="Stream">List client workspaces associated with the specified stream.</param>
+		/// <param name="UserName">List only client workspaces owned by this user.</param>
+		/// <param name="CancellationToken">Token used to cancel the operation</param>
 		/// <returns>Response from the server; either a client record or error code</returns>
-		public PerforceResponseList<ClientsRecord> Clients(ClientsOptions Options, string UserName)
+		public async Task<List<ClientsRecord>> GetClientsAsync(ClientsOptions Options, string? Filter, int MaxResults, string? Stream, string? UserName, CancellationToken CancellationToken)
 		{
-			return Clients(Options, null, -1, null, UserName);
+			return (await TryGetClientsAsync(Options, Filter, MaxResults, Stream, UserName, CancellationToken)).Data;
 		}
 
 		/// <summary>
@@ -1044,11 +1100,12 @@ namespace EpicGames.Perforce
 		/// <param name="MaxResults">Limit the number of results to return. -1 for all.</param>
 		/// <param name="Stream">List client workspaces associated with the specified stream.</param>
 		/// <param name="UserName">List only client workspaces owned by this user.</param>
+		/// <param name="CancellationToken">Token used to cancel the operation</param>
 		/// <returns>Response from the server; either a client record or error code</returns>
-		public PerforceResponseList<ClientsRecord> Clients(ClientsOptions Options, string Filter, int MaxResults, string Stream, string UserName)
+		public Task<PerforceResponseList<ClientsRecord>> TryGetClientsAsync(ClientsOptions Options, string? Filter, int MaxResults, string? Stream, string? UserName, CancellationToken CancellationToken)
 		{
 			StringBuilder Arguments = new StringBuilder("clients");
-			if((Options & ClientsOptions.All) != 0)
+			if ((Options & ClientsOptions.All) != 0)
 			{
 				Arguments.Append(" -a");
 			}
@@ -1063,27 +1120,27 @@ namespace EpicGames.Perforce
 					Arguments.AppendFormat(" -E \"{0}\"", Filter);
 				}
 			}
-			if(MaxResults != -1)
+			if (MaxResults != -1)
 			{
 				Arguments.AppendFormat(" -m {0}", MaxResults);
 			}
-			if(Stream != null)
+			if (Stream != null)
 			{
 				Arguments.AppendFormat(" -S \"{0}\"", Stream);
 			}
-			if((Options & ClientsOptions.WithTimes) != 0)
+			if ((Options & ClientsOptions.WithTimes) != 0)
 			{
 				Arguments.Append(" -t");
 			}
-			if(UserName != null)
+			if (UserName != null)
 			{
 				Arguments.AppendFormat(" -u \"{0}\"", UserName);
 			}
-			if((Options & ClientsOptions.Unloaded) != 0)
+			if ((Options & ClientsOptions.Unloaded) != 0)
 			{
 				Arguments.Append(" -U");
 			}
-			return Command<ClientsRecord>(Arguments.ToString(), null);
+			return CommandAsync<ClientsRecord>(Arguments.ToString(), null, CancellationToken);
 		}
 
 		#endregion
@@ -1095,31 +1152,44 @@ namespace EpicGames.Perforce
 		/// </summary>
 		/// <param name="Options">Options for the command</param>
 		/// <param name="FileSpecs">List of file specifications to query</param>
+		/// <param name="CancellationToken">Token used to cancel the operation</param>
+		/// <returns>List of responses from the server</returns>
+		public async Task<List<DeleteRecord>> DeleteAsync(int ChangeNumber, DeleteOptions Options, string[] FileSpecs, CancellationToken CancellationToken)
+		{
+			return (await TryDeleteAsync(ChangeNumber, Options, FileSpecs, CancellationToken)).Data;
+		}
+
+		/// <summary>
+		/// Execute the 'delete' command
+		/// </summary>
+		/// <param name="Options">Options for the command</param>
+		/// <param name="FileSpecs">List of file specifications to query</param>
+		/// <param name="CancellationToken">Token used to cancel the operation</param>
 		/// <returns>List of responses from the server</returns>
-		public PerforceResponseList<DeleteRecord> Delete(int ChangeNumber, DeleteOptions Options, params string[] FileSpecs)
+		public Task<PerforceResponseList<DeleteRecord>> TryDeleteAsync(int ChangeNumber, DeleteOptions Options, string[] FileSpecs, CancellationToken CancellationToken)
 		{
 			StringBuilder Arguments = new StringBuilder("delete");
-			if(ChangeNumber != -1)
+			if (ChangeNumber != -1)
 			{
 				Arguments.AppendFormat(" -c {0}", ChangeNumber);
 			}
-			if((Options & DeleteOptions.PreviewOnly) != 0)
+			if ((Options & DeleteOptions.PreviewOnly) != 0)
 			{
 				Arguments.Append(" -n");
 			}
-			if((Options & DeleteOptions.KeepWorkspaceFiles) != 0)
+			if ((Options & DeleteOptions.KeepWorkspaceFiles) != 0)
 			{
 				Arguments.Append(" -k");
 			}
-			if((Options & DeleteOptions.WithoutSyncing) != 0)
+			if ((Options & DeleteOptions.WithoutSyncing) != 0)
 			{
 				Arguments.Append(" -v");
 			}
-			foreach(string FileSpec in FileSpecs)
+			foreach (string FileSpec in FileSpecs)
 			{
 				Arguments.AppendFormat(" \"{0}\"", FileSpec);
 			}
-			return Command<DeleteRecord>(Arguments.ToString(), null);
+			return CommandAsync<DeleteRecord>(Arguments.ToString(), null, CancellationToken);
 		}
 
 		#endregion
@@ -1130,10 +1200,22 @@ namespace EpicGames.Perforce
 		/// Describes a single changelist
 		/// </summary>
 		/// <param name="ChangeNumber">The changelist number to retrieve description for</param>
+		/// <param name="CancellationToken">Token used to cancel the operation</param>
+		/// <returns>Response from the server; either a describe record or error code</returns>
+		public async Task<DescribeRecord> DescribeAsync(int ChangeNumber, CancellationToken CancellationToken)
+		{
+			return (await TryDescribeAsync(ChangeNumber, CancellationToken)).Data;
+		}
+
+		/// <summary>
+		/// Describes a single changelist
+		/// </summary>
+		/// <param name="ChangeNumber">The changelist number to retrieve description for</param>
+		/// <param name="CancellationToken">Token used to cancel the operation</param>
 		/// <returns>Response from the server; either a describe record or error code</returns>
-		public PerforceResponse<DescribeRecord> Describe(int ChangeNumber)
+		public async Task<PerforceResponse<DescribeRecord>> TryDescribeAsync(int ChangeNumber, CancellationToken CancellationToken)
 		{
-			PerforceResponseList<DescribeRecord> Records = Describe(new int[] { ChangeNumber });
+			PerforceResponseList<DescribeRecord> Records = await TryDescribeAsync(new int[] { ChangeNumber }, CancellationToken);
 			if (Records.Count != 1)
 			{
 				throw new PerforceException("Expected only one record returned from p4 describe command, got {0}", Records.Count);
@@ -1145,50 +1227,74 @@ namespace EpicGames.Perforce
 		/// Describes a set of changelists
 		/// </summary>
 		/// <param name="ChangeNumbers">The changelist numbers to retrieve descriptions for</param>
+		/// <param name="CancellationToken">Token used to cancel the operation</param>
+		/// <returns>List of responses from the server</returns>
+		public async Task<List<DescribeRecord>> DescribeAsync(int[] ChangeNumbers, CancellationToken CancellationToken)
+		{
+			return (await TryDescribeAsync(ChangeNumbers, CancellationToken)).Data;
+		}
+
+		/// <summary>
+		/// Describes a set of changelists
+		/// </summary>
+		/// <param name="ChangeNumbers">The changelist numbers to retrieve descriptions for</param>
+		/// <param name="CancellationToken">Token used to cancel the operation</param>
 		/// <returns>List of responses from the server</returns>
-		public PerforceResponseList<DescribeRecord> Describe(params int[] ChangeNumbers)
+		public Task<PerforceResponseList<DescribeRecord>> TryDescribeAsync(int[] ChangeNumbers, CancellationToken CancellationToken)
 		{
 			StringBuilder Arguments = new StringBuilder("describe -s");
-			foreach(int ChangeNumber in ChangeNumbers)
+			foreach (int ChangeNumber in ChangeNumbers)
 			{
 				Arguments.AppendFormat(" {0}", ChangeNumber);
 			}
-			return Command<DescribeRecord>(Arguments.ToString(), null);
+			return CommandAsync<DescribeRecord>(Arguments.ToString(), null, CancellationToken);
+		}
+
+		/// <summary>
+		/// Describes a set of changelists
+		/// </summary>
+		/// <param name="ChangeNumbers">The changelist numbers to retrieve descriptions for</param>
+		/// <param name="CancellationToken">Token used to cancel the operation</param>
+		/// <returns>List of responses from the server</returns>
+		public async Task<List<DescribeRecord>> DescribeAsync(DescribeOptions Options, int MaxNumFiles, int[] ChangeNumbers, CancellationToken CancellationToken)
+		{
+			return (await TryDescribeAsync(Options, MaxNumFiles, ChangeNumbers, CancellationToken)).Data;
 		}
 
 		/// <summary>
 		/// Describes a set of changelists
 		/// </summary>
 		/// <param name="ChangeNumbers">The changelist numbers to retrieve descriptions for</param>
+		/// <param name="CancellationToken">Token used to cancel the operation</param>
 		/// <returns>List of responses from the server</returns>
-		public PerforceResponseList<DescribeRecord> Describe(DescribeOptions Options, int MaxResults, params int[] ChangeNumbers)
+		public Task<PerforceResponseList<DescribeRecord>> TryDescribeAsync(DescribeOptions Options, int MaxNumFiles, int[] ChangeNumbers, CancellationToken CancellationToken)
 		{
 			StringBuilder Arguments = new StringBuilder("describe -s");
-			if((Options & DescribeOptions.ShowDescriptionForRestrictedChanges) != 0)
+			if ((Options & DescribeOptions.ShowDescriptionForRestrictedChanges) != 0)
 			{
 				Arguments.Append(" -f");
 			}
-			if((Options & DescribeOptions.Identity) != 0)
+			if ((Options & DescribeOptions.Identity) != 0)
 			{
 				Arguments.Append(" -I");
 			}
-			if(MaxResults != -1)
+			if (MaxNumFiles != -1)
 			{
-				Arguments.AppendFormat(" -m{0}", MaxResults);
+				Arguments.AppendFormat(" -m{0}", MaxNumFiles);
 			}
-			if((Options & DescribeOptions.OriginalChangeNumber) != 0)
+			if ((Options & DescribeOptions.OriginalChangeNumber) != 0)
 			{
 				Arguments.Append(" -O");
 			}
-			if((Options & DescribeOptions.Shelved) != 0)
+			if ((Options & DescribeOptions.Shelved) != 0)
 			{
 				Arguments.Append(" -S");
 			}
-			foreach(int ChangeNumber in ChangeNumbers)
+			foreach (int ChangeNumber in ChangeNumbers)
 			{
 				Arguments.AppendFormat(" {0}", ChangeNumber);
 			}
-			return Command<DescribeRecord>(Arguments.ToString(), null);
+			return CommandAsync<DescribeRecord>(Arguments.ToString(), null, CancellationToken);
 		}
 
 		#endregion
@@ -1200,10 +1306,23 @@ namespace EpicGames.Perforce
 		/// </summary>
 		/// <param name="ChangeNumber">Changelist to add files to</param>
 		/// <param name="FileSpecs">Files to be opened for edit</param>
+		/// <param name="CancellationToken">Token used to cancel the operation</param>
+		/// <returns>Response from the server</returns>
+		public async Task<List<EditRecord>> EditAsync(int ChangeNumber, string[] FileSpecs, CancellationToken CancellationToken)
+		{
+			return (await TryEditAsync(ChangeNumber, FileSpecs, CancellationToken)).Data;
+		}
+
+		/// <summary>
+		/// Opens files for edit
+		/// </summary>
+		/// <param name="ChangeNumber">Changelist to add files to</param>
+		/// <param name="FileSpecs">Files to be opened for edit</param>
+		/// <param name="CancellationToken">Token used to cancel the operation</param>
 		/// <returns>Response from the server</returns>
-		public PerforceResponseList<EditRecord> Edit(int ChangeNumber, params string[] FileSpecs)
+		public Task<PerforceResponseList<EditRecord>> TryEditAsync(int ChangeNumber, string[] FileSpecs, CancellationToken CancellationToken)
 		{
-			return Edit(ChangeNumber, null, EditOptions.None, FileSpecs);
+			return TryEditAsync(ChangeNumber, null, EditOptions.None, FileSpecs, CancellationToken);
 		}
 
 		/// <summary>
@@ -1213,32 +1332,47 @@ namespace EpicGames.Perforce
 		/// <param name="FileType">Type for new files</param>
 		/// <param name="Options">Options for the command</param>
 		/// <param name="FileSpecs">Files to be opened for edit</param>
+		/// <param name="CancellationToken">Token used to cancel the operation</param>
 		/// <returns>Response from the server</returns>
-		public PerforceResponseList<EditRecord> Edit(int ChangeNumber, string FileType, EditOptions Options, params string[] FileSpecs)
+		public async Task<List<EditRecord>> EditAsync(int ChangeNumber, string FileType, EditOptions Options, string[] FileSpecs, CancellationToken CancellationToken)
+		{
+			return (await TryEditAsync(ChangeNumber, FileType, Options, FileSpecs, CancellationToken)).Data;
+		}
+
+		/// <summary>
+		/// Opens files for edit
+		/// </summary>
+		/// <param name="ChangeNumber">Changelist to add files to</param>
+		/// <param name="FileType">Type for new files</param>
+		/// <param name="Options">Options for the command</param>
+		/// <param name="FileSpecs">Files to be opened for edit</param>
+		/// <param name="CancellationToken">Token used to cancel the operation</param>
+		/// <returns>Response from the server</returns>
+		public Task<PerforceResponseList<EditRecord>> TryEditAsync(int ChangeNumber, string? FileType, EditOptions Options, string[] FileSpecs, CancellationToken CancellationToken)
 		{
 			StringBuilder Arguments = new StringBuilder("edit");
-			if(ChangeNumber != -1)
+			if (ChangeNumber != -1)
 			{
 				Arguments.AppendFormat(" -c {0}", ChangeNumber);
 			}
-			if((Options & EditOptions.KeepWorkspaceFiles) != 0)
+			if ((Options & EditOptions.KeepWorkspaceFiles) != 0)
 			{
 				Arguments.Append(" -k");
 			}
-			if((Options & EditOptions.PreviewOnly) != 0)
+			if ((Options & EditOptions.PreviewOnly) != 0)
 			{
 				Arguments.Append(" -n");
 			}
-			if(FileType != null)
+			if (FileType != null)
 			{
 				Arguments.AppendFormat(" -t \"{0}\"", FileType);
 			}
-			foreach(string FileSpec in FileSpecs)
+			foreach (string FileSpec in FileSpecs)
 			{
 				Arguments.AppendFormat(" \"{0}\"", FileSpec);
 			}
 
-			return Command<EditRecord>(Arguments.ToString(), null);
+			return CommandAsync<EditRecord>(Arguments.ToString(), null, CancellationToken);
 		}
 
 		#endregion
@@ -1250,22 +1384,63 @@ namespace EpicGames.Perforce
 		/// </summary>
 		/// <param name="Options">Options for the command</param>
 		/// <param name="FileSpecs">List of file specifications to query</param>
+		/// <param name="CancellationToken">Token used to cancel the operation</param>
+		/// <returns>List of responses from the server</returns>
+		public async Task<List<FileLogRecord>> FileLogAsync(FileLogOptions Options, string[] FileSpecs, CancellationToken CancellationToken)
+		{
+			return (await TryFileLogAsync(Options, FileSpecs, CancellationToken)).Data;
+		}
+
+		/// <summary>
+		/// Execute the 'filelog' command
+		/// </summary>
+		/// <param name="Options">Options for the command</param>
+		/// <param name="FileSpecs">List of file specifications to query</param>
+		/// <param name="CancellationToken">Token used to cancel the operation</param>
+		/// <returns>List of responses from the server</returns>
+		public Task<PerforceResponseList<FileLogRecord>> TryFileLogAsync(FileLogOptions Options, string[] FileSpecs, CancellationToken CancellationToken)
+		{
+			return TryFileLogAsync(-1, -1, Options, FileSpecs, CancellationToken);
+		}
+
+		/// <summary>
+		/// Execute the 'filelog' command
+		/// </summary>
+		/// <param name="MaxChanges">Number of changelists to show. Ignored if zero or negative.</param>
+		/// <param name="Options">Options for the command</param>
+		/// <param name="FileSpecs">List of file specifications to query</param>
+		/// <param name="CancellationToken">Token used to cancel the operation</param>
+		/// <returns>List of responses from the server</returns>
+		public async Task<List<FileLogRecord>> FileLogAsync(int MaxChanges, FileLogOptions Options, string[] FileSpecs, CancellationToken CancellationToken)
+		{
+			return (await TryFileLogAsync(MaxChanges, Options, FileSpecs, CancellationToken)).Data;
+		}
+
+		/// <summary>
+		/// Execute the 'filelog' command
+		/// </summary>
+		/// <param name="MaxChanges">Number of changelists to show. Ignored if zero or negative.</param>
+		/// <param name="Options">Options for the command</param>
+		/// <param name="FileSpecs">List of file specifications to query</param>
+		/// <param name="CancellationToken">Token used to cancel the operation</param>
 		/// <returns>List of responses from the server</returns>
-		public PerforceResponseList<FileLogRecord> FileLog(FileLogOptions Options, params string[] FileSpecs)
+		public Task<PerforceResponseList<FileLogRecord>> TryFileLogAsync(int MaxChanges, FileLogOptions Options, string[] FileSpecs, CancellationToken CancellationToken)
 		{
-			return FileLog(-1, -1, Options, FileSpecs);
+			return TryFileLogAsync(-1, MaxChanges, Options, FileSpecs, CancellationToken);
 		}
 
 		/// <summary>
 		/// Execute the 'filelog' command
 		/// </summary>
+		/// <param name="ChangeNumber">Show only files modified by this changelist. Ignored if zero or negative.</param>
 		/// <param name="MaxChanges">Number of changelists to show. Ignored if zero or negative.</param>
 		/// <param name="Options">Options for the command</param>
 		/// <param name="FileSpecs">List of file specifications to query</param>
+		/// <param name="CancellationToken">Token used to cancel the operation</param>
 		/// <returns>List of responses from the server</returns>
-		public PerforceResponseList<FileLogRecord> FileLog(int MaxChanges, FileLogOptions Options, params string[] FileSpecs)
+		public async Task<List<FileLogRecord>> FileLogAsync(int ChangeNumber, int MaxChanges, FileLogOptions Options, string[] FileSpecs, CancellationToken CancellationToken)
 		{
-			return FileLog(-1, MaxChanges, Options, FileSpecs);
+			return (await TryFileLogAsync(ChangeNumber, MaxChanges, Options, FileSpecs, CancellationToken)).Data;
 		}
 
 		/// <summary>
@@ -1275,40 +1450,41 @@ namespace EpicGames.Perforce
 		/// <param name="MaxChanges">Number of changelists to show. Ignored if zero or negative.</param>
 		/// <param name="Options">Options for the command</param>
 		/// <param name="FileSpecs">List of file specifications to query</param>
+		/// <param name="CancellationToken">Token used to cancel the operation</param>
 		/// <returns>List of responses from the server</returns>
-		public PerforceResponseList<FileLogRecord> FileLog(int ChangeNumber, int MaxChanges, FileLogOptions Options, params string[] FileSpecs)
+		public Task<PerforceResponseList<FileLogRecord>> TryFileLogAsync(int ChangeNumber, int MaxChanges, FileLogOptions Options, string[] FileSpecs, CancellationToken CancellationToken)
 		{
 			// Build the argument list
 			StringBuilder Arguments = new StringBuilder("filelog");
-			if(ChangeNumber > 0)
+			if (ChangeNumber > 0)
 			{
 				Arguments.AppendFormat(" -c {0}", ChangeNumber);
 			}
-			if((Options & FileLogOptions.ContentHistory) != 0)
+			if ((Options & FileLogOptions.ContentHistory) != 0)
 			{
 				Arguments.Append(" -h");
 			}
-			if((Options & FileLogOptions.FollowAcrossBranches) != 0)
+			if ((Options & FileLogOptions.FollowAcrossBranches) != 0)
 			{
 				Arguments.Append(" -i");
 			}
-			if((Options & FileLogOptions.FullDescriptions) != 0)
+			if ((Options & FileLogOptions.FullDescriptions) != 0)
 			{
 				Arguments.Append(" -l");
 			}
-			if((Options & FileLogOptions.LongDescriptions) != 0)
+			if ((Options & FileLogOptions.LongDescriptions) != 0)
 			{
 				Arguments.Append(" -L");
 			}
-			if(MaxChanges > 0)
+			if (MaxChanges > 0)
 			{
 				Arguments.AppendFormat(" -m {0}", MaxChanges);
 			}
-			if((Options & FileLogOptions.DoNotFollowPromotedTaskStreams) != 0)
+			if ((Options & FileLogOptions.DoNotFollowPromotedTaskStreams) != 0)
 			{
 				Arguments.Append(" -p");
 			}
-			if((Options & FileLogOptions.IgnoreNonContributoryIntegrations) != 0)
+			if ((Options & FileLogOptions.IgnoreNonContributoryIntegrations) != 0)
 			{
 				Arguments.Append(" -s");
 			}
@@ -1316,14 +1492,8 @@ namespace EpicGames.Perforce
 			// Always include times to simplify parsing
 			Arguments.Append(" -t");
 
-			// Add the file arguments
-			foreach(string FileSpec in FileSpecs)
-			{
-				Arguments.AppendFormat(" \"{0}\"", FileSpec);
-			}
-
-			// Execute the command
-			return Command<FileLogRecord>(Arguments.ToString(), null);
+			// Append all the arguments
+			return BatchedCommandAsync<FileLogRecord>(Arguments.ToString(), FileSpecs, null, CancellationToken);
 		}
 
 		#endregion
@@ -1335,22 +1505,65 @@ namespace EpicGames.Perforce
 		/// </summary>
 		/// <param name="Options">Options for the command</param>
 		/// <param name="FileSpecs">List of file specifications to query</param>
+		/// <param name="CancellationToken">Token used to cancel the operation</param>
+		/// <returns>List of responses from the server</returns>
+		public async Task<List<FStatRecord>> FStatAsync(FStatOptions Options, string[] FileSpecs, CancellationToken CancellationToken)
+		{
+			return (await TryFStatAsync(Options, FileSpecs, CancellationToken)).Data;
+		}
+
+		/// <summary>
+		/// Execute the 'fstat' command
+		/// </summary>
+		/// <param name="Options">Options for the command</param>
+		/// <param name="FileSpecs">List of file specifications to query</param>
+		/// <param name="CancellationToken">Token used to cancel the operation</param>
+		/// <returns>List of responses from the server</returns>
+		public Task<PerforceResponseList<FStatRecord>> TryFStatAsync(FStatOptions Options, string[] FileSpecs, CancellationToken CancellationToken)
+		{
+			return TryFStatAsync(-1, Options, FileSpecs, CancellationToken);
+		}
+
+		/// <summary>
+		/// Execute the 'fstat' command
+		/// </summary>
+		/// <param name="MaxFiles">Produce fstat output for only the first max files.</param>
+		/// <param name="Options">Options for the command</param>
+		/// <param name="FileSpecs">List of file specifications to query</param>
+		/// <param name="CancellationToken">Token used to cancel the operation</param>
+		/// <returns>List of responses from the server</returns>
+		public async Task<List<FStatRecord>> FStatAsync(int MaxFiles, FStatOptions Options, string[] FileSpecs, CancellationToken CancellationToken)
+		{
+			return (await TryFStatAsync(MaxFiles, Options, FileSpecs, CancellationToken)).Data;
+		}
+
+		/// <summary>
+		/// Execute the 'fstat' command
+		/// </summary>
+		/// <param name="MaxFiles">Produce fstat output for only the first max files.</param>
+		/// <param name="Options">Options for the command</param>
+		/// <param name="FileSpecs">List of file specifications to query</param>
+		/// <param name="CancellationToken">Token used to cancel the operation</param>
 		/// <returns>List of responses from the server</returns>
-		public PerforceResponseList<FStatRecord> FStat(FStatOptions Options, params string[] FileSpecs)
+		public Task<PerforceResponseList<FStatRecord>> TryFStatAsync(int MaxFiles, FStatOptions Options, string[] FileSpecs, CancellationToken CancellationToken)
 		{
-			return FStat(-1, Options, FileSpecs);
+			return TryFStatAsync(-1, -1, null, MaxFiles, Options, FileSpecs, CancellationToken);
 		}
 
 		/// <summary>
 		/// Execute the 'fstat' command
 		/// </summary>
+		/// <param name="AfterChangeNumber">Return only files affected after the given changelist number.</param>
+		/// <param name="OnlyChangeNumber">Return only files affected by the given changelist number.</param>
+		/// <param name="Filter">List only those files that match the criteria specified.</param>
 		/// <param name="MaxFiles">Produce fstat output for only the first max files.</param>
 		/// <param name="Options">Options for the command</param>
 		/// <param name="FileSpecs">List of file specifications to query</param>
+		/// <param name="CancellationToken">Token used to cancel the operation</param>
 		/// <returns>List of responses from the server</returns>
-		public PerforceResponseList<FStatRecord> FStat(int MaxFiles, FStatOptions Options, params string[] FileSpecs)
+		public async Task<List<FStatRecord>> FStatAsync(int AfterChangeNumber, int OnlyChangeNumber, string? Filter, int MaxFiles, FStatOptions Options, string[] FileSpecs, CancellationToken CancellationToken)
 		{
-			return FStat(-1, -1, null, MaxFiles, Options, FileSpecs);
+			return (await TryFStatAsync(AfterChangeNumber, OnlyChangeNumber, Filter, MaxFiles, Options, FileSpecs, CancellationToken)).Data;
 		}
 
 		/// <summary>
@@ -1362,8 +1575,9 @@ namespace EpicGames.Perforce
 		/// <param name="MaxFiles">Produce fstat output for only the first max files.</param>
 		/// <param name="Options">Options for the command</param>
 		/// <param name="FileSpecs">List of file specifications to query</param>
+		/// <param name="CancellationToken">Token used to cancel the operation</param>
 		/// <returns>List of responses from the server</returns>
-		public PerforceResponseList<FStatRecord> FStat(int AfterChangeNumber, int OnlyChangeNumber, string Filter, int MaxFiles, FStatOptions Options, params string[] FileSpecs)
+		public Task<PerforceResponseList<FStatRecord>> TryFStatAsync(int AfterChangeNumber, int OnlyChangeNumber, string? Filter, int MaxFiles, FStatOptions Options, string[] FileSpecs, CancellationToken CancellationToken)
 		{
 			// Build the argument list
 			StringBuilder Arguments = new StringBuilder("fstat");
@@ -1460,14 +1674,8 @@ namespace EpicGames.Perforce
 				Arguments.Append(" -U");
 			}
 
-			// Add the file arguments
-			foreach(string FileSpec in FileSpecs)
-			{
-				Arguments.AppendFormat(" \"{0}\"", FileSpec);
-			}
-
 			// Execute the command
-			return Command<FStatRecord>(Arguments.ToString(), null);
+			return BatchedCommandAsync<FStatRecord>(Arguments.ToString(), FileSpecs, null, CancellationToken);
 		}
 
 		#endregion
@@ -1478,8 +1686,20 @@ namespace EpicGames.Perforce
 		/// Execute the 'info' command
 		/// </summary>
 		/// <param name="Options">Options for the command</param>
+		/// <param name="CancellationToken">Token used to cancel the operation</param>
+		/// <returns>Response from the server; an InfoRecord or error code</returns>
+		public async Task<InfoRecord> GetInfoAsync(InfoOptions Options, CancellationToken CancellationToken)
+		{
+			return (await TryGetInfoAsync(Options, CancellationToken)).Data;
+		}
+
+		/// <summary>
+		/// Execute the 'info' command
+		/// </summary>
+		/// <param name="Options">Options for the command</param>
+		/// <param name="CancellationToken">Token used to cancel the operation</param>
 		/// <returns>Response from the server; an InfoRecord or error code</returns>
-		public PerforceResponse<InfoRecord> Info(InfoOptions Options)
+		public Task<PerforceResponse<InfoRecord>> TryGetInfoAsync(InfoOptions Options, CancellationToken CancellationToken)
 		{
 			// Build the argument list
 			StringBuilder Arguments = new StringBuilder("info");
@@ -1487,7 +1707,91 @@ namespace EpicGames.Perforce
 			{
 				Arguments.Append(" -s");
 			}
-			return SingleResponseCommand<InfoRecord>(Arguments.ToString(), null);
+			return SingleResponseCommandAsync<InfoRecord>(Arguments.ToString(), null, CancellationToken);
+		}
+
+		#endregion
+
+		#region p4 login
+
+		/// <summary>
+		/// Execute the 'login' command
+		/// </summary>
+		/// <param name="Password">Password to use to login</param>
+		/// <param name="CancellationToken">Token used to cancel the operation</param>
+		/// <returns>Response from the server</returns>
+		public async Task LoginAsync(string Password, CancellationToken CancellationToken)
+		{
+			// Some versions of P4.EXE do not support marshalled output for P4 login calls, so we only support this as a basic text query.
+			byte[] PasswordBytes = Encoding.UTF8.GetBytes(Password);
+			using (ManagedProcessGroup ChildProcessGroup = new ManagedProcessGroup())
+			{
+				string Executable = PerforceChildProcess.GetExecutable();
+				string FullArgumentList = GetFullCommandLine("login");
+
+				Logger.LogDebug("Running {0} {1}", Executable, FullArgumentList);
+				using (ManagedProcess ChildProcess = new ManagedProcess(ChildProcessGroup, Executable, FullArgumentList, null, null, PasswordBytes, ProcessPriorityClass.Normal))
+				{
+					StringBuilder Lines = new StringBuilder("Unable to log in: ");
+					for(; ;)
+					{
+						string? Line = await ChildProcess.ReadLineAsync();
+						if(Line == null)
+						{
+							break;
+						}
+
+						Lines.Append($"\n  {Line}");
+
+						if(Line.StartsWith("Enter password:", StringComparison.OrdinalIgnoreCase))
+						{
+							Lines.Append($" {new string('*', Password.Length)}");
+						}
+					}
+
+					ChildProcess.WaitForExit();
+
+					if (ChildProcess.ExitCode != 0)
+					{
+						throw new PerforceException($"{Lines}\n  (exit code {ChildProcess.ExitCode})");
+					}
+				}
+			}
+		}
+
+		#endregion
+
+		#region p4 merge
+
+		/// <summary>
+		/// Execute the 'merge' command
+		/// </summary>
+		/// <param name="Options">Options for the merge</param>
+		/// <param name="Change"></param>
+		/// <param name="MaxFiles">Maximum number of files to merge</param>
+		/// <param name="SourceFileSpec">The source filespec and revision range</param>
+		/// <param name="TargetFileSpec">The target filespec</param>
+		/// <param name="CancellationToken">Cancellation token</param>
+		/// <returns>List of records</returns>
+		public Task<PerforceResponseList<MergeRecord>> MergeAsync(MergeOptions Options, int Change, int MaxFiles, string SourceFileSpec, string TargetFileSpec, CancellationToken CancellationToken)
+		{
+			StringBuilder Arguments = new StringBuilder("merge");
+			if ((Options & MergeOptions.Preview) != 0)
+			{
+				Arguments.AppendFormat(" -n");
+			}
+			if (Change != -1)
+			{
+				Arguments.AppendFormat(" -c {0}", Change);
+			}
+			if (MaxFiles != -1)
+			{
+				Arguments.AppendFormat(" -m {0}", MaxFiles);
+			}
+			Arguments.AppendFormat(" \"{0}\"", SourceFileSpec);
+			Arguments.AppendFormat(" \"{0}\"", TargetFileSpec);
+
+			return CommandAsync<MergeRecord>(Arguments.ToString(), null, CancellationToken);
 		}
 
 		#endregion
@@ -1503,8 +1807,25 @@ namespace EpicGames.Perforce
 		/// <param name="UserName">List only files that are opened by the given user</param>
 		/// <param name="MaxResults">Maximum number of results to return</param>
 		/// <param name="FileSpecs">Specification for the files to list</param>
+		/// <param name="CancellationToken">Token used to cancel the operation</param>
+		/// <returns>Response from the server</returns>
+		public async Task<List<FStatRecord>> GetOpenFilesAsync(OpenedOptions Options, int ChangeNumber, string? ClientName, string? UserName, int MaxResults, string[] FileSpecs, CancellationToken CancellationToken)
+		{
+			return (await TryGetOpenFilesAsync(Options, ChangeNumber, ClientName, UserName, MaxResults, FileSpecs, CancellationToken)).Data;
+		}
+
+		/// <summary>
+		/// Execute the 'opened' command
+		/// </summary>
+		/// <param name="Options">Options for the command</param>
+		/// <param name="ChangeNumber">List the files in pending changelist change. To list files in the default changelist, use DefaultChange.</param>
+		/// <param name="ClientName">List only files that are open in the given client</param>
+		/// <param name="UserName">List only files that are opened by the given user</param>
+		/// <param name="MaxResults">Maximum number of results to return</param>
+		/// <param name="FileSpecs">Specification for the files to list</param>
+		/// <param name="CancellationToken">Token used to cancel the operation</param>
 		/// <returns>Response from the server</returns>
-		public PerforceResponseList<FStatRecord> Opened(OpenedOptions Options, int ChangeNumber, string ClientName, string UserName, int MaxResults, params string[] FileSpecs)
+		public Task<PerforceResponseList<FStatRecord>> TryGetOpenFilesAsync(OpenedOptions Options, int ChangeNumber, string? ClientName, string? UserName, int MaxResults, string[] FileSpecs, CancellationToken CancellationToken)
 		{
 			// Build the argument list
 			StringBuilder Arguments = new StringBuilder("opened");
@@ -1524,7 +1845,7 @@ namespace EpicGames.Perforce
 			{
 				Arguments.AppendFormat(" -u \"{0}\"", UserName);
 			}
-			if(MaxResults == DefaultChange)
+			if(MaxResults == PerforceReflection.DefaultChange)
 			{
 				Arguments.AppendFormat(" -m default");
 			}
@@ -1540,32 +1861,45 @@ namespace EpicGames.Perforce
 			{
 				Arguments.AppendFormat(" \"{0}\"", FileSpec);
 			}
-			return Command<FStatRecord>(Arguments.ToString(), null);
+			return CommandAsync<FStatRecord>(Arguments.ToString(), null, CancellationToken);
+		}
+
+		#endregion
+
+		#region p4 print
+
+		/// <summary>
+		/// Execute the 'print' command
+		/// </summary>
+		/// <param name="OutputFile">Output file to redirect output to</param>
+		/// <param name="FileSpec">Specification for the files to print</param>
+		/// <param name="CancellationToken">Token used to cancel the operation</param>
+		/// <returns>Response from the server</returns>
+		public async Task<PrintRecord> PrintAsync(string OutputFile, string FileSpec, CancellationToken CancellationToken)
+		{
+			return (await TryPrintAsync(OutputFile, FileSpec, CancellationToken)).Data;
 		}
 
-		#endregion
-
-		#region p4 print
-
 		/// <summary>
 		/// Execute the 'print' command
 		/// </summary>
 		/// <param name="OutputFile">Output file to redirect output to</param>
 		/// <param name="FileSpec">Specification for the files to print</param>
+		/// <param name="CancellationToken">Token used to cancel the operation</param>
 		/// <returns>Response from the server</returns>
-		public PerforceResponse<PrintRecord> Print(string OutputFile, string FileSpec)
+		public Task<PerforceResponse<PrintRecord>> TryPrintAsync(string OutputFile, string FileSpec, CancellationToken CancellationToken)
 		{
 			// Build the argument list
 			StringBuilder Arguments = new StringBuilder("print");
 			Arguments.AppendFormat(" -o \"{0}\"", OutputFile);
 			Arguments.AppendFormat(" \"{0}\"", FileSpec);
-			return SingleResponseCommand<PrintRecord>(Arguments.ToString(), null);
+			return SingleResponseCommandAsync<PrintRecord>(Arguments.ToString(), null, CancellationToken);
 		}
 
 		class PrintHandler : IDisposable
 		{
 			Dictionary<string, string> DepotFileToLocalFile;
-			FileStream OutputStream;
+			FileStream? OutputStream;
 
 			public PrintHandler(Dictionary<string, string> DepotFileToLocalFile)
 			{
@@ -1603,9 +1937,9 @@ namespace EpicGames.Perforce
 				string Value = (string)Fields[0].Value;
 				if(Value == "stat")
 				{
-					string DepotFile = Fields.First(x => x.Key == "depotFile").Value.ToString();
+					string DepotFile = Fields.First(x => x.Key == "depotFile").Value.ToString() ?? String.Empty;
 
-					string LocalFile;
+					string? LocalFile;
 					if(!DepotFileToLocalFile.TryGetValue(DepotFile, out LocalFile))
 					{
 						throw new PerforceException("Depot file '{0}' not found in input dictionary", DepotFile);
@@ -1616,7 +1950,7 @@ namespace EpicGames.Perforce
 				else if(Value == "binary" || Value == "text")
 				{
 					byte[] Data = (byte[])Fields.First(x => x.Key == "data").Value;
-					OutputStream.Write(Data, 0, Data.Length);
+					OutputStream!.Write(Data, 0, Data.Length);
 				}
 				else
 				{
@@ -1625,44 +1959,32 @@ namespace EpicGames.Perforce
 			}
 		}
 
+		#endregion
+
+		#region p4 reconcile
+
 		/// <summary>
-		/// Execute the 'print' command
+		/// Open files for add, delete, and/or edit in order to reconcile a workspace with changes made outside of Perforce.
 		/// </summary>
-		/// <param name="OutputFile">Output file to redirect output to</param>
-		/// <param name="FileSpec">Specification for the files to print</param>
+		/// <param name="ChangeNumber">Changelist to open files to</param>
+		/// <param name="Options">Options for the command</param>
+		/// <param name="FileSpecs">Files to be reverted</param>
+		/// <param name="CancellationToken">Token used to cancel the operation</param>
 		/// <returns>Response from the server</returns>
-		public void Print(Dictionary<string, string> DepotFileToLocalFile)
+		public async Task<List<ReconcileRecord>> ReconcileAsync(int ChangeNumber, ReconcileOptions Options, string[] FileSpecs, CancellationToken CancellationToken)
 		{
-			string ListFileName = Path.GetTempFileName();
-			try
-			{
-				// Write the list of depot files
-				File.WriteAllLines(ListFileName, DepotFileToLocalFile.Keys);
-
-				// Execute Perforce, consuming the binary output into a memory stream
-				using(PrintHandler Handler = new PrintHandler(DepotFileToLocalFile))
-				{
-					Command(String.Format("-x \"{0}\" print", ListFileName), null, Handler.HandleRecord);
-				}
-			}
-			finally
-			{
-				File.Delete(ListFileName);
-			}
+			return (await TryReconcileAsync(ChangeNumber, Options, FileSpecs, CancellationToken)).Data;
 		}
 
-		#endregion
-
-		#region p4 reconcile
-		
 		/// <summary>
 		/// Open files for add, delete, and/or edit in order to reconcile a workspace with changes made outside of Perforce.
 		/// </summary>
 		/// <param name="ChangeNumber">Changelist to open files to</param>
 		/// <param name="Options">Options for the command</param>
 		/// <param name="FileSpecs">Files to be reverted</param>
+		/// <param name="CancellationToken">Token used to cancel the operation</param>
 		/// <returns>Response from the server</returns>
-		public PerforceResponseList<ReconcileRecord> Reconcile(int ChangeNumber, ReconcileOptions Options, params string[] FileSpecs)
+		public Task<PerforceResponseList<ReconcileRecord>> TryReconcileAsync(int ChangeNumber, ReconcileOptions Options, string[] FileSpecs, CancellationToken CancellationToken)
 		{
 			StringBuilder Arguments = new StringBuilder("reconcile");
 			if(ChangeNumber != -1)
@@ -1702,11 +2024,27 @@ namespace EpicGames.Perforce
 				Arguments.AppendFormat(" \"{0}\"", FileSpec);
 			}
 
-			return Command<ReconcileRecord>(Arguments.ToString(), null);
+			return CommandAsync<ReconcileRecord>(Arguments.ToString(), null, CancellationToken);
 		}
 
 		#endregion
 
+		#region p4 reload
+
+		/// <summary>
+		/// Reloads a client workspace
+		/// </summary>
+		/// <param name="ClientName">Name of the client to reload</param>
+		/// <param name="SourceServerId">The source server id</param>
+		/// <param name="CancellationToken">Token used to cancel the operation</param>
+		/// <returns>Response from the server</returns>
+		public Task<List<PerforceResponse>> ReloadClient(string ClientName, string SourceServerId, CancellationToken CancellationToken)
+		{
+			return CommandAsync($"reload -c \"{ClientName}\" -p \"{SourceServerId}\"", null, null, CancellationToken);
+		}
+
+		#endregion p4 reload
+
 		#region p4 resolve
 
 		/// <summary>
@@ -1715,8 +2053,22 @@ namespace EpicGames.Perforce
 		/// <param name="ChangeNumber">Changelist to open files to</param>
 		/// <param name="Options">Options for the command</param>
 		/// <param name="FileSpecs">Files to be reverted</param>
+		/// <param name="CancellationToken">Token used to cancel the operation</param>
+		/// <returns>Response from the server</returns>
+		public async Task<List<ResolveRecord>> ResolveAsync(int ChangeNumber, ResolveOptions Options, string[] FileSpecs, CancellationToken CancellationToken)
+		{
+			return (await TryResolveAsync(ChangeNumber, Options, FileSpecs, CancellationToken)).Data;
+		}
+
+		/// <summary>
+		/// Resolve conflicts between file revisions.
+		/// </summary>
+		/// <param name="ChangeNumber">Changelist to open files to</param>
+		/// <param name="Options">Options for the command</param>
+		/// <param name="FileSpecs">Files to be reverted</param>
+		/// <param name="CancellationToken">Token used to cancel the operation</param>
 		/// <returns>Response from the server</returns>
-		public PerforceResponseList<ResolveRecord> Resolve(int ChangeNumber, ResolveOptions Options, params string[] FileSpecs)
+		public Task<PerforceResponseList<ResolveRecord>> TryResolveAsync(int ChangeNumber, ResolveOptions Options, string[] FileSpecs, CancellationToken CancellationToken)
 		{
 			StringBuilder Arguments = new StringBuilder("resolve");
 			if ((Options & ResolveOptions.Automatic) != 0)
@@ -1764,7 +2116,7 @@ namespace EpicGames.Perforce
 				Arguments.AppendFormat(" \"{0}\"", FileSpec);
 			}
 
-			return Command<ResolveRecord>(Arguments.ToString(), null);
+			return CommandAsync<ResolveRecord>(Arguments.ToString(), null, CancellationToken);
 		}
 
 		#endregion
@@ -1778,8 +2130,23 @@ namespace EpicGames.Perforce
 		/// <param name="ClientName">Revert another user’s open files. </param>
 		/// <param name="Options">Options for the command</param>
 		/// <param name="FileSpecs">Files to be reverted</param>
+		/// <param name="CancellationToken">Token used to cancel the operation</param>
+		/// <returns>Response from the server</returns>
+		public async Task<List<RevertRecord>> RevertAsync(int ChangeNumber, string? ClientName, RevertOptions Options, string[] FileSpecs, CancellationToken CancellationToken)
+		{
+			return (await TryRevertAsync(ChangeNumber, ClientName, Options, FileSpecs, CancellationToken)).Data;
+		}
+
+		/// <summary>
+		/// Reverts files that have been added to a pending changelist.
+		/// </summary>
+		/// <param name="ChangeNumber">Changelist to add files to</param>
+		/// <param name="ClientName">Revert another user’s open files. </param>
+		/// <param name="Options">Options for the command</param>
+		/// <param name="FileSpecs">Files to be reverted</param>
+		/// <param name="CancellationToken">Token used to cancel the operation</param>
 		/// <returns>Response from the server</returns>
-		public PerforceResponseList<RevertRecord> Revert(int ChangeNumber, string ClientName, RevertOptions Options, params string[] FileSpecs)
+		public async Task<PerforceResponseList<RevertRecord>> TryRevertAsync(int ChangeNumber, string? ClientName, RevertOptions Options, string[] FileSpecs, CancellationToken CancellationToken)
 		{
 			StringBuilder Arguments = new StringBuilder("revert");
 			if((Options & RevertOptions.Unchanged) != 0)
@@ -1811,7 +2178,82 @@ namespace EpicGames.Perforce
 				Arguments.AppendFormat(" \"{0}\"", FileSpec);
 			}
 
-			return Command<RevertRecord>(Arguments.ToString(), null);
+			PerforceResponseList<RevertRecord> Records = await CommandAsync<RevertRecord>(Arguments.ToString(), null, CancellationToken);
+			if(Records.Count == 1 && Records[0].Error != null && Records[0].Error!.Generic == PerforceGenericCode.Empty)
+			{
+				Records.Clear();
+			}
+			return Records;
+		}
+
+		#endregion
+
+		#region p4 set
+
+		/// <summary>
+		/// Sets an environment variable
+		/// </summary>
+		/// <param name="Name">Name of the variable to set</param>
+		/// <param name="Value">Value for the variable</param>
+		/// <param name="CancellationToken">Token used to cancel the operation</param>
+		/// <returns>Response from the server</returns>
+		public async Task SetAsync(string Name, string Value, CancellationToken CancellationToken)
+		{
+			Tuple<bool, string> Result = await TrySetAsync(Name, Value, CancellationToken);
+			if (!Result.Item1)
+			{
+				throw new PerforceException(Result.Item2);
+			}
+		}
+
+		/// <summary>
+		/// Sets an environment variable
+		/// </summary>
+		/// <param name="Name">Name of the variable to set</param>
+		/// <param name="Value">Value for the variable</param>
+		/// <param name="CancellationToken">Token used to cancel the operation</param>
+		/// <returns>Response from the server</returns>
+		public async Task<Tuple<bool, string>> TrySetAsync(string Name, string Value, CancellationToken CancellationToken)
+		{
+			StringBuilder Arguments = new StringBuilder("set");
+			Arguments.Append($" {Name}=");
+			Arguments.Append($"\"{Value}\"");
+
+			using (PerforceChildProcess ChildProcess = new PerforceChildProcess(null, Arguments.ToString(), Logger))
+			{
+				return await ChildProcess.TryReadToEndAsync(CancellationToken);
+			}
+		}
+
+		/// <summary>
+		/// Gets the setting of a Perforce variable
+		/// </summary>
+		/// <param name="Name">Name of the variable to get</param>
+		/// <param name="CancellationToken">Cancellation token for the request</param>
+		/// <returns>Value of the variable</returns>
+		public async Task<string?> TryGetSettingAsync(string Name, CancellationToken CancellationToken)
+		{
+			using (PerforceChildProcess ChildProcess = new PerforceChildProcess(null, $"set {Name}=", Logger))
+			{
+				Tuple<bool, string> Response = await ChildProcess.TryReadToEndAsync(CancellationToken);
+				if (Response.Item1)
+				{
+					string TrimResponse = Response.Item2.Trim();
+					if (TrimResponse.Length == 0)
+					{
+						return null;
+					}
+
+					string Prefix = Name + "=";
+					if (TrimResponse.StartsWith(Prefix, StringComparison.OrdinalIgnoreCase))
+					{
+						return TrimResponse.Substring(Prefix.Length);
+					}
+				}
+
+				Logger.LogDebug($"Unable to get '{Name}' variable: {Response.Item2}");
+				return null;
+			}
 		}
 
 		#endregion
@@ -1824,8 +2266,22 @@ namespace EpicGames.Perforce
 		/// <param name="ChangeNumber">The change number to receive the shelved files</param>
 		/// <param name="Options">Options for the command</param>
 		/// <param name="FileSpecs">Files to sync</param>
+		/// <param name="CancellationToken">Token used to cancel the operation</param>
+		/// <returns>Response from the server</returns>
+		public async Task<List<ShelveRecord>> ShelveAsync(int ChangeNumber, ShelveOptions Options, string[] FileSpecs, CancellationToken CancellationToken)
+		{
+			return (await TryShelveAsync(ChangeNumber, Options, FileSpecs, CancellationToken)).Data;
+		}
+
+		/// <summary>
+		/// Shelves a set of files
+		/// </summary>
+		/// <param name="ChangeNumber">The change number to receive the shelved files</param>
+		/// <param name="Options">Options for the command</param>
+		/// <param name="FileSpecs">Files to sync</param>
+		/// <param name="CancellationToken">Token used to cancel the operation</param>
 		/// <returns>Response from the server</returns>
-		public PerforceResponseList<ShelveRecord> Shelve(int ChangeNumber, ShelveOptions Options, params string[] FileSpecs)
+		public Task<PerforceResponseList<ShelveRecord>> TryShelveAsync(int ChangeNumber, ShelveOptions Options, string[] FileSpecs, CancellationToken CancellationToken)
 		{
 			StringBuilder Arguments = new StringBuilder("shelve");
 			Arguments.AppendFormat(" -c {0}", ChangeNumber);
@@ -1842,7 +2298,25 @@ namespace EpicGames.Perforce
 				Arguments.AppendFormat(" \"{0}\"", FileSpec);
 			}
 
-			return Command<ShelveRecord>(Arguments.ToString(), null);
+			return CommandAsync<ShelveRecord>(Arguments.ToString(), null, CancellationToken);
+		}
+
+		/// <summary>
+		/// Deletes files from a shelved changelist
+		/// </summary>
+		/// <param name="ChangeNumber">Changelist containing shelved files to be deleted</param>
+		/// <param name="FileSpecs">Files to delete</param>
+		/// <param name="CancellationToken">Token used to cancel the operation</param>
+		/// <returns>Response from the server</returns>
+		public async Task DeleteShelvedFilesAsync(int ChangeNumber, string[] FileSpecs, CancellationToken CancellationToken)
+		{
+			List<PerforceResponse> Responses = await TryDeleteShelvedFilesAsync(ChangeNumber, FileSpecs, CancellationToken);
+
+			PerforceResponse? ErrorResponse = Responses.FirstOrDefault(x => x.Error != null && x.Error.Generic != PerforceGenericCode.Empty);
+			if(ErrorResponse != null)
+			{
+				throw new PerforceException(ErrorResponse.Error!);
+			}
 		}
 
 		/// <summary>
@@ -1850,8 +2324,9 @@ namespace EpicGames.Perforce
 		/// </summary>
 		/// <param name="ChangeNumber">Changelist containing shelved files to be deleted</param>
 		/// <param name="FileSpecs">Files to delete</param>
+		/// <param name="CancellationToken">Token used to cancel the operation</param>
 		/// <returns>Response from the server</returns>
-		public PerforceResponse DeleteShelvedFiles(int ChangeNumber, params string[] FileSpecs)
+		public Task<List<PerforceResponse>> TryDeleteShelvedFilesAsync(int ChangeNumber, string[] FileSpecs, CancellationToken CancellationToken)
 		{
 			StringBuilder Arguments = new StringBuilder("shelve -d");
 			if(ChangeNumber != -1)
@@ -1863,7 +2338,7 @@ namespace EpicGames.Perforce
 				Arguments.AppendFormat(" \"{0}\"", FileSpec);
 			}
 
-			return SingleResponseCommand(Arguments.ToString(), null, null);
+			return CommandAsync(Arguments.ToString(), null, null, CancellationToken);
 		}
 
 		#endregion
@@ -1874,10 +2349,36 @@ namespace EpicGames.Perforce
 		/// Enumerates all streams in a depot
 		/// </summary>
 		/// <param name="StreamPath">The path for streams to enumerate (eg. "//UE4/...")</param>
+		/// <param name="CancellationToken">Token used to cancel the operation</param>
+		/// <returns>List of streams matching the given criteria</returns>
+		public async Task<List<StreamsRecord>> GetStreamsAsync(string StreamPath, CancellationToken CancellationToken)
+		{
+			return (await TryGetStreamsAsync(StreamPath, CancellationToken)).Data;
+		}
+
+		/// <summary>
+		/// Enumerates all streams in a depot
+		/// </summary>
+		/// <param name="StreamPath">The path for streams to enumerate (eg. "//UE4/...")</param>
+		/// <param name="CancellationToken">Token used to cancel the operation</param>
+		/// <returns>List of streams matching the given criteria</returns>
+		public Task<PerforceResponseList<StreamsRecord>> TryGetStreamsAsync(string StreamPath, CancellationToken CancellationToken)
+		{
+			return TryGetStreamsAsync(StreamPath, -1, null, false, CancellationToken);
+		}
+
+		/// <summary>
+		/// Enumerates all streams in a depot
+		/// </summary>
+		/// <param name="StreamPath">The path for streams to enumerate (eg. "//UE4/...")</param>
+		/// <param name="MaxResults">Maximum number of results to return</param>
+		/// <param name="Filter">Additional filter to be applied to the results</param>
+		/// <param name="bUnloaded">Whether to enumerate unloaded workspaces</param>
+		/// <param name="CancellationToken">Token used to cancel the operation</param>
 		/// <returns>List of streams matching the given criteria</returns>
-		public PerforceResponseList<StreamRecord> Streams(string StreamPath)
+		public async Task<List<StreamsRecord>> GetStreamsAsync(string StreamPath, int MaxResults, string? Filter, bool bUnloaded, CancellationToken CancellationToken)
 		{
-			return Streams(StreamPath, -1, null, false);
+			return (await TryGetStreamsAsync(StreamPath, MaxResults, Filter, bUnloaded, CancellationToken)).Data;
 		}
 
 		/// <summary>
@@ -1887,8 +2388,9 @@ namespace EpicGames.Perforce
 		/// <param name="MaxResults">Maximum number of results to return</param>
 		/// <param name="Filter">Additional filter to be applied to the results</param>
 		/// <param name="bUnloaded">Whether to enumerate unloaded workspaces</param>
+		/// <param name="CancellationToken">Token used to cancel the operation</param>
 		/// <returns>List of streams matching the given criteria</returns>
-		public PerforceResponseList<StreamRecord> Streams(string StreamPath, int MaxResults, string Filter, bool bUnloaded)
+		public Task<PerforceResponseList<StreamsRecord>> TryGetStreamsAsync(string StreamPath, int MaxResults, string? Filter, bool bUnloaded, CancellationToken CancellationToken)
 		{
 			// Build the command line
 			StringBuilder Arguments = new StringBuilder("streams");
@@ -1907,7 +2409,7 @@ namespace EpicGames.Perforce
 			Arguments.AppendFormat(" \"{0}\"", StreamPath);
 
 			// Execute the command
-			return Command<StreamRecord>(Arguments.ToString(), null);
+			return CommandAsync<StreamsRecord>(Arguments.ToString(), null, CancellationToken);
 		}
 
 		#endregion
@@ -1919,8 +2421,21 @@ namespace EpicGames.Perforce
 		/// </summary>
 		/// <param name="ChangeNumber">The changelist to submit</param>
 		/// <param name="Options">Options for the command</param>
+		/// <param name="CancellationToken">Token used to cancel the operation</param>
+		/// <returns>Response from the server</returns>
+		public async Task<SubmitRecord> SubmitAsync(int ChangeNumber, SubmitOptions Options, CancellationToken CancellationToken)
+		{
+			return (await TrySubmitAsync(ChangeNumber, Options, CancellationToken)).Data;
+		}
+
+		/// <summary>
+		/// Submits a pending changelist
+		/// </summary>
+		/// <param name="ChangeNumber">The changelist to submit</param>
+		/// <param name="Options">Options for the command</param>
+		/// <param name="CancellationToken">Token used to cancel the operation</param>
 		/// <returns>Response from the server</returns>
-		public PerforceResponse<SubmitRecord> Submit(int ChangeNumber, SubmitOptions Options)
+		public async Task<PerforceResponse<SubmitRecord>> TrySubmitAsync(int ChangeNumber, SubmitOptions Options, CancellationToken CancellationToken)
 		{
 			StringBuilder Arguments = new StringBuilder("submit");
 			if((Options & SubmitOptions.ReopenAsEdit) != 0)
@@ -1929,7 +2444,7 @@ namespace EpicGames.Perforce
 			}
 			Arguments.AppendFormat(" -c {0}", ChangeNumber);
 
-			return Command<SubmitRecord>(Arguments.ToString(), null)[0];
+			return (await CommandAsync<SubmitRecord>(Arguments.ToString(), null, CancellationToken))[0];
 		}
 
 		#endregion
@@ -1940,10 +2455,35 @@ namespace EpicGames.Perforce
 		/// Syncs files from the server
 		/// </summary>
 		/// <param name="FileSpecs">Files to sync</param>
+		/// <param name="CancellationToken">Token used to cancel the operation</param>
+		/// <returns>Response from the server</returns>
+		public async Task<List<SyncRecord>> SyncAsync(string[] FileSpecs, CancellationToken CancellationToken)
+		{
+			return (await TrySyncAsync(FileSpecs, CancellationToken)).Data;
+		}
+
+		/// <summary>
+		/// Syncs files from the server
+		/// </summary>
+		/// <param name="FileSpecs">Files to sync</param>
+		/// <param name="CancellationToken">Token used to cancel the operation</param>
+		/// <returns>Response from the server</returns>
+		public Task<PerforceResponseList<SyncRecord>> TrySyncAsync(string[] FileSpecs, CancellationToken CancellationToken)
+		{
+			return TrySyncAsync(SyncOptions.None, -1, FileSpecs, CancellationToken);
+		}
+
+		/// <summary>
+		/// Syncs files from the server
+		/// </summary>
+		/// <param name="Options">Options for the command</param>
+		/// <param name="MaxFiles">Syncs only the first number of files specified.</param>
+		/// <param name="FileSpecs">Files to sync</param>
+		/// <param name="CancellationToken">Token used to cancel the operation</param>
 		/// <returns>Response from the server</returns>
-		public PerforceResponseList<SyncRecord> Sync(params string[] FileSpecs)
+		public async Task<List<SyncRecord>> SyncAsync(SyncOptions Options, int MaxFiles, string[] FileSpecs, CancellationToken CancellationToken)
 		{
-			return Sync(SyncOptions.None, -1, FileSpecs);
+			return (await TrySyncAsync(Options, MaxFiles, FileSpecs, CancellationToken)).Data;
 		}
 
 		/// <summary>
@@ -1952,10 +2492,29 @@ namespace EpicGames.Perforce
 		/// <param name="Options">Options for the command</param>
 		/// <param name="MaxFiles">Syncs only the first number of files specified.</param>
 		/// <param name="FileSpecs">Files to sync</param>
+		/// <param name="CancellationToken">Token used to cancel the operation</param>
+		/// <returns>Response from the server</returns>
+		public Task<PerforceResponseList<SyncRecord>> TrySyncAsync(SyncOptions Options, int MaxFiles, string[] FileSpecs, CancellationToken CancellationToken)
+		{
+			return TrySyncAsync(Options, -1, -1, -1, -1, -1, -1, FileSpecs, CancellationToken);
+		}
+
+		/// <summary>
+		/// Syncs files from the server
+		/// </summary>
+		/// <param name="Options">Options for the command</param>
+		/// <param name="MaxFiles">Syncs only the first number of files specified</param>
+		/// <param name="NumThreads">Sync in parallel using the given number of threads</param>
+		/// <param name="Batch">The number of files in a batch</param>
+		/// <param name="BatchSize">The number of bytes in a batch</param>
+		/// <param name="Min">Minimum number of files in a parallel sync</param>
+		/// <param name="MinSize">Minimum number of bytes in a parallel sync</param>
+		/// <param name="FileSpecs">Files to sync</param>
+		/// <param name="CancellationToken">Token used to cancel the operation</param>
 		/// <returns>Response from the server</returns>
-		public PerforceResponseList<SyncRecord> Sync(SyncOptions Options, int MaxFiles, params string[] FileSpecs)
+		public async Task<List<SyncRecord>> SyncAsync(SyncOptions Options, int MaxFiles, int NumThreads, int Batch, int BatchSize, int Min, int MinSize, string[] FileSpecs, CancellationToken CancellationToken)
 		{
-			return Sync(Options, -1, -1, -1, -1, -1, -1, FileSpecs);
+			return (await TrySyncAsync(Options, MaxFiles, NumThreads, Batch, BatchSize, Min, MinSize, FileSpecs, CancellationToken)).Data;
 		}
 
 		/// <summary>
@@ -1969,11 +2528,27 @@ namespace EpicGames.Perforce
 		/// <param name="Min">Minimum number of files in a parallel sync</param>
 		/// <param name="MinSize">Minimum number of bytes in a parallel sync</param>
 		/// <param name="FileSpecs">Files to sync</param>
+		/// <param name="CancellationToken">Token used to cancel the operation</param>
+		/// <returns>Response from the server</returns>
+		public async Task<PerforceResponseList<SyncRecord>> TrySyncAsync(SyncOptions Options, int MaxFiles, int NumThreads, int Batch, int BatchSize, int Min, int MinSize, string[] FileSpecs, CancellationToken CancellationToken)
+		{
+			// Perforce annoyingly returns 'up-to-date' as an error. Ignore it.
+			PerforceResponseList<SyncRecord> Records = await SyncInternalAsync<SyncRecord>(Options, MaxFiles, NumThreads, Batch, BatchSize, Min, MinSize, FileSpecs, false, CancellationToken);
+			Records.RemoveAll(x => x.Error != null && x.Error.Generic == PerforceGenericCode.Empty);
+			return Records;
+		}
+
+		/// <summary>
+		/// Syncs files from the server
+		/// </summary>
+		/// <param name="Options">Options for the command</param>
+		/// <param name="MaxFiles">Syncs only the first number of files specified.</param>
+		/// <param name="FileSpecs">Files to sync</param>
+		/// <param name="CancellationToken">Token used to cancel the operation</param>
 		/// <returns>Response from the server</returns>
-		public PerforceResponseList<SyncRecord> Sync(SyncOptions Options, int MaxFiles, int NumThreads, int Batch, int BatchSize, int Min, int MinSize, params string[] FileSpecs)
+		public async Task<List<SyncSummaryRecord>> SyncQuietAsync(SyncOptions Options, int MaxFiles, string[] FileSpecs, CancellationToken CancellationToken)
 		{
-			string Arguments = GetSyncArguments(Options, MaxFiles, NumThreads, Batch, BatchSize, Min, MinSize, FileSpecs, false);
-			return Command<SyncRecord>(Arguments.ToString(), null);
+			return (await TrySyncQuietAsync(Options, MaxFiles, FileSpecs, CancellationToken)).Data;
 		}
 
 		/// <summary>
@@ -1982,10 +2557,30 @@ namespace EpicGames.Perforce
 		/// <param name="Options">Options for the command</param>
 		/// <param name="MaxFiles">Syncs only the first number of files specified.</param>
 		/// <param name="FileSpecs">Files to sync</param>
+		/// <param name="CancellationToken">Token used to cancel the operation</param>
+		/// <returns>Response from the server</returns>
+		public Task<PerforceResponseList<SyncSummaryRecord>> TrySyncQuietAsync(SyncOptions Options, int MaxFiles, string[] FileSpecs, CancellationToken CancellationToken)
+		{
+			return TrySyncQuietAsync(Options, MaxFiles, -1, -1, -1, -1, -1, FileSpecs, null, CancellationToken);
+		}
+
+		/// <summary>
+		/// Syncs files from the server without returning detailed file info
+		/// </summary>
+		/// <param name="Options">Options for the command</param>
+		/// <param name="MaxFiles">Syncs only the first number of files specified</param>
+		/// <param name="NumThreads">Sync in parallel using the given number of threads</param>
+		/// <param name="Batch">The number of files in a batch</param>
+		/// <param name="BatchSize">The number of bytes in a batch</param>
+		/// <param name="Min">Minimum number of files in a parallel sync</param>
+		/// <param name="MinSize">Minimum number of bytes in a parallel sync</param>
+		/// <param name="FileSpecs">Files to sync</param>
+		/// <param name="ResponseFile">Response file for list of files to sync</param>
+		/// <param name="CancellationToken">Token used to cancel the operation</param>
 		/// <returns>Response from the server</returns>
-		public PerforceResponseList<SyncSummaryRecord> SyncQuiet(SyncOptions Options, int MaxFiles, params string[] FileSpecs)
+		public async Task<List<SyncSummaryRecord>> SyncQuietAsync(SyncOptions Options, int MaxFiles, int NumThreads, int Batch, int BatchSize, int Min, int MinSize, string[] FileSpecs, FileReference? ResponseFile, CancellationToken CancellationToken)
 		{
-			return SyncQuiet(Options, -1, -1, -1, -1, -1, -1, FileSpecs);
+			return (await TrySyncQuietAsync(Options, MaxFiles, NumThreads, Batch, BatchSize, Min, MinSize, FileSpecs, ResponseFile, CancellationToken)).Data;
 		}
 
 		/// <summary>
@@ -1999,11 +2594,12 @@ namespace EpicGames.Perforce
 		/// <param name="Min">Minimum number of files in a parallel sync</param>
 		/// <param name="MinSize">Minimum number of bytes in a parallel sync</param>
 		/// <param name="FileSpecs">Files to sync</param>
+		/// <param name="ResponseFile">Response file for list of files to sync</param>
+		/// <param name="CancellationToken">Token used to cancel the operation</param>
 		/// <returns>Response from the server</returns>
-		public PerforceResponseList<SyncSummaryRecord> SyncQuiet(SyncOptions Options, int MaxFiles, int NumThreads, int Batch, int BatchSize, int Min, int MinSize, params string[] FileSpecs)
+		public Task<PerforceResponseList<SyncSummaryRecord>> TrySyncQuietAsync(SyncOptions Options, int MaxFiles, int NumThreads, int Batch, int BatchSize, int Min, int MinSize, string[] FileSpecs, FileReference? ResponseFile, CancellationToken CancellationToken)
 		{
-			string Arguments = GetSyncArguments(Options, MaxFiles, NumThreads, Batch, BatchSize, Min, MinSize, FileSpecs, true);
-			return Command<SyncSummaryRecord>(Arguments.ToString(), null);
+			return SyncInternalAsync<SyncSummaryRecord>(Options, MaxFiles, NumThreads, Batch, BatchSize, Min, MinSize, FileSpecs, true, CancellationToken);
 		}
 
 		/// <summary>
@@ -2018,8 +2614,9 @@ namespace EpicGames.Perforce
 		/// <param name="MinSize">Minimum number of bytes in a parallel sync</param>
 		/// <param name="FileSpecs">Files to sync</param>
 		/// <param name="bQuiet">Whether to use quiet output</param>
+		/// <param name="CancellationToken">Token used to cancel the operation</param>
 		/// <returns>Arguments for the command</returns>
-		static string GetSyncArguments(SyncOptions Options, int MaxFiles, int NumThreads, int Batch, int BatchSize, int Min, int MinSize, string[] FileSpecs, bool bQuiet)
+		private async Task<PerforceResponseList<T>> SyncInternalAsync<T>(SyncOptions Options, int MaxFiles, int NumThreads, int Batch, int BatchSize, int Min, int MinSize, string[] FileSpecs, bool bQuiet, CancellationToken CancellationToken) where T : class
 		{
 			StringBuilder Arguments = new StringBuilder("sync");
 			if((Options & SyncOptions.Force) != 0)
@@ -2062,31 +2659,47 @@ namespace EpicGames.Perforce
 			{
 				Arguments.AppendFormat(" -m {0}", MaxFiles);
 			}
-/*			if(NumThreads != -1)
+			/*			if(NumThreads != -1)
+						{
+							Arguments.AppendFormat(" --parallel-threads={0}", NumThreads);
+							if(Batch != -1)
+							{
+								Arguments.AppendFormat(",batch={0}", Batch);
+							}
+							if(BatchSize != -1)
+							{
+								Arguments.AppendFormat(",batchsize={0}", BatchSize);
+							}
+							if(Min != -1)
+							{
+								Arguments.AppendFormat(",min={0}", Min);
+							}
+							if(MinSize != -1)
+							{
+								Arguments.AppendFormat(",minsize={0}", MinSize);
+							}
+						}*/
+			if (FileSpecs.Length < 10)
 			{
-				Arguments.AppendFormat(" --parallel-threads={0}", NumThreads);
-				if(Batch != -1)
-				{
-					Arguments.AppendFormat(",batch={0}", Batch);
-				}
-				if(BatchSize != -1)
+				foreach (string FileSpec in FileSpecs)
 				{
-					Arguments.AppendFormat(",batchsize={0}", BatchSize);
+					Arguments.AppendFormat(" \"{0}\"", FileSpec);
 				}
-				if(Min != -1)
+				return await CommandAsync<T>(Arguments.ToString(), null, CancellationToken);
+			}
+			else
+			{
+				string TempFileName = Path.GetTempFileName();
+				try
 				{
-					Arguments.AppendFormat(",min={0}", Min);
+					Arguments.Insert(0, $"-x\"{TempFileName}\" ");
+					return await CommandAsync<T>(Arguments.ToString(), null, CancellationToken);
 				}
-				if(MinSize != -1)
+				finally
 				{
-					Arguments.AppendFormat(",minsize={0}", MinSize);
+					File.Delete(TempFileName);
 				}
-			}*/
-			foreach (string FileSpec in FileSpecs)
-			{
-				Arguments.AppendFormat(" \"{0}\"", FileSpec);
 			}
-			return Arguments.ToString();
 		}
 
 		#endregion
@@ -2103,8 +2716,26 @@ namespace EpicGames.Perforce
 		/// <param name="ForceParentStream">Unshelve to the specified parent stream. Overrides the parent defined in the source stream specification.</param>
 		/// <param name="Options">Options for the command</param>
 		/// <param name="FileSpecs">Files to unshelve</param>
+		/// <param name="CancellationToken">Token used to cancel the operation</param>
+		/// <returns>Response from the server</returns>
+		public async Task<List<UnshelveRecord>> UnshelveAsync(int ChangeNumber, int IntoChangeNumber, string UsingBranchSpec, string UsingStream, string ForceParentStream, UnshelveOptions Options, string[] FileSpecs, CancellationToken CancellationToken)
+		{
+			return (await TryUnshelveAsync(ChangeNumber, IntoChangeNumber, UsingBranchSpec, UsingStream, ForceParentStream, Options, FileSpecs, CancellationToken)).Data;
+		}
+
+		/// <summary>
+		/// Restore shelved files from a pending change into a workspace
+		/// </summary>
+		/// <param name="ChangeNumber">The changelist containing shelved files</param>
+		/// <param name="IntoChangeNumber">The changelist to receive the unshelved files</param>
+		/// <param name="UsingBranchSpec">The branchspec to use when unshelving files</param>
+		/// <param name="UsingStream">Specifies the use of a stream-derived branch view to map the shelved files between the specified stream and its parent stream.</param>
+		/// <param name="ForceParentStream">Unshelve to the specified parent stream. Overrides the parent defined in the source stream specification.</param>
+		/// <param name="Options">Options for the command</param>
+		/// <param name="FileSpecs">Files to unshelve</param>
+		/// <param name="CancellationToken">Token used to cancel the operation</param>
 		/// <returns>Response from the server</returns>
-		public PerforceResponseList<UnshelveRecord> Unshelve(int ChangeNumber, int IntoChangeNumber, string UsingBranchSpec, string UsingStream, string ForceParentStream, UnshelveOptions Options, params string[] FileSpecs)
+		public Task<PerforceResponseList<UnshelveRecord>> TryUnshelveAsync(int ChangeNumber, int IntoChangeNumber, string UsingBranchSpec, string UsingStream, string ForceParentStream, UnshelveOptions Options, string[] FileSpecs, CancellationToken CancellationToken)
 		{
 			StringBuilder Arguments = new StringBuilder("unshelve");
 			Arguments.AppendFormat(" -s {0}", ChangeNumber);
@@ -2137,7 +2768,7 @@ namespace EpicGames.Perforce
 				Arguments.AppendFormat(" \"{0}\"", FileSpec);
 			}
 
-			return Command<UnshelveRecord>(Arguments.ToString(), null);
+			return CommandAsync<UnshelveRecord>(Arguments.ToString(), null, CancellationToken);
 		}
 
 		#endregion
@@ -2148,12 +2779,24 @@ namespace EpicGames.Perforce
 		/// Enumerates all streams in a depot
 		/// </summary>
 		/// <param name="UserName">Name of the user to fetch information for</param>
+		/// <param name="CancellationToken">Token used to cancel the operation</param>
+		/// <returns>Response from the server</returns>
+		public async Task<UserRecord> GetUserAsync(string UserName, CancellationToken CancellationToken)
+		{
+			return (await TryGetUserAsync(UserName, CancellationToken)).Data;
+		}
+
+		/// <summary>
+		/// Enumerates all streams in a depot
+		/// </summary>
+		/// <param name="UserName">Name of the user to fetch information for</param>
+		/// <param name="CancellationToken">Token used to cancel the operation</param>
 		/// <returns>Response from the server</returns>
-		public PerforceResponse<UserRecord> User(string UserName)
+		public Task<PerforceResponse<UserRecord>> TryGetUserAsync(string UserName, CancellationToken CancellationToken)
 		{
 			StringBuilder Arguments = new StringBuilder("user");
 			Arguments.AppendFormat(" -o \"{0}\"", UserName);
-			return SingleResponseCommand<UserRecord>(Arguments.ToString(), null);
+			return SingleResponseCommandAsync<UserRecord>(Arguments.ToString(), null, CancellationToken);
 		}
 
 		#endregion
@@ -2164,15 +2807,37 @@ namespace EpicGames.Perforce
 		/// Retrieves the location of a file of set of files in the workspace
 		/// </summary>
 		/// <param name="FileSpecs">Patterns for the files to query</param>
+		/// <param name="CancellationToken">Token used to cancel the operation</param>
 		/// <returns>List of responses from the server</returns>
-		public PerforceResponseList<WhereRecord> Where(params string[] FileSpecs)
+		public async Task<List<WhereRecord>> WhereAsync(string[] FileSpecs, CancellationToken CancellationToken)
 		{
-			StringBuilder Arguments = new StringBuilder("where");
-			foreach(string FileSpec in FileSpecs)
+			List<WhereRecord> Records = (await TryWhere(FileSpecs, CancellationToken)).Data;
+			if (Records.Count != FileSpecs.Length)
 			{
-				Arguments.AppendFormat(" \"{0}\"", FileSpec);
+				throw new PerforceException($"Unexpected response count; expected {FileSpecs.Length}, got {Records.Count}.");
+			}
+			return Records;
+		}
+
+		/// <summary>
+		/// Retrieves the location of a file of set of files in the workspace
+		/// </summary>
+		/// <param name="FileSpecs">Patterns for the files to query</param>
+		/// <param name="CancellationToken">Token used to cancel the operation</param>
+		/// <returns>List of responses from the server</returns>
+		public async Task<PerforceResponseList<WhereRecord>> TryWhere(string[] FileSpecs, CancellationToken CancellationToken)
+		{
+			PerforceResponseList<WhereRecord> Responses = new PerforceResponseList<WhereRecord>();
+			for (int Idx = 0; Idx < FileSpecs.Length; )
+			{
+				StringBuilder Arguments = new StringBuilder("where");
+				for (; Arguments.Length < 2048 && Idx < FileSpecs.Length; Idx++)
+				{
+					Arguments.AppendFormat(" \"{0}\"", FileSpecs[Idx]);
+				}
+				Responses.AddRange(await CommandAsync<WhereRecord>(Arguments.ToString(), null, CancellationToken));
 			}
-			return Command<WhereRecord>(Arguments.ToString(), null);
+			return Responses;
 		}
 
 		#endregion
diff --git a/Engine/Source/Programs/Shared/EpicGames.Perforce/PerforceError.cs b/Engine/Source/Programs/Shared/EpicGames.Perforce/PerforceError.cs
--- a/Engine/Source/Programs/Shared/EpicGames.Perforce/PerforceError.cs
+++ b/Engine/Source/Programs/Shared/EpicGames.Perforce/PerforceError.cs
@@ -1,4 +1,4 @@
-// Copyright Epic Games, Inc. All Rights Reserved.
+// Copyright 1998-2019 Epic Games, Inc. All Rights Reserved.
 
 using System;
 using System.Collections.Generic;
@@ -69,6 +69,14 @@ namespace EpicGames.Perforce
 		[PerforceTagAttribute("data")]
 		public string Data;
 
+		/// <summary>
+		/// Private constructor for serialization
+		/// </summary>
+		private PerforceError()
+		{
+			Data = null!;
+		}
+
 		/// <summary>
 		/// Formats this error for display in the debugger
 		/// </summary>
diff --git a/Engine/Source/Programs/Shared/EpicGames.Perforce/PerforceException.cs b/Engine/Source/Programs/Shared/EpicGames.Perforce/PerforceException.cs
--- a/Engine/Source/Programs/Shared/EpicGames.Perforce/PerforceException.cs
+++ b/Engine/Source/Programs/Shared/EpicGames.Perforce/PerforceException.cs
@@ -1,4 +1,4 @@
-// Copyright Epic Games, Inc. All Rights Reserved.
+// Copyright 1998-2019 Epic Games, Inc. All Rights Reserved.
 
 using System;
 using System.Collections.Generic;
@@ -16,7 +16,7 @@ namespace EpicGames.Perforce
 		/// <summary>
 		/// For errors returned by the server, contains the error record
 		/// </summary>
-		public PerforceError Error;
+		public PerforceError? Error;
 
 		/// <summary>
 		/// Constructor
diff --git a/Engine/Source/Programs/Shared/EpicGames.Perforce/PerforceInfo.cs b/Engine/Source/Programs/Shared/EpicGames.Perforce/PerforceInfo.cs
--- a/Engine/Source/Programs/Shared/EpicGames.Perforce/PerforceInfo.cs
+++ b/Engine/Source/Programs/Shared/EpicGames.Perforce/PerforceInfo.cs
@@ -1,4 +1,4 @@
-// Copyright Epic Games, Inc. All Rights Reserved.
+// Copyright 1998-2019 Epic Games, Inc. All Rights Reserved.
 
 using System;
 using System.Collections.Generic;
@@ -25,11 +25,19 @@ namespace EpicGames.Perforce
 		[PerforceTag("data")]
 		public string Data;
 
+		/// <summary>
+		/// Private constructor for serialization
+		/// </summary>
+		private PerforceInfo()
+		{
+			Data = null!;
+		}
+
 		/// <summary>
 		/// Formats this error for display in the debugger
 		/// </summary>
 		/// <returns>String representation of this object</returns>
-		public override string ToString()
+		public override string? ToString()
 		{
 			return Data;
 		}
diff --git a/Engine/Source/Programs/Shared/EpicGames.Perforce/PerforceRecord.cs b/Engine/Source/Programs/Shared/EpicGames.Perforce/PerforceRecord.cs
new file mode 100644
--- /dev/null
+++ b/Engine/Source/Programs/Shared/EpicGames.Perforce/PerforceRecord.cs
@@ -0,0 +1,263 @@
+using EpicGames.Core;
+using System;
+using System.Buffers.Binary;
+using System.Buffers.Text;
+using System.Collections.Generic;
+using System.Text;
+using System.Text.RegularExpressions;
+using System.Text.Unicode;
+
+namespace EpicGames.Perforce
+{
+	/// <summary>
+	/// The type of a value returned by Perforce
+	/// </summary>
+	public enum PerforceValueType
+	{
+		/// <summary>
+		/// A utf-8 encoded string
+		/// </summary>
+		String,
+
+		/// <summary>
+		/// A 32-bit integer
+		/// </summary>
+		Integer,
+	}
+
+	/// <summary>
+	/// Wrapper for values returned by Perforce
+	/// </summary>
+	public struct PerforceValue
+	{
+		/// <summary>
+		/// The raw data for the value, including type, size, and payload
+		/// </summary>
+		public readonly ReadOnlyMemory<byte> Data;
+
+		/// <summary>
+		/// Constructor
+		/// </summary>
+		/// <param name="Data">The data to construct this value from</param>
+		public PerforceValue(ReadOnlyMemory<byte> Data)
+		{
+			this.Data = Data;
+		}
+
+		/// <summary>
+		/// Determines if the value is empty
+		/// </summary>
+		/// <returns>True if the value is empty</returns>
+		public bool IsEmpty
+		{
+			get { return Data.IsEmpty; }
+		}
+
+		/// <summary>
+		/// Accessor for the type of value stored by this struct
+		/// </summary>
+		public PerforceValueType Type
+		{
+			get
+			{
+				switch (Data.Span[0])
+				{
+					case (byte)'s':
+						return PerforceValueType.String;
+					case (byte)'i':
+						return PerforceValueType.Integer;
+					default:
+						throw new NotImplementedException("Unknown/unsupported value type");
+				}
+			}
+		}
+
+		/// <summary>
+		/// Converts the value to a boolean
+		/// </summary>
+		/// <returns>The boolean value</returns>
+		public bool AsBool()
+		{
+			ReadOnlyUtf8String String = GetString();
+			return String.Length == 0 || String == StringConstants.True;
+		}
+
+		/// <summary>
+		/// Converts the value to an integer
+		/// </summary>
+		/// <returns>Integer value</returns>
+		public int AsInteger()
+		{
+			if (Type == PerforceValueType.Integer)
+			{
+				return GetInteger();
+			}
+			else if (Type == PerforceValueType.String)
+			{
+				ReadOnlyUtf8String String = GetString();
+
+				int Value;
+				int BytesConsumed;
+				if (Utf8Parser.TryParse(String.Span, out Value, out BytesConsumed) && BytesConsumed == String.Length)
+				{
+					return Value;
+				}
+				else if (String == StringConstants.New || String == StringConstants.None)
+				{
+					return -1;
+				}
+				else if (String.Length > 0 && String[0] == '#' && Utf8Parser.TryParse(String.Span.Slice(1), out Value, out BytesConsumed) && BytesConsumed + 1 == String.Length)
+				{
+					return Value;
+				}
+				else if (String == StringConstants.Default)
+				{
+					return PerforceReflection.DefaultChange;
+				}
+				else
+				{
+					throw new PerforceException($"Unable to parse {String} as an integer");
+				}
+			}
+			else
+			{
+				throw new NotImplementedException();
+			}
+		}
+
+		/// <summary>
+		/// Converts this value to a long
+		/// </summary>
+		/// <returns>The converted value</returns>
+		public long AsLong()
+		{
+			ReadOnlyUtf8String String = AsString();
+
+			long Value;
+			int BytesConsumed;
+			if (!Utf8Parser.TryParse(AsString().Span, out Value, out BytesConsumed) || BytesConsumed != String.Length)
+			{
+				throw new PerforceException($"Unable to parse {ToString()} as a long value");
+			}
+
+			return Value;
+		}
+
+		/// <summary>
+		/// Converts this value to a DateTimeOffset
+		/// </summary>
+		/// <returns>The converted value</returns>
+		public DateTimeOffset AsDateTimeOffset()
+		{
+			string Text = ToString();
+			return DateTimeOffset.Parse(Regex.Replace(Text, "[a-zA-Z ]*$", "")); // Strip timezone name (eg. "EST")
+		}
+
+		/// <summary>
+		/// Converts this value to a string
+		/// </summary>
+		/// <returns>The converted value</returns>
+		public ReadOnlyUtf8String AsString()
+		{
+			switch(Type)
+			{
+				case PerforceValueType.String:
+					return GetString();
+				case PerforceValueType.Integer:
+					return new ReadOnlyUtf8String(GetInteger().ToString());
+				default:
+					throw new NotImplementedException();
+			}
+		}
+
+		/// <summary>
+		/// Gets the contents of this string
+		/// </summary>
+		/// <returns></returns>
+		public int GetInteger()
+		{
+			if (Type != PerforceValueType.Integer)
+			{
+				throw new InvalidOperationException("Value is not an integer");
+			}
+			return BinaryPrimitives.ReadInt32LittleEndian(Data.Span.Slice(1));
+		}
+
+		/// <summary>
+		/// Gets the contents of this string
+		/// </summary>
+		/// <returns></returns>
+		public ReadOnlyUtf8String GetString()
+		{
+			if(Type != PerforceValueType.String)
+			{
+				throw new InvalidOperationException("Value is not a string");
+			}
+			return new ReadOnlyUtf8String(Data.Slice(5));
+		}
+
+		/// <summary>
+		/// Convert to a string
+		/// </summary>
+		/// <returns>String representation of the value</returns>
+		public override string ToString()
+		{
+			return AsString().ToString();
+		}
+	}
+
+	/// <summary>
+	/// Low-overhead record type for generic responses
+	/// </summary>
+	public struct PerforceRecord
+	{
+		/// <summary>
+		/// The rows in this record
+		/// </summary>
+		public List<KeyValuePair<ReadOnlyUtf8String, PerforceValue>> Rows;
+
+		/// <summary>
+		/// Copy this record into the given array of values. This method is O(n) if every record key being in the list of keys in the same order, but O(n^2) if not.
+		/// </summary>
+		/// <param name="Keys">List of keys to parse</param>
+		/// <param name="Values">Array to receive the list of values</param>
+		public void CopyInto(ReadOnlyUtf8String[] Keys, PerforceValue[] Values)
+		{
+			// Clear out the current values array
+			for (int ValueIdx = 0; ValueIdx < Values.Length; ValueIdx++)
+			{
+				Values[ValueIdx] = new PerforceValue();
+			}
+
+			// Parse all the keys
+			int RowIdx = 0;
+			while (RowIdx < Rows.Count)
+			{
+				int InitialRowIdx = RowIdx;
+
+				// Find the key that matches this row.
+				ReadOnlyUtf8String RowKey = Rows[RowIdx].Key;
+				for (int KeyIdx = 0; KeyIdx < Keys.Length; KeyIdx++)
+				{
+					ReadOnlyUtf8String Key = Keys[KeyIdx];
+					if (RowKey == Key)
+					{
+						// Copy the value to the output array, then move to the next row, and try to match that against the next key.
+						Values[KeyIdx] = Rows[RowIdx].Value;
+						if (++RowIdx == Rows.Count)
+						{
+							break;
+						}
+						RowKey = Rows[RowIdx].Key;
+					}
+				}
+
+				// If we didn't match any key name for this row, move to the next one
+				if(RowIdx == InitialRowIdx)
+				{
+					RowIdx++;
+				}
+			}
+		}
+	}
+}
diff --git a/Engine/Source/Programs/Shared/EpicGames.Perforce/PerforceReflection.cs b/Engine/Source/Programs/Shared/EpicGames.Perforce/PerforceReflection.cs
new file mode 100644
--- /dev/null
+++ b/Engine/Source/Programs/Shared/EpicGames.Perforce/PerforceReflection.cs
@@ -0,0 +1,565 @@
+using System;
+using System.Buffers.Text;
+using System.Collections.Concurrent;
+using System.Collections.Generic;
+using System.Diagnostics.CodeAnalysis;
+using System.Linq;
+using System.Reflection;
+using System.Reflection.Emit;
+using System.Reflection.Metadata.Ecma335;
+using System.Text;
+using System.Text.RegularExpressions;
+using EpicGames.Core;
+
+namespace EpicGames.Perforce
+{
+	/// <summary>
+	/// String constants for perforce values
+	/// </summary>
+	static class StringConstants
+	{
+		public static readonly ReadOnlyUtf8String True = new ReadOnlyUtf8String("true");
+		public static readonly ReadOnlyUtf8String New = new ReadOnlyUtf8String("new");
+		public static readonly ReadOnlyUtf8String None = new ReadOnlyUtf8String("none");
+		public static readonly ReadOnlyUtf8String Default = new ReadOnlyUtf8String("default");
+	}
+
+	/// <summary>
+	/// Stores cached information about a field with a P4Tag attribute
+	/// </summary>
+	class CachedTagInfo
+	{
+		/// <summary>
+		/// Name of the tag. Specified in the attribute or inferred from the field name.
+		/// </summary>
+		public ReadOnlyUtf8String Name;
+
+		/// <summary>
+		/// Whether this tag is optional or not.
+		/// </summary>
+		public bool Optional;
+
+		/// <summary>
+		/// The field containing the value of this data.
+		/// </summary>
+		public FieldInfo Field;
+
+		/// <summary>
+		/// Parser for this field type
+		/// </summary>
+		public Action<object, int> SetFromInteger;
+
+		/// <summary>
+		/// Parser for this field type
+		/// </summary>
+		public Action<object, ReadOnlyUtf8String> SetFromString;
+
+		/// <summary>
+		/// Index into the bitmask of required types
+		/// </summary>
+		public ulong RequiredTagBitMask;
+
+		/// <summary>
+		/// Constructor
+		/// </summary>
+		/// <param name="Name"></param>
+		/// <param name="Optional"></param>
+		/// <param name="Field"></param>
+		/// <param name="RequiredTagBitMask"></param>
+		public CachedTagInfo(ReadOnlyUtf8String Name, bool Optional, FieldInfo Field, ulong RequiredTagBitMask)
+		{
+			this.Name = Name;
+			this.Optional = Optional;
+			this.Field = Field;
+			this.RequiredTagBitMask = RequiredTagBitMask;
+			this.SetFromInteger = (Obj, Value) => throw new PerforceException($"Field {Name} was not expecting an integer value.");
+			this.SetFromString = (Obj, String) => throw new PerforceException($"Field {Name} was not expecting a string value.");
+		}
+	}
+
+	/// <summary>
+	/// Stores cached information about a record
+	/// </summary>
+	class CachedRecordInfo
+	{
+		/// <summary>
+		/// Delegate type for creating a record instance
+		/// </summary>
+		/// <returns>New instance</returns>
+		public delegate object CreateRecordDelegate();
+
+		/// <summary>
+		/// Type of the record
+		/// </summary>
+		public Type Type;
+
+		/// <summary>
+		/// Method to construct this record
+		/// </summary>
+		public CreateRecordDelegate CreateInstance;
+
+		/// <summary>
+		/// List of fields in the record. These should be ordered to match P4 output for maximum efficiency.
+		/// </summary>
+		public List<CachedTagInfo> Fields = new List<CachedTagInfo>();
+
+		/// <summary>
+		/// Map of name to tag info
+		/// </summary>
+		public Dictionary<ReadOnlyUtf8String, CachedTagInfo> NameToInfo = new Dictionary<ReadOnlyUtf8String, CachedTagInfo>();
+
+		/// <summary>
+		/// Bitmask of all the required tags. Formed by bitwise-or'ing the RequiredTagBitMask fields for each required CachedTagInfo.
+		/// </summary>
+		public ulong RequiredTagsBitMask;
+
+		/// <summary>
+		/// The type of records to create for subelements
+		/// </summary>
+		public Type? SubElementType;
+
+		/// <summary>
+		/// The cached record info for the subelement type
+		/// </summary>
+		public CachedRecordInfo? SubElementRecordInfo;
+
+		/// <summary>
+		/// Field containing subelements
+		/// </summary>
+		public FieldInfo? SubElementField;
+
+		/// <summary>
+		/// Constructor
+		/// </summary>
+		/// <param name="Type">The record type</param>
+		public CachedRecordInfo(Type Type)
+		{
+			this.Type = Type;
+
+			ConstructorInfo? Constructor = Type.GetConstructor(BindingFlags.Public | BindingFlags.NonPublic | BindingFlags.Instance, null, Type.EmptyTypes, null);
+			if(Constructor == null)
+			{
+				throw new PerforceException($"Unable to find default constructor for {Type}");
+			}
+
+			DynamicMethod DynamicMethod = new DynamicMethod("_", Type, null);
+			ILGenerator Generator = DynamicMethod.GetILGenerator();
+			Generator.Emit(OpCodes.Newobj, Constructor);
+			Generator.Emit(OpCodes.Ret);
+			CreateInstance = (CreateRecordDelegate)DynamicMethod.CreateDelegate(typeof(CreateRecordDelegate));
+		}
+	}
+
+	/// <summary>
+	/// Information about an enum
+	/// </summary>
+	class CachedEnumInfo
+	{
+		/// <summary>
+		/// The enum type
+		/// </summary>
+		public Type EnumType;
+
+		/// <summary>
+		/// Whether the enum has the [Flags] attribute
+		/// </summary>
+		public bool bHasFlagsAttribute;
+
+		/// <summary>
+		/// Map of name to value
+		/// </summary>
+		public Dictionary<ReadOnlyUtf8String, int> NameToValue = new Dictionary<ReadOnlyUtf8String, int>();
+
+		/// <summary>
+		/// List of name/value pairs
+		/// </summary>
+		public List<KeyValuePair<string, int>> NameValuePairs = new List<KeyValuePair<string, int>>();
+
+		/// <summary>
+		/// Constructor
+		/// </summary>
+		/// <param name="EnumType">The type to construct from</param>
+		public CachedEnumInfo(Type EnumType)
+		{
+			this.EnumType = EnumType;
+
+			bHasFlagsAttribute = EnumType.GetCustomAttribute<FlagsAttribute>() != null;
+
+			FieldInfo[] Fields = EnumType.GetFields(BindingFlags.Public | BindingFlags.Static);
+			foreach (FieldInfo Field in Fields)
+			{
+				PerforceEnumAttribute? Attribute = Field.GetCustomAttribute<PerforceEnumAttribute>();
+				if (Attribute != null)
+				{
+					object? Value = Field.GetValue(null);
+					if (Value != null)
+					{
+						ReadOnlyUtf8String Name = new ReadOnlyUtf8String(Attribute.Name);
+						NameToValue[Name] = (int)Value;
+
+						NameValuePairs.Add(new KeyValuePair<string, int>(Attribute.Name, (int)Value));
+					}
+				}
+			}
+		}
+
+		/// <summary>
+		/// Parses the given integer as an enum
+		/// </summary>
+		/// <param name="Value">The value to convert to an enum</param>
+		/// <returns>The enum value corresponding to the given value</returns>
+		public object ParseInteger(int Value)
+		{
+			return Enum.ToObject(EnumType, Value);
+		}
+
+		/// <summary>
+		/// Parses the given text as an enum
+		/// </summary>
+		/// <param name="Text">The text to parse</param>
+		/// <returns>The enum value corresponding to the given text</returns>
+		public object ParseString(ReadOnlyUtf8String Text)
+		{
+			return Enum.ToObject(EnumType, ParseToInteger(Text));
+		}
+			
+		/// <summary>
+		/// Parses the given text as an enum
+		/// </summary>
+		/// <param name="Text"></param>
+		/// <returns></returns>
+		public int ParseToInteger(ReadOnlyUtf8String Name)
+		{
+			if (bHasFlagsAttribute)
+			{
+				int Result = 0;
+				for (int Offset = 0; Offset < Name.Length;)
+				{
+					if (Name.Span[Offset] == (byte)' ')
+					{
+						Offset++;
+					}
+					else
+					{
+						// Find the end of this name
+						int StartOffset = ++Offset;
+						while (Offset < Name.Length && Name.Span[Offset] != (byte)' ')
+						{
+							Offset++;
+						}
+
+						// Take the subset
+						ReadOnlyUtf8String Item = Name.Slice(StartOffset, Offset - StartOffset);
+
+						// Get the value
+						int ItemValue;
+						if (NameToValue.TryGetValue(Item, out ItemValue))
+						{
+							Result |= ItemValue;
+						}
+					}
+				}
+				return Result;
+			}
+			else
+			{
+				int Result;
+				NameToValue.TryGetValue(Name, out Result);
+				return Result;
+			}
+		}
+
+		/// <summary>
+		/// Parses an enum value, using PerforceEnumAttribute markup for names.
+		/// </summary>
+		/// <param name="Value">Value of the enum.</param>
+		/// <returns>Text for the enum.</returns>
+		public string GetEnumText(int Value)
+		{
+			if (bHasFlagsAttribute)
+			{
+				List<string> Names = new List<string>();
+
+				int CombinedIntegerValue = 0;
+				foreach (KeyValuePair<string, int> Pair in NameValuePairs)
+				{
+					if ((Value & Pair.Value) != 0)
+					{
+						Names.Add(Pair.Key);
+						CombinedIntegerValue |= Pair.Value;
+					}
+				}
+
+				if (CombinedIntegerValue != Value)
+				{
+					throw new ArgumentException($"Invalid enum value {Value}");
+				}
+
+				return String.Join(" ", Names);
+			}
+			else
+			{
+				string? Name = null;
+				foreach (KeyValuePair<string, int> Pair in NameValuePairs)
+				{
+					if (Value == Pair.Value)
+					{
+						Name = Pair.Key;
+						break;
+					}
+				}
+
+				if (Name == null)
+				{
+					throw new ArgumentException($"Invalid enum value {Value}");
+				}
+				return Name;
+			}
+		}
+	}
+
+	/// <summary>
+	/// Utility methods for converting to/from native types
+	/// </summary>
+	static class PerforceReflection
+	{
+		/// <summary>
+		/// Unix epoch; used for converting times back into C# datetime objects
+		/// </summary>
+		public static readonly DateTime UnixEpoch = new DateTime(1970, 1, 1, 0, 0, 0, DateTimeKind.Utc);
+
+		/// <summary>
+		/// Constant for the default changelist, where valid.
+		/// </summary>
+		public const int DefaultChange = -2;
+
+		/// <summary>
+		/// Cached map of enum types to a lookup mapping from p4 strings to enum values.
+		/// </summary>
+		static ConcurrentDictionary<Type, CachedEnumInfo> EnumTypeToInfo = new ConcurrentDictionary<Type, CachedEnumInfo>();
+
+		/// <summary>
+		/// Cached set of record 
+		/// </summary>
+		static ConcurrentDictionary<Type, CachedRecordInfo> RecordTypeToInfo = new ConcurrentDictionary<Type, CachedRecordInfo>();
+
+		/// <summary>
+		/// Default type for info
+		/// </summary>
+		public static CachedRecordInfo InfoRecordInfo = GetCachedRecordInfo(typeof(PerforceInfo));
+
+		/// <summary>
+		/// Default type for errors
+		/// </summary>
+		public static CachedRecordInfo ErrorRecordInfo = GetCachedRecordInfo(typeof(PerforceError));
+
+		/// <summary>
+		/// Gets a mapping of flags to enum values for the given type
+		/// </summary>
+		/// <param name="EnumType">The enum type to retrieve flags for</param>
+		/// <returns>Map of name to enum value</returns>
+		static CachedEnumInfo GetCachedEnumInfo(Type EnumType)
+		{
+			CachedEnumInfo? EnumInfo;
+			if (!EnumTypeToInfo.TryGetValue(EnumType, out EnumInfo))
+			{
+				EnumInfo = new CachedEnumInfo(EnumType);
+				if (!EnumTypeToInfo.TryAdd(EnumType, EnumInfo))
+				{
+					EnumInfo = EnumTypeToInfo[EnumType];
+				}
+			}
+			return EnumInfo;
+		}
+
+		/// <summary>
+		/// Parses an enum value, using PerforceEnumAttribute markup for names.
+		/// </summary>
+		/// <param name="EnumType">Type of the enum to parse.</param>
+		/// <param name="Value">Value of the enum.</param>
+		/// <returns>Text for the enum.</returns>
+		public static string GetEnumText(Type EnumType, object Value)
+		{
+			return GetCachedEnumInfo(EnumType).GetEnumText((int)Value);
+		}
+
+		/// <summary>
+		/// Gets reflection data for the given record type
+		/// </summary>
+		/// <param name="RecordType">The type to retrieve record info for</param>
+		/// <returns>The cached reflection information for the given type</returns>
+		public static CachedRecordInfo GetCachedRecordInfo(Type RecordType)
+		{
+			CachedRecordInfo? Record;
+			if (!RecordTypeToInfo.TryGetValue(RecordType, out Record))
+			{
+				Record = new CachedRecordInfo(RecordType);
+
+				// Get all the fields for this type
+				FieldInfo[] Fields = RecordType.GetFields(BindingFlags.Public | BindingFlags.NonPublic | BindingFlags.Instance);
+
+				// Build the map of all tags for this record
+				foreach (FieldInfo Field in Fields)
+				{
+					PerforceTagAttribute? TagAttribute = Field.GetCustomAttribute<PerforceTagAttribute>();
+					if (TagAttribute != null)
+					{
+						string TagName = TagAttribute.Name ?? Field.Name;
+
+						ulong RequiredTagBitMask = 0;
+						if (!TagAttribute.Optional)
+						{
+							RequiredTagBitMask = Record.RequiredTagsBitMask + 1;
+							if (RequiredTagBitMask == 0)
+							{
+								throw new PerforceException("Too many required tags in {0}; max is {1}", RecordType.Name, sizeof(ulong) * 8);
+							}
+							Record.RequiredTagsBitMask |= RequiredTagBitMask;
+						}
+
+						CachedTagInfo TagInfo = new CachedTagInfo(new ReadOnlyUtf8String(TagName), TagAttribute.Optional, Field, RequiredTagBitMask);
+
+						Type FieldType = Field.FieldType;
+
+						FieldInfo FieldCopy = Field;
+						if (FieldType == typeof(DateTime))
+						{
+							TagInfo.SetFromString = (Obj, String) => FieldCopy.SetValue(Obj, ParseStringAsDateTime(String));
+						}
+						else if (FieldType == typeof(bool))
+						{
+							TagInfo.SetFromString = (Obj, String) => FieldCopy.SetValue(Obj, ParseStringAsBool(String));
+						}
+						else if (FieldType == typeof(Nullable<bool>))
+						{
+							TagInfo.SetFromString = (Obj, String) => FieldCopy.SetValue(Obj, ParseStringAsNullableBool(String));
+						}
+						else if (FieldType == typeof(int))
+						{
+							TagInfo.SetFromInteger = (Obj, Int) => FieldCopy.SetValue(Obj, Int);
+							TagInfo.SetFromString = (Obj, String) => FieldCopy.SetValue(Obj, ParseStringAsInt(String));
+						}
+						else if (FieldType == typeof(long))
+						{
+							TagInfo.SetFromString = (Obj, String) => FieldCopy.SetValue(Obj, ParseStringAsLong(String));
+						}
+						else if (FieldType == typeof(string))
+						{
+							TagInfo.SetFromString = (Obj, String) => FieldCopy.SetValue(Obj, ParseString(String));
+						}
+						else if (FieldType.IsEnum)
+						{
+							CachedEnumInfo EnumInfo = GetCachedEnumInfo(FieldType);
+							TagInfo.SetFromInteger = (Obj, Int) => FieldCopy.SetValue(Obj, EnumInfo.ParseInteger(Int));
+							TagInfo.SetFromString = (Obj, String) => FieldCopy.SetValue(Obj, EnumInfo.ParseString(String));
+						}
+						else if (FieldType == typeof(DateTimeOffset?))
+						{
+							TagInfo.SetFromString = (Obj, String) => FieldCopy.SetValue(Obj, ParseStringAsNullableDateTimeOffset(String));
+						}
+						else if (FieldType == typeof(List<string>))
+						{
+							TagInfo.SetFromString = (Obj, String) => ((List<string>)FieldCopy.GetValue(Obj)!).Add(String.ToString());
+						}
+						else
+						{
+							throw new PerforceException("Unsupported type of {0}.{1} for tag '{2}'", RecordType.Name, FieldType.Name, TagName);
+						}
+
+						Record.Fields.Add(TagInfo);
+					}
+
+					Record.NameToInfo = Record.Fields.ToDictionary(x => x.Name, x => x);
+
+					PerforceRecordListAttribute? SubElementAttribute = Field.GetCustomAttribute<PerforceRecordListAttribute>();
+					if (SubElementAttribute != null)
+					{
+						Record.SubElementField = Field;
+						Record.SubElementType = Field.FieldType.GenericTypeArguments[0];
+						Record.SubElementRecordInfo = GetCachedRecordInfo(Record.SubElementType);
+					}
+				}
+
+				// Try to save the record info, or get the version that's already in the cache
+				if (!RecordTypeToInfo.TryAdd(RecordType, Record))
+				{
+					Record = RecordTypeToInfo[RecordType];
+				}
+			}
+			return Record;
+		}
+
+		static object ParseString(ReadOnlyUtf8String String)
+		{
+			return String.ToString();
+		}
+
+		static object ParseStringAsDateTime(ReadOnlyUtf8String String)
+		{
+			string Text = String.ToString();
+
+			DateTime Time;
+			if (DateTime.TryParse(Text, out Time))
+			{
+				return Time;
+			}
+			else
+			{
+				return PerforceReflection.UnixEpoch + TimeSpan.FromSeconds(long.Parse(Text));
+			}
+		}
+
+		static object ParseStringAsBool(ReadOnlyUtf8String String)
+		{
+			return String.Length == 0 || String == StringConstants.True;
+		}
+
+		static object ParseStringAsNullableBool(ReadOnlyUtf8String String)
+		{
+			return String == StringConstants.True;
+		}
+
+		static object ParseStringAsInt(ReadOnlyUtf8String String)
+		{
+			int Value;
+			int BytesConsumed;
+			if (Utf8Parser.TryParse(String.Span, out Value, out BytesConsumed) && BytesConsumed == String.Length)
+			{
+				return Value;
+			}
+			else if(String == StringConstants.New || String == StringConstants.None)
+			{
+				return -1;
+			}
+			else if(String.Length > 0 && String[0] == '#')
+			{
+				return ParseStringAsInt(String.Slice(1));
+			}
+			else if(String == StringConstants.Default)
+			{
+				return DefaultChange;
+			}
+			else
+			{
+				throw new PerforceException($"Unable to parse {String} as an integer");
+			}
+		}
+
+		static object ParseStringAsLong(ReadOnlyUtf8String String)
+		{
+			long Value;
+			int BytesConsumed;
+			if (!Utf8Parser.TryParse(String.Span, out Value, out BytesConsumed) || BytesConsumed != String.Length)
+			{
+				throw new PerforceException($"Unable to parse {String} as a long value");
+			}
+			return Value;
+		}
+
+		static object ParseStringAsNullableDateTimeOffset(ReadOnlyUtf8String String)
+		{
+			string Text = String.ToString();
+			return DateTimeOffset.Parse(Regex.Replace(Text, "[a-zA-Z ]*$", "")); // Strip timezone name (eg. "EST")
+		}
+	}
+}
diff --git a/Engine/Source/Programs/Shared/EpicGames.Perforce/PerforceResponse.cs b/Engine/Source/Programs/Shared/EpicGames.Perforce/PerforceResponse.cs
--- a/Engine/Source/Programs/Shared/EpicGames.Perforce/PerforceResponse.cs
+++ b/Engine/Source/Programs/Shared/EpicGames.Perforce/PerforceResponse.cs
@@ -1,4 +1,4 @@
-// Copyright Epic Games, Inc. All Rights Reserved.
+// Copyright 1998-2019 Epic Games, Inc. All Rights Reserved.
 
 using System;
 using System.Collections.Generic;
@@ -16,7 +16,7 @@ namespace EpicGames.Perforce
 		/// <summary>
 		/// Stores the response data
 		/// </summary>
-		private object InternalData;
+		protected object InternalData;
 
 		/// <summary>
 		/// Constructor
@@ -50,7 +50,7 @@ namespace EpicGames.Perforce
 		{
 			get
 			{
-				RequireSuccess();
+				EnsureSuccess();
 				return InternalData;
 			}
 		}
@@ -58,11 +58,11 @@ namespace EpicGames.Perforce
 		/// <summary>
 		/// Returns the info data.
 		/// </summary>
-		public PerforceInfo Info
+		public PerforceInfo? Info
 		{
 			get
 			{
-				RequireSuccess();
+				EnsureSuccess();
 				return InternalData as PerforceInfo;
 			}
 		}
@@ -70,7 +70,7 @@ namespace EpicGames.Perforce
 		/// <summary>
 		/// Returns the error data, or null if this is a succesful response.
 		/// </summary>
-		public PerforceError Error
+		public PerforceError? Error
 		{
 			get { return InternalData as PerforceError; }
 		}
@@ -78,9 +78,9 @@ namespace EpicGames.Perforce
 		/// <summary>
 		/// Throws an exception if the response is an error
 		/// </summary>
-		public void RequireSuccess()
+		public void EnsureSuccess()
 		{
-			PerforceError Error = InternalData as PerforceError;
+			PerforceError? Error = InternalData as PerforceError;
 			if(Error != null)
 			{
 				throw new PerforceException(Error);
@@ -91,7 +91,7 @@ namespace EpicGames.Perforce
 		/// Returns a string representation of this object for debugging
 		/// </summary>
 		/// <returns>String representation of the object for debugging</returns>
-		public override string ToString()
+		public override string? ToString()
 		{
 			return InternalData.ToString();
 		}
@@ -102,7 +102,7 @@ namespace EpicGames.Perforce
 	/// text if the response value is attempted to be accessed and an error has occurred.
 	/// </summary>
 	/// <typeparam name="T">Type of data returned on success</typeparam>
-	public class PerforceResponse<T> : PerforceResponse
+	public class PerforceResponse<T> : PerforceResponse where T : class
 	{
 		/// <summary>
 		/// Constructor
@@ -146,7 +146,26 @@ namespace EpicGames.Perforce
 		/// </summary>
 		public new T Data
 		{
-			get { return (T)base.Data; }
+			get
+			{
+				T? Result = InternalData as T;
+				if (Result == null)
+				{
+					if (InternalData is PerforceInfo)
+					{
+						throw new PerforceException($"Expected record of type '{typeof(T).Name}', got info: {InternalData}");
+					}
+					else if(InternalData is PerforceError)
+					{
+						throw new PerforceException($"Expected record of type '{typeof(T).Name}', got error: {InternalData}");
+					}
+					else
+					{
+						throw new PerforceException($"Expected record of type '{typeof(T).Name}', got: {InternalData}");
+					}
+				}
+				return Result;
+			}
 		}
 	}
 }
diff --git a/Engine/Source/Programs/Shared/EpicGames.Perforce/PerforceResponseList.cs b/Engine/Source/Programs/Shared/EpicGames.Perforce/PerforceResponseList.cs
--- a/Engine/Source/Programs/Shared/EpicGames.Perforce/PerforceResponseList.cs
+++ b/Engine/Source/Programs/Shared/EpicGames.Perforce/PerforceResponseList.cs
@@ -1,4 +1,4 @@
-// Copyright Epic Games, Inc. All Rights Reserved.
+// Copyright 1998-2019 Epic Games, Inc. All Rights Reserved.
 
 using System;
 using System.Collections.Generic;
@@ -34,9 +34,9 @@ namespace EpicGames.Perforce
 		/// <summary>
 		/// Sequence of all the data objects from the responses.
 		/// </summary>
-		public IEnumerable<T> Data
+		public List<T> Data
 		{
-			get { return this.Select(x => x.Data).Where(x => x != null); }
+			get { return this.Select(x => x.Data).Where(x => x != null).ToList(); }
 		}
 
 		/// <summary>
@@ -44,17 +44,27 @@ namespace EpicGames.Perforce
 		/// </summary>
 		public IEnumerable<PerforceError> Errors
 		{
-			get { return this.Where(x => !x.Succeeded).Select(x => x.Error); }
+			get
+			{
+				foreach (PerforceResponse<T> Response in this)
+				{
+					PerforceError? Error = Response.Error;
+					if (Error != null)
+					{
+						yield return Error;
+					}
+				}
+			}
 		}
 
 		/// <summary>
 		/// Throws an exception if any response is an error
 		/// </summary>
-		public void RequireSuccess()
+		public void EnsureSuccess()
 		{
 			foreach (PerforceResponse<T> Response in this)
 			{
-				Response.RequireSuccess();
+				Response.EnsureSuccess();
 			}
 		}
 	}
diff --git a/Engine/Source/Programs/Shared/EpicGames.Perforce/PerforceUtils.cs b/Engine/Source/Programs/Shared/EpicGames.Perforce/PerforceUtils.cs
--- a/Engine/Source/Programs/Shared/EpicGames.Perforce/PerforceUtils.cs
+++ b/Engine/Source/Programs/Shared/EpicGames.Perforce/PerforceUtils.cs
@@ -1,10 +1,12 @@
-// Copyright Epic Games, Inc. All Rights Reserved.
+// Copyright 1998-2019 Epic Games, Inc. All Rights Reserved.
 
+using EpicGames.Core;
 using System;
 using System.Collections.Generic;
 using System.Linq;
 using System.Text;
 using System.Threading.Tasks;
+using System.Threading.Tasks.Dataflow;
 
 namespace EpicGames.Perforce
 {
@@ -30,5 +32,48 @@ namespace EpicGames.Perforce
 			NewPath = NewPath.Replace("%25", "%");
 			return NewPath;
 		}
+
+		static public ReadOnlyUtf8String UnescapePath(ReadOnlyUtf8String Path)
+		{
+			ReadOnlySpan<byte> PathSpan = Path.Span;
+			for (int InputIdx = 0; InputIdx < PathSpan.Length - 2; InputIdx++)
+			{
+				if(PathSpan[InputIdx] == '%')
+				{
+					// Allocate the output buffer
+					byte[] Buffer = new byte[Path.Length];
+					PathSpan.Slice(0, InputIdx).CopyTo(Buffer.AsSpan());
+
+					// Copy the data to the output buffer
+					int OutputIdx = InputIdx;
+					while (InputIdx < PathSpan.Length)
+					{
+						// Parse the character code
+						int Value = StringUtils.ParseHexByte(PathSpan, InputIdx + 1);
+						if (Value == -1)
+						{
+							Buffer[OutputIdx++] = (byte)'%';
+							InputIdx++;
+						}
+						else
+						{
+							Buffer[OutputIdx++] = (byte)Value;
+							InputIdx += 3;
+						}
+
+						// Keep copying until we get to another percent character
+						while (InputIdx < PathSpan.Length && (PathSpan[InputIdx] != '%' || InputIdx + 2 >= PathSpan.Length))
+						{
+							Buffer[OutputIdx++] = PathSpan[InputIdx++];
+						}
+					}
+
+					// Copy the last chunk of data to the output buffer
+					Path = new ReadOnlyUtf8String(Buffer.AsMemory(0, OutputIdx));
+					break;
+				}
+			}
+			return Path;
+		}
 	}
 }
diff --git a/Engine/Source/Programs/Shared/EpicGames.Perforce/Records/AddOptions.cs b/Engine/Source/Programs/Shared/EpicGames.Perforce/Records/AddOptions.cs
--- a/Engine/Source/Programs/Shared/EpicGames.Perforce/Records/AddOptions.cs
+++ b/Engine/Source/Programs/Shared/EpicGames.Perforce/Records/AddOptions.cs
@@ -1,4 +1,4 @@
-// Copyright Epic Games, Inc. All Rights Reserved.
+// Copyright 1998-2019 Epic Games, Inc. All Rights Reserved.
 
 using System;
 using System.Collections.Generic;
diff --git a/Engine/Source/Programs/Shared/EpicGames.Perforce/Records/AddRecord.cs b/Engine/Source/Programs/Shared/EpicGames.Perforce/Records/AddRecord.cs
--- a/Engine/Source/Programs/Shared/EpicGames.Perforce/Records/AddRecord.cs
+++ b/Engine/Source/Programs/Shared/EpicGames.Perforce/Records/AddRecord.cs
@@ -1,4 +1,4 @@
-// Copyright Epic Games, Inc. All Rights Reserved.
+// Copyright 1998-2019 Epic Games, Inc. All Rights Reserved.
 
 using System;
 using System.Collections.Generic;
@@ -43,11 +43,22 @@ namespace EpicGames.Perforce
 		[PerforceTag("type")]
 		public string Type;
 
+		/// <summary>
+		/// Private constructor for serialization
+		/// </summary>
+		private AddRecord()
+		{
+			DepotFile = null!;
+			ClientFile = null!;
+			Action = null!;
+			Type = null!;
+		}
+
 		/// <summary>
 		/// Format this record for display in the debugger
 		/// </summary>
 		/// <returns>Summary of this revision</returns>
-		public override string ToString()
+		public override string? ToString()
 		{
 			return DepotFile;
 		}
diff --git a/Engine/Source/Programs/Shared/EpicGames.Perforce/Records/ChangeRecord.cs b/Engine/Source/Programs/Shared/EpicGames.Perforce/Records/ChangeRecord.cs
--- a/Engine/Source/Programs/Shared/EpicGames.Perforce/Records/ChangeRecord.cs
+++ b/Engine/Source/Programs/Shared/EpicGames.Perforce/Records/ChangeRecord.cs
@@ -1,4 +1,4 @@
-// Copyright Epic Games, Inc. All Rights Reserved.
+// Copyright 1998-2019 Epic Games, Inc. All Rights Reserved.
 
 using System;
 using System.Collections.Generic;
@@ -29,13 +29,13 @@ namespace EpicGames.Perforce
 		/// The user that owns or submitted the change
 		/// </summary>
 		[PerforceTag("User")]
-		public string User;
+		public string? User;
 
 		/// <summary>
 		/// The client that owns the change
 		/// </summary>
 		[PerforceTag("Client")]
-		public string Client;
+		public string? Client;
 
 		/// <summary>
 		/// Current changelist status
@@ -53,7 +53,7 @@ namespace EpicGames.Perforce
 		/// Description for the changelist
 		/// </summary>
 		[PerforceTag("Description")]
-		public string Description;
+		public string? Description;
 
 		/// <summary>
 		/// Files that are open in this changelist
diff --git a/Engine/Source/Programs/Shared/EpicGames.Perforce/Records/ChangeStatus.cs b/Engine/Source/Programs/Shared/EpicGames.Perforce/Records/ChangeStatus.cs
--- a/Engine/Source/Programs/Shared/EpicGames.Perforce/Records/ChangeStatus.cs
+++ b/Engine/Source/Programs/Shared/EpicGames.Perforce/Records/ChangeStatus.cs
@@ -1,4 +1,4 @@
-// Copyright Epic Games, Inc. All Rights Reserved.
+// Copyright 1998-2019 Epic Games, Inc. All Rights Reserved.
 
 using System;
 using System.Collections.Generic;
diff --git a/Engine/Source/Programs/Shared/EpicGames.Perforce/Records/ChangeType.cs b/Engine/Source/Programs/Shared/EpicGames.Perforce/Records/ChangeType.cs
--- a/Engine/Source/Programs/Shared/EpicGames.Perforce/Records/ChangeType.cs
+++ b/Engine/Source/Programs/Shared/EpicGames.Perforce/Records/ChangeType.cs
@@ -1,4 +1,4 @@
-// Copyright Epic Games, Inc. All Rights Reserved.
+// Copyright 1998-2019 Epic Games, Inc. All Rights Reserved.
 
 using System;
 using System.Collections.Generic;
diff --git a/Engine/Source/Programs/Shared/EpicGames.Perforce/Records/ChangesOptions.cs b/Engine/Source/Programs/Shared/EpicGames.Perforce/Records/ChangesOptions.cs
--- a/Engine/Source/Programs/Shared/EpicGames.Perforce/Records/ChangesOptions.cs
+++ b/Engine/Source/Programs/Shared/EpicGames.Perforce/Records/ChangesOptions.cs
@@ -1,4 +1,4 @@
-// Copyright Epic Games, Inc. All Rights Reserved.
+// Copyright 1998-2019 Epic Games, Inc. All Rights Reserved.
 
 using System;
 using System.Collections.Generic;
diff --git a/Engine/Source/Programs/Shared/EpicGames.Perforce/Records/ChangesRecord.cs b/Engine/Source/Programs/Shared/EpicGames.Perforce/Records/ChangesRecord.cs
--- a/Engine/Source/Programs/Shared/EpicGames.Perforce/Records/ChangesRecord.cs
+++ b/Engine/Source/Programs/Shared/EpicGames.Perforce/Records/ChangesRecord.cs
@@ -1,4 +1,4 @@
-// Copyright Epic Games, Inc. All Rights Reserved.
+// Copyright 1998-2019 Epic Games, Inc. All Rights Reserved.
 
 using System;
 using System.Collections.Generic;
@@ -29,13 +29,13 @@ namespace EpicGames.Perforce
 		/// The user that owns or submitted the change
 		/// </summary>
 		[PerforceTag("user")]
-		public string User;
+		public string? User;
 
 		/// <summary>
 		/// The client that owns the change
 		/// </summary>
 		[PerforceTag("client")]
-		public string Client;
+		public string? Client;
 
 		/// <summary>
 		/// Current changelist status
@@ -53,12 +53,12 @@ namespace EpicGames.Perforce
 		/// The path affected by this change.
 		/// </summary>
 		[PerforceTag("path", Optional = true)]
-		public string Path;
+		public string? Path;
 
 		/// <summary>
 		/// Description for the changelist
 		/// </summary>
 		[PerforceTag("desc")]
-		public string Description;
+		public string? Description;
 	}
 }
diff --git a/Engine/Source/Programs/Shared/EpicGames.Perforce/Records/CleanOptions.cs b/Engine/Source/Programs/Shared/EpicGames.Perforce/Records/CleanOptions.cs
new file mode 100644
--- /dev/null
+++ b/Engine/Source/Programs/Shared/EpicGames.Perforce/Records/CleanOptions.cs
@@ -0,0 +1,53 @@
+using System;
+using System.Collections.Generic;
+using System.Text;
+
+namespace EpicGames.Perforce
+{
+	/// <summary>
+	/// Options for the p4 clean command
+	/// </summary>
+	[Flags]
+	public enum CleanOptions
+	{
+		/// <summary>
+		/// User default options
+		/// </summary>
+		Default = 0,
+
+		/// <summary>
+		/// Edited files: Find files in the workspace that have been modified and restore them to the last file version that has synced from the depot.
+		/// </summary>
+		Edited = 1,
+
+		/// <summary>
+		/// Added files: Find files in the workspace that have no corresponding files in the depot and delete them.
+		/// </summary>
+		Added = 2,
+
+		/// <summary>
+		/// Deleted files: Find those files in the depot that do not exist in your workspace and add them to the workspace.
+		/// </summary>
+		Deleted = 4,
+
+		/// <summary>
+		/// Preview the results of the operation without performing any action.
+		/// </summary>
+		Preview = 8,
+
+		/// <summary>
+		/// Do not perform any ignore checking; ignore any settings specified by P4IGNORE for added files.
+		/// </summary>
+		NoIgnoreChecking = 16,
+
+		/// <summary>
+		/// Include file paths in local syntax
+		/// </summary>
+		LocalSyntax = 32,
+
+		/// <summary>
+		/// Use modified times to determine whether files are out of date (see https://community.perforce.com/s/article/15133)
+		/// </summary>
+		ModifiedTimes = 64,
+	};
+}
diff --git a/Engine/Source/Programs/Shared/EpicGames.Perforce/Records/CleanRecord.cs b/Engine/Source/Programs/Shared/EpicGames.Perforce/Records/CleanRecord.cs
new file mode 100644
--- /dev/null
+++ b/Engine/Source/Programs/Shared/EpicGames.Perforce/Records/CleanRecord.cs
@@ -0,0 +1,42 @@
+using System;
+using System.Collections.Generic;
+using System.Text;
+
+namespace EpicGames.Perforce
+{
+	/// <summary>
+	/// Record returned by the p4 clean command
+	/// </summary>
+	public class CleanRecord
+	{
+		/// <summary>
+		/// Depot path of the file, according to the workspace mapping
+		/// </summary>
+		[PerforceTag("depotFile", Optional = true)]
+		public string? DepotFile;
+
+		/// <summary>
+		/// Client path to the file (in local path syntax)
+		/// </summary>
+		[PerforceTag("clientFile", Optional = true)]
+		public string? ClientFile;
+
+		/// <summary>
+		/// Path to the file in local syntax
+		/// </summary>
+		[PerforceTag("localFile", Optional = true)]
+		public string? LocalFile;
+
+		/// <summary>
+		/// Revision of the file
+		/// </summary>
+		[PerforceTag("rev", Optional = true)]
+		public string? Revision;
+
+		/// <summary>
+		/// Action to be taken to the file
+		/// </summary>
+		[PerforceTag("action", Optional = true)]
+		public FileAction Action;
+	}
+}
diff --git a/Engine/Source/Programs/Shared/EpicGames.Perforce/Records/ClientLineEndings.cs b/Engine/Source/Programs/Shared/EpicGames.Perforce/Records/ClientLineEndings.cs
--- a/Engine/Source/Programs/Shared/EpicGames.Perforce/Records/ClientLineEndings.cs
+++ b/Engine/Source/Programs/Shared/EpicGames.Perforce/Records/ClientLineEndings.cs
@@ -1,4 +1,4 @@
-// Copyright Epic Games, Inc. All Rights Reserved.
+// Copyright 1998-2019 Epic Games, Inc. All Rights Reserved.
 
 using System;
 using System.Collections.Generic;
diff --git a/Engine/Source/Programs/Shared/EpicGames.Perforce/Records/ClientOptions.cs b/Engine/Source/Programs/Shared/EpicGames.Perforce/Records/ClientOptions.cs
--- a/Engine/Source/Programs/Shared/EpicGames.Perforce/Records/ClientOptions.cs
+++ b/Engine/Source/Programs/Shared/EpicGames.Perforce/Records/ClientOptions.cs
@@ -1,4 +1,4 @@
-// Copyright Epic Games, Inc. All Rights Reserved.
+// Copyright 1998-2019 Epic Games, Inc. All Rights Reserved.
 
 using System;
 using System.Collections.Generic;
diff --git a/Engine/Source/Programs/Shared/EpicGames.Perforce/Records/ClientRecord.cs b/Engine/Source/Programs/Shared/EpicGames.Perforce/Records/ClientRecord.cs
--- a/Engine/Source/Programs/Shared/EpicGames.Perforce/Records/ClientRecord.cs
+++ b/Engine/Source/Programs/Shared/EpicGames.Perforce/Records/ClientRecord.cs
@@ -1,4 +1,4 @@
-// Copyright Epic Games, Inc. All Rights Reserved.
+// Copyright 1998-2019 Epic Games, Inc. All Rights Reserved.
 
 using System;
 using System.Collections.Generic;
@@ -41,13 +41,13 @@ namespace EpicGames.Perforce
 		/// The name of the workstation on which this workspace resides.
 		/// </summary>
 		[PerforceTag("Host", Optional = true)]
-		public string Host;
+		public string? Host;
 
 		/// <summary>
 		/// A textual description of the workspace.
 		/// </summary>
 		[PerforceTag("Description", Optional = true)]
-		public string Description;
+		public string? Description;
 
 		/// <summary>
 		/// The directory (on the local host) relative to which all the files in the View: are specified. 
@@ -77,18 +77,41 @@ namespace EpicGames.Perforce
 		/// Associates the workspace with the specified stream.
 		/// </summary>
 		[PerforceTag("Stream", Optional = true)]
-		public string Stream;
+		public string? Stream;
 
 		/// <summary>
 		/// The type of client.
 		/// </summary>
 		[PerforceTag("Type", Optional = true)]
-		public string Type;
+		public string? Type;
 
 		/// <summary>
 		/// Specifies the mappings between files in the depot and files in the workspace.
 		/// </summary>
 		[PerforceTag("View")]
 		public List<string> View = new List<string>();
+
+		/// <summary>
+		/// Private constructor for serialization
+		/// </summary>
+		private ClientRecord()
+		{
+			Name = null!;
+			Owner = null!;
+			Root = null!;
+		}
+
+		/// <summary>
+		/// Constructor
+		/// </summary>
+		/// <param name="Name">Name of the client</param>
+		/// <param name="Owner">Owner of the client</param>
+		/// <param name="Root">The root directory to sync</param>
+		public ClientRecord(string Name, string Owner, string Root)
+		{
+			this.Name = Name;
+			this.Owner = Owner;
+			this.Root = Root;
+		}
 	}
 }
diff --git a/Engine/Source/Programs/Shared/EpicGames.Perforce/Records/ClientSubmitOptions.cs b/Engine/Source/Programs/Shared/EpicGames.Perforce/Records/ClientSubmitOptions.cs
--- a/Engine/Source/Programs/Shared/EpicGames.Perforce/Records/ClientSubmitOptions.cs
+++ b/Engine/Source/Programs/Shared/EpicGames.Perforce/Records/ClientSubmitOptions.cs
@@ -1,4 +1,4 @@
-// Copyright Epic Games, Inc. All Rights Reserved.
+// Copyright 1998-2019 Epic Games, Inc. All Rights Reserved.
 
 using System;
 using System.Collections.Generic;
diff --git a/Engine/Source/Programs/Shared/EpicGames.Perforce/Records/ClientsOptions.cs b/Engine/Source/Programs/Shared/EpicGames.Perforce/Records/ClientsOptions.cs
--- a/Engine/Source/Programs/Shared/EpicGames.Perforce/Records/ClientsOptions.cs
+++ b/Engine/Source/Programs/Shared/EpicGames.Perforce/Records/ClientsOptions.cs
@@ -1,4 +1,4 @@
-// Copyright Epic Games, Inc. All Rights Reserved.
+// Copyright 1998-2019 Epic Games, Inc. All Rights Reserved.
 
 using System;
 using System.Collections.Generic;
diff --git a/Engine/Source/Programs/Shared/EpicGames.Perforce/Records/ClientsRecord.cs b/Engine/Source/Programs/Shared/EpicGames.Perforce/Records/ClientsRecord.cs
--- a/Engine/Source/Programs/Shared/EpicGames.Perforce/Records/ClientsRecord.cs
+++ b/Engine/Source/Programs/Shared/EpicGames.Perforce/Records/ClientsRecord.cs
@@ -1,4 +1,4 @@
-// Copyright Epic Games, Inc. All Rights Reserved.
+// Copyright 1998-2019 Epic Games, Inc. All Rights Reserved.
 
 using System;
 using System.Collections.Generic;
@@ -41,13 +41,13 @@ namespace EpicGames.Perforce
 		/// The name of the workstation on which this workspace resides.
 		/// </summary>
 		[PerforceTag("Host", Optional = true)]
-		public string Host;
+		public string? Host;
 
 		/// <summary>
 		/// A textual description of the workspace.
 		/// </summary>
 		[PerforceTag("Description", Optional = true)]
-		public string Description;
+		public string? Description;
 
 		/// <summary>
 		/// The directory (on the local host) relative to which all the files in the View: are specified. 
@@ -77,13 +77,29 @@ namespace EpicGames.Perforce
 		/// Associates the workspace with the specified stream.
 		/// </summary>
 		[PerforceTag("Stream", Optional = true)]
-		public string Stream;
+		public string? Stream;
+
+		/// <summary>
+		/// The edge server ID
+		/// </summary>
+		[PerforceTag("ServerID", Optional = true)]
+		public string? ServerID;
+
+		/// <summary>
+		/// Private constructor for serialization
+		/// </summary>
+		private ClientsRecord()
+		{
+			Name = null!;
+			Owner = null!;
+			Root = null!;
+		}
 
 		/// <summary>
 		/// Format this record for display in the debugger
 		/// </summary>
 		/// <returns>Summary of this revision</returns>
-		public override string ToString()
+		public override string? ToString()
 		{
 			return Name;
 		}
diff --git a/Engine/Source/Programs/Shared/EpicGames.Perforce/Records/DeleteChangeOptions.cs b/Engine/Source/Programs/Shared/EpicGames.Perforce/Records/DeleteChangeOptions.cs
--- a/Engine/Source/Programs/Shared/EpicGames.Perforce/Records/DeleteChangeOptions.cs
+++ b/Engine/Source/Programs/Shared/EpicGames.Perforce/Records/DeleteChangeOptions.cs
@@ -1,4 +1,4 @@
-// Copyright Epic Games, Inc. All Rights Reserved.
+// Copyright 1998-2019 Epic Games, Inc. All Rights Reserved.
 
 using System;
 using System.Collections.Generic;
diff --git a/Engine/Source/Programs/Shared/EpicGames.Perforce/Records/DeleteClientOptions.cs b/Engine/Source/Programs/Shared/EpicGames.Perforce/Records/DeleteClientOptions.cs
--- a/Engine/Source/Programs/Shared/EpicGames.Perforce/Records/DeleteClientOptions.cs
+++ b/Engine/Source/Programs/Shared/EpicGames.Perforce/Records/DeleteClientOptions.cs
@@ -1,4 +1,4 @@
-// Copyright Epic Games, Inc. All Rights Reserved.
+// Copyright 1998-2019 Epic Games, Inc. All Rights Reserved.
 
 using System;
 using System.Collections.Generic;
diff --git a/Engine/Source/Programs/Shared/EpicGames.Perforce/Records/DeleteOptions.cs b/Engine/Source/Programs/Shared/EpicGames.Perforce/Records/DeleteOptions.cs
--- a/Engine/Source/Programs/Shared/EpicGames.Perforce/Records/DeleteOptions.cs
+++ b/Engine/Source/Programs/Shared/EpicGames.Perforce/Records/DeleteOptions.cs
@@ -1,4 +1,4 @@
-// Copyright Epic Games, Inc. All Rights Reserved.
+// Copyright 1998-2019 Epic Games, Inc. All Rights Reserved.
 
 using System;
 using System.Collections.Generic;
diff --git a/Engine/Source/Programs/Shared/EpicGames.Perforce/Records/DeleteRecord.cs b/Engine/Source/Programs/Shared/EpicGames.Perforce/Records/DeleteRecord.cs
--- a/Engine/Source/Programs/Shared/EpicGames.Perforce/Records/DeleteRecord.cs
+++ b/Engine/Source/Programs/Shared/EpicGames.Perforce/Records/DeleteRecord.cs
@@ -1,4 +1,4 @@
-// Copyright Epic Games, Inc. All Rights Reserved.
+// Copyright 1998-2019 Epic Games, Inc. All Rights Reserved.
 
 using System;
 using System.Collections.Generic;
@@ -43,11 +43,22 @@ namespace EpicGames.Perforce
 		[PerforceTag("type")]
 		public string Type;
 
+		/// <summary>
+		/// Private constructor for serialization
+		/// </summary>
+		private DeleteRecord()
+		{
+			DepotFile = null!;
+			ClientFile = null!;
+			Action = null!;
+			Type = null!;
+		}
+
 		/// <summary>
 		/// Format this record for display in the debugger
 		/// </summary>
 		/// <returns>Summary of this revision</returns>
-		public override string ToString()
+		public override string? ToString()
 		{
 			return DepotFile;
 		}
diff --git a/Engine/Source/Programs/Shared/EpicGames.Perforce/Records/DescribeFileRecord.cs b/Engine/Source/Programs/Shared/EpicGames.Perforce/Records/DescribeFileRecord.cs
--- a/Engine/Source/Programs/Shared/EpicGames.Perforce/Records/DescribeFileRecord.cs
+++ b/Engine/Source/Programs/Shared/EpicGames.Perforce/Records/DescribeFileRecord.cs
@@ -1,4 +1,4 @@
-// Copyright Epic Games, Inc. All Rights Reserved.
+// Copyright 1998-2019 Epic Games, Inc. All Rights Reserved.
 
 using System;
 using System.Collections.Generic;
@@ -47,7 +47,16 @@ namespace EpicGames.Perforce
 		/// Digest of the file, or null if not specified
 		/// </summary>
 		[PerforceTag("digest", Optional = true)]
-		public readonly string Digest;
+		public readonly string? Digest;
+
+		/// <summary>
+		/// Private constructor for serialization
+		/// </summary>
+		private DescribeFileRecord()
+		{
+			DepotFile = null!;
+			Type = null!;
+		}
 
 		/// <summary>
 		/// Format this record for display in the debugger
diff --git a/Engine/Source/Programs/Shared/EpicGames.Perforce/Records/DescribeOptions.cs b/Engine/Source/Programs/Shared/EpicGames.Perforce/Records/DescribeOptions.cs
--- a/Engine/Source/Programs/Shared/EpicGames.Perforce/Records/DescribeOptions.cs
+++ b/Engine/Source/Programs/Shared/EpicGames.Perforce/Records/DescribeOptions.cs
@@ -1,4 +1,4 @@
-// Copyright Epic Games, Inc. All Rights Reserved.
+// Copyright 1998-2019 Epic Games, Inc. All Rights Reserved.
 
 using System;
 using System.Collections.Generic;
diff --git a/Engine/Source/Programs/Shared/EpicGames.Perforce/Records/DescribeRecord.cs b/Engine/Source/Programs/Shared/EpicGames.Perforce/Records/DescribeRecord.cs
--- a/Engine/Source/Programs/Shared/EpicGames.Perforce/Records/DescribeRecord.cs
+++ b/Engine/Source/Programs/Shared/EpicGames.Perforce/Records/DescribeRecord.cs
@@ -1,4 +1,4 @@
-// Copyright Epic Games, Inc. All Rights Reserved.
+// Copyright 1998-2019 Epic Games, Inc. All Rights Reserved.
 
 using System;
 using System.Collections.Generic;
@@ -59,7 +59,7 @@ namespace EpicGames.Perforce
 		/// Narrowest path which contains all files affected by this change
 		/// </summary>
 		[PerforceTag("path", Optional = true)]
-		public string Path;
+		public string? Path;
 
 		/// <summary>
 		/// The files affected by this change
@@ -67,6 +67,16 @@ namespace EpicGames.Perforce
 		[PerforceRecordList]
 		public List<DescribeFileRecord> Files = new List<DescribeFileRecord>();
 
+		/// <summary>
+		/// Private constructor for serialization
+		/// </summary>
+		private DescribeRecord()
+		{
+			User = null!;
+			Client = null!;
+			Description = null!;
+		}
+
 		/// <summary>
 		/// Format this record for display in the debugger
 		/// </summary>
diff --git a/Engine/Source/Programs/Shared/EpicGames.Perforce/Records/EditOptions.cs b/Engine/Source/Programs/Shared/EpicGames.Perforce/Records/EditOptions.cs
--- a/Engine/Source/Programs/Shared/EpicGames.Perforce/Records/EditOptions.cs
+++ b/Engine/Source/Programs/Shared/EpicGames.Perforce/Records/EditOptions.cs
@@ -1,4 +1,4 @@
-// Copyright Epic Games, Inc. All Rights Reserved.
+// Copyright 1998-2019 Epic Games, Inc. All Rights Reserved.
 
 using System;
 using System.Collections.Generic;
diff --git a/Engine/Source/Programs/Shared/EpicGames.Perforce/Records/EditRecord.cs b/Engine/Source/Programs/Shared/EpicGames.Perforce/Records/EditRecord.cs
--- a/Engine/Source/Programs/Shared/EpicGames.Perforce/Records/EditRecord.cs
+++ b/Engine/Source/Programs/Shared/EpicGames.Perforce/Records/EditRecord.cs
@@ -1,4 +1,4 @@
-// Copyright Epic Games, Inc. All Rights Reserved.
+// Copyright 1998-2019 Epic Games, Inc. All Rights Reserved.
 
 using System;
 using System.Collections.Generic;
@@ -43,6 +43,17 @@ namespace EpicGames.Perforce
 		[PerforceTag("type")]
 		public string Type;
 
+		/// <summary>
+		/// Private constructor for serialization
+		/// </summary>
+		private EditRecord()
+		{
+			DepotFile = null!;
+			ClientFile = null!;
+			Action = null!;
+			Type = null!;
+		}
+
 		/// <summary>
 		/// Format this record for display in the debugger
 		/// </summary>
diff --git a/Engine/Source/Programs/Shared/EpicGames.Perforce/Records/FStatOptions.cs b/Engine/Source/Programs/Shared/EpicGames.Perforce/Records/FStatOptions.cs
--- a/Engine/Source/Programs/Shared/EpicGames.Perforce/Records/FStatOptions.cs
+++ b/Engine/Source/Programs/Shared/EpicGames.Perforce/Records/FStatOptions.cs
@@ -1,4 +1,4 @@
-// Copyright Epic Games, Inc. All Rights Reserved.
+// Copyright 1998-2019 Epic Games, Inc. All Rights Reserved.
 
 using System;
 using System.Collections.Generic;
diff --git a/Engine/Source/Programs/Shared/EpicGames.Perforce/Records/FStatRecord.cs b/Engine/Source/Programs/Shared/EpicGames.Perforce/Records/FStatRecord.cs
--- a/Engine/Source/Programs/Shared/EpicGames.Perforce/Records/FStatRecord.cs
+++ b/Engine/Source/Programs/Shared/EpicGames.Perforce/Records/FStatRecord.cs
@@ -1,4 +1,4 @@
-// Copyright Epic Games, Inc. All Rights Reserved.
+// Copyright 1998-2019 Epic Games, Inc. All Rights Reserved.
 
 using System;
 using System.Collections.Generic;
@@ -17,19 +17,19 @@ namespace EpicGames.Perforce
 		/// Depot path to file
 		/// </summary>
 		[PerforceTag("depotFile", Optional = true)]
-		public string DepotFile;
+		public string? DepotFile;
 
 		/// <summary>
 		/// Local path to file (in local syntax by default, or in Perforce syntax with the FStatOptions.ClientFileInPerforceSyntax option)
 		/// </summary>
 		[PerforceTag("clientFile", Optional = true)]
-		public string ClientFile;
+		public string? ClientFile;
 
 		/// <summary>
 		/// Local path to file
 		/// </summary>
 		[PerforceTag("path", Optional = true)]
-		public string Path;
+		public string? Path;
 
 		/// <summary>
 		/// Set if the file is open for add and it is mapped to current client workspace
@@ -53,7 +53,7 @@ namespace EpicGames.Perforce
 		/// Head revision type, if in depot
 		/// </summary>
 		[PerforceTag("headType", Optional = true)]
-		public string HeadType;
+		public string? HeadType;
 
 		/// <summary>
 		/// Head revision changelist time, if in depot.
@@ -95,13 +95,13 @@ namespace EpicGames.Perforce
 		/// Changelist description (if using -e changelist and if the file was part of changelist)
 		/// </summary>
 		[PerforceTag("desc", Optional = true)]
-		public string Description;
+		public string? Description;
 
 		/// <summary>
 		/// MD5 digest of a file (requires -Ol option)
 		/// </summary>
 		[PerforceTag("digest", Optional = true)]
-		public string Digest;
+		public string? Digest;
 
 		/// <summary>
 		/// File length in bytes (requires -Ol option)
@@ -125,7 +125,7 @@ namespace EpicGames.Perforce
 		/// User who opened the file, if opens
 		/// </summary>
 		[PerforceTag("actionOwner", Optional = true)]
-		public string ActionOwner;
+		public string? ActionOwner;
 
 		/// <summary>
 		/// Open changelist number, if opened in your workspace
@@ -137,7 +137,7 @@ namespace EpicGames.Perforce
 		/// Summarize this record for display in the debugger
 		/// </summary>
 		/// <returns>Formatted record</returns>
-		public override string ToString()
+		public override string? ToString()
 		{
 			return DepotFile ?? ClientFile ?? Path ?? base.ToString();
 		}
diff --git a/Engine/Source/Programs/Shared/EpicGames.Perforce/Records/FileAction.cs b/Engine/Source/Programs/Shared/EpicGames.Perforce/Records/FileAction.cs
--- a/Engine/Source/Programs/Shared/EpicGames.Perforce/Records/FileAction.cs
+++ b/Engine/Source/Programs/Shared/EpicGames.Perforce/Records/FileAction.cs
@@ -1,4 +1,4 @@
-// Copyright Epic Games, Inc. All Rights Reserved.
+// Copyright 1998-2019 Epic Games, Inc. All Rights Reserved.
 
 using System;
 using System.Collections.Generic;
diff --git a/Engine/Source/Programs/Shared/EpicGames.Perforce/Records/FileLogOptions.cs b/Engine/Source/Programs/Shared/EpicGames.Perforce/Records/FileLogOptions.cs
--- a/Engine/Source/Programs/Shared/EpicGames.Perforce/Records/FileLogOptions.cs
+++ b/Engine/Source/Programs/Shared/EpicGames.Perforce/Records/FileLogOptions.cs
@@ -1,4 +1,4 @@
-// Copyright Epic Games, Inc. All Rights Reserved.
+// Copyright 1998-2019 Epic Games, Inc. All Rights Reserved.
 
 using System;
 using System.Collections.Generic;
diff --git a/Engine/Source/Programs/Shared/EpicGames.Perforce/Records/FileLogRecord.cs b/Engine/Source/Programs/Shared/EpicGames.Perforce/Records/FileLogRecord.cs
--- a/Engine/Source/Programs/Shared/EpicGames.Perforce/Records/FileLogRecord.cs
+++ b/Engine/Source/Programs/Shared/EpicGames.Perforce/Records/FileLogRecord.cs
@@ -1,4 +1,4 @@
-// Copyright Epic Games, Inc. All Rights Reserved.
+// Copyright 1998-2019 Epic Games, Inc. All Rights Reserved.
 
 using System;
 using System.Collections.Generic;
@@ -25,6 +25,25 @@ namespace EpicGames.Perforce
 		[PerforceRecordList]
 		public List<RevisionRecord> Revisions = new List<RevisionRecord>();
 
+		/// <summary>
+		/// Private constructor for serialization
+		/// </summary>
+		private FileLogRecord()
+		{
+			DepotPath = null!;
+		}
+
+		/// <summary>
+		/// Constructor
+		/// </summary>
+		/// <param name="DepotPath">Path to the file in the depot</param>
+		/// <param name="Revisions">Revisions of this file</param>
+		public FileLogRecord(string DepotPath, List<RevisionRecord> Revisions)
+		{
+			this.DepotPath = DepotPath;
+			this.Revisions = Revisions;
+		}
+
 		/// <summary>
 		/// Format this record for display in the debugger
 		/// </summary>
diff --git a/Engine/Source/Programs/Shared/EpicGames.Perforce/Records/GetChangeOptions.cs b/Engine/Source/Programs/Shared/EpicGames.Perforce/Records/GetChangeOptions.cs
--- a/Engine/Source/Programs/Shared/EpicGames.Perforce/Records/GetChangeOptions.cs
+++ b/Engine/Source/Programs/Shared/EpicGames.Perforce/Records/GetChangeOptions.cs
@@ -1,4 +1,4 @@
-// Copyright Epic Games, Inc. All Rights Reserved.
+// Copyright 1998-2019 Epic Games, Inc. All Rights Reserved.
 
 using System;
 using System.Collections.Generic;
diff --git a/Engine/Source/Programs/Shared/EpicGames.Perforce/Records/InfoOptions.cs b/Engine/Source/Programs/Shared/EpicGames.Perforce/Records/InfoOptions.cs
--- a/Engine/Source/Programs/Shared/EpicGames.Perforce/Records/InfoOptions.cs
+++ b/Engine/Source/Programs/Shared/EpicGames.Perforce/Records/InfoOptions.cs
@@ -1,4 +1,4 @@
-// Copyright Epic Games, Inc. All Rights Reserved.
+// Copyright 1998-2019 Epic Games, Inc. All Rights Reserved.
 
 using System;
 using System.Collections.Generic;
diff --git a/Engine/Source/Programs/Shared/EpicGames.Perforce/Records/InfoRecord.cs b/Engine/Source/Programs/Shared/EpicGames.Perforce/Records/InfoRecord.cs
--- a/Engine/Source/Programs/Shared/EpicGames.Perforce/Records/InfoRecord.cs
+++ b/Engine/Source/Programs/Shared/EpicGames.Perforce/Records/InfoRecord.cs
@@ -1,4 +1,4 @@
-// Copyright Epic Games, Inc. All Rights Reserved.
+// Copyright 1998-2019 Epic Games, Inc. All Rights Reserved.
 
 using System;
 using System.Collections.Generic;
@@ -17,42 +17,60 @@ namespace EpicGames.Perforce
 		/// The current user according to the Perforce environment
 		/// </summary>
 		[PerforceTag("userName", Optional = true)]
-		public string UserName;
+		public string? UserName;
 
 		/// <summary>
 		/// The current client
 		/// </summary>
 		[PerforceTag("clientName", Optional = true)]
-		public string ClientName;
+		public string? ClientName;
 
 		/// <summary>
 		/// The current host
 		/// </summary>
 		[PerforceTag("clientHost", Optional = true)]
-		public string ClientHost;
+		public string? ClientHost;
 
 		/// <summary>
 		/// Root directory for the current client
 		/// </summary>
 		[PerforceTag("clientRoot", Optional = true)]
-		public string ClientRoot;
+		public string? ClientRoot;
 
 		/// <summary>
 		/// Selected stream in the current client
 		/// </summary>
 		[PerforceTag("clientStream", Optional = true)]
-		public string ClientStream;
+		public string? ClientStream;
 
 		/// <summary>
 		/// Address of the Perforce server
 		/// </summary>
 		[PerforceTag("serverAddress", Optional = true)]
-		public string ServerAddress;
+		public string? ServerAddress;
+
+		/// <summary>
+		/// Date and time on the server
+		/// </summary>
+		[PerforceTag("serverDate", Optional = true)]
+		public DateTimeOffset? ServerDate;
 
 		/// <summary>
 		/// Case handling setting on the server
 		/// </summary>
 		[PerforceTag("caseHandling", Optional = true)]
-		public string CaseHandling;
+		public string? CaseHandling;
+
+		/// <summary>
+		/// List of services provided by this server
+		/// </summary>
+		[PerforceTag("serverServices", Optional = true)]
+		public string? Services;
+
+		/// <summary>
+		/// The server unique id
+		/// </summary>
+		[PerforceTag("ServerID", Optional = true)]
+		public string? ServerID;
 	}
 }
diff --git a/Engine/Source/Programs/Shared/EpicGames.Perforce/Records/IntegrationAction.cs b/Engine/Source/Programs/Shared/EpicGames.Perforce/Records/IntegrationAction.cs
--- a/Engine/Source/Programs/Shared/EpicGames.Perforce/Records/IntegrationAction.cs
+++ b/Engine/Source/Programs/Shared/EpicGames.Perforce/Records/IntegrationAction.cs
@@ -1,4 +1,4 @@
-// Copyright Epic Games, Inc. All Rights Reserved.
+// Copyright 1998-2019 Epic Games, Inc. All Rights Reserved.
 
 using System;
 using System.Collections.Generic;
diff --git a/Engine/Source/Programs/Shared/EpicGames.Perforce/Records/IntegrationRecord.cs b/Engine/Source/Programs/Shared/EpicGames.Perforce/Records/IntegrationRecord.cs
--- a/Engine/Source/Programs/Shared/EpicGames.Perforce/Records/IntegrationRecord.cs
+++ b/Engine/Source/Programs/Shared/EpicGames.Perforce/Records/IntegrationRecord.cs
@@ -1,4 +1,4 @@
-// Copyright Epic Games, Inc. All Rights Reserved.
+// Copyright 1998-2019 Epic Games, Inc. All Rights Reserved.
 
 using System;
 using System.Collections.Generic;
@@ -37,6 +37,14 @@ namespace EpicGames.Perforce
 		[PerforceTag("erev")]
 		public int EndRevisionNumber;
 
+		/// <summary>
+		/// Private constructor for serialization
+		/// </summary>
+		private IntegrationRecord()
+		{
+			OtherFile = null!;
+		}
+
 		/// <summary>
 		/// Summarize this record for display in the debugger
 		/// </summary>
diff --git a/Engine/Source/Programs/Shared/EpicGames.Perforce/Records/MergeOptions.cs b/Engine/Source/Programs/Shared/EpicGames.Perforce/Records/MergeOptions.cs
new file mode 100644
--- /dev/null
+++ b/Engine/Source/Programs/Shared/EpicGames.Perforce/Records/MergeOptions.cs
@@ -0,0 +1,18 @@
+using System;
+using System.Collections.Generic;
+using System.Text;
+
+namespace EpicGames.Perforce
+{
+	/// <summary>
+	/// Options for the merge command
+	/// </summary>
+	[Flags]
+	public enum MergeOptions
+	{
+		/// <summary>
+		/// Preview the merge
+		/// </summary>
+		Preview = 1,
+	}
+}
diff --git a/Engine/Source/Programs/Shared/EpicGames.Perforce/Records/MergeRecord.cs b/Engine/Source/Programs/Shared/EpicGames.Perforce/Records/MergeRecord.cs
new file mode 100644
--- /dev/null
+++ b/Engine/Source/Programs/Shared/EpicGames.Perforce/Records/MergeRecord.cs
@@ -0,0 +1,35 @@
+using System;
+using System.Collections.Generic;
+using System.Text;
+
+namespace EpicGames.Perforce
+{
+	/// <summary>
+	/// Response from a merge
+	/// </summary>
+	public class MergeRecord
+	{
+		/// <summary>
+		/// The file being merged from
+		/// </summary>
+		[PerforceTag("fromFile")]
+		public string FromFile = null!;
+
+		/// <summary>
+		/// The starting revision, as an integer (null if 'none')
+		/// </summary>
+		public int? StartFromRev => String.Equals(StartFromRevText, "none") ? (int?)null : int.Parse(StartFromRevText);
+
+		/// <summary>
+		/// The first revision being merged (or 'none')
+		/// </summary>
+		[PerforceTag("startFromRev")]
+		public string StartFromRevText = null!;
+
+		/// <summary>
+		/// The last revision being merged
+		/// </summary>
+		[PerforceTag("endFromRev")]
+		public int EndFromRev;
+	}
+}
diff --git a/Engine/Source/Programs/Shared/EpicGames.Perforce/Records/OpenedOptions.cs b/Engine/Source/Programs/Shared/EpicGames.Perforce/Records/OpenedOptions.cs
--- a/Engine/Source/Programs/Shared/EpicGames.Perforce/Records/OpenedOptions.cs
+++ b/Engine/Source/Programs/Shared/EpicGames.Perforce/Records/OpenedOptions.cs
@@ -1,4 +1,4 @@
-// Copyright Epic Games, Inc. All Rights Reserved.
+// Copyright 1998-2019 Epic Games, Inc. All Rights Reserved.
 
 using System;
 using System.Collections.Generic;
diff --git a/Engine/Source/Programs/Shared/EpicGames.Perforce/Records/PrintRecord.cs b/Engine/Source/Programs/Shared/EpicGames.Perforce/Records/PrintRecord.cs
--- a/Engine/Source/Programs/Shared/EpicGames.Perforce/Records/PrintRecord.cs
+++ b/Engine/Source/Programs/Shared/EpicGames.Perforce/Records/PrintRecord.cs
@@ -1,4 +1,4 @@
-// Copyright Epic Games, Inc. All Rights Reserved.
+// Copyright 1998-2019 Epic Games, Inc. All Rights Reserved.
 
 using System;
 using System.Collections.Generic;
@@ -54,5 +54,14 @@ namespace EpicGames.Perforce
 		/// </summary>
 		[PerforceTag("fileSize")]
 		public long FileSize;
+
+		/// <summary>
+		/// Private constructor for serialization
+		/// </summary>
+		private PrintRecord()
+		{
+			DepotFile = null!;
+			Type = null!;
+		}
 	}
 }
diff --git a/Engine/Source/Programs/Shared/EpicGames.Perforce/Records/ReconcileOptions.cs b/Engine/Source/Programs/Shared/EpicGames.Perforce/Records/ReconcileOptions.cs
--- a/Engine/Source/Programs/Shared/EpicGames.Perforce/Records/ReconcileOptions.cs
+++ b/Engine/Source/Programs/Shared/EpicGames.Perforce/Records/ReconcileOptions.cs
@@ -1,4 +1,4 @@
-// Copyright Epic Games, Inc. All Rights Reserved.
+// Copyright 1998-2019 Epic Games, Inc. All Rights Reserved.
 
 using System;
 using System.Collections.Generic;
diff --git a/Engine/Source/Programs/Shared/EpicGames.Perforce/Records/ReconcileRecord.cs b/Engine/Source/Programs/Shared/EpicGames.Perforce/Records/ReconcileRecord.cs
--- a/Engine/Source/Programs/Shared/EpicGames.Perforce/Records/ReconcileRecord.cs
+++ b/Engine/Source/Programs/Shared/EpicGames.Perforce/Records/ReconcileRecord.cs
@@ -1,4 +1,4 @@
-// Copyright Epic Games, Inc. All Rights Reserved.
+// Copyright 1998-2019 Epic Games, Inc. All Rights Reserved.
 
 using System;
 using System.Collections.Generic;
@@ -43,6 +43,17 @@ namespace EpicGames.Perforce
 		[PerforceTag("type")]
 		public string Type;
 
+		/// <summary>
+		/// Private constructor for serialization
+		/// </summary>
+		private ReconcileRecord()
+		{
+			DepotFile = null!;
+			ClientFile = null!;
+			Action = null!;
+			Type = null!;
+		}
+
 		/// <summary>
 		/// Format this record for display in the debugger
 		/// </summary>
diff --git a/Engine/Source/Programs/Shared/EpicGames.Perforce/Records/ResolveOptions.cs b/Engine/Source/Programs/Shared/EpicGames.Perforce/Records/ResolveOptions.cs
--- a/Engine/Source/Programs/Shared/EpicGames.Perforce/Records/ResolveOptions.cs
+++ b/Engine/Source/Programs/Shared/EpicGames.Perforce/Records/ResolveOptions.cs
@@ -1,4 +1,4 @@
-// Copyright Epic Games, Inc. All Rights Reserved.
+// Copyright 1998-2019 Epic Games, Inc. All Rights Reserved.
 
 using System;
 using System.Collections.Generic;
diff --git a/Engine/Source/Programs/Shared/EpicGames.Perforce/Records/ResolveRecord.cs b/Engine/Source/Programs/Shared/EpicGames.Perforce/Records/ResolveRecord.cs
--- a/Engine/Source/Programs/Shared/EpicGames.Perforce/Records/ResolveRecord.cs
+++ b/Engine/Source/Programs/Shared/EpicGames.Perforce/Records/ResolveRecord.cs
@@ -1,4 +1,4 @@
-// Copyright Epic Games, Inc. All Rights Reserved.
+// Copyright 1998-2019 Epic Games, Inc. All Rights Reserved.
 
 using System;
 using System.Collections.Generic;
@@ -47,7 +47,17 @@ namespace EpicGames.Perforce
 		/// For content resolves, the type of resolve to be performed
 		/// </summary>
 		[PerforceTag("contentResolveType", Optional = true)]
-		public string ContentResolveType;
+		public string? ContentResolveType;
+
+		/// <summary>
+		/// Private constructor for serialization
+		/// </summary>
+		private ResolveRecord()
+		{
+			ClientFile = null!;
+			FromFile = null!;
+			ResolveType = null!;
+		}
 
 		/// <summary>
 		/// Format this record for display in the debugger
diff --git a/Engine/Source/Programs/Shared/EpicGames.Perforce/Records/RevertOptions.cs b/Engine/Source/Programs/Shared/EpicGames.Perforce/Records/RevertOptions.cs
--- a/Engine/Source/Programs/Shared/EpicGames.Perforce/Records/RevertOptions.cs
+++ b/Engine/Source/Programs/Shared/EpicGames.Perforce/Records/RevertOptions.cs
@@ -1,4 +1,4 @@
-// Copyright Epic Games, Inc. All Rights Reserved.
+// Copyright 1998-2019 Epic Games, Inc. All Rights Reserved.
 
 using System;
 using System.Collections.Generic;
diff --git a/Engine/Source/Programs/Shared/EpicGames.Perforce/Records/RevertRecord.cs b/Engine/Source/Programs/Shared/EpicGames.Perforce/Records/RevertRecord.cs
--- a/Engine/Source/Programs/Shared/EpicGames.Perforce/Records/RevertRecord.cs
+++ b/Engine/Source/Programs/Shared/EpicGames.Perforce/Records/RevertRecord.cs
@@ -1,4 +1,4 @@
-// Copyright Epic Games, Inc. All Rights Reserved.
+// Copyright 1998-2019 Epic Games, Inc. All Rights Reserved.
 
 using System;
 using System.Collections.Generic;
@@ -43,6 +43,17 @@ namespace EpicGames.Perforce
 		[PerforceTag("action")]
 		public string Action;
 
+		/// <summary>
+		/// Private constructor for serialization
+		/// </summary>
+		private RevertRecord()
+		{
+			DepotFile = null!;
+			ClientFile = null!;
+			OldAction = null!;
+			Action = null!;
+		}
+
 		/// <summary>
 		/// Format this record for display in the debugger
 		/// </summary>
diff --git a/Engine/Source/Programs/Shared/EpicGames.Perforce/Records/RevisionRecord.cs b/Engine/Source/Programs/Shared/EpicGames.Perforce/Records/RevisionRecord.cs
--- a/Engine/Source/Programs/Shared/EpicGames.Perforce/Records/RevisionRecord.cs
+++ b/Engine/Source/Programs/Shared/EpicGames.Perforce/Records/RevisionRecord.cs
@@ -1,4 +1,4 @@
-// Copyright Epic Games, Inc. All Rights Reserved.
+// Copyright 1998-2019 Epic Games, Inc. All Rights Reserved.
 
 using System;
 using System.Collections.Generic;
@@ -65,7 +65,7 @@ namespace EpicGames.Perforce
 		/// Digest of the file, or null if not specified
 		/// </summary>
 		[PerforceTag("digest", Optional = true)]
-		public readonly string Digest;
+		public readonly string? Digest;
 
 		/// <summary>
 		/// Description of this changelist
@@ -79,6 +79,17 @@ namespace EpicGames.Perforce
 		[PerforceRecordList]
 		public List<IntegrationRecord> Integrations = new List<IntegrationRecord>();
 
+		/// <summary>
+		/// Private constructor for serialization
+		/// </summary>
+		private RevisionRecord()
+		{
+			Type = null!;
+			UserName = null!;
+			ClientName = null!;
+			Description = null!;
+		}
+
 		/// <summary>
 		/// Format this record for display in the debugger
 		/// </summary>
diff --git a/Engine/Source/Programs/Shared/EpicGames.Perforce/Records/ShelveOptions.cs b/Engine/Source/Programs/Shared/EpicGames.Perforce/Records/ShelveOptions.cs
--- a/Engine/Source/Programs/Shared/EpicGames.Perforce/Records/ShelveOptions.cs
+++ b/Engine/Source/Programs/Shared/EpicGames.Perforce/Records/ShelveOptions.cs
@@ -1,4 +1,4 @@
-// Copyright Epic Games, Inc. All Rights Reserved.
+// Copyright 1998-2019 Epic Games, Inc. All Rights Reserved.
 
 using System;
 using System.Collections.Generic;
diff --git a/Engine/Source/Programs/Shared/EpicGames.Perforce/Records/ShelveRecord.cs b/Engine/Source/Programs/Shared/EpicGames.Perforce/Records/ShelveRecord.cs
--- a/Engine/Source/Programs/Shared/EpicGames.Perforce/Records/ShelveRecord.cs
+++ b/Engine/Source/Programs/Shared/EpicGames.Perforce/Records/ShelveRecord.cs
@@ -1,4 +1,4 @@
-// Copyright Epic Games, Inc. All Rights Reserved.
+// Copyright 1998-2019 Epic Games, Inc. All Rights Reserved.
 
 using System;
 using System.Collections.Generic;
@@ -37,6 +37,15 @@ namespace EpicGames.Perforce
 		[PerforceTag("action")]
 		public string Action;
 
+		/// <summary>
+		/// Private constructor for serialization
+		/// </summary>
+		private ShelveRecord()
+		{
+			DepotFile = null!;
+			Action = null!;
+		}
+
 		/// <summary>
 		/// Format this record for display in the debugger
 		/// </summary>
diff --git a/Engine/Source/Programs/Shared/EpicGames.Perforce/Records/StreamOptions.cs b/Engine/Source/Programs/Shared/EpicGames.Perforce/Records/StreamOptions.cs
--- a/Engine/Source/Programs/Shared/EpicGames.Perforce/Records/StreamOptions.cs
+++ b/Engine/Source/Programs/Shared/EpicGames.Perforce/Records/StreamOptions.cs
@@ -1,77 +1,79 @@
-// Copyright Epic Games, Inc. All Rights Reserved.
-
 using System;
 using System.Collections.Generic;
-using System.Linq;
 using System.Text;
-using System.Threading.Tasks;
+using EpicGames.Perforce;
 
 namespace EpicGames.Perforce
 {
 	/// <summary>
-	/// Options for a stream definition
+	/// Options for a stream
 	/// </summary>
 	[Flags]
 	public enum StreamOptions
 	{
 		/// <summary>
-		/// The stream is locked
+		/// No options are set
 		/// </summary>
-		[PerforceEnum("locked")]
-		Locked = 1,
+		None = 0,
 
 		/// <summary>
-		/// The stream is unlocked
+		/// Enable other users' ability to edit or delete the stream. If locked, the stream specification cannot be deleted, and only its owner can modify it.
 		/// </summary>
 		[PerforceEnum("unlocked")]
-		Unlocked = 2,
+		Unlocked = 1,
 
 		/// <summary>
-		/// Only the owner may submit to the stream
+		/// Disable other users' ability to edit or delete the stream. If locked, the stream specification cannot be deleted, and only its owner can modify it.
 		/// </summary>
-		[PerforceEnum("ownersubmit")]
-		OwnerSubmit = 4,
+		[PerforceEnum("locked")]
+		Locked = 2,
 
 		/// <summary>
-		/// Anyone may submit to the stream
+		/// Specifies that all users can submit changes to the stream.
 		/// </summary>
 		[PerforceEnum("allsubmit")]
-		AllSubmit = 8,
+		AllSubmit = 4,
 
 		/// <summary>
-		/// Integrations from this stream to its parent are expected
+		/// Specifies that only the owner of the stream can submit changes to the stream.
+		/// </summary>
+		[PerforceEnum("ownersubmit")]
+		OwnerSubmit = 8,
+
+		/// <summary>
+		/// Specifies whether integrations from the stream to its parent are expected.
 		/// </summary>
 		[PerforceEnum("toparent")]
 		ToParent = 16,
 
 		/// <summary>
-		/// Integrations from this stream from its parent are expected
+		/// Specifies whether integrations from the stream to its parent are expected.
 		/// </summary>
 		[PerforceEnum("notoparent")]
-		NotToParent = 32,
+		NoToParent = 32,
 
 		/// <summary>
-		/// Integrations from this stream from its parent are expected
+		/// Specifies whether integrations to the stream from its parent are expected.
 		/// </summary>
 		[PerforceEnum("fromparent")]
 		FromParent = 64,
 
 		/// <summary>
-		/// Integrations from this stream from its parent are expected
+		/// Specifies whether integrations to the stream from its parent are expected.
 		/// </summary>
-		[PerforceEnum("nofromparent")]
-		NotFromParent = 128,
+		[PerforceEnum("notoparent")]
+		NoFromParent = 128,
 
 		/// <summary>
-		/// Undocumented?
+		/// Specifies whether the merge flow is restricted or whether merge is permitted from any other stream.
 		/// </summary>
-		[PerforceEnum("mergedown")]
-		MergeDown = 256,
+		[PerforceEnum("mergeany")]
+		MergeAny = 256,
 
 		/// <summary>
-		/// Undocumented?
+		/// Specifies whether the merge flow is restricted or whether merge is permitted from any other stream.
 		/// </summary>
-		[PerforceEnum("mergeany")]
-		MergeAny = 512,
+		[PerforceEnum("mergedown")]
+		MergeDown = 512,
 	}
 }
diff --git a/Engine/Source/Programs/Shared/EpicGames.Perforce/Records/StreamRecord.cs b/Engine/Source/Programs/Shared/EpicGames.Perforce/Records/StreamRecord.cs
--- a/Engine/Source/Programs/Shared/EpicGames.Perforce/Records/StreamRecord.cs
+++ b/Engine/Source/Programs/Shared/EpicGames.Perforce/Records/StreamRecord.cs
@@ -1,104 +1,89 @@
-// Copyright Epic Games, Inc. All Rights Reserved.
-
 using System;
 using System.Collections.Generic;
-using System.Diagnostics;
-using System.Linq;
 using System.Text;
-using System.Threading.Tasks;
+using EpicGames.Perforce;
 
 namespace EpicGames.Perforce
 {
 	/// <summary>
-	/// Contains information about a stream, as returned by the 'p4 streams' command
+	/// Represents a Perforce stream spec
 	/// </summary>
 	public class StreamRecord
 	{
 		/// <summary>
-		/// Path to the stream
+		/// The stream path
 		/// </summary>
 		[PerforceTag("Stream")]
 		public string Stream;
 
 		/// <summary>
-		/// Last time the stream definition was updated
+		/// The stream name
 		/// </summary>
-		[PerforceTag("Update")]
-		public DateTime Update;
+		[PerforceTag("Name")]
+		public string Name;
 
 		/// <summary>
-		/// Last time the stream definition was accessed
+		/// The Perforce user name of the user who owns the stream.
 		/// </summary>
-		[PerforceTag("Access")]
-		public DateTime Access;
+		[PerforceTag("Owner")]
+		public string Owner;
 
 		/// <summary>
-		/// Owner of this stream
+		/// The date the stream specification was last modified.
 		/// </summary>
-		[PerforceTag("Owner")]
-		public string Owner;
+		[PerforceTag("Update", Optional = true)]
+		public DateTime Update;
 
 		/// <summary>
-		/// Name of the stream. This may be modified after the stream is initially created, but it's underlying depot path will not change.
+		/// The date and time that the stream was last used in any way.
 		/// </summary>
-		[PerforceTag("Name")]
-		public string Name;
+		[PerforceTag("Access", Optional = true)]
+		public DateTime Access;
 
 		/// <summary>
 		/// The parent stream
 		/// </summary>
-		[PerforceTag("Parent")]
-		public string Parent;
+		[PerforceTag("Parent", Optional = true)]
+		public string? Parent;
 
 		/// <summary>
-		/// Type of the stream
+		/// The stream type
 		/// </summary>
-		[PerforceTag("Type")]
-		public StreamType Type;
+		[PerforceTag("Type", Optional = true)]
+		public string? Type;
 
 		/// <summary>
-		/// User supplied description of the stream
+		/// A textual description of the stream.
 		/// </summary>
-		[PerforceTag("desc")]
-		public string Description;
+		[PerforceTag("Description", Optional = true)]
+		public string? Description;
 
 		/// <summary>
-		/// Options for the stream definition
+		/// Options for this stream
 		/// </summary>
 		[PerforceTag("Options")]
 		public StreamOptions Options;
 
 		/// <summary>
-		/// Whether this stream is more stable than the parent stream
-		/// </summary>
-		[PerforceTag("firmerThanParent")]
-		public Nullable<bool> FirmerThanParent;
-
-		/// <summary>
-		/// Whether changes from this stream flow to the parent stream
-		/// </summary>
-		[PerforceTag("changeFlowsToParent")]
-		public bool ChangeFlowsToParent;
-
-		/// <summary>
-		/// Whether changes from this stream flow from the parent stream
+		/// List of paths in the stream spec
 		/// </summary>
-		[PerforceTag("changeFlowsFromParent")]
-		public bool ChangeFlowsFromParent;
+		[PerforceTag("Paths")]
+		public List<string> Paths = new List<string>();
 
 		/// <summary>
-		/// The mainline branch associated with this stream
+		/// Computed view for the stream
 		/// </summary>
-		[PerforceTag("baseParent")]
-		public string BaseParent;
+		[PerforceTag("View")]
+		public List<string> View = new List<string>();
 
 		/// <summary>
-		/// Summarize this record for display in the debugger
+		/// Private constructor for serialization
 		/// </summary>
-		/// <returns>Formatted stream information</returns>
-		public override string ToString()
+		private StreamRecord()
 		{
-			return Name;
+			Stream = null!;
+			Name = null!;
+			Owner = null!;
 		}
 	}
 }
diff --git a/Engine/Source/Programs/Shared/EpicGames.Perforce/Records/StreamType.cs b/Engine/Source/Programs/Shared/EpicGames.Perforce/Records/StreamType.cs
--- a/Engine/Source/Programs/Shared/EpicGames.Perforce/Records/StreamType.cs
+++ b/Engine/Source/Programs/Shared/EpicGames.Perforce/Records/StreamType.cs
@@ -1,4 +1,4 @@
-// Copyright Epic Games, Inc. All Rights Reserved.
+// Copyright 1998-2019 Epic Games, Inc. All Rights Reserved.
 
 using System;
 using System.Collections.Generic;
diff --git a/Engine/Source/Programs/Shared/EpicGames.Perforce/Records/StreamsOptions.cs b/Engine/Source/Programs/Shared/EpicGames.Perforce/Records/StreamsOptions.cs
new file mode 100644
--- /dev/null
+++ b/Engine/Source/Programs/Shared/EpicGames.Perforce/Records/StreamsOptions.cs
@@ -0,0 +1,77 @@
+// Copyright 1998-2019 Epic Games, Inc. All Rights Reserved.
+
+using System;
+using System.Collections.Generic;
+using System.Linq;
+using System.Text;
+using System.Threading.Tasks;
+
+namespace EpicGames.Perforce
+{
+	/// <summary>
+	/// Options for a stream definition
+	/// </summary>
+	[Flags]
+	public enum StreamsOptions
+	{
+		/// <summary>
+		/// The stream is locked
+		/// </summary>
+		[PerforceEnum("locked")]
+		Locked = 1,
+
+		/// <summary>
+		/// The stream is unlocked
+		/// </summary>
+		[PerforceEnum("unlocked")]
+		Unlocked = 2,
+
+		/// <summary>
+		/// Only the owner may submit to the stream
+		/// </summary>
+		[PerforceEnum("ownersubmit")]
+		OwnerSubmit = 4,
+
+		/// <summary>
+		/// Anyone may submit to the stream
+		/// </summary>
+		[PerforceEnum("allsubmit")]
+		AllSubmit = 8,
+
+		/// <summary>
+		/// Integrations from this stream to its parent are expected
+		/// </summary>
+		[PerforceEnum("toparent")]
+		ToParent = 16,
+
+		/// <summary>
+		/// Integrations from this stream from its parent are expected
+		/// </summary>
+		[PerforceEnum("notoparent")]
+		NotToParent = 32,
+
+		/// <summary>
+		/// Integrations from this stream from its parent are expected
+		/// </summary>
+		[PerforceEnum("fromparent")]
+		FromParent = 64,
+
+		/// <summary>
+		/// Integrations from this stream from its parent are expected
+		/// </summary>
+		[PerforceEnum("nofromparent")]
+		NotFromParent = 128,
+
+		/// <summary>
+		/// Undocumented?
+		/// </summary>
+		[PerforceEnum("mergedown")]
+		MergeDown = 256,
+
+		/// <summary>
+		/// Undocumented?
+		/// </summary>
+		[PerforceEnum("mergeany")]
+		MergeAny = 512,
+	}
+}
diff --git a/Engine/Source/Programs/Shared/EpicGames.Perforce/Records/StreamsRecord.cs b/Engine/Source/Programs/Shared/EpicGames.Perforce/Records/StreamsRecord.cs
new file mode 100644
--- /dev/null
+++ b/Engine/Source/Programs/Shared/EpicGames.Perforce/Records/StreamsRecord.cs
@@ -0,0 +1,114 @@
+// Copyright 1998-2019 Epic Games, Inc. All Rights Reserved.
+
+using System;
+using System.Collections.Generic;
+using System.Diagnostics;
+using System.Linq;
+using System.Text;
+using System.Threading.Tasks;
+
+namespace EpicGames.Perforce
+{
+	/// <summary>
+	/// Contains information about a stream, as returned by the 'p4 streams' command
+	/// </summary>
+	public class StreamsRecord
+	{
+		/// <summary>
+		/// Path to the stream
+		/// </summary>
+		[PerforceTag("Stream")]
+		public string Stream;
+
+		/// <summary>
+		/// Last time the stream definition was updated
+		/// </summary>
+		[PerforceTag("Update")]
+		public DateTime Update;
+
+		/// <summary>
+		/// Last time the stream definition was accessed
+		/// </summary>
+		[PerforceTag("Access")]
+		public DateTime Access;
+
+		/// <summary>
+		/// Owner of this stream
+		/// </summary>
+		[PerforceTag("Owner")]
+		public string Owner;
+
+		/// <summary>
+		/// Name of the stream. This may be modified after the stream is initially created, but it's underlying depot path will not change.
+		/// </summary>
+		[PerforceTag("Name")]
+		public string Name;
+
+		/// <summary>
+		/// The parent stream
+		/// </summary>
+		[PerforceTag("Parent", Optional = true)]
+		public string? Parent;
+
+		/// <summary>
+		/// Type of the stream
+		/// </summary>
+		[PerforceTag("Type")]
+		public StreamType Type;
+
+		/// <summary>
+		/// User supplied description of the stream
+		/// </summary>
+		[PerforceTag("desc", Optional = true)]
+		public string? Description;
+
+		/// <summary>
+		/// Options for the stream definition
+		/// </summary>
+		[PerforceTag("Options")]
+		public StreamsOptions Options;
+
+		/// <summary>
+		/// Whether this stream is more stable than the parent stream
+		/// </summary>
+		[PerforceTag("firmerThanParent")]
+		public Nullable<bool> FirmerThanParent;
+
+		/// <summary>
+		/// Whether changes from this stream flow to the parent stream
+		/// </summary>
+		[PerforceTag("changeFlowsToParent")]
+		public bool ChangeFlowsToParent;
+
+		/// <summary>
+		/// Whether changes from this stream flow from the parent stream
+		/// </summary>
+		[PerforceTag("changeFlowsFromParent")]
+		public bool ChangeFlowsFromParent;
+
+		/// <summary>
+		/// The mainline branch associated with this stream
+		/// </summary>
+		[PerforceTag("baseParent", Optional = true)]
+		public string? BaseParent;
+
+		/// <summary>
+		/// Private constructor for serialization
+		/// </summary>
+		private StreamsRecord()
+		{
+			Stream = null!;
+			Owner = null!;
+			Name = null!;
+		}
+
+		/// <summary>
+		/// Summarize this record for display in the debugger
+		/// </summary>
+		/// <returns>Formatted stream information</returns>
+		public override string? ToString()
+		{
+			return Name;
+		}
+	}
+}
diff --git a/Engine/Source/Programs/Shared/EpicGames.Perforce/Records/SubmitOptions.cs b/Engine/Source/Programs/Shared/EpicGames.Perforce/Records/SubmitOptions.cs
--- a/Engine/Source/Programs/Shared/EpicGames.Perforce/Records/SubmitOptions.cs
+++ b/Engine/Source/Programs/Shared/EpicGames.Perforce/Records/SubmitOptions.cs
@@ -1,4 +1,4 @@
-// Copyright Epic Games, Inc. All Rights Reserved.
+// Copyright 1998-2019 Epic Games, Inc. All Rights Reserved.
 
 using System;
 using System.Collections.Generic;
diff --git a/Engine/Source/Programs/Shared/EpicGames.Perforce/Records/SubmitRecord.cs b/Engine/Source/Programs/Shared/EpicGames.Perforce/Records/SubmitRecord.cs
--- a/Engine/Source/Programs/Shared/EpicGames.Perforce/Records/SubmitRecord.cs
+++ b/Engine/Source/Programs/Shared/EpicGames.Perforce/Records/SubmitRecord.cs
@@ -1,4 +1,4 @@
-// Copyright Epic Games, Inc. All Rights Reserved.
+// Copyright 1998-2019 Epic Games, Inc. All Rights Reserved.
 
 using System;
 using System.Collections.Generic;
diff --git a/Engine/Source/Programs/Shared/EpicGames.Perforce/Records/SwitchClientOptions.cs b/Engine/Source/Programs/Shared/EpicGames.Perforce/Records/SwitchClientOptions.cs
--- a/Engine/Source/Programs/Shared/EpicGames.Perforce/Records/SwitchClientOptions.cs
+++ b/Engine/Source/Programs/Shared/EpicGames.Perforce/Records/SwitchClientOptions.cs
@@ -1,4 +1,4 @@
-// Copyright Epic Games, Inc. All Rights Reserved.
+// Copyright 1998-2019 Epic Games, Inc. All Rights Reserved.
 
 using System;
 using System.Collections.Generic;
diff --git a/Engine/Source/Programs/Shared/EpicGames.Perforce/Records/SyncOptions.cs b/Engine/Source/Programs/Shared/EpicGames.Perforce/Records/SyncOptions.cs
--- a/Engine/Source/Programs/Shared/EpicGames.Perforce/Records/SyncOptions.cs
+++ b/Engine/Source/Programs/Shared/EpicGames.Perforce/Records/SyncOptions.cs
@@ -1,4 +1,4 @@
-// Copyright Epic Games, Inc. All Rights Reserved.
+// Copyright 1998-2019 Epic Games, Inc. All Rights Reserved.
 
 using System;
 using System.Collections.Generic;
diff --git a/Engine/Source/Programs/Shared/EpicGames.Perforce/Records/SyncRecord.cs b/Engine/Source/Programs/Shared/EpicGames.Perforce/Records/SyncRecord.cs
--- a/Engine/Source/Programs/Shared/EpicGames.Perforce/Records/SyncRecord.cs
+++ b/Engine/Source/Programs/Shared/EpicGames.Perforce/Records/SyncRecord.cs
@@ -1,4 +1,4 @@
-// Copyright Epic Games, Inc. All Rights Reserved.
+// Copyright 1998-2019 Epic Games, Inc. All Rights Reserved.
 
 using System;
 using System.Collections.Generic;
@@ -61,6 +61,16 @@ namespace EpicGames.Perforce
 		[PerforceTag("change", Optional = true)]
 		public int Change;
 
+		/// <summary>
+		/// Private constructor for serialization
+		/// </summary>
+		private SyncRecord()
+		{
+			DepotFile = null!;
+			ClientFile = null!;
+			Action = null!;
+		}
+
 		/// <summary>
 		/// Format this record for display in the debugger
 		/// </summary>
diff --git a/Engine/Source/Programs/Shared/EpicGames.Perforce/Records/SyncSummaryRecord.cs b/Engine/Source/Programs/Shared/EpicGames.Perforce/Records/SyncSummaryRecord.cs
--- a/Engine/Source/Programs/Shared/EpicGames.Perforce/Records/SyncSummaryRecord.cs
+++ b/Engine/Source/Programs/Shared/EpicGames.Perforce/Records/SyncSummaryRecord.cs
@@ -1,4 +1,4 @@
-// Copyright Epic Games, Inc. All Rights Reserved.
+// Copyright 1998-2019 Epic Games, Inc. All Rights Reserved.
 
 using System;
 using System.Collections.Generic;
diff --git a/Engine/Source/Programs/Shared/EpicGames.Perforce/Records/UnshelveOptions.cs b/Engine/Source/Programs/Shared/EpicGames.Perforce/Records/UnshelveOptions.cs
--- a/Engine/Source/Programs/Shared/EpicGames.Perforce/Records/UnshelveOptions.cs
+++ b/Engine/Source/Programs/Shared/EpicGames.Perforce/Records/UnshelveOptions.cs
@@ -1,4 +1,4 @@
-// Copyright Epic Games, Inc. All Rights Reserved.
+// Copyright 1998-2019 Epic Games, Inc. All Rights Reserved.
 
 using System;
 using System.Collections.Generic;
diff --git a/Engine/Source/Programs/Shared/EpicGames.Perforce/Records/UnshelveRecord.cs b/Engine/Source/Programs/Shared/EpicGames.Perforce/Records/UnshelveRecord.cs
--- a/Engine/Source/Programs/Shared/EpicGames.Perforce/Records/UnshelveRecord.cs
+++ b/Engine/Source/Programs/Shared/EpicGames.Perforce/Records/UnshelveRecord.cs
@@ -1,4 +1,4 @@
-// Copyright Epic Games, Inc. All Rights Reserved.
+// Copyright 1998-2019 Epic Games, Inc. All Rights Reserved.
 
 using System;
 using System.Collections.Generic;
@@ -31,6 +31,15 @@ namespace EpicGames.Perforce
 		[PerforceTag("action")]
 		public string Action;
 
+		/// <summary>
+		/// Private constructor for serialization
+		/// </summary>
+		private UnshelveRecord()
+		{
+			DepotFile = null!;
+			Action = null!;
+		}
+
 		/// <summary>
 		/// Format this record for display in the debugger
 		/// </summary>
diff --git a/Engine/Source/Programs/Shared/EpicGames.Perforce/Records/UpdateChangeOptions.cs b/Engine/Source/Programs/Shared/EpicGames.Perforce/Records/UpdateChangeOptions.cs
--- a/Engine/Source/Programs/Shared/EpicGames.Perforce/Records/UpdateChangeOptions.cs
+++ b/Engine/Source/Programs/Shared/EpicGames.Perforce/Records/UpdateChangeOptions.cs
@@ -1,4 +1,4 @@
-// Copyright Epic Games, Inc. All Rights Reserved.
+// Copyright 1998-2019 Epic Games, Inc. All Rights Reserved.
 
 using System;
 using System.Collections.Generic;
diff --git a/Engine/Source/Programs/Shared/EpicGames.Perforce/Records/UserRecord.cs b/Engine/Source/Programs/Shared/EpicGames.Perforce/Records/UserRecord.cs
--- a/Engine/Source/Programs/Shared/EpicGames.Perforce/Records/UserRecord.cs
+++ b/Engine/Source/Programs/Shared/EpicGames.Perforce/Records/UserRecord.cs
@@ -1,4 +1,4 @@
-// Copyright Epic Games, Inc. All Rights Reserved.
+// Copyright 1998-2019 Epic Games, Inc. All Rights Reserved.
 
 using System;
 using System.Collections.Generic;
@@ -59,7 +59,18 @@ namespace EpicGames.Perforce
 		/// Method used to authenticate
 		/// </summary>
 		[PerforceTag("AuthMethod")]
-		public string AuthMethod;
+		public string? AuthMethod;
+
+		/// <summary>
+		/// Private constructor for serialization
+		/// </summary>
+		private UserRecord()
+		{
+			UserName = null!;
+			Email = null!;
+			FullName = null!;
+			Type = null!;
+		}
 
 		/// <summary>
 		/// Summarize this record for display in the debugger
diff --git a/Engine/Source/Programs/Shared/EpicGames.Perforce/Records/WhereRecord.cs b/Engine/Source/Programs/Shared/EpicGames.Perforce/Records/WhereRecord.cs
--- a/Engine/Source/Programs/Shared/EpicGames.Perforce/Records/WhereRecord.cs
+++ b/Engine/Source/Programs/Shared/EpicGames.Perforce/Records/WhereRecord.cs
@@ -1,4 +1,4 @@
-// Copyright Epic Games, Inc. All Rights Reserved.
+// Copyright 1998-2019 Epic Games, Inc. All Rights Reserved.
 
 using System;
 using System.Collections.Generic;
@@ -37,6 +37,16 @@ namespace EpicGames.Perforce
 		[PerforceTag("unmap", Optional = true)]
 		public bool Unmap;
 
+		/// <summary>
+		/// Private constructor for serialization
+		/// </summary>
+		private WhereRecord()
+		{
+			DepotFile = null!;
+			ClientFile = null!;
+			Path = null!;
+		}
+
 		/// <summary>
 		/// Summarize this record for display in the debugger
 		/// </summary>
