Various minor fixes to Perforce library.

#rnx

diff --git a/Engine/Source/Programs/Shared/EpicGames.Perforce/PerforceConnection.cs b/Engine/Source/Programs/Shared/EpicGames.Perforce/PerforceConnection.cs
--- a/Engine/Source/Programs/Shared/EpicGames.Perforce/PerforceConnection.cs
+++ b/Engine/Source/Programs/Shared/EpicGames.Perforce/PerforceConnection.cs
@@ -912,7 +912,7 @@ namespace EpicGames.Perforce
 		/// <param name="ClientName">Name of the client. Specify null for the current client.</param>
 		/// <param name="CancellationToken">Token used to cancel the operation</param>
 		/// <returns>Response from the server; either a client record or error code</returns>
-		public async Task<ClientRecord> GetClientAsync(string ClientName, CancellationToken CancellationToken)
+		public async Task<ClientRecord> GetClientAsync(string? ClientName, CancellationToken CancellationToken)
 		{
 			return (await TryGetClientAsync(ClientName, CancellationToken)).Data;
 		}
@@ -923,7 +923,7 @@ namespace EpicGames.Perforce
 		/// <param name="ClientName">Name of the client. Specify null for the current client.</param>
 		/// <param name="CancellationToken">Token used to cancel the operation</param>
 		/// <returns>Response from the server; either a client record or error code</returns>
-		public Task<PerforceResponse<ClientRecord>> TryGetClientAsync(string ClientName, CancellationToken CancellationToken)
+		public Task<PerforceResponse<ClientRecord>> TryGetClientAsync(string? ClientName, CancellationToken CancellationToken)
 		{
 			StringBuilder Arguments = new StringBuilder("client -o");
 			if (ClientName != null)
@@ -1334,7 +1334,7 @@ namespace EpicGames.Perforce
 		/// <param name="FileSpecs">Files to be opened for edit</param>
 		/// <param name="CancellationToken">Token used to cancel the operation</param>
 		/// <returns>Response from the server</returns>
-		public async Task<List<EditRecord>> EditAsync(int ChangeNumber, string FileType, EditOptions Options, string[] FileSpecs, CancellationToken CancellationToken)
+		public async Task<List<EditRecord>> EditAsync(int ChangeNumber, string? FileType, EditOptions Options, string[] FileSpecs, CancellationToken CancellationToken)
 		{
 			return (await TryEditAsync(ChangeNumber, FileType, Options, FileSpecs, CancellationToken)).Data;
 		}
@@ -1796,6 +1796,63 @@ namespace EpicGames.Perforce
 
 		#endregion
 
+		#region p4 move
+
+		/// <summary>
+		/// Opens files for move
+		/// </summary>
+		/// <param name="ChangeNumber">Changelist to add files to</param>
+		/// <param name="FileType">Type for new files</param>
+		/// <param name="Options">Options for the command</param>
+		/// <param name="SourceFileSpec">The source file(s)</param>
+		/// <param name="TargetFileSpec">The target file(s)</param>
+		/// <param name="CancellationToken">Token used to cancel the operation</param>
+		/// <returns>Response from the server</returns>
+		public async Task<List<MoveRecord>> MoveAsync(int ChangeNumber, string? FileType, MoveOptions Options, string SourceFileSpec, string TargetFileSpec, CancellationToken CancellationToken)
+		{
+			return (await TryMoveAsync(ChangeNumber, FileType, Options, SourceFileSpec, TargetFileSpec, CancellationToken)).Data;
+		}
+
+		/// <summary>
+		/// Opens files for move
+		/// </summary>
+		/// <param name="ChangeNumber">Changelist to add files to</param>
+		/// <param name="FileType">Type for new files</param>
+		/// <param name="Options">Options for the command</param>
+		/// <param name="SourceFileSpec">The source file(s)</param>
+		/// <param name="TargetFileSpec">The target file(s)</param>
+		/// <param name="CancellationToken">Token used to cancel the operation</param>
+		/// <returns>Response from the server</returns>
+		public Task<PerforceResponseList<MoveRecord>> TryMoveAsync(int ChangeNumber, string? FileType, MoveOptions Options, string SourceFileSpec, string TargetFileSpec, CancellationToken CancellationToken)
+		{
+			StringBuilder Arguments = new StringBuilder("move");
+			if (ChangeNumber != -1)
+			{
+				Arguments.AppendFormat(" -c {0}", ChangeNumber);
+			}
+			if ((Options & MoveOptions.KeepWorkspaceFiles) != 0)
+			{
+				Arguments.Append(" -k");
+			}
+			if ((Options & MoveOptions.RenameOnly) != 0)
+			{
+				Arguments.Append(" -r");
+			}
+			if ((Options & MoveOptions.PreviewOnly) != 0)
+			{
+				Arguments.Append(" -n");
+			}
+			if (FileType != null)
+			{
+				Arguments.AppendFormat(" -t \"{0}\"", FileType);
+			}
+			Arguments.AppendFormat(" \"{0}\" \"{1}\"", SourceFileSpec, TargetFileSpec);
+
+			return CommandAsync<MoveRecord>(Arguments.ToString(), null, CancellationToken);
+		}
+
+		#endregion
+
 		#region p4 opened
 
 		/// <summary>
@@ -1809,7 +1866,7 @@ namespace EpicGames.Perforce
 		/// <param name="FileSpecs">Specification for the files to list</param>
 		/// <param name="CancellationToken">Token used to cancel the operation</param>
 		/// <returns>Response from the server</returns>
-		public async Task<List<FStatRecord>> GetOpenFilesAsync(OpenedOptions Options, int ChangeNumber, string? ClientName, string? UserName, int MaxResults, string[] FileSpecs, CancellationToken CancellationToken)
+		public async Task<List<FStatRecord>> GetOpenFilesAsync(OpenedOptions Options, int ChangeNumber, string? ClientName, string? UserName, int MaxResults, string[]? FileSpecs, CancellationToken CancellationToken)
 		{
 			return (await TryGetOpenFilesAsync(Options, ChangeNumber, ClientName, UserName, MaxResults, FileSpecs, CancellationToken)).Data;
 		}
@@ -1825,7 +1882,7 @@ namespace EpicGames.Perforce
 		/// <param name="FileSpecs">Specification for the files to list</param>
 		/// <param name="CancellationToken">Token used to cancel the operation</param>
 		/// <returns>Response from the server</returns>
-		public Task<PerforceResponseList<FStatRecord>> TryGetOpenFilesAsync(OpenedOptions Options, int ChangeNumber, string? ClientName, string? UserName, int MaxResults, string[] FileSpecs, CancellationToken CancellationToken)
+		public Task<PerforceResponseList<FStatRecord>> TryGetOpenFilesAsync(OpenedOptions Options, int ChangeNumber, string? ClientName, string? UserName, int MaxResults, string[]? FileSpecs, CancellationToken CancellationToken)
 		{
 			// Build the argument list
 			StringBuilder Arguments = new StringBuilder("opened");
@@ -1833,7 +1890,11 @@ namespace EpicGames.Perforce
 			{
 				Arguments.AppendFormat(" -a");
 			}
-			if(ChangeNumber != -1)
+			if (ChangeNumber == PerforceReflection.DefaultChange)
+			{
+				Arguments.AppendFormat(" -c default");
+			}
+			else if (ChangeNumber != -1)
 			{
 				Arguments.AppendFormat(" -c {0}", ChangeNumber);
 			}
@@ -1845,11 +1906,7 @@ namespace EpicGames.Perforce
 			{
 				Arguments.AppendFormat(" -u \"{0}\"", UserName);
 			}
-			if(MaxResults == PerforceReflection.DefaultChange)
-			{
-				Arguments.AppendFormat(" -m default");
-			}
-			else if(MaxResults != -1)
+			if(MaxResults != -1)
 			{
 				Arguments.AppendFormat(" -m {0}", MaxResults);
 			}
@@ -1857,9 +1914,12 @@ namespace EpicGames.Perforce
 			{
 				Arguments.AppendFormat(" -s");
 			}
-			foreach(string FileSpec in FileSpecs)
+			if (FileSpecs != null)
 			{
-				Arguments.AppendFormat(" \"{0}\"", FileSpec);
+				foreach (string FileSpec in FileSpecs)
+				{
+					Arguments.AppendFormat(" \"{0}\"", FileSpec);
+				}
 			}
 			return CommandAsync<FStatRecord>(Arguments.ToString(), null, CancellationToken);
 		}
@@ -2029,6 +2089,54 @@ namespace EpicGames.Perforce
 
 		#endregion
 
+		#region p4 reopen
+
+		/// <summary>
+		/// Reopen a file
+		/// </summary>
+		/// <param name="ChangeNumber">Changelist to open files to</param>
+		/// <param name="FileType">New filetype</param>
+		/// <param name="FileSpec">Files to be reverted</param>
+		/// <param name="CancellationToken">Token used to cancel the operation</param>
+		/// <returns>Response from the server</returns>
+		public async Task<List<ReopenRecord>> ReopenAsync(int? ChangeNumber, string? FileType, string FileSpec, CancellationToken CancellationToken)
+		{
+			return (await TryReopenAsync(ChangeNumber, FileType, FileSpec, CancellationToken)).Data;
+		}
+
+		/// <summary>
+		/// Reopen a file
+		/// </summary>
+		/// <param name="ChangeNumber">Changelist to open files to</param>
+		/// <param name="FileType">New filetype</param>
+		/// <param name="FileSpec">Files to be reverted</param>
+		/// <param name="CancellationToken">Token used to cancel the operation</param>
+		/// <returns>Response from the server</returns>
+		public Task<PerforceResponseList<ReopenRecord>> TryReopenAsync(int? ChangeNumber, string? FileType, string FileSpec, CancellationToken CancellationToken)
+		{
+			StringBuilder Arguments = new StringBuilder("reopen");
+			if (ChangeNumber != null)
+			{
+				if (ChangeNumber == PerforceReflection.DefaultChange)
+				{
+					Arguments.Append(" -c default");
+				}
+				else
+				{
+					Arguments.AppendFormat(" -c {0}", ChangeNumber);
+				}
+			}
+			if (FileType != null)
+			{
+				Arguments.AppendFormat(" -t \"{0}\"", FileType);
+			}
+			Arguments.AppendFormat(" \"{0}\"", FileSpec);
+
+			return CommandAsync<ReopenRecord>(Arguments.ToString(), null, CancellationToken);
+		}
+
+		#endregion
+
 		#region p4 reload
 
 		/// <summary>
@@ -2055,7 +2163,7 @@ namespace EpicGames.Perforce
 		/// <param name="FileSpecs">Files to be reverted</param>
 		/// <param name="CancellationToken">Token used to cancel the operation</param>
 		/// <returns>Response from the server</returns>
-		public async Task<List<ResolveRecord>> ResolveAsync(int ChangeNumber, ResolveOptions Options, string[] FileSpecs, CancellationToken CancellationToken)
+		public async Task<List<ResolveRecord>> ResolveAsync(int ChangeNumber, ResolveOptions Options, string[]? FileSpecs, CancellationToken CancellationToken)
 		{
 			return (await TryResolveAsync(ChangeNumber, Options, FileSpecs, CancellationToken)).Data;
 		}
@@ -2068,7 +2176,7 @@ namespace EpicGames.Perforce
 		/// <param name="FileSpecs">Files to be reverted</param>
 		/// <param name="CancellationToken">Token used to cancel the operation</param>
 		/// <returns>Response from the server</returns>
-		public Task<PerforceResponseList<ResolveRecord>> TryResolveAsync(int ChangeNumber, ResolveOptions Options, string[] FileSpecs, CancellationToken CancellationToken)
+		public Task<PerforceResponseList<ResolveRecord>> TryResolveAsync(int ChangeNumber, ResolveOptions Options, string[]? FileSpecs, CancellationToken CancellationToken)
 		{
 			StringBuilder Arguments = new StringBuilder("resolve");
 			if ((Options & ResolveOptions.Automatic) != 0)
@@ -2111,9 +2219,12 @@ namespace EpicGames.Perforce
 			{
 				Arguments.Append(" -n");
 			}
-			foreach(string FileSpec in FileSpecs)
+			if (FileSpecs != null)
 			{
-				Arguments.AppendFormat(" \"{0}\"", FileSpec);
+				foreach (string FileSpec in FileSpecs)
+				{
+					Arguments.AppendFormat(" \"{0}\"", FileSpec);
+				}
 			}
 
 			return CommandAsync<ResolveRecord>(Arguments.ToString(), null, CancellationToken);
@@ -2811,7 +2922,7 @@ namespace EpicGames.Perforce
 		/// <returns>List of responses from the server</returns>
 		public async Task<List<WhereRecord>> WhereAsync(string[] FileSpecs, CancellationToken CancellationToken)
 		{
-			List<WhereRecord> Records = (await TryWhere(FileSpecs, CancellationToken)).Data;
+			List<WhereRecord> Records = (await TryWhereAsync(FileSpecs, CancellationToken)).Data;
 			if (Records.Count != FileSpecs.Length)
 			{
 				throw new PerforceException($"Unexpected response count; expected {FileSpecs.Length}, got {Records.Count}.");
@@ -2825,7 +2936,7 @@ namespace EpicGames.Perforce
 		/// <param name="FileSpecs">Patterns for the files to query</param>
 		/// <param name="CancellationToken">Token used to cancel the operation</param>
 		/// <returns>List of responses from the server</returns>
-		public async Task<PerforceResponseList<WhereRecord>> TryWhere(string[] FileSpecs, CancellationToken CancellationToken)
+		public async Task<PerforceResponseList<WhereRecord>> TryWhereAsync(string[] FileSpecs, CancellationToken CancellationToken)
 		{
 			PerforceResponseList<WhereRecord> Responses = new PerforceResponseList<WhereRecord>();
 			for (int Idx = 0; Idx < FileSpecs.Length; )
diff --git a/Engine/Source/Programs/Shared/EpicGames.Perforce/PerforceResponseList.cs b/Engine/Source/Programs/Shared/EpicGames.Perforce/PerforceResponseList.cs
--- a/Engine/Source/Programs/Shared/EpicGames.Perforce/PerforceResponseList.cs
+++ b/Engine/Source/Programs/Shared/EpicGames.Perforce/PerforceResponseList.cs
@@ -36,7 +36,18 @@ namespace EpicGames.Perforce
 		/// </summary>
 		public List<T> Data
 		{
-			get { return this.Select(x => x.Data).Where(x => x != null).ToList(); }
+			get
+			{
+				if (Count == 1)
+				{
+					PerforceError? Error = this[0].Error;
+					if (Error != null && Error.Generic == PerforceGenericCode.Empty)
+					{
+						return new List<T>();
+					}
+				}
+				return this.Select(x => x.Data).Where(x => x != null).ToList();
+			}
 		}
 
 		/// <summary>
diff --git a/Engine/Source/Programs/Shared/EpicGames.Perforce/Records/FStatRecord.cs b/Engine/Source/Programs/Shared/EpicGames.Perforce/Records/FStatRecord.cs
--- a/Engine/Source/Programs/Shared/EpicGames.Perforce/Records/FStatRecord.cs
+++ b/Engine/Source/Programs/Shared/EpicGames.Perforce/Records/FStatRecord.cs
@@ -25,6 +25,12 @@ namespace EpicGames.Perforce
 		[PerforceTag("clientFile", Optional = true)]
 		public string? ClientFile;
 
+		/// <summary>
+		/// For move operations, shows the counterpart file
+		/// </summary>
+		[PerforceTag("movedFile", Optional = true)]
+		public string? MovedFile;
+
 		/// <summary>
 		/// Local path to file
 		/// </summary>
@@ -61,6 +67,12 @@ namespace EpicGames.Perforce
 		[PerforceTag("headTime", Optional = true)]
 		public DateTime HeadTime;
 
+		/// <summary>
+		/// Revision for opened files
+		/// </summary>
+		[PerforceTag("rev", Optional = true)]
+		public int Revision;
+
 		/// <summary>
 		/// Head revision number, if in depot
 		/// </summary>
@@ -119,7 +131,7 @@ namespace EpicGames.Perforce
 		/// Open type, if opened in your workspace
 		/// </summary>
 		[PerforceTag("type", Optional = true)]
-		public FileAction Type;
+		public string? Type;
 
 		/// <summary>
 		/// User who opened the file, if opens
diff --git a/Engine/Source/Programs/Shared/EpicGames.Perforce/Records/MoveOptions.cs b/Engine/Source/Programs/Shared/EpicGames.Perforce/Records/MoveOptions.cs
new file mode 100644
--- /dev/null
+++ b/Engine/Source/Programs/Shared/EpicGames.Perforce/Records/MoveOptions.cs
@@ -0,0 +1,37 @@
+// Copyright 1998-2019 Epic Games, Inc. All Rights Reserved.
+
+using System;
+using System.Collections.Generic;
+using System.Linq;
+using System.Text;
+using System.Threading.Tasks;
+
+namespace EpicGames.Perforce
+{
+	/// <summary>
+	/// Options for the p4 edit command
+	/// </summary>
+	[Flags]
+	public enum MoveOptions
+	{
+		/// <summary>
+		/// No options specified
+		/// </summary>
+		None = 0,
+
+		/// <summary>
+		/// Keep existing workspace files; mark the file as open for edit even if the file is not in the client view. 
+		/// </summary>
+		KeepWorkspaceFiles = 1,
+
+		/// <summary>
+		/// Rename existing files without altering content or type to a new target.
+		/// </summary>
+		RenameOnly = 2, 
+
+		/// <summary>
+		/// Preview which files would be opened for edit, without actually changing any files or metadata.
+		/// </summary>
+		PreviewOnly = 4,
+	}
+}
diff --git a/Engine/Source/Programs/Shared/EpicGames.Perforce/Records/MoveRecord.cs b/Engine/Source/Programs/Shared/EpicGames.Perforce/Records/MoveRecord.cs
new file mode 100644
--- /dev/null
+++ b/Engine/Source/Programs/Shared/EpicGames.Perforce/Records/MoveRecord.cs
@@ -0,0 +1,66 @@
+// Copyright 1998-2019 Epic Games, Inc. All Rights Reserved.
+
+using System;
+using System.Collections.Generic;
+using System.Linq;
+using System.Text;
+using System.Threading.Tasks;
+
+namespace EpicGames.Perforce
+{
+	/// <summary>
+	/// Information about a file opened for move
+	/// </summary>
+	public class MoveRecord
+	{
+		/// <summary>
+		/// Path to the file in the depot
+		/// </summary>
+		[PerforceTag("depotFile")]
+		public string DepotFile;
+
+		/// <summary>
+		/// Path to the file in the workspace
+		/// </summary>
+		[PerforceTag("clientFile")]
+		public string ClientFile;
+
+		/// <summary>
+		/// The working revision number of the file that was synced
+		/// </summary>
+		[PerforceTag("workRev", Optional = true)]
+		public int WorkingRevision;
+
+		/// <summary>
+		/// Action taken when syncing the file
+		/// </summary>
+		[PerforceTag("action")]
+		public string Action;
+
+		/// <summary>
+		/// Type of the file
+		/// </summary>
+		[PerforceTag("type")]
+		public string Type;
+
+		/// <summary>
+		/// Private constructor for serialization
+		/// </summary>
+		private MoveRecord()
+		{
+			DepotFile = null!;
+			ClientFile = null!;
+			Action = null!;
+			Type = null!;
+		}
+
+		/// <summary>
+		/// Format this record for display in the debugger
+		/// </summary>
+		/// <returns>Summary of this revision</returns>
+		public override string ToString()
+		{
+			return DepotFile;
+		}
+	}
+}
diff --git a/Engine/Source/Programs/Shared/EpicGames.Perforce/Records/ReopenRecord.cs b/Engine/Source/Programs/Shared/EpicGames.Perforce/Records/ReopenRecord.cs
new file mode 100644
--- /dev/null
+++ b/Engine/Source/Programs/Shared/EpicGames.Perforce/Records/ReopenRecord.cs
@@ -0,0 +1,13 @@
+using System;
+using System.Collections.Generic;
+using System.Text;
+
+namespace EpicGames.Perforce
+{
+	/// <summary>
+	/// Record returned from a reopen command
+	/// </summary>
+	public class ReopenRecord
+	{
+	}
+}
