Fix copyright notices in Perforce library.

diff --git a/Engine/Source/Programs/Shared/EpicGames.Perforce/PerforceAttributes.cs b/Engine/Source/Programs/Shared/EpicGames.Perforce/PerforceAttributes.cs
--- a/Engine/Source/Programs/Shared/EpicGames.Perforce/PerforceAttributes.cs
+++ b/Engine/Source/Programs/Shared/EpicGames.Perforce/PerforceAttributes.cs
@@ -1,4 +1,4 @@
-// Copyright 1998-2019 Epic Games, Inc. All Rights Reserved.
+// Copyright Epic Games, Inc. All Rights Reserved.
 
 using System;
 using System.Collections.Generic;
diff --git a/Engine/Source/Programs/Shared/EpicGames.Perforce/PerforceChildProcess.cs b/Engine/Source/Programs/Shared/EpicGames.Perforce/PerforceChildProcess.cs
--- a/Engine/Source/Programs/Shared/EpicGames.Perforce/PerforceChildProcess.cs
+++ b/Engine/Source/Programs/Shared/EpicGames.Perforce/PerforceChildProcess.cs
@@ -1,4 +1,4 @@
-// Copyright 1998-2019 Epic Games, Inc. All Rights Reserved.
+// Copyright Epic Games, Inc. All Rights Reserved.
 
 using Datadog.Trace;
 using EpicGames.Core;
diff --git a/Engine/Source/Programs/Shared/EpicGames.Perforce/PerforceClientConnection.cs b/Engine/Source/Programs/Shared/EpicGames.Perforce/PerforceClientConnection.cs
--- a/Engine/Source/Programs/Shared/EpicGames.Perforce/PerforceClientConnection.cs
+++ b/Engine/Source/Programs/Shared/EpicGames.Perforce/PerforceClientConnection.cs
@@ -1,3 +1,5 @@
+// Copyright Epic Games, Inc. All Rights Reserved.
+
 using Microsoft.Extensions.Logging;
 using System;
 using System.Collections.Generic;
diff --git a/Engine/Source/Programs/Shared/EpicGames.Perforce/PerforceConnection.cs b/Engine/Source/Programs/Shared/EpicGames.Perforce/PerforceConnection.cs
--- a/Engine/Source/Programs/Shared/EpicGames.Perforce/PerforceConnection.cs
+++ b/Engine/Source/Programs/Shared/EpicGames.Perforce/PerforceConnection.cs
@@ -1,4 +1,4 @@
-// Copyright 1998-2019 Epic Games, Inc. All Rights Reserved.
+// Copyright Epic Games, Inc. All Rights Reserved.
 
 using EpicGames.Core;
 using Microsoft.Extensions.Logging;
diff --git a/Engine/Source/Programs/Shared/EpicGames.Perforce/PerforceError.cs b/Engine/Source/Programs/Shared/EpicGames.Perforce/PerforceError.cs
--- a/Engine/Source/Programs/Shared/EpicGames.Perforce/PerforceError.cs
+++ b/Engine/Source/Programs/Shared/EpicGames.Perforce/PerforceError.cs
@@ -1,4 +1,4 @@
-// Copyright 1998-2019 Epic Games, Inc. All Rights Reserved.
+// Copyright Epic Games, Inc. All Rights Reserved.
 
 using System;
 using System.Collections.Generic;
diff --git a/Engine/Source/Programs/Shared/EpicGames.Perforce/PerforceException.cs b/Engine/Source/Programs/Shared/EpicGames.Perforce/PerforceException.cs
--- a/Engine/Source/Programs/Shared/EpicGames.Perforce/PerforceException.cs
+++ b/Engine/Source/Programs/Shared/EpicGames.Perforce/PerforceException.cs
@@ -1,4 +1,4 @@
-// Copyright 1998-2019 Epic Games, Inc. All Rights Reserved.
+// Copyright Epic Games, Inc. All Rights Reserved.
 
 using System;
 using System.Collections.Generic;
diff --git a/Engine/Source/Programs/Shared/EpicGames.Perforce/PerforceInfo.cs b/Engine/Source/Programs/Shared/EpicGames.Perforce/PerforceInfo.cs
--- a/Engine/Source/Programs/Shared/EpicGames.Perforce/PerforceInfo.cs
+++ b/Engine/Source/Programs/Shared/EpicGames.Perforce/PerforceInfo.cs
@@ -1,4 +1,4 @@
-// Copyright 1998-2019 Epic Games, Inc. All Rights Reserved.
+// Copyright Epic Games, Inc. All Rights Reserved.
 
 using System;
 using System.Collections.Generic;
diff --git a/Engine/Source/Programs/Shared/EpicGames.Perforce/PerforceRecord.cs b/Engine/Source/Programs/Shared/EpicGames.Perforce/PerforceRecord.cs
--- a/Engine/Source/Programs/Shared/EpicGames.Perforce/PerforceRecord.cs
+++ b/Engine/Source/Programs/Shared/EpicGames.Perforce/PerforceRecord.cs
@@ -1,3 +1,5 @@
+// Copyright Epic Games, Inc. All Rights Reserved.
+
 using EpicGames.Core;
 using System;
 using System.Buffers.Binary;
diff --git a/Engine/Source/Programs/Shared/EpicGames.Perforce/PerforceReflection.cs b/Engine/Source/Programs/Shared/EpicGames.Perforce/PerforceReflection.cs
--- a/Engine/Source/Programs/Shared/EpicGames.Perforce/PerforceReflection.cs
+++ b/Engine/Source/Programs/Shared/EpicGames.Perforce/PerforceReflection.cs
@@ -1,3 +1,5 @@
+// Copyright Epic Games, Inc. All Rights Reserved.
+
 using System;
 using System.Buffers.Text;
 using System.Collections.Concurrent;
diff --git a/Engine/Source/Programs/Shared/EpicGames.Perforce/PerforceResponse.cs b/Engine/Source/Programs/Shared/EpicGames.Perforce/PerforceResponse.cs
--- a/Engine/Source/Programs/Shared/EpicGames.Perforce/PerforceResponse.cs
+++ b/Engine/Source/Programs/Shared/EpicGames.Perforce/PerforceResponse.cs
@@ -1,4 +1,4 @@
-// Copyright 1998-2019 Epic Games, Inc. All Rights Reserved.
+// Copyright Epic Games, Inc. All Rights Reserved.
 
 using System;
 using System.Collections.Generic;
diff --git a/Engine/Source/Programs/Shared/EpicGames.Perforce/PerforceResponseList.cs b/Engine/Source/Programs/Shared/EpicGames.Perforce/PerforceResponseList.cs
--- a/Engine/Source/Programs/Shared/EpicGames.Perforce/PerforceResponseList.cs
+++ b/Engine/Source/Programs/Shared/EpicGames.Perforce/PerforceResponseList.cs
@@ -1,4 +1,4 @@
-// Copyright 1998-2019 Epic Games, Inc. All Rights Reserved.
+// Copyright Epic Games, Inc. All Rights Reserved.
 
 using System;
 using System.Collections.Generic;
diff --git a/Engine/Source/Programs/Shared/EpicGames.Perforce/PerforceUtils.cs b/Engine/Source/Programs/Shared/EpicGames.Perforce/PerforceUtils.cs
--- a/Engine/Source/Programs/Shared/EpicGames.Perforce/PerforceUtils.cs
+++ b/Engine/Source/Programs/Shared/EpicGames.Perforce/PerforceUtils.cs
@@ -1,4 +1,4 @@
-// Copyright 1998-2019 Epic Games, Inc. All Rights Reserved.
+// Copyright Epic Games, Inc. All Rights Reserved.
 
 using EpicGames.Core;
 using System;
diff --git a/Engine/Source/Programs/Shared/EpicGames.Perforce/Records/AddOptions.cs b/Engine/Source/Programs/Shared/EpicGames.Perforce/Records/AddOptions.cs
--- a/Engine/Source/Programs/Shared/EpicGames.Perforce/Records/AddOptions.cs
+++ b/Engine/Source/Programs/Shared/EpicGames.Perforce/Records/AddOptions.cs
@@ -1,4 +1,4 @@
-// Copyright 1998-2019 Epic Games, Inc. All Rights Reserved.
+// Copyright Epic Games, Inc. All Rights Reserved.
 
 using System;
 using System.Collections.Generic;
diff --git a/Engine/Source/Programs/Shared/EpicGames.Perforce/Records/AddRecord.cs b/Engine/Source/Programs/Shared/EpicGames.Perforce/Records/AddRecord.cs
--- a/Engine/Source/Programs/Shared/EpicGames.Perforce/Records/AddRecord.cs
+++ b/Engine/Source/Programs/Shared/EpicGames.Perforce/Records/AddRecord.cs
@@ -1,4 +1,4 @@
-// Copyright 1998-2019 Epic Games, Inc. All Rights Reserved.
+// Copyright Epic Games, Inc. All Rights Reserved.
 
 using System;
 using System.Collections.Generic;
diff --git a/Engine/Source/Programs/Shared/EpicGames.Perforce/Records/ChangeRecord.cs b/Engine/Source/Programs/Shared/EpicGames.Perforce/Records/ChangeRecord.cs
--- a/Engine/Source/Programs/Shared/EpicGames.Perforce/Records/ChangeRecord.cs
+++ b/Engine/Source/Programs/Shared/EpicGames.Perforce/Records/ChangeRecord.cs
@@ -1,4 +1,4 @@
-// Copyright 1998-2019 Epic Games, Inc. All Rights Reserved.
+// Copyright Epic Games, Inc. All Rights Reserved.
 
 using System;
 using System.Collections.Generic;
diff --git a/Engine/Source/Programs/Shared/EpicGames.Perforce/Records/ChangeStatus.cs b/Engine/Source/Programs/Shared/EpicGames.Perforce/Records/ChangeStatus.cs
--- a/Engine/Source/Programs/Shared/EpicGames.Perforce/Records/ChangeStatus.cs
+++ b/Engine/Source/Programs/Shared/EpicGames.Perforce/Records/ChangeStatus.cs
@@ -1,4 +1,4 @@
-// Copyright 1998-2019 Epic Games, Inc. All Rights Reserved.
+// Copyright Epic Games, Inc. All Rights Reserved.
 
 using System;
 using System.Collections.Generic;
diff --git a/Engine/Source/Programs/Shared/EpicGames.Perforce/Records/ChangeType.cs b/Engine/Source/Programs/Shared/EpicGames.Perforce/Records/ChangeType.cs
--- a/Engine/Source/Programs/Shared/EpicGames.Perforce/Records/ChangeType.cs
+++ b/Engine/Source/Programs/Shared/EpicGames.Perforce/Records/ChangeType.cs
@@ -1,4 +1,4 @@
-// Copyright 1998-2019 Epic Games, Inc. All Rights Reserved.
+// Copyright Epic Games, Inc. All Rights Reserved.
 
 using System;
 using System.Collections.Generic;
diff --git a/Engine/Source/Programs/Shared/EpicGames.Perforce/Records/ChangesOptions.cs b/Engine/Source/Programs/Shared/EpicGames.Perforce/Records/ChangesOptions.cs
--- a/Engine/Source/Programs/Shared/EpicGames.Perforce/Records/ChangesOptions.cs
+++ b/Engine/Source/Programs/Shared/EpicGames.Perforce/Records/ChangesOptions.cs
@@ -1,4 +1,4 @@
-// Copyright 1998-2019 Epic Games, Inc. All Rights Reserved.
+// Copyright Epic Games, Inc. All Rights Reserved.
 
 using System;
 using System.Collections.Generic;
diff --git a/Engine/Source/Programs/Shared/EpicGames.Perforce/Records/ChangesRecord.cs b/Engine/Source/Programs/Shared/EpicGames.Perforce/Records/ChangesRecord.cs
--- a/Engine/Source/Programs/Shared/EpicGames.Perforce/Records/ChangesRecord.cs
+++ b/Engine/Source/Programs/Shared/EpicGames.Perforce/Records/ChangesRecord.cs
@@ -1,4 +1,4 @@
-// Copyright 1998-2019 Epic Games, Inc. All Rights Reserved.
+// Copyright Epic Games, Inc. All Rights Reserved.
 
 using System;
 using System.Collections.Generic;
diff --git a/Engine/Source/Programs/Shared/EpicGames.Perforce/Records/CleanOptions.cs b/Engine/Source/Programs/Shared/EpicGames.Perforce/Records/CleanOptions.cs
--- a/Engine/Source/Programs/Shared/EpicGames.Perforce/Records/CleanOptions.cs
+++ b/Engine/Source/Programs/Shared/EpicGames.Perforce/Records/CleanOptions.cs
@@ -1,3 +1,5 @@
+// Copyright Epic Games, Inc. All Rights Reserved.
+
 using System;
 using System.Collections.Generic;
 using System.Text;
diff --git a/Engine/Source/Programs/Shared/EpicGames.Perforce/Records/CleanRecord.cs b/Engine/Source/Programs/Shared/EpicGames.Perforce/Records/CleanRecord.cs
--- a/Engine/Source/Programs/Shared/EpicGames.Perforce/Records/CleanRecord.cs
+++ b/Engine/Source/Programs/Shared/EpicGames.Perforce/Records/CleanRecord.cs
@@ -1,3 +1,5 @@
+// Copyright Epic Games, Inc. All Rights Reserved.
+
 using System;
 using System.Collections.Generic;
 using System.Text;
diff --git a/Engine/Source/Programs/Shared/EpicGames.Perforce/Records/ClientLineEndings.cs b/Engine/Source/Programs/Shared/EpicGames.Perforce/Records/ClientLineEndings.cs
--- a/Engine/Source/Programs/Shared/EpicGames.Perforce/Records/ClientLineEndings.cs
+++ b/Engine/Source/Programs/Shared/EpicGames.Perforce/Records/ClientLineEndings.cs
@@ -1,4 +1,4 @@
-// Copyright 1998-2019 Epic Games, Inc. All Rights Reserved.
+// Copyright Epic Games, Inc. All Rights Reserved.
 
 using System;
 using System.Collections.Generic;
diff --git a/Engine/Source/Programs/Shared/EpicGames.Perforce/Records/ClientOptions.cs b/Engine/Source/Programs/Shared/EpicGames.Perforce/Records/ClientOptions.cs
--- a/Engine/Source/Programs/Shared/EpicGames.Perforce/Records/ClientOptions.cs
+++ b/Engine/Source/Programs/Shared/EpicGames.Perforce/Records/ClientOptions.cs
@@ -1,4 +1,4 @@
-// Copyright 1998-2019 Epic Games, Inc. All Rights Reserved.
+// Copyright Epic Games, Inc. All Rights Reserved.
 
 using System;
 using System.Collections.Generic;
diff --git a/Engine/Source/Programs/Shared/EpicGames.Perforce/Records/ClientRecord.cs b/Engine/Source/Programs/Shared/EpicGames.Perforce/Records/ClientRecord.cs
--- a/Engine/Source/Programs/Shared/EpicGames.Perforce/Records/ClientRecord.cs
+++ b/Engine/Source/Programs/Shared/EpicGames.Perforce/Records/ClientRecord.cs
@@ -1,4 +1,4 @@
-// Copyright 1998-2019 Epic Games, Inc. All Rights Reserved.
+// Copyright Epic Games, Inc. All Rights Reserved.
 
 using System;
 using System.Collections.Generic;
diff --git a/Engine/Source/Programs/Shared/EpicGames.Perforce/Records/ClientSubmitOptions.cs b/Engine/Source/Programs/Shared/EpicGames.Perforce/Records/ClientSubmitOptions.cs
--- a/Engine/Source/Programs/Shared/EpicGames.Perforce/Records/ClientSubmitOptions.cs
+++ b/Engine/Source/Programs/Shared/EpicGames.Perforce/Records/ClientSubmitOptions.cs
@@ -1,4 +1,4 @@
-// Copyright 1998-2019 Epic Games, Inc. All Rights Reserved.
+// Copyright Epic Games, Inc. All Rights Reserved.
 
 using System;
 using System.Collections.Generic;
diff --git a/Engine/Source/Programs/Shared/EpicGames.Perforce/Records/ClientsOptions.cs b/Engine/Source/Programs/Shared/EpicGames.Perforce/Records/ClientsOptions.cs
--- a/Engine/Source/Programs/Shared/EpicGames.Perforce/Records/ClientsOptions.cs
+++ b/Engine/Source/Programs/Shared/EpicGames.Perforce/Records/ClientsOptions.cs
@@ -1,4 +1,4 @@
-// Copyright 1998-2019 Epic Games, Inc. All Rights Reserved.
+// Copyright Epic Games, Inc. All Rights Reserved.
 
 using System;
 using System.Collections.Generic;
diff --git a/Engine/Source/Programs/Shared/EpicGames.Perforce/Records/ClientsRecord.cs b/Engine/Source/Programs/Shared/EpicGames.Perforce/Records/ClientsRecord.cs
--- a/Engine/Source/Programs/Shared/EpicGames.Perforce/Records/ClientsRecord.cs
+++ b/Engine/Source/Programs/Shared/EpicGames.Perforce/Records/ClientsRecord.cs
@@ -1,4 +1,4 @@
-// Copyright 1998-2019 Epic Games, Inc. All Rights Reserved.
+// Copyright Epic Games, Inc. All Rights Reserved.
 
 using System;
 using System.Collections.Generic;
diff --git a/Engine/Source/Programs/Shared/EpicGames.Perforce/Records/DeleteChangeOptions.cs b/Engine/Source/Programs/Shared/EpicGames.Perforce/Records/DeleteChangeOptions.cs
--- a/Engine/Source/Programs/Shared/EpicGames.Perforce/Records/DeleteChangeOptions.cs
+++ b/Engine/Source/Programs/Shared/EpicGames.Perforce/Records/DeleteChangeOptions.cs
@@ -1,4 +1,4 @@
-// Copyright 1998-2019 Epic Games, Inc. All Rights Reserved.
+// Copyright Epic Games, Inc. All Rights Reserved.
 
 using System;
 using System.Collections.Generic;
diff --git a/Engine/Source/Programs/Shared/EpicGames.Perforce/Records/DeleteClientOptions.cs b/Engine/Source/Programs/Shared/EpicGames.Perforce/Records/DeleteClientOptions.cs
--- a/Engine/Source/Programs/Shared/EpicGames.Perforce/Records/DeleteClientOptions.cs
+++ b/Engine/Source/Programs/Shared/EpicGames.Perforce/Records/DeleteClientOptions.cs
@@ -1,4 +1,4 @@
-// Copyright 1998-2019 Epic Games, Inc. All Rights Reserved.
+// Copyright Epic Games, Inc. All Rights Reserved.
 
 using System;
 using System.Collections.Generic;
diff --git a/Engine/Source/Programs/Shared/EpicGames.Perforce/Records/DeleteOptions.cs b/Engine/Source/Programs/Shared/EpicGames.Perforce/Records/DeleteOptions.cs
--- a/Engine/Source/Programs/Shared/EpicGames.Perforce/Records/DeleteOptions.cs
+++ b/Engine/Source/Programs/Shared/EpicGames.Perforce/Records/DeleteOptions.cs
@@ -1,4 +1,4 @@
-// Copyright 1998-2019 Epic Games, Inc. All Rights Reserved.
+// Copyright Epic Games, Inc. All Rights Reserved.
 
 using System;
 using System.Collections.Generic;
diff --git a/Engine/Source/Programs/Shared/EpicGames.Perforce/Records/DeleteRecord.cs b/Engine/Source/Programs/Shared/EpicGames.Perforce/Records/DeleteRecord.cs
--- a/Engine/Source/Programs/Shared/EpicGames.Perforce/Records/DeleteRecord.cs
+++ b/Engine/Source/Programs/Shared/EpicGames.Perforce/Records/DeleteRecord.cs
@@ -1,4 +1,4 @@
-// Copyright 1998-2019 Epic Games, Inc. All Rights Reserved.
+// Copyright Epic Games, Inc. All Rights Reserved.
 
 using System;
 using System.Collections.Generic;
diff --git a/Engine/Source/Programs/Shared/EpicGames.Perforce/Records/DescribeFileRecord.cs b/Engine/Source/Programs/Shared/EpicGames.Perforce/Records/DescribeFileRecord.cs
--- a/Engine/Source/Programs/Shared/EpicGames.Perforce/Records/DescribeFileRecord.cs
+++ b/Engine/Source/Programs/Shared/EpicGames.Perforce/Records/DescribeFileRecord.cs
@@ -1,4 +1,4 @@
-// Copyright 1998-2019 Epic Games, Inc. All Rights Reserved.
+// Copyright Epic Games, Inc. All Rights Reserved.
 
 using System;
 using System.Collections.Generic;
diff --git a/Engine/Source/Programs/Shared/EpicGames.Perforce/Records/DescribeOptions.cs b/Engine/Source/Programs/Shared/EpicGames.Perforce/Records/DescribeOptions.cs
--- a/Engine/Source/Programs/Shared/EpicGames.Perforce/Records/DescribeOptions.cs
+++ b/Engine/Source/Programs/Shared/EpicGames.Perforce/Records/DescribeOptions.cs
@@ -1,4 +1,4 @@
-// Copyright 1998-2019 Epic Games, Inc. All Rights Reserved.
+// Copyright Epic Games, Inc. All Rights Reserved.
 
 using System;
 using System.Collections.Generic;
diff --git a/Engine/Source/Programs/Shared/EpicGames.Perforce/Records/DescribeRecord.cs b/Engine/Source/Programs/Shared/EpicGames.Perforce/Records/DescribeRecord.cs
--- a/Engine/Source/Programs/Shared/EpicGames.Perforce/Records/DescribeRecord.cs
+++ b/Engine/Source/Programs/Shared/EpicGames.Perforce/Records/DescribeRecord.cs
@@ -1,4 +1,4 @@
-// Copyright 1998-2019 Epic Games, Inc. All Rights Reserved.
+// Copyright Epic Games, Inc. All Rights Reserved.
 
 using System;
 using System.Collections.Generic;
diff --git a/Engine/Source/Programs/Shared/EpicGames.Perforce/Records/EditOptions.cs b/Engine/Source/Programs/Shared/EpicGames.Perforce/Records/EditOptions.cs
--- a/Engine/Source/Programs/Shared/EpicGames.Perforce/Records/EditOptions.cs
+++ b/Engine/Source/Programs/Shared/EpicGames.Perforce/Records/EditOptions.cs
@@ -1,4 +1,4 @@
-// Copyright 1998-2019 Epic Games, Inc. All Rights Reserved.
+// Copyright Epic Games, Inc. All Rights Reserved.
 
 using System;
 using System.Collections.Generic;
diff --git a/Engine/Source/Programs/Shared/EpicGames.Perforce/Records/EditRecord.cs b/Engine/Source/Programs/Shared/EpicGames.Perforce/Records/EditRecord.cs
--- a/Engine/Source/Programs/Shared/EpicGames.Perforce/Records/EditRecord.cs
+++ b/Engine/Source/Programs/Shared/EpicGames.Perforce/Records/EditRecord.cs
@@ -1,4 +1,4 @@
-// Copyright 1998-2019 Epic Games, Inc. All Rights Reserved.
+// Copyright Epic Games, Inc. All Rights Reserved.
 
 using System;
 using System.Collections.Generic;
diff --git a/Engine/Source/Programs/Shared/EpicGames.Perforce/Records/FStatOptions.cs b/Engine/Source/Programs/Shared/EpicGames.Perforce/Records/FStatOptions.cs
--- a/Engine/Source/Programs/Shared/EpicGames.Perforce/Records/FStatOptions.cs
+++ b/Engine/Source/Programs/Shared/EpicGames.Perforce/Records/FStatOptions.cs
@@ -1,4 +1,4 @@
-// Copyright 1998-2019 Epic Games, Inc. All Rights Reserved.
+// Copyright Epic Games, Inc. All Rights Reserved.
 
 using System;
 using System.Collections.Generic;
diff --git a/Engine/Source/Programs/Shared/EpicGames.Perforce/Records/FStatRecord.cs b/Engine/Source/Programs/Shared/EpicGames.Perforce/Records/FStatRecord.cs
--- a/Engine/Source/Programs/Shared/EpicGames.Perforce/Records/FStatRecord.cs
+++ b/Engine/Source/Programs/Shared/EpicGames.Perforce/Records/FStatRecord.cs
@@ -1,4 +1,4 @@
-// Copyright 1998-2019 Epic Games, Inc. All Rights Reserved.
+// Copyright Epic Games, Inc. All Rights Reserved.
 
 using System;
 using System.Collections.Generic;
diff --git a/Engine/Source/Programs/Shared/EpicGames.Perforce/Records/FileAction.cs b/Engine/Source/Programs/Shared/EpicGames.Perforce/Records/FileAction.cs
--- a/Engine/Source/Programs/Shared/EpicGames.Perforce/Records/FileAction.cs
+++ b/Engine/Source/Programs/Shared/EpicGames.Perforce/Records/FileAction.cs
@@ -1,4 +1,4 @@
-// Copyright 1998-2019 Epic Games, Inc. All Rights Reserved.
+// Copyright Epic Games, Inc. All Rights Reserved.
 
 using System;
 using System.Collections.Generic;
diff --git a/Engine/Source/Programs/Shared/EpicGames.Perforce/Records/FileLogOptions.cs b/Engine/Source/Programs/Shared/EpicGames.Perforce/Records/FileLogOptions.cs
--- a/Engine/Source/Programs/Shared/EpicGames.Perforce/Records/FileLogOptions.cs
+++ b/Engine/Source/Programs/Shared/EpicGames.Perforce/Records/FileLogOptions.cs
@@ -1,4 +1,4 @@
-// Copyright 1998-2019 Epic Games, Inc. All Rights Reserved.
+// Copyright Epic Games, Inc. All Rights Reserved.
 
 using System;
 using System.Collections.Generic;
diff --git a/Engine/Source/Programs/Shared/EpicGames.Perforce/Records/FileLogRecord.cs b/Engine/Source/Programs/Shared/EpicGames.Perforce/Records/FileLogRecord.cs
--- a/Engine/Source/Programs/Shared/EpicGames.Perforce/Records/FileLogRecord.cs
+++ b/Engine/Source/Programs/Shared/EpicGames.Perforce/Records/FileLogRecord.cs
@@ -1,4 +1,4 @@
-// Copyright 1998-2019 Epic Games, Inc. All Rights Reserved.
+// Copyright Epic Games, Inc. All Rights Reserved.
 
 using System;
 using System.Collections.Generic;
diff --git a/Engine/Source/Programs/Shared/EpicGames.Perforce/Records/GetChangeOptions.cs b/Engine/Source/Programs/Shared/EpicGames.Perforce/Records/GetChangeOptions.cs
--- a/Engine/Source/Programs/Shared/EpicGames.Perforce/Records/GetChangeOptions.cs
+++ b/Engine/Source/Programs/Shared/EpicGames.Perforce/Records/GetChangeOptions.cs
@@ -1,4 +1,4 @@
-// Copyright 1998-2019 Epic Games, Inc. All Rights Reserved.
+// Copyright Epic Games, Inc. All Rights Reserved.
 
 using System;
 using System.Collections.Generic;
diff --git a/Engine/Source/Programs/Shared/EpicGames.Perforce/Records/InfoOptions.cs b/Engine/Source/Programs/Shared/EpicGames.Perforce/Records/InfoOptions.cs
--- a/Engine/Source/Programs/Shared/EpicGames.Perforce/Records/InfoOptions.cs
+++ b/Engine/Source/Programs/Shared/EpicGames.Perforce/Records/InfoOptions.cs
@@ -1,4 +1,4 @@
-// Copyright 1998-2019 Epic Games, Inc. All Rights Reserved.
+// Copyright Epic Games, Inc. All Rights Reserved.
 
 using System;
 using System.Collections.Generic;
diff --git a/Engine/Source/Programs/Shared/EpicGames.Perforce/Records/InfoRecord.cs b/Engine/Source/Programs/Shared/EpicGames.Perforce/Records/InfoRecord.cs
--- a/Engine/Source/Programs/Shared/EpicGames.Perforce/Records/InfoRecord.cs
+++ b/Engine/Source/Programs/Shared/EpicGames.Perforce/Records/InfoRecord.cs
@@ -1,4 +1,4 @@
-// Copyright 1998-2019 Epic Games, Inc. All Rights Reserved.
+// Copyright Epic Games, Inc. All Rights Reserved.
 
 using System;
 using System.Collections.Generic;
diff --git a/Engine/Source/Programs/Shared/EpicGames.Perforce/Records/IntegrationAction.cs b/Engine/Source/Programs/Shared/EpicGames.Perforce/Records/IntegrationAction.cs
--- a/Engine/Source/Programs/Shared/EpicGames.Perforce/Records/IntegrationAction.cs
+++ b/Engine/Source/Programs/Shared/EpicGames.Perforce/Records/IntegrationAction.cs
@@ -1,4 +1,4 @@
-// Copyright 1998-2019 Epic Games, Inc. All Rights Reserved.
+// Copyright Epic Games, Inc. All Rights Reserved.
 
 using System;
 using System.Collections.Generic;
diff --git a/Engine/Source/Programs/Shared/EpicGames.Perforce/Records/IntegrationRecord.cs b/Engine/Source/Programs/Shared/EpicGames.Perforce/Records/IntegrationRecord.cs
--- a/Engine/Source/Programs/Shared/EpicGames.Perforce/Records/IntegrationRecord.cs
+++ b/Engine/Source/Programs/Shared/EpicGames.Perforce/Records/IntegrationRecord.cs
@@ -1,4 +1,4 @@
-// Copyright 1998-2019 Epic Games, Inc. All Rights Reserved.
+// Copyright Epic Games, Inc. All Rights Reserved.
 
 using System;
 using System.Collections.Generic;
diff --git a/Engine/Source/Programs/Shared/EpicGames.Perforce/Records/MergeOptions.cs b/Engine/Source/Programs/Shared/EpicGames.Perforce/Records/MergeOptions.cs
--- a/Engine/Source/Programs/Shared/EpicGames.Perforce/Records/MergeOptions.cs
+++ b/Engine/Source/Programs/Shared/EpicGames.Perforce/Records/MergeOptions.cs
@@ -1,3 +1,5 @@
+// Copyright Epic Games, Inc. All Rights Reserved.
+
 using System;
 using System.Collections.Generic;
 using System.Text;
diff --git a/Engine/Source/Programs/Shared/EpicGames.Perforce/Records/MergeRecord.cs b/Engine/Source/Programs/Shared/EpicGames.Perforce/Records/MergeRecord.cs
--- a/Engine/Source/Programs/Shared/EpicGames.Perforce/Records/MergeRecord.cs
+++ b/Engine/Source/Programs/Shared/EpicGames.Perforce/Records/MergeRecord.cs
@@ -1,3 +1,5 @@
+// Copyright Epic Games, Inc. All Rights Reserved.
+
 using System;
 using System.Collections.Generic;
 using System.Text;
diff --git a/Engine/Source/Programs/Shared/EpicGames.Perforce/Records/MoveOptions.cs b/Engine/Source/Programs/Shared/EpicGames.Perforce/Records/MoveOptions.cs
--- a/Engine/Source/Programs/Shared/EpicGames.Perforce/Records/MoveOptions.cs
+++ b/Engine/Source/Programs/Shared/EpicGames.Perforce/Records/MoveOptions.cs
@@ -1,4 +1,4 @@
-// Copyright 1998-2019 Epic Games, Inc. All Rights Reserved.
+// Copyright Epic Games, Inc. All Rights Reserved.
 
 using System;
 using System.Collections.Generic;
diff --git a/Engine/Source/Programs/Shared/EpicGames.Perforce/Records/MoveRecord.cs b/Engine/Source/Programs/Shared/EpicGames.Perforce/Records/MoveRecord.cs
--- a/Engine/Source/Programs/Shared/EpicGames.Perforce/Records/MoveRecord.cs
+++ b/Engine/Source/Programs/Shared/EpicGames.Perforce/Records/MoveRecord.cs
@@ -1,4 +1,4 @@
-// Copyright 1998-2019 Epic Games, Inc. All Rights Reserved.
+// Copyright Epic Games, Inc. All Rights Reserved.
 
 using System;
 using System.Collections.Generic;
diff --git a/Engine/Source/Programs/Shared/EpicGames.Perforce/Records/OpenedOptions.cs b/Engine/Source/Programs/Shared/EpicGames.Perforce/Records/OpenedOptions.cs
--- a/Engine/Source/Programs/Shared/EpicGames.Perforce/Records/OpenedOptions.cs
+++ b/Engine/Source/Programs/Shared/EpicGames.Perforce/Records/OpenedOptions.cs
@@ -1,4 +1,4 @@
-// Copyright 1998-2019 Epic Games, Inc. All Rights Reserved.
+// Copyright Epic Games, Inc. All Rights Reserved.
 
 using System;
 using System.Collections.Generic;
diff --git a/Engine/Source/Programs/Shared/EpicGames.Perforce/Records/PrintRecord.cs b/Engine/Source/Programs/Shared/EpicGames.Perforce/Records/PrintRecord.cs
--- a/Engine/Source/Programs/Shared/EpicGames.Perforce/Records/PrintRecord.cs
+++ b/Engine/Source/Programs/Shared/EpicGames.Perforce/Records/PrintRecord.cs
@@ -1,4 +1,4 @@
-// Copyright 1998-2019 Epic Games, Inc. All Rights Reserved.
+// Copyright Epic Games, Inc. All Rights Reserved.
 
 using System;
 using System.Collections.Generic;
diff --git a/Engine/Source/Programs/Shared/EpicGames.Perforce/Records/ReconcileOptions.cs b/Engine/Source/Programs/Shared/EpicGames.Perforce/Records/ReconcileOptions.cs
--- a/Engine/Source/Programs/Shared/EpicGames.Perforce/Records/ReconcileOptions.cs
+++ b/Engine/Source/Programs/Shared/EpicGames.Perforce/Records/ReconcileOptions.cs
@@ -1,4 +1,4 @@
-// Copyright 1998-2019 Epic Games, Inc. All Rights Reserved.
+// Copyright Epic Games, Inc. All Rights Reserved.
 
 using System;
 using System.Collections.Generic;
diff --git a/Engine/Source/Programs/Shared/EpicGames.Perforce/Records/ReconcileRecord.cs b/Engine/Source/Programs/Shared/EpicGames.Perforce/Records/ReconcileRecord.cs
--- a/Engine/Source/Programs/Shared/EpicGames.Perforce/Records/ReconcileRecord.cs
+++ b/Engine/Source/Programs/Shared/EpicGames.Perforce/Records/ReconcileRecord.cs
@@ -1,4 +1,4 @@
-// Copyright 1998-2019 Epic Games, Inc. All Rights Reserved.
+// Copyright Epic Games, Inc. All Rights Reserved.
 
 using System;
 using System.Collections.Generic;
diff --git a/Engine/Source/Programs/Shared/EpicGames.Perforce/Records/ReopenRecord.cs b/Engine/Source/Programs/Shared/EpicGames.Perforce/Records/ReopenRecord.cs
--- a/Engine/Source/Programs/Shared/EpicGames.Perforce/Records/ReopenRecord.cs
+++ b/Engine/Source/Programs/Shared/EpicGames.Perforce/Records/ReopenRecord.cs
@@ -1,3 +1,5 @@
+// Copyright Epic Games, Inc. All Rights Reserved.
+
 using System;
 using System.Collections.Generic;
 using System.Text;
diff --git a/Engine/Source/Programs/Shared/EpicGames.Perforce/Records/ResolveOptions.cs b/Engine/Source/Programs/Shared/EpicGames.Perforce/Records/ResolveOptions.cs
--- a/Engine/Source/Programs/Shared/EpicGames.Perforce/Records/ResolveOptions.cs
+++ b/Engine/Source/Programs/Shared/EpicGames.Perforce/Records/ResolveOptions.cs
@@ -1,4 +1,4 @@
-// Copyright 1998-2019 Epic Games, Inc. All Rights Reserved.
+// Copyright Epic Games, Inc. All Rights Reserved.
 
 using System;
 using System.Collections.Generic;
diff --git a/Engine/Source/Programs/Shared/EpicGames.Perforce/Records/ResolveRecord.cs b/Engine/Source/Programs/Shared/EpicGames.Perforce/Records/ResolveRecord.cs
--- a/Engine/Source/Programs/Shared/EpicGames.Perforce/Records/ResolveRecord.cs
+++ b/Engine/Source/Programs/Shared/EpicGames.Perforce/Records/ResolveRecord.cs
@@ -1,4 +1,4 @@
-// Copyright 1998-2019 Epic Games, Inc. All Rights Reserved.
+// Copyright Epic Games, Inc. All Rights Reserved.
 
 using System;
 using System.Collections.Generic;
diff --git a/Engine/Source/Programs/Shared/EpicGames.Perforce/Records/RevertOptions.cs b/Engine/Source/Programs/Shared/EpicGames.Perforce/Records/RevertOptions.cs
--- a/Engine/Source/Programs/Shared/EpicGames.Perforce/Records/RevertOptions.cs
+++ b/Engine/Source/Programs/Shared/EpicGames.Perforce/Records/RevertOptions.cs
@@ -1,4 +1,4 @@
-// Copyright 1998-2019 Epic Games, Inc. All Rights Reserved.
+// Copyright Epic Games, Inc. All Rights Reserved.
 
 using System;
 using System.Collections.Generic;
diff --git a/Engine/Source/Programs/Shared/EpicGames.Perforce/Records/RevertRecord.cs b/Engine/Source/Programs/Shared/EpicGames.Perforce/Records/RevertRecord.cs
--- a/Engine/Source/Programs/Shared/EpicGames.Perforce/Records/RevertRecord.cs
+++ b/Engine/Source/Programs/Shared/EpicGames.Perforce/Records/RevertRecord.cs
@@ -1,4 +1,4 @@
-// Copyright 1998-2019 Epic Games, Inc. All Rights Reserved.
+// Copyright Epic Games, Inc. All Rights Reserved.
 
 using System;
 using System.Collections.Generic;
diff --git a/Engine/Source/Programs/Shared/EpicGames.Perforce/Records/RevisionRecord.cs b/Engine/Source/Programs/Shared/EpicGames.Perforce/Records/RevisionRecord.cs
--- a/Engine/Source/Programs/Shared/EpicGames.Perforce/Records/RevisionRecord.cs
+++ b/Engine/Source/Programs/Shared/EpicGames.Perforce/Records/RevisionRecord.cs
@@ -1,4 +1,4 @@
-// Copyright 1998-2019 Epic Games, Inc. All Rights Reserved.
+// Copyright Epic Games, Inc. All Rights Reserved.
 
 using System;
 using System.Collections.Generic;
diff --git a/Engine/Source/Programs/Shared/EpicGames.Perforce/Records/ShelveOptions.cs b/Engine/Source/Programs/Shared/EpicGames.Perforce/Records/ShelveOptions.cs
--- a/Engine/Source/Programs/Shared/EpicGames.Perforce/Records/ShelveOptions.cs
+++ b/Engine/Source/Programs/Shared/EpicGames.Perforce/Records/ShelveOptions.cs
@@ -1,4 +1,4 @@
-// Copyright 1998-2019 Epic Games, Inc. All Rights Reserved.
+// Copyright Epic Games, Inc. All Rights Reserved.
 
 using System;
 using System.Collections.Generic;
diff --git a/Engine/Source/Programs/Shared/EpicGames.Perforce/Records/ShelveRecord.cs b/Engine/Source/Programs/Shared/EpicGames.Perforce/Records/ShelveRecord.cs
--- a/Engine/Source/Programs/Shared/EpicGames.Perforce/Records/ShelveRecord.cs
+++ b/Engine/Source/Programs/Shared/EpicGames.Perforce/Records/ShelveRecord.cs
@@ -1,4 +1,4 @@
-// Copyright 1998-2019 Epic Games, Inc. All Rights Reserved.
+// Copyright Epic Games, Inc. All Rights Reserved.
 
 using System;
 using System.Collections.Generic;
diff --git a/Engine/Source/Programs/Shared/EpicGames.Perforce/Records/StreamOptions.cs b/Engine/Source/Programs/Shared/EpicGames.Perforce/Records/StreamOptions.cs
--- a/Engine/Source/Programs/Shared/EpicGames.Perforce/Records/StreamOptions.cs
+++ b/Engine/Source/Programs/Shared/EpicGames.Perforce/Records/StreamOptions.cs
@@ -1,3 +1,5 @@
+// Copyright Epic Games, Inc. All Rights Reserved.
+
 using System;
 using System.Collections.Generic;
 using System.Text;
diff --git a/Engine/Source/Programs/Shared/EpicGames.Perforce/Records/StreamRecord.cs b/Engine/Source/Programs/Shared/EpicGames.Perforce/Records/StreamRecord.cs
--- a/Engine/Source/Programs/Shared/EpicGames.Perforce/Records/StreamRecord.cs
+++ b/Engine/Source/Programs/Shared/EpicGames.Perforce/Records/StreamRecord.cs
@@ -1,3 +1,5 @@
+// Copyright Epic Games, Inc. All Rights Reserved.
+
 using System;
 using System.Collections.Generic;
 using System.Text;
diff --git a/Engine/Source/Programs/Shared/EpicGames.Perforce/Records/StreamType.cs b/Engine/Source/Programs/Shared/EpicGames.Perforce/Records/StreamType.cs
--- a/Engine/Source/Programs/Shared/EpicGames.Perforce/Records/StreamType.cs
+++ b/Engine/Source/Programs/Shared/EpicGames.Perforce/Records/StreamType.cs
@@ -1,4 +1,4 @@
-// Copyright 1998-2019 Epic Games, Inc. All Rights Reserved.
+// Copyright Epic Games, Inc. All Rights Reserved.
 
 using System;
 using System.Collections.Generic;
diff --git a/Engine/Source/Programs/Shared/EpicGames.Perforce/Records/StreamsOptions.cs b/Engine/Source/Programs/Shared/EpicGames.Perforce/Records/StreamsOptions.cs
--- a/Engine/Source/Programs/Shared/EpicGames.Perforce/Records/StreamsOptions.cs
+++ b/Engine/Source/Programs/Shared/EpicGames.Perforce/Records/StreamsOptions.cs
@@ -1,4 +1,4 @@
-// Copyright 1998-2019 Epic Games, Inc. All Rights Reserved.
+// Copyright Epic Games, Inc. All Rights Reserved.
 
 using System;
 using System.Collections.Generic;
diff --git a/Engine/Source/Programs/Shared/EpicGames.Perforce/Records/StreamsRecord.cs b/Engine/Source/Programs/Shared/EpicGames.Perforce/Records/StreamsRecord.cs
--- a/Engine/Source/Programs/Shared/EpicGames.Perforce/Records/StreamsRecord.cs
+++ b/Engine/Source/Programs/Shared/EpicGames.Perforce/Records/StreamsRecord.cs
@@ -1,4 +1,4 @@
-// Copyright 1998-2019 Epic Games, Inc. All Rights Reserved.
+// Copyright Epic Games, Inc. All Rights Reserved.
 
 using System;
 using System.Collections.Generic;
diff --git a/Engine/Source/Programs/Shared/EpicGames.Perforce/Records/SubmitOptions.cs b/Engine/Source/Programs/Shared/EpicGames.Perforce/Records/SubmitOptions.cs
--- a/Engine/Source/Programs/Shared/EpicGames.Perforce/Records/SubmitOptions.cs
+++ b/Engine/Source/Programs/Shared/EpicGames.Perforce/Records/SubmitOptions.cs
@@ -1,4 +1,4 @@
-// Copyright 1998-2019 Epic Games, Inc. All Rights Reserved.
+// Copyright Epic Games, Inc. All Rights Reserved.
 
 using System;
 using System.Collections.Generic;
diff --git a/Engine/Source/Programs/Shared/EpicGames.Perforce/Records/SubmitRecord.cs b/Engine/Source/Programs/Shared/EpicGames.Perforce/Records/SubmitRecord.cs
--- a/Engine/Source/Programs/Shared/EpicGames.Perforce/Records/SubmitRecord.cs
+++ b/Engine/Source/Programs/Shared/EpicGames.Perforce/Records/SubmitRecord.cs
@@ -1,4 +1,4 @@
-// Copyright 1998-2019 Epic Games, Inc. All Rights Reserved.
+// Copyright Epic Games, Inc. All Rights Reserved.
 
 using System;
 using System.Collections.Generic;
diff --git a/Engine/Source/Programs/Shared/EpicGames.Perforce/Records/SwitchClientOptions.cs b/Engine/Source/Programs/Shared/EpicGames.Perforce/Records/SwitchClientOptions.cs
--- a/Engine/Source/Programs/Shared/EpicGames.Perforce/Records/SwitchClientOptions.cs
+++ b/Engine/Source/Programs/Shared/EpicGames.Perforce/Records/SwitchClientOptions.cs
@@ -1,4 +1,4 @@
-// Copyright 1998-2019 Epic Games, Inc. All Rights Reserved.
+// Copyright Epic Games, Inc. All Rights Reserved.
 
 using System;
 using System.Collections.Generic;
diff --git a/Engine/Source/Programs/Shared/EpicGames.Perforce/Records/SyncOptions.cs b/Engine/Source/Programs/Shared/EpicGames.Perforce/Records/SyncOptions.cs
--- a/Engine/Source/Programs/Shared/EpicGames.Perforce/Records/SyncOptions.cs
+++ b/Engine/Source/Programs/Shared/EpicGames.Perforce/Records/SyncOptions.cs
@@ -1,4 +1,4 @@
-// Copyright 1998-2019 Epic Games, Inc. All Rights Reserved.
+// Copyright Epic Games, Inc. All Rights Reserved.
 
 using System;
 using System.Collections.Generic;
diff --git a/Engine/Source/Programs/Shared/EpicGames.Perforce/Records/SyncRecord.cs b/Engine/Source/Programs/Shared/EpicGames.Perforce/Records/SyncRecord.cs
--- a/Engine/Source/Programs/Shared/EpicGames.Perforce/Records/SyncRecord.cs
+++ b/Engine/Source/Programs/Shared/EpicGames.Perforce/Records/SyncRecord.cs
@@ -1,4 +1,4 @@
-// Copyright 1998-2019 Epic Games, Inc. All Rights Reserved.
+// Copyright Epic Games, Inc. All Rights Reserved.
 
 using System;
 using System.Collections.Generic;
diff --git a/Engine/Source/Programs/Shared/EpicGames.Perforce/Records/SyncSummaryRecord.cs b/Engine/Source/Programs/Shared/EpicGames.Perforce/Records/SyncSummaryRecord.cs
--- a/Engine/Source/Programs/Shared/EpicGames.Perforce/Records/SyncSummaryRecord.cs
+++ b/Engine/Source/Programs/Shared/EpicGames.Perforce/Records/SyncSummaryRecord.cs
@@ -1,4 +1,4 @@
-// Copyright 1998-2019 Epic Games, Inc. All Rights Reserved.
+// Copyright Epic Games, Inc. All Rights Reserved.
 
 using System;
 using System.Collections.Generic;
diff --git a/Engine/Source/Programs/Shared/EpicGames.Perforce/Records/UnshelveOptions.cs b/Engine/Source/Programs/Shared/EpicGames.Perforce/Records/UnshelveOptions.cs
--- a/Engine/Source/Programs/Shared/EpicGames.Perforce/Records/UnshelveOptions.cs
+++ b/Engine/Source/Programs/Shared/EpicGames.Perforce/Records/UnshelveOptions.cs
@@ -1,4 +1,4 @@
-// Copyright 1998-2019 Epic Games, Inc. All Rights Reserved.
+// Copyright Epic Games, Inc. All Rights Reserved.
 
 using System;
 using System.Collections.Generic;
diff --git a/Engine/Source/Programs/Shared/EpicGames.Perforce/Records/UnshelveRecord.cs b/Engine/Source/Programs/Shared/EpicGames.Perforce/Records/UnshelveRecord.cs
--- a/Engine/Source/Programs/Shared/EpicGames.Perforce/Records/UnshelveRecord.cs
+++ b/Engine/Source/Programs/Shared/EpicGames.Perforce/Records/UnshelveRecord.cs
@@ -1,4 +1,4 @@
-// Copyright 1998-2019 Epic Games, Inc. All Rights Reserved.
+// Copyright Epic Games, Inc. All Rights Reserved.
 
 using System;
 using System.Collections.Generic;
diff --git a/Engine/Source/Programs/Shared/EpicGames.Perforce/Records/UpdateChangeOptions.cs b/Engine/Source/Programs/Shared/EpicGames.Perforce/Records/UpdateChangeOptions.cs
--- a/Engine/Source/Programs/Shared/EpicGames.Perforce/Records/UpdateChangeOptions.cs
+++ b/Engine/Source/Programs/Shared/EpicGames.Perforce/Records/UpdateChangeOptions.cs
@@ -1,4 +1,4 @@
-// Copyright 1998-2019 Epic Games, Inc. All Rights Reserved.
+// Copyright Epic Games, Inc. All Rights Reserved.
 
 using System;
 using System.Collections.Generic;
diff --git a/Engine/Source/Programs/Shared/EpicGames.Perforce/Records/UserRecord.cs b/Engine/Source/Programs/Shared/EpicGames.Perforce/Records/UserRecord.cs
--- a/Engine/Source/Programs/Shared/EpicGames.Perforce/Records/UserRecord.cs
+++ b/Engine/Source/Programs/Shared/EpicGames.Perforce/Records/UserRecord.cs
@@ -1,4 +1,4 @@
-// Copyright 1998-2019 Epic Games, Inc. All Rights Reserved.
+// Copyright Epic Games, Inc. All Rights Reserved.
 
 using System;
 using System.Collections.Generic;
diff --git a/Engine/Source/Programs/Shared/EpicGames.Perforce/Records/WhereRecord.cs b/Engine/Source/Programs/Shared/EpicGames.Perforce/Records/WhereRecord.cs
--- a/Engine/Source/Programs/Shared/EpicGames.Perforce/Records/WhereRecord.cs
+++ b/Engine/Source/Programs/Shared/EpicGames.Perforce/Records/WhereRecord.cs
@@ -1,4 +1,4 @@
-// Copyright 1998-2019 Epic Games, Inc. All Rights Reserved.
+// Copyright Epic Games, Inc. All Rights Reserved.
 
 using System;
 using System.Collections.Generic;
