Make the sink for tracing data pluggable, and remove dependency on Datadog tracing library from shared Perforce library.

Fixes long shutdown time for small applications, due to DD library waiting for TCP connection to timeout.

diff --git a/Engine/Source/Programs/Shared/EpicGames.Core/TraceSpan.cs b/Engine/Source/Programs/Shared/EpicGames.Core/TraceSpan.cs
--- a/Engine/Source/Programs/Shared/EpicGames.Core/TraceSpan.cs
+++ b/Engine/Source/Programs/Shared/EpicGames.Core/TraceSpan.cs
@@ -26,13 +26,30 @@ namespace EpicGames.Core
 	}
 
 	/// <summary>
-	/// Methods for creating ITraceScope instances
+	/// Sink for tracing information
 	/// </summary>
-	public static class TraceSpan
+	public interface ITraceSink
 	{
 		/// <summary>
-		/// Concrete implementation of ITraceScope
+		/// Create a new trace span
 		/// </summary>
+		/// <param name="Operation">Name of the operation</param>
+		/// <param name="Resource">Resource that the operation is being performed on</param>
+		/// <param name="Service">Name of the service</param>
+		/// <returns>New span instance</returns>
+		ITraceSpan Create(string Operation, string? Resource, string? Service);
+
+		/// <summary>
+		/// Flush all the current trace data
+		/// </summary>
+		void Flush();
+	}
+
+	/// <summary>
+	/// Default implementation of ITraceSink. Writes trace information to a JSON file on application exit if the UE_TELEMETRY_DIR environment variable is set.
+	/// </summary>
+	public class JsonTraceSink : ITraceSink
+	{
 		class TraceSpanImpl : ITraceSpan
 		{
 			public string Name;
@@ -57,22 +74,36 @@ namespace EpicGames.Core
 
 			public void Dispose()
 			{
-				if(FinishTime == null)
+				if (FinishTime == null)
 				{
 					FinishTime = DateTimeOffset.Now;
 				}
 			}
 		}
 
+		/// <summary>
+		/// Output directory for telemetry info
+		/// </summary>
+		DirectoryReference TelemetryDir;
+
 		/// <summary>
 		/// The current scope provider
 		/// </summary>
-		static List<TraceSpanImpl> Spans = new List<TraceSpanImpl>();
+		List<TraceSpanImpl> Spans = new List<TraceSpanImpl>();
+
+		/// <summary>
+		/// Constructor
+		/// </summary>
+		/// <param name="TelemetryDir">Directory to store telemetry files</param>
+		public JsonTraceSink(DirectoryReference TelemetryDir)
+		{
+			this.TelemetryDir = TelemetryDir;
+		}
 
 		/// <summary>
 		/// Creates a scope using the current provider
 		/// </summary>
-		public static ITraceSpan Create(string Name, string? Resource = null, string? Service = null)
+		public ITraceSpan Create(string Name, string? Resource = null, string? Service = null)
 		{
 			TraceSpanImpl Span = new TraceSpanImpl(Name, Resource, Service);
 			Spans.Add(Span);
@@ -82,53 +113,150 @@ namespace EpicGames.Core
 		/// <summary>
 		/// Saves all the scope information to a file
 		/// </summary>
-		public static void Flush()
+		public void Flush()
 		{
-			string? TelemetryDir = Environment.GetEnvironmentVariable("UE_TELEMETRY_DIR");
-			if (TelemetryDir != null)
+			FileReference File;
+			using (Process Process = Process.GetCurrentProcess())
 			{
-				FileReference File;
-				using (Process Process = Process.GetCurrentProcess())
-				{
-					DirectoryReference TelemetryDirRef = new DirectoryReference(TelemetryDir);
-					DirectoryReference.CreateDirectory(TelemetryDirRef);
+				DirectoryReference.CreateDirectory(TelemetryDir);
 
-					string FileName = String.Format("{0}.{1}.json", Path.GetFileName(Assembly.GetEntryAssembly()!.Location), Process.Id, Process.StartTime.Ticks);
-					File = FileReference.Combine(TelemetryDirRef, FileName);
-				}
+				string FileName = String.Format("{0}.{1}.json", Path.GetFileName(Assembly.GetEntryAssembly()!.Location), Process.Id, Process.StartTime.Ticks);
+				File = FileReference.Combine(TelemetryDir, FileName);
+			}
 
-				using (JsonWriter Writer = new JsonWriter(File))
+			using (JsonWriter Writer = new JsonWriter(File))
+			{
+				Writer.WriteObjectStart();
+				Writer.WriteArrayStart("Spans");
+				foreach (TraceSpanImpl Span in Spans)
 				{
-					Writer.WriteObjectStart();
-					Writer.WriteArrayStart("Spans");
-					foreach (TraceSpanImpl Span in Spans)
+					if (Span.FinishTime != null)
 					{
-						if (Span.FinishTime != null)
+						Writer.WriteObjectStart();
+						Writer.WriteValue("Name", Span.Name);
+						if (Span.Resource != null)
 						{
-							Writer.WriteObjectStart();
-							Writer.WriteValue("Name", Span.Name);
-							if (Span.Resource != null)
-							{
-								Writer.WriteValue("Resource", Span.Resource);
-							}
-							if (Span.Service != null)
-							{
-								Writer.WriteValue("Service", Span.Service);
-							}
-							Writer.WriteValue("StartTime", Span.StartTime.ToString("o", CultureInfo.InvariantCulture));
-							Writer.WriteValue("FinishTime", Span.FinishTime.Value.ToString("o", CultureInfo.InvariantCulture));
-							Writer.WriteObjectStart("Metadata");
-							foreach (KeyValuePair<string, string> Pair in Span.Metadata)
-							{
-								Writer.WriteValue(Pair.Key, Pair.Value);
-							}
-							Writer.WriteObjectEnd();
-							Writer.WriteObjectEnd();
+							Writer.WriteValue("Resource", Span.Resource);
 						}
+						if (Span.Service != null)
+						{
+							Writer.WriteValue("Service", Span.Service);
+						}
+						Writer.WriteValue("StartTime", Span.StartTime.ToString("o", CultureInfo.InvariantCulture));
+						Writer.WriteValue("FinishTime", Span.FinishTime.Value.ToString("o", CultureInfo.InvariantCulture));
+						Writer.WriteObjectStart("Metadata");
+						foreach (KeyValuePair<string, string> Pair in Span.Metadata)
+						{
+							Writer.WriteValue(Pair.Key, Pair.Value);
+						}
+						Writer.WriteObjectEnd();
+						Writer.WriteObjectEnd();
 					}
-					Writer.WriteArrayEnd();
-					Writer.WriteObjectEnd();
 				}
+				Writer.WriteArrayEnd();
+				Writer.WriteObjectEnd();
+			}
+		}
+	}
+
+	/// <summary>
+	/// Methods for creating ITraceScope instances
+	/// </summary>
+	public static class TraceSpan
+	{
+		class CombinedTraceSpan : ITraceSpan
+		{
+			ITraceSpan[] Spans;
+
+			public CombinedTraceSpan(ITraceSpan[] Spans)
+			{
+				this.Spans = Spans;
+			}
+
+			public void AddMetadata(string Name, string Value)
+			{
+				foreach(ITraceSpan Span in Spans)
+				{
+					Span.AddMetadata(Name, Value);
+				}
+			}
+
+			public void Dispose()
+			{
+				foreach (ITraceSpan Span in Spans)
+				{
+					Span.Dispose();
+				}
+			}
+		}
+
+		/// <summary>
+		/// The sinks to use
+		/// </summary>
+		static List<ITraceSink> Sinks = GetDefaultSinks();
+
+		/// <summary>
+		/// Build a list of default sinks
+		/// </summary>
+		/// <returns></returns>
+		static List<ITraceSink> GetDefaultSinks()
+		{
+			List<ITraceSink> Sinks = new List<ITraceSink>();
+
+			string? TelemetryDir = Environment.GetEnvironmentVariable("UE_TELEMETRY_DIR");
+			if (TelemetryDir != null)
+			{
+				Sinks.Add(new JsonTraceSink(new DirectoryReference(TelemetryDir)));
+			}
+
+			return Sinks;
+		}
+
+		/// <summary>
+		/// Adds a new sink
+		/// </summary>
+		/// <param name="Sink">The sink to add</param>
+		public static void AddSink(ITraceSink Sink)
+		{
+			Sinks.Add(Sink);
+		}
+
+		/// <summary>
+		/// Remove a sink from the current list
+		/// </summary>
+		/// <param name="Sink">The sink to remove</param>
+		public static void RemoveSink(ITraceSink Sink)
+		{
+			Sinks.Remove(Sink);
+		}
+
+		/// <summary>
+		/// Creates a scope using the current provider
+		/// </summary>
+		public static ITraceSpan Create(string Operation, string? Resource = null, string? Service = null)
+		{
+			if (Sinks.Count == 0)
+			{
+				return new CombinedTraceSpan(Array.Empty<ITraceSpan>());
+			}
+			else if (Sinks.Count == 1)
+			{
+				return Sinks[0].Create(Operation, Resource, Service);
+			}
+			else
+			{
+				return new CombinedTraceSpan(Sinks.ConvertAll(x => x.Create(Operation, Resource, Service)).ToArray());
+			}
+		}
+
+		/// <summary>
+		/// Saves all the scope information to a file
+		/// </summary>
+		public static void Flush()
+		{
+			foreach (ITraceSink Sink in Sinks)
+			{
+				Sink.Flush();
 			}
 		}
 	}
diff --git a/Engine/Source/Programs/Shared/EpicGames.Perforce/EpicGames.Perforce.csproj b/Engine/Source/Programs/Shared/EpicGames.Perforce/EpicGames.Perforce.csproj
--- a/Engine/Source/Programs/Shared/EpicGames.Perforce/EpicGames.Perforce.csproj
+++ b/Engine/Source/Programs/Shared/EpicGames.Perforce/EpicGames.Perforce.csproj
@@ -27,10 +27,6 @@
     </Compile>
   </ItemGroup>
 
-  <ItemGroup>
-    <PackageReference Include="Datadog.Trace" Version="1.18.0" />
-  </ItemGroup>
-
   <ItemGroup>
     <ProjectReference Include="..\EpicGames.Core\EpicGames.Core.csproj" />
   </ItemGroup>
diff --git a/Engine/Source/Programs/Shared/EpicGames.Perforce/PerforceChildProcess.cs b/Engine/Source/Programs/Shared/EpicGames.Perforce/PerforceChildProcess.cs
--- a/Engine/Source/Programs/Shared/EpicGames.Perforce/PerforceChildProcess.cs
+++ b/Engine/Source/Programs/Shared/EpicGames.Perforce/PerforceChildProcess.cs
@@ -1,6 +1,5 @@
 // Copyright Epic Games, Inc. All Rights Reserved.
 
-using Datadog.Trace;
 using EpicGames.Core;
 using Microsoft.Extensions.Logging;
 using System;
@@ -55,7 +54,7 @@ namespace EpicGames.Perforce
 		/// <summary>
 		/// Scope object for tracing
 		/// </summary>
-		Scope Scope;
+		ITraceSpan Scope;
 
 		/// <summary>
 		/// Buffer for data from the child process
@@ -90,8 +89,8 @@ namespace EpicGames.Perforce
 			string FullArgumentList = "-G " + Arguments;
 			Logger.LogDebug("Running {0} {1}", PerforceFileName, FullArgumentList);
 
-			Scope = Tracer.Instance.StartActive($"{GetCommandName(Arguments)}", serviceName: "perforce");
-			Scope.Span.SetTag("arguments", FullArgumentList);
+			Scope = TraceSpan.Create($"{GetCommandName(Arguments)}", Service: "perforce");
+			Scope.AddMetadata("arguments", FullArgumentList);
 
 			ChildProcessGroup = new ManagedProcessGroup();
 			ChildProcess = new ManagedProcess(ChildProcessGroup, PerforceFileName, FullArgumentList, null, null, InputData, ProcessPriorityClass.Normal);
