Add extension methods for creating new default logger instances.

diff --git a/Engine/Source/Programs/Shared/EpicGames.Core/EpicGames.Core.csproj b/Engine/Source/Programs/Shared/EpicGames.Core/EpicGames.Core.csproj
--- a/Engine/Source/Programs/Shared/EpicGames.Core/EpicGames.Core.csproj
+++ b/Engine/Source/Programs/Shared/EpicGames.Core/EpicGames.Core.csproj
@@ -26,6 +26,7 @@
     </Compile>
   </ItemGroup>
   <ItemGroup>
+    <PackageReference Include="Microsoft.Extensions.Logging" Version="5.0.0" />
     <PackageReference Include="Microsoft.Extensions.Logging.Abstractions" Version="5.0.0" />
     <PackageReference Include="System.Memory" Version="4.5.4" />
   </ItemGroup>
diff --git a/Engine/Source/Programs/Shared/EpicGames.Core/Log.cs b/Engine/Source/Programs/Shared/EpicGames.Core/Log.cs
--- a/Engine/Source/Programs/Shared/EpicGames.Core/Log.cs
+++ b/Engine/Source/Programs/Shared/EpicGames.Core/Log.cs
@@ -1,5 +1,6 @@
 // Copyright Epic Games, Inc. All Rights Reserved.
 
+using Microsoft.Extensions.DependencyInjection;
 using Microsoft.Extensions.Logging;
 using System;
 using System.Collections.Generic;
@@ -342,7 +343,7 @@ namespace EpicGames.Core
 		/// <param name="LogFileName">The log filename to display, if any</param>
 		public static void WriteException(Exception Ex, string LogFileName)
 		{
-			string LogSuffix = (LogFileName == null)? "" : String.Format("\n(see {0} for full exception trace)", LogFileName);
+			string LogSuffix = (LogFileName == null) ? "" : String.Format("\n(see {0} for full exception trace)", LogFileName);
 			TraceLog("==============================================================================");
 			TraceError("{0}{1}", ExceptionUtils.FormatException(Ex), LogSuffix);
 			TraceLog("\n{0}", ExceptionUtils.FormatExceptionDetails(Ex));
@@ -1171,4 +1172,36 @@ namespace EpicGames.Core
 			}
 		}
 	}
+
+	/// <summary>
+	/// Provider for default logger instances
+	/// </summary>
+	public class DefaultLoggerProvider : ILoggerProvider
+	{
+		/// <inheritdoc/>
+		public ILogger CreateLogger(string CategoryName)
+		{
+			return new DefaultLogger();
+		}
+
+		/// <inheritdoc/>
+		public void Dispose()
+		{
+		}
+	}
+
+	/// <summary>
+	/// Extension methods to support the default logger
+	/// </summary>
+	public static class DefaultLoggerExtensions
+	{
+		/// <summary>
+		/// Adds a regular Epic logger to the builder
+		/// </summary>
+		/// <param name="Builder">Logging builder</param>
+		public static void AddEpicDefault(this ILoggingBuilder Builder)
+		{
+			Builder.Services.AddSingleton<ILoggerProvider, DefaultLoggerProvider>();
+		}
+	}
 }
