Allow direct access to stdin/stdout/stderr in ManagedProcess.

diff --git a/Engine/Source/Programs/Shared/EpicGames.Core/ManagedProcess.cs b/Engine/Source/Programs/Shared/EpicGames.Core/ManagedProcess.cs
--- a/Engine/Source/Programs/Shared/EpicGames.Core/ManagedProcess.cs
+++ b/Engine/Source/Programs/Shared/EpicGames.Core/ManagedProcess.cs
@@ -15,6 +15,23 @@ using System.Threading.Tasks;
 
 namespace EpicGames.Core
 {
+	/// <summary>
+	/// Flags for the managed process
+	/// </summary>
+	[Flags]
+	public enum ManagedProcessFlags
+	{
+		/// <summary>
+		/// No flags 
+		/// </summary>
+		None = 0,
+
+		/// <summary>
+		/// Merge stdout and stderr
+		/// </summary>
+		MergeOutputPipes = 1,
+	}
+
 	/// <summary>
 	/// Tracks a set of processes, and destroys them when the object is disposed.
 	/// </summary>
@@ -223,7 +240,7 @@ namespace EpicGames.Core
             IntPtr hSourceProcessHandle,
             SafeHandle hSourceHandle,
             IntPtr hTargetProcess,
-            out SafeWaitHandle targetHandle,
+            out SafeFileHandle targetHandle,
             int dwDesiredAccess,
             [MarshalAs(UnmanagedType.Bool)] bool bInheritHandle,
             int dwOptions
@@ -260,15 +277,35 @@ namespace EpicGames.Core
 		/// </summary>
 		SafeFileHandle? StdOutRead;
 
+		/// <summary>
+		/// The read end of the child process' stderr pipe.
+		/// </summary>
+		SafeFileHandle? StdErrRead;
+
+		/// <summary>
+		/// Input stream for the child process.
+		/// </summary>
+		public FileStream StdIn { get; private set; } = null!;
+
 		/// <summary>
 		/// Output stream for the child process.
 		/// </summary>
-		FileStream? InnerStream;
+		public FileStream StdOut { get; private set; } = null!;
 
 		/// <summary>
 		/// Reader for the process' output stream.
 		/// </summary>
-		StreamReader? ReadStream;
+		public StreamReader StdOutText { get; private set; } = null!;
+
+		/// <summary>
+		/// Output stream for the child process.
+		/// </summary>
+		public FileStream StdErr { get; private set; } = null!;
+
+		/// <summary>
+		/// Reader for the process' output stream.
+		/// </summary>
+		public StreamReader StdErrText { get; private set; } = null!;
 
 		/// <summary>
 		/// Standard process implementation for non-Windows platforms.
@@ -291,7 +328,26 @@ namespace EpicGames.Core
 		/// <param name="Environment">Environment variables for the new process. May be null, in which case the current process' environment is inherited</param>
 		/// <param name="Input">Text to be passed via stdin to the new process. May be null.</param>
 		/// <param name="Priority">Priority for the child process</param>
-		public ManagedProcess(ManagedProcessGroup? Group, string FileName, string CommandLine, string? WorkingDirectory, IReadOnlyDictionary<string, string>? Environment, byte[]? Input, ProcessPriorityClass Priority)
+		public ManagedProcess(ManagedProcessGroup? Group, string FileName, string CommandLine, string? WorkingDirectory, IReadOnlyDictionary<string, string>? Environment, byte[]? Input, ProcessPriorityClass Priority, ManagedProcessFlags Flags = ManagedProcessFlags.MergeOutputPipes)
+			: this(Group, FileName, CommandLine, WorkingDirectory, Environment, Priority, Flags)
+		{
+			if (Input != null)
+			{
+				StdIn.Write(Input, 0, Input.Length);
+				StdIn.Flush();
+			}
+		}
+
+		/// <summary>
+		/// Spawns a new managed process.
+		/// </summary>
+		/// <param name="Group">The managed process group to add to</param>
+		/// <param name="FileName">Path to the executable to be run</param>
+		/// <param name="CommandLine">Command line arguments for the process</param>
+		/// <param name="WorkingDirectory">Working directory for the new process. May be null to use the current working directory.</param>
+		/// <param name="Environment">Environment variables for the new process. May be null, in which case the current process' environment is inherited</param>
+		/// <param name="Priority">Priority for the child process</param>
+		public ManagedProcess(ManagedProcessGroup? Group, string FileName, string CommandLine, string? WorkingDirectory, IReadOnlyDictionary<string, string>? Environment, ProcessPriorityClass Priority, ManagedProcessFlags Flags = ManagedProcessFlags.MergeOutputPipes)
 		{
 			// Create the child process
 			// NOTE: Child process must be created in a separate method to avoid stomping exception callstacks (https://stackoverflow.com/a/2494150)
@@ -299,11 +355,11 @@ namespace EpicGames.Core
 			{
 				if (ManagedProcessGroup.SupportsJobObjects)
 				{
-					CreateManagedProcessWin32(Group, FileName, CommandLine, WorkingDirectory, Environment, Input, Priority);
+					CreateManagedProcessWin32(Group, FileName, CommandLine, WorkingDirectory, Environment, Priority, Flags);
 				}
 				else
 				{
-					CreateManagedProcessPortable(Group, FileName, CommandLine, WorkingDirectory, Environment, Input, Priority);
+					CreateManagedProcessPortable(Group, FileName, CommandLine, WorkingDirectory, Environment, Priority, Flags);
 				}
 			}
 			catch (Exception Ex)
@@ -321,9 +377,9 @@ namespace EpicGames.Core
 		/// <param name="CommandLine">Command line arguments for the process</param>
 		/// <param name="WorkingDirectory">Working directory for the new process. May be null to use the current working directory.</param>
 		/// <param name="Environment">Environment variables for the new process. May be null, in which case the current process' environment is inherited</param>
-		/// <param name="Input">Text to be passed via stdin to the new process. May be null.</param>
 		/// <param name="Priority">Priority for the child process</param>
-		private void CreateManagedProcessWin32(ManagedProcessGroup? Group, string FileName, string CommandLine, string? WorkingDirectory, IReadOnlyDictionary<string, string>? Environment, byte[]? Input, ProcessPriorityClass Priority)
+		/// <param name="ManagedFlags">Flags controlling how the new process is created</param>
+		private void CreateManagedProcessWin32(ManagedProcessGroup? Group, string FileName, string CommandLine, string? WorkingDirectory, IReadOnlyDictionary<string, string>? Environment, ProcessPriorityClass Priority, ManagedProcessFlags ManagedFlags)
 		{
 			IntPtr EnvironmentBlock = IntPtr.Zero;
 			try
@@ -381,7 +437,7 @@ namespace EpicGames.Core
 					{
 						SafeFileHandle? StdInRead = null;
 						SafeFileHandle? StdOutWrite = null;
-						SafeWaitHandle? StdErrWrite = null;
+						SafeFileHandle? StdErrWrite = null;
 						try
 						{
 							// Create stdin and stdout pipes for the child process. We'll close the handles for the child process' ends after it's been created.
@@ -397,9 +453,20 @@ namespace EpicGames.Core
 							{
 								throw new Win32ExceptionWithCode("Unable to create stdout pipe");
 							}
-							if (DuplicateHandle(GetCurrentProcess(), StdOutWrite, GetCurrentProcess(), out StdErrWrite, 0, true, DUPLICATE_SAME_ACCESS) == 0)
+
+							if ((ManagedFlags & ManagedProcessFlags.MergeOutputPipes) != 0)
+							{
+								if (DuplicateHandle(GetCurrentProcess(), StdOutWrite, GetCurrentProcess(), out StdErrWrite, 0, true, DUPLICATE_SAME_ACCESS) == 0)
+								{
+									throw new Win32ExceptionWithCode("Unable to duplicate stdout handle");
+								}
+							}
+							else
 							{
-								throw new Win32ExceptionWithCode("Unable to duplicate stdout handle");
+								if (CreatePipe(out StdErrRead, out StdErrWrite, SecurityAttributes, 1024 * 1024) == 0 || SetHandleInformation(StdErrRead, HANDLE_FLAG_INHERIT, 0) == 0)
+								{
+									throw new Win32ExceptionWithCode("Unable to create stderr pipe");
+								}
 							}
 
 							// Create the new process as suspended, so we can modify it before it starts executing (and potentially preempting us)
@@ -470,18 +537,23 @@ namespace EpicGames.Core
 					}
 
 					// If we have any input text, write it to stdin now
-					using (FileStream Stream = new FileStream(StdInWrite, FileAccess.Write, 4096, false))
-					{
-						if (Input != null)
-						{
-							Stream.Write(Input, 0, Input.Length);
-							Stream.Flush();
-						}
-					}
+					StdIn = new FileStream(StdInWrite, FileAccess.Write, 4096, false);
 
 					// Create the stream objects for reading the process output
-					InnerStream = new FileStream(StdOutRead, FileAccess.Read, 4096, false);
-					ReadStream = new StreamReader(InnerStream, Console.OutputEncoding);
+					StdOut = new FileStream(StdOutRead, FileAccess.Read, 4096, false);
+					StdOutText = new StreamReader(StdOut, Console.OutputEncoding);
+
+					// Do the same for the stderr output
+					if (StdErrRead != null)
+					{
+						StdErr = new FileStream(StdErrRead, FileAccess.Read, 4096, false);
+						StdErrText = new StreamReader(StdErr, Console.OutputEncoding);
+					}
+					else
+					{
+						StdErr = StdOut;
+						StdErrText = StdOutText;
+					}
 
 					// Wrap the process handle in a SafeFileHandle
 					ProcessHandle = new SafeFileHandle(ProcessInfo.hProcess, true);
@@ -516,9 +588,9 @@ namespace EpicGames.Core
 		/// <param name="CommandLine">Command line arguments for the process</param>
 		/// <param name="WorkingDirectory">Working directory for the new process. May be null to use the current working directory.</param>
 		/// <param name="Environment">Environment variables for the new process. May be null, in which case the current process' environment is inherited</param>
-		/// <param name="Input">Text to be passed via stdin to the new process. May be null.</param>
 		/// <param name="Priority">Priority for the child process</param>
-		private void CreateManagedProcessPortable(ManagedProcessGroup? Group, string FileName, string CommandLine, string? WorkingDirectory, IReadOnlyDictionary<string, string>? Environment, byte[]? Input, ProcessPriorityClass Priority)
+		/// <param name="Flags">Flags for the new process</param>
+		private void CreateManagedProcessPortable(ManagedProcessGroup? Group, string FileName, string CommandLine, string? WorkingDirectory, IReadOnlyDictionary<string, string>? Environment, ProcessPriorityClass Priority, ManagedProcessFlags Flags)
 		{
 			// Fallback for Mono platforms
 			FrameworkProcess = new Process();
@@ -547,12 +619,6 @@ namespace EpicGames.Core
 			{
 
 			}
-			// If we have any input text, write it to stdin now
-			if (Input != null)
-			{
-				FrameworkProcess.StandardInput.BaseStream.Write(Input, 0, Input.Length);
-				FrameworkProcess.StandardInput.BaseStream.Close();
-			}
 		}
 
 		/// <summary>
@@ -609,7 +675,7 @@ namespace EpicGames.Core
 			Task<int> ReadTask;
 			if(FrameworkProcess == null)
 			{
-				ReadTask = InnerStream!.ReadAsync(Buffer, Offset, Count);
+				ReadTask = StdOut!.ReadAsync(Buffer, Offset, Count);
 			}
 			else
 			{
@@ -633,7 +699,7 @@ namespace EpicGames.Core
 		{
 			if (FrameworkProcess == null)
 			{
-				return await InnerStream!.ReadAsync(Buffer, Offset, Count, CancellationToken);
+				return await StdOut!.ReadAsync(Buffer, Offset, Count, CancellationToken);
 			}
 			else
 			{
@@ -651,7 +717,7 @@ namespace EpicGames.Core
 		{
 			if(FrameworkProcess == null)
 			{
-				return InnerStream!.CopyToAsync(OutputStream, CancellationToken);
+				return StdOut!.CopyToAsync(OutputStream, CancellationToken);
 			}
 			else
 			{
@@ -723,7 +789,7 @@ namespace EpicGames.Core
 		{
 			if (FrameworkProcess == null)
 			{
-				return await ReadStream!.ReadLineAsync();
+				return await StdOutText!.ReadLineAsync();
 			}
 			else
 			{
@@ -745,7 +811,7 @@ namespace EpicGames.Core
 				Task<string?> ReadLineTask;
 				if(FrameworkProcess == null)
 				{
-					ReadLineTask = ReadStream!.ReadLineAsync();
+					ReadLineTask = StdOutText!.ReadLineAsync();
 				}
 				else
 				{
