P4VUtils: Handle invalid or missing arguments gracefully.

#rnx

diff --git a/Engine/Extras/P4VUtils/Commands/CherryPickCommand.cs b/Engine/Extras/P4VUtils/Commands/CherryPickCommand.cs
--- a/Engine/Extras/P4VUtils/Commands/CherryPickCommand.cs
+++ b/Engine/Extras/P4VUtils/Commands/CherryPickCommand.cs
@@ -18,7 +18,17 @@ namespace P4VUtils.Commands
 
 		public override async Task<int> Execute(string[] Args, IReadOnlyDictionary<string, string> ConfigValues, ILogger Logger)
 		{
-			int Change = int.Parse(Args[1]);
+			int Change;
+			if (Args.Length < 2)
+			{
+				Logger.LogError("Missing changelist number");
+				return 1;
+			}
+			else if (!int.TryParse(Args[1], out Change))
+			{
+				Logger.LogError("'{Argument}' is not a numbered changelist", Args[1]);
+				return 1;
+			}
 
 			PerforceConnection Perforce = new PerforceConnection(null, null, null, Logger);
 			int NewChange = await MergeAsync(Perforce, Change, Logger);
diff --git a/Engine/Extras/P4VUtils/Commands/ConvertToEditCommand.cs b/Engine/Extras/P4VUtils/Commands/ConvertToEditCommand.cs
--- a/Engine/Extras/P4VUtils/Commands/ConvertToEditCommand.cs
+++ b/Engine/Extras/P4VUtils/Commands/ConvertToEditCommand.cs
@@ -21,7 +21,18 @@ namespace P4VUtils.Commands
 
 		public override async Task<int> Execute(string[] Args, IReadOnlyDictionary<string, string> ConfigValues, ILogger Logger)
 		{
-			int Change = int.Parse(Args[1]);
+			int Change;
+			if (Args.Length < 2)
+			{
+				Logger.LogError("Missing changelist number");
+				return 1;
+			}
+			else if (!int.TryParse(Args[1], out Change))
+			{
+				Logger.LogError("'{Argument}' is not a numbered changelist", Args[1]);
+				return 1;
+			}
+
 			bool Debug = Args.Any(x => x.Equals("-Debug", StringComparison.OrdinalIgnoreCase));
 
 			PerforceConnection Perforce = new PerforceConnection(null, null, Logger);
diff --git a/Engine/Extras/P4VUtils/Commands/DescribeCommand.cs b/Engine/Extras/P4VUtils/Commands/DescribeCommand.cs
--- a/Engine/Extras/P4VUtils/Commands/DescribeCommand.cs
+++ b/Engine/Extras/P4VUtils/Commands/DescribeCommand.cs
@@ -20,7 +20,17 @@ namespace P4VUtils.Commands
 
 		public override async Task<int> Execute(string[] Args, IReadOnlyDictionary<string, string> ConfigValues, ILogger Logger)
 		{
-			int Change = int.Parse(Args[1]);
+			int Change;
+			if (Args.Length < 2)
+			{
+				Logger.LogError("Missing changelist number");
+				return 1;
+			}
+			else if (!int.TryParse(Args[1], out Change))
+			{
+				Logger.LogError("'{Argument}' is not a numbered changelist", Args[1]);
+				return 1;
+			}
 
 			PerforceConnection Perforce = new PerforceConnection(null, null, Logger);
 			DescribeRecord DescribeRecord = await Perforce.DescribeAsync(Change, CancellationToken.None);
diff --git a/Engine/Extras/P4VUtils/Commands/EdigrateCommand.cs b/Engine/Extras/P4VUtils/Commands/EdigrateCommand.cs
--- a/Engine/Extras/P4VUtils/Commands/EdigrateCommand.cs
+++ b/Engine/Extras/P4VUtils/Commands/EdigrateCommand.cs
@@ -19,7 +19,18 @@ namespace P4VUtils.Commands
 
 		public override async Task<int> Execute(string[] Args, IReadOnlyDictionary<string, string> ConfigValues, ILogger Logger)
 		{
-			int Change = int.Parse(Args[1]);
+			int Change;
+			if (Args.Length < 2)
+			{
+				Logger.LogError("Missing changelist number");
+				return 1;
+			}
+			else if (!int.TryParse(Args[1], out Change))
+			{
+				Logger.LogError("'{Argument}' is not a numbered changelist", Args[1]);
+				return 1;
+			}
+
 			bool Debug = Args.Any(x => x.Equals("-Debug", StringComparison.OrdinalIgnoreCase));
 
 			PerforceConnection Perforce = new PerforceConnection(null, null, null, Logger);
diff --git a/Engine/Extras/P4VUtils/Commands/PreflightCommand.cs b/Engine/Extras/P4VUtils/Commands/PreflightCommand.cs
--- a/Engine/Extras/P4VUtils/Commands/PreflightCommand.cs
+++ b/Engine/Extras/P4VUtils/Commands/PreflightCommand.cs
@@ -25,7 +25,17 @@ namespace P4VUtils.Commands
 
 		public override async Task<int> Execute(string[] Args, IReadOnlyDictionary<string, string> ConfigValues, ILogger Logger)
 		{
-			int Change = int.Parse(Args[1]);
+			int Change;
+			if (Args.Length < 2)
+			{
+				Logger.LogError("Missing changelist number");
+				return 1;
+			}
+			else if (!int.TryParse(Args[1], out Change))
+			{
+				Logger.LogError("'{Argument}' is not a numbered changelist", Args[1]);
+				return 1;
+			}
 
 			PerforceConnection Perforce = new PerforceConnection(null, null, Logger);
 
diff --git a/Engine/Source/Programs/Shared/EpicGames.Perforce/PerforceConnection.cs b/Engine/Source/Programs/Shared/EpicGames.Perforce/PerforceConnection.cs
--- a/Engine/Source/Programs/Shared/EpicGames.Perforce/PerforceConnection.cs
+++ b/Engine/Source/Programs/Shared/EpicGames.Perforce/PerforceConnection.cs
@@ -2394,7 +2394,7 @@ namespace EpicGames.Perforce
 		/// <param name="FileSpecs">Files to sync</param>
 		/// <param name="CancellationToken">Token used to cancel the operation</param>
 		/// <returns>Response from the server</returns>
-		public Task<PerforceResponseList<ShelveRecord>> TryShelveAsync(int ChangeNumber, ShelveOptions Options, string[] FileSpecs, CancellationToken CancellationToken)
+		public async Task<PerforceResponseList<ShelveRecord>> TryShelveAsync(int ChangeNumber, ShelveOptions Options, string[] FileSpecs, CancellationToken CancellationToken)
 		{
 			StringBuilder Arguments = new StringBuilder("shelve");
 			Arguments.AppendFormat(" -c {0}", ChangeNumber);
@@ -2411,7 +2411,9 @@ namespace EpicGames.Perforce
 				Arguments.AppendFormat(" \"{0}\"", FileSpec);
 			}
 
-			return CommandAsync<ShelveRecord>(Arguments.ToString(), null, CancellationToken);
+			PerforceResponseList<ShelveRecord> Records = await CommandAsync<ShelveRecord>(Arguments.ToString(), null, CancellationToken);
+			Records.RemoveAll(x => x.Error != null && x.Error.Generic == PerforceGenericCode.Empty);
+			return Records;
 		}
 
 		/// <summary>
