Close stdin after writing input data to child processes. Perforce waits for the handle to be closed before processing the data.

#rnx

diff --git a/Engine/Source/Programs/Shared/EpicGames.Core/ManagedProcess.cs b/Engine/Source/Programs/Shared/EpicGames.Core/ManagedProcess.cs
--- a/Engine/Source/Programs/Shared/EpicGames.Core/ManagedProcess.cs
+++ b/Engine/Source/Programs/Shared/EpicGames.Core/ManagedProcess.cs
@@ -334,8 +334,8 @@ namespace EpicGames.Core
 			if (Input != null)
 			{
 				StdIn.Write(Input, 0, Input.Length);
-				StdIn.Flush();
 			}
+			StdIn.Close();
 		}
 
 		/// <summary>
