Add a utility method for formatting a ReadOnlyMemory<byte> as a hex string.

#rnx

diff --git a/Engine/Source/Programs/Shared/EpicGames.Core/StringUtils.cs b/Engine/Source/Programs/Shared/EpicGames.Core/StringUtils.cs
--- a/Engine/Source/Programs/Shared/EpicGames.Core/StringUtils.cs
+++ b/Engine/Source/Programs/Shared/EpicGames.Core/StringUtils.cs
@@ -370,11 +370,21 @@ namespace EpicGames.Core
 		/// <param name="Bytes">An array of bytes</param>
 		/// <returns>String representation of the array</returns>
 		public static string FormatHexString(byte[] Bytes)
+		{
+			return FormatHexString(Bytes.AsSpan());
+		}
+
+		/// <summary>
+		/// Formats an array of bytes as a hexadecimal string
+		/// </summary>
+		/// <param name="Bytes">An array of bytes</param>
+		/// <returns>String representation of the array</returns>
+		public static string FormatHexString(ReadOnlySpan<byte> Bytes)
 		{
 			const string HexDigits = "0123456789abcdef";
 
 			char[] Characters = new char[Bytes.Length * 2];
-			for(int Idx = 0; Idx < Bytes.Length; Idx++)
+			for (int Idx = 0; Idx < Bytes.Length; Idx++)
 			{
 				Characters[Idx * 2 + 0] = HexDigits[Bytes[Idx] >> 4];
 				Characters[Idx * 2 + 1] = HexDigits[Bytes[Idx] & 15];
