Added missing copyright boilerplate

diff --git a/Engine/Source/Programs/Shared/EpicGames.Core/ZipArchiveExtensions.cs b/Engine/Source/Programs/Shared/EpicGames.Core/ZipArchiveExtensions.cs
--- a/Engine/Source/Programs/Shared/EpicGames.Core/ZipArchiveExtensions.cs
+++ b/Engine/Source/Programs/Shared/EpicGames.Core/ZipArchiveExtensions.cs
@@ -1,3 +1,5 @@
+// Copyright Epic Games, Inc. All Rights Reserved.
+
 using System;
 using System.Collections.Generic;
 using System.Text;
