Fix UBT log not containing all log levels.

#rnx

diff --git a/Engine/Source/Programs/Shared/EpicGames.Core/Log.cs b/Engine/Source/Programs/Shared/EpicGames.Core/Log.cs
--- a/Engine/Source/Programs/Shared/EpicGames.Core/Log.cs
+++ b/Engine/Source/Programs/Shared/EpicGames.Core/Log.cs
@@ -969,7 +969,7 @@ namespace EpicGames.Core
 				WriteOnceSet.Add(Formatted);
 			}
 
-			if (Verbosity > OutputLevel)
+			if (Verbosity >= OutputLevel)
 			{
 				lock (SyncObject)
 				{
