Fix permissions when extracting a Windows zip file on Mac.

diff --git a/Engine/Source/Programs/Shared/EpicGames.Core/ZipArchiveExtensions.cs b/Engine/Source/Programs/Shared/EpicGames.Core/ZipArchiveExtensions.cs
--- a/Engine/Source/Programs/Shared/EpicGames.Core/ZipArchiveExtensions.cs
+++ b/Engine/Source/Programs/Shared/EpicGames.Core/ZipArchiveExtensions.cs
@@ -5,6 +5,7 @@ using System.Collections.Generic;
 using System.Text;
 using System.Runtime.InteropServices;
 using System.IO.Compression;
+using System.Reflection;
 
 namespace EpicGames.Core
 {
@@ -35,6 +36,11 @@ namespace EpicGames.Core
 			return Entry;
 		}
 
+		/// <summary>
+		/// Internal field storing information about the platform that created a ZipArchiveEntry. Cannot interpret how to treat the attribute bits without reading this.
+		/// </summary>
+		static FieldInfo VersionMadeByPlatformField = typeof(ZipArchiveEntry).GetField("_versionMadeByPlatform", BindingFlags.NonPublic | BindingFlags.Instance)!;
+
 		/// <summary>
 		/// Extract a zip archive entry, preserving platform mode bits
 		/// </summary>
@@ -46,7 +52,11 @@ namespace EpicGames.Core
 			ZipFileExtensions.ExtractToFile(Entry, TargetFileName, Overwrite);
 			if (RuntimeInformation.IsOSPlatform(OSPlatform.OSX))
 			{
-				FileUtils.SetFileMode_Mac(TargetFileName, (ushort)(Entry.ExternalAttributes >> 16));
+				int MadeByPlatform = Convert.ToInt32(VersionMadeByPlatformField.GetValue(Entry)!);
+				if (MadeByPlatform == 3 || MadeByPlatform == 19) // Unix or OSX
+				{
+					FileUtils.SetFileMode_Mac(TargetFileName, (ushort)(Entry.ExternalAttributes >> 16));
+				}
 			}
 		}
 	}
