P4VUtils: Fix edigrate operations on large changelists.

diff --git a/Engine/Source/Programs/Shared/EpicGames.Perforce/PerforceConnection.cs b/Engine/Source/Programs/Shared/EpicGames.Perforce/PerforceConnection.cs
--- a/Engine/Source/Programs/Shared/EpicGames.Perforce/PerforceConnection.cs
+++ b/Engine/Source/Programs/Shared/EpicGames.Perforce/PerforceConnection.cs
@@ -2289,16 +2289,9 @@ namespace EpicGames.Perforce
 			{
 				Arguments.AppendFormat(" -C \"{0}\"", ClientName);
 			}
-			foreach(string FileSpec in FileSpecs)
-			{
-				Arguments.AppendFormat(" \"{0}\"", FileSpec);
-			}
 
-			PerforceResponseList<RevertRecord> Records = await CommandAsync<RevertRecord>(Arguments.ToString(), null, CancellationToken);
-			if(Records.Count == 1 && Records[0].Error != null && Records[0].Error!.Generic == PerforceGenericCode.Empty)
-			{
-				Records.Clear();
-			}
+			PerforceResponseList<RevertRecord> Records = await BatchedCommandAsync<RevertRecord>(Arguments.ToString(), FileSpecs, null, CancellationToken);
+			Records.RemoveAll(x => x.Error != null && x.Error.Generic == PerforceGenericCode.Empty);
 			return Records;
 		}
 
