P4VUtils: Fix some exceptions caused by long command lines.

#rnx

diff --git a/Engine/Extras/P4VUtils/Commands/ConvertToEditCommand.cs b/Engine/Extras/P4VUtils/Commands/ConvertToEditCommand.cs
--- a/Engine/Extras/P4VUtils/Commands/ConvertToEditCommand.cs
+++ b/Engine/Extras/P4VUtils/Commands/ConvertToEditCommand.cs
@@ -221,7 +221,7 @@ namespace P4VUtils.Commands
 			{
 				Logger.LogInformation("Re-opening the following files for {Operation}:", Operation);
 
-				string[] LocalFiles = Records.Select(x => OpenedClientToLocalMap[x.ClientFile!]).ToArray();
+				string[] LocalFiles = Records.Select(x => PerforceUtils.EscapePath(OpenedClientToLocalMap[x.ClientFile!])).ToArray();
 				foreach (string LocalFile in LocalFiles)
 				{
 					Logger.LogInformation("    {LocalFile}", LocalFile);
diff --git a/Engine/Source/Programs/Shared/EpicGames.Perforce/PerforceConnection.cs b/Engine/Source/Programs/Shared/EpicGames.Perforce/PerforceConnection.cs
--- a/Engine/Source/Programs/Shared/EpicGames.Perforce/PerforceConnection.cs
+++ b/Engine/Source/Programs/Shared/EpicGames.Perforce/PerforceConnection.cs
@@ -394,12 +394,8 @@ namespace EpicGames.Perforce
 			{
 				Arguments.AppendFormat(" -t \"{0}\"", FileType);
 			}
-			foreach (string FileName in FileNames)
-			{
-				Arguments.AppendFormat(" \"{0}\"", FileName);
-			}
 
-			return CommandAsync<AddRecord>(Arguments.ToString(), null, CancellationToken);
+			return BatchedCommandAsync<AddRecord>(Arguments.ToString(), FileNames, null, CancellationToken);
 		}
 
 		#endregion
@@ -1188,11 +1184,8 @@ namespace EpicGames.Perforce
 			{
 				Arguments.Append(" -v");
 			}
-			foreach (string FileSpec in FileSpecs)
-			{
-				Arguments.AppendFormat(" \"{0}\"", FileSpec);
-			}
-			return CommandAsync<DeleteRecord>(Arguments.ToString(), null, CancellationToken);
+
+			return BatchedCommandAsync<DeleteRecord>(Arguments.ToString(), FileSpecs, null, CancellationToken);
 		}
 
 		#endregion
@@ -1370,12 +1363,8 @@ namespace EpicGames.Perforce
 			{
 				Arguments.AppendFormat(" -t \"{0}\"", FileType);
 			}
-			foreach (string FileSpec in FileSpecs)
-			{
-				Arguments.AppendFormat(" \"{0}\"", FileSpec);
-			}
 
-			return CommandAsync<EditRecord>(Arguments.ToString(), null, CancellationToken);
+			return BatchedCommandAsync<EditRecord>(Arguments.ToString(), FileSpecs, null, CancellationToken);
 		}
 
 		#endregion
