Merge utility library changes from GitSync.

diff --git a/Engine/Source/Programs/AutomationTool/AutomationUtils/HelpUtils.cs b/Engine/Source/Programs/AutomationTool/AutomationUtils/HelpUtils.cs
--- a/Engine/Source/Programs/AutomationTool/AutomationUtils/HelpUtils.cs
+++ b/Engine/Source/Programs/AutomationTool/AutomationUtils/HelpUtils.cs
@@ -146,7 +146,7 @@ namespace AutomationTool
 			}
 
 			Log.TraceInformation("");
-			HelpUtils.PrintHelp(String.Format("{0} Help:", Command.Name), Description, ParamDict.ToList());
+			HelpUtils.PrintHelp(String.Format("{0} Help:", Command.Name), Description, ParamDict.ToList(), HelpUtils.WindowWidth, Log.Logger);
 		}
 
 		/// <summary>
diff --git a/Engine/Source/Programs/AutomationTool/BuildGraph/Graph.cs b/Engine/Source/Programs/AutomationTool/BuildGraph/Graph.cs
--- a/Engine/Source/Programs/AutomationTool/BuildGraph/Graph.cs
+++ b/Engine/Source/Programs/AutomationTool/BuildGraph/Graph.cs
@@ -803,7 +803,7 @@ namespace AutomationTool
 					CommandUtils.LogInformation("");
 					CommandUtils.LogInformation("Options:");
 					CommandUtils.LogInformation("");
-					HelpUtils.PrintTable(Parameters, 4, 24);
+					HelpUtils.PrintTable(Parameters, 4, 24, Log.Logger);
 				}
 			}
 
diff --git a/Engine/Source/Programs/Shared/EpicGames.Core/BinaryArchiveReader.cs b/Engine/Source/Programs/Shared/EpicGames.Core/BinaryArchiveReader.cs
--- a/Engine/Source/Programs/Shared/EpicGames.Core/BinaryArchiveReader.cs
+++ b/Engine/Source/Programs/Shared/EpicGames.Core/BinaryArchiveReader.cs
@@ -2,6 +2,7 @@
 
 using System;
 using System.Collections.Generic;
+using System.Diagnostics.CodeAnalysis;
 using System.IO;
 using System.Linq;
 using System.Text;
@@ -527,5 +528,36 @@ namespace EpicGames.Core
 				return (T)Objects[Index];
 			}
 		}
+
+		/// <summary>
+		/// Helper method for validating that deserialized objects are not null
+		/// </summary>
+		/// <typeparam name="T">Type of the deserialized object</typeparam>
+		/// <param name="Param">The object instance</param>
+		/// <returns>The object instance</returns>
+		[return: NotNull]
+		public static T NotNull<T>(T Param) where T : class
+		{
+			if (Param == null)
+			{
+				throw new InvalidDataException("Object stored in archive is not allowed to be null.");
+			}
+			return Param;
+		}
+
+		/// <summary>
+		/// Helper method for validating that deserialized objects are not null
+		/// </summary>
+		/// <typeparam name="T">Type of the deserialized object</typeparam>
+		/// <param name="Param">The object instance</param>
+		/// <returns>The object instance</returns>
+		public static T NotNullStruct<T>(T? Param) where T : struct
+		{
+			if (Param == null)
+			{
+				throw new InvalidDataException("Object stored in archive is not allowed to be null.");
+			}
+			return Param.Value;
+		}
 	}
 }
diff --git a/Engine/Source/Programs/Shared/EpicGames.Core/CommandLineArguments.cs b/Engine/Source/Programs/Shared/EpicGames.Core/CommandLineArguments.cs
--- a/Engine/Source/Programs/Shared/EpicGames.Core/CommandLineArguments.cs
+++ b/Engine/Source/Programs/Shared/EpicGames.Core/CommandLineArguments.cs
@@ -1,17 +1,17 @@
 // Copyright Epic Games, Inc. All Rights Reserved.
 
+using Microsoft.Extensions.Logging;
 using System;
 using System.Collections;
 using System.Collections.Generic;
 using System.ComponentModel;
 using System.Diagnostics;
+using System.Diagnostics.CodeAnalysis;
 using System.Linq;
 using System.Reflection;
 using System.Text;
 using System.Threading.Tasks;
 
-#nullable disable
-
 namespace EpicGames.Core
 {
 	/// <summary>
@@ -291,7 +291,7 @@ namespace EpicGames.Core
 		/// </summary>
 		/// <param name="Argument">Receives the argument that was read, on success</param>
 		/// <returns>True if an argument was read</returns>
-		public bool TryGetPositionalArgument(out string Argument)
+		public bool TryGetPositionalArgument([NotNullWhen(true)] out string? Argument)
 		{
 			for (int Idx = 0; Idx < PositionalArgumentIndices.Count; Idx++)
 			{
@@ -331,7 +331,7 @@ namespace EpicGames.Core
 		/// <returns>Value of the argument</returns>
 		public string GetString(string Prefix)
 		{
-			string Value;
+			string? Value;
 			if(!TryGetValue(Prefix, out Value))
 			{
 				throw new CommandLineArgumentException(String.Format("Missing '{0}...' argument", Prefix));
@@ -361,7 +361,7 @@ namespace EpicGames.Core
 		/// <returns>Value of the argument</returns>
 		public FileReference GetFileReference(string Prefix)
 		{
-			FileReference Value;
+			FileReference? Value;
 			if(!TryGetValue(Prefix, out Value))
 			{
 				throw new CommandLineArgumentException(String.Format("Missing '{0}...' argument", Prefix));
@@ -376,7 +376,7 @@ namespace EpicGames.Core
 		/// <returns>Value of the argument</returns>
 		public DirectoryReference GetDirectoryReference(string Prefix)
 		{
-			DirectoryReference Value;
+			DirectoryReference? Value;
 			if(!TryGetValue(Prefix, out Value))
 			{
 				throw new CommandLineArgumentException(String.Format("Missing '{0}...' argument", Prefix));
@@ -405,9 +405,10 @@ namespace EpicGames.Core
 		/// <param name="Prefix">The argument prefix (eg. "-Foo="). Must end with an '=' character.</param>
 		/// <param name="DefaultValue">Default value for the argument</param>
 		/// <returns>Value of the argument</returns>
-		public string GetStringOrDefault(string Prefix, string DefaultValue)
+		[return: NotNullIfNotNull("DefaultValue")]
+		public string? GetStringOrDefault(string Prefix, string? DefaultValue)
 		{
-			string Value;
+			string? Value;
 			if(!TryGetValue(Prefix, out Value))
 			{
 				Value = DefaultValue;
@@ -437,9 +438,10 @@ namespace EpicGames.Core
 		/// <param name="Prefix">The argument prefix (eg. "-Foo="). Must end with an '=' character.</param>
 		/// <param name="DefaultValue">Default value for the argument</param>
 		/// <returns>Value of the argument</returns>
-		public FileReference GetFileReferenceOrDefault(string Prefix, FileReference DefaultValue)
+		[return: NotNullIfNotNull("DefaultValue")]
+		public FileReference? GetFileReferenceOrDefault(string Prefix, FileReference? DefaultValue)
 		{
-			FileReference Value;
+			FileReference? Value;
 			if(!TryGetValue(Prefix, out Value))
 			{
 				Value = DefaultValue;
@@ -453,9 +455,10 @@ namespace EpicGames.Core
 		/// <param name="Prefix">The argument prefix (eg. "-Foo="). Must end with an '=' character.</param>
 		/// <param name="DefaultValue">Default value for the argument</param>
 		/// <returns>Value of the argument</returns>
-		public DirectoryReference GetDirectoryReferenceOrDefault(string Prefix, DirectoryReference DefaultValue)
+		[return: NotNullIfNotNull("DefaultValue")]
+		public DirectoryReference? GetDirectoryReferenceOrDefault(string Prefix, DirectoryReference? DefaultValue)
 		{
-			DirectoryReference Value;
+			DirectoryReference? Value;
 			if(!TryGetValue(Prefix, out Value))
 			{
 				Value = DefaultValue;
@@ -485,7 +488,7 @@ namespace EpicGames.Core
 		/// <param name="Prefix">The argument prefix (eg. "-Foo="). Must end with an '=' character.</param>
 		/// <param name="Value">Value of the argument, if found</param>
 		/// <returns>True if the argument was found (and Value was set), false otherwise.</returns>
-		public bool TryGetValue(string Prefix, out string Value)
+		public bool TryGetValue(string Prefix, [NotNullWhen(true)] out string? Value)
 		{
 			CheckValidPrefix(Prefix);
 
@@ -515,7 +518,7 @@ namespace EpicGames.Core
 		public bool TryGetValue(string Prefix, out int Value)
 		{
 			// Try to get the string value of this argument
-			string StringValue;
+			string? StringValue;
 			if(!TryGetValue(Prefix, out StringValue))
 			{
 				Value = 0;
@@ -540,10 +543,10 @@ namespace EpicGames.Core
 		/// <param name="Prefix">The argument prefix (eg. "-Foo="). Must end with an '=' character.</param>
 		/// <param name="Value">Value of the argument, if found</param>
 		/// <returns>True if the argument was found (and Value was set), false otherwise.</returns>
-		public bool TryGetValue(string Prefix, out FileReference Value)
+		public bool TryGetValue(string Prefix, [NotNullWhen(true)] out FileReference? Value)
 		{
 			// Try to get the string value of this argument
-			string StringValue;
+			string? StringValue;
 			if(!TryGetValue(Prefix, out StringValue))
 			{
 				Value = null;
@@ -568,10 +571,10 @@ namespace EpicGames.Core
 		/// <param name="Prefix">The argument prefix (eg. "-Foo="). Must end with an '=' character.</param>
 		/// <param name="Value">Value of the argument, if found</param>
 		/// <returns>True if the argument was found (and Value was set), false otherwise.</returns>
-		public bool TryGetValue(string Prefix, out DirectoryReference Value)
+		public bool TryGetValue(string Prefix, [NotNullWhen(true)] out DirectoryReference? Value)
 		{
 			// Try to get the string value of this argument
-			string StringValue;
+			string? StringValue;
 			if(!TryGetValue(Prefix, out StringValue))
 			{
 				Value = null;
@@ -599,7 +602,7 @@ namespace EpicGames.Core
 		public bool TryGetValue<T>(string Prefix, out T Value) where T : struct
 		{
 			// Try to get the string value of this argument
-			string StringValue;
+			string? StringValue;
 			if(!TryGetValue(Prefix, out StringValue))
 			{
 				Value = new T();
@@ -670,7 +673,7 @@ namespace EpicGames.Core
 				ValueType = ValueType.GetGenericArguments()[0];
 			}
 
-			string Prefix = Attribute.Prefix;
+			string? Prefix = Attribute.Prefix;
 			if (Prefix == null)
 			{
 				if (ValueType == typeof(bool))
@@ -696,19 +699,30 @@ namespace EpicGames.Core
 		/// Applies these arguments to fields with the [CommandLine] attribute in the given object.
 		/// </summary>
 		/// <param name="TargetObject">The object to configure</param>
+		/// <param name="Logger">Sink for error/warning messages</param>
 		public void ApplyTo(object TargetObject)
+		{
+			ApplyTo(TargetObject, Log.Logger);
+		}
+
+		/// <summary>
+		/// Applies these arguments to fields with the [CommandLine] attribute in the given object.
+		/// </summary>
+		/// <param name="TargetObject">The object to configure</param>
+		/// <param name="Logger">Sink for error/warning messages</param>
+		public void ApplyTo(object TargetObject, ILogger Logger)
 		{
 			// Build a mapping from name to field and attribute for this object
 			List<string> MissingArguments = new List<string>();
-			for(Type TargetType = TargetObject.GetType(); TargetType != typeof(object); TargetType = TargetType.BaseType)
+			for(Type TargetType = TargetObject.GetType(); TargetType != typeof(object); TargetType = TargetType.BaseType!)
 			{
 				foreach(FieldInfo FieldInfo in TargetType.GetFields(BindingFlags.Instance | BindingFlags.GetField | BindingFlags.GetProperty | BindingFlags.Public | BindingFlags.NonPublic | BindingFlags.DeclaredOnly))
 				{
 					// If any attribute is required, keep track of it so we can include an error for it
-					string RequiredPrefix = null;
+					string? RequiredPrefix = null;
 
 					// Keep track of whether a value has already been assigned to this field
-					string AssignedArgument = null;
+					string? AssignedArgument = null;
 
 					// Loop through all the attributes for different command line options that can modify it
 					IEnumerable<CommandLineAttribute> Attributes = FieldInfo.GetCustomAttributes<CommandLineAttribute>();
@@ -744,7 +758,7 @@ namespace EpicGames.Core
 								// Apply the value to the field
 								if(Attribute.ListSeparator == 0)
 								{
-									if(ApplyArgument(TargetObject, FieldInfo, Argument, ValueText, AssignedArgument))
+									if(ApplyArgument(TargetObject, FieldInfo, Argument, ValueText, AssignedArgument, Logger))
 									{
 										AssignedArgument = Argument;
 									}
@@ -753,7 +767,7 @@ namespace EpicGames.Core
 								{
 									foreach(string ItemValueText in ValueText.Split(Attribute.ListSeparator))
 									{
-										if(ApplyArgument(TargetObject, FieldInfo, Argument, ItemValueText, AssignedArgument))
+										if(ApplyArgument(TargetObject, FieldInfo, Argument, ItemValueText, AssignedArgument, Logger))
 										{
 											AssignedArgument = Argument;
 										}
@@ -802,7 +816,7 @@ namespace EpicGames.Core
 		public static List<KeyValuePair<string, string>> GetParameters(Type Type)
 		{
 			List<KeyValuePair<string, string>> Parameters = new List<KeyValuePair<string, string>>();
-			for (Type TargetType = Type; TargetType != typeof(object); TargetType = TargetType.BaseType)
+			for (Type? TargetType = Type; TargetType != typeof(object) && TargetType != null; TargetType = TargetType.BaseType)
 			{
 				foreach (FieldInfo FieldInfo in TargetType.GetFields(BindingFlags.Instance | BindingFlags.GetField | BindingFlags.GetProperty | BindingFlags.Public | BindingFlags.NonPublic | BindingFlags.DeclaredOnly))
 				{
@@ -998,6 +1012,14 @@ namespace EpicGames.Core
 		/// Checks that there are no unused arguments (and warns if there are)
 		/// </summary>
 		public void CheckAllArgumentsUsed()
+		{
+			CheckAllArgumentsUsed(Log.Logger);
+		}
+
+		/// <summary>
+		/// Checks that there are no unused arguments (and warns if there are)
+		/// </summary>
+		public void CheckAllArgumentsUsed(ILogger Logger)
 		{
 			// Find all the unused arguments
 			List<string> RemainingArguments = new List<string>();
@@ -1014,11 +1036,11 @@ namespace EpicGames.Core
 			{
 				if(RemainingArguments.Count == 1)
 				{
-					Log.TraceWarning("Invalid argument: {0}", RemainingArguments[0]);
+					Logger.LogWarning(String.Format("Invalid argument: {0}", RemainingArguments[0]));
 				}
 				else
 				{
-					Log.TraceWarning("Invalid arguments:\n{0}", String.Join("\n", RemainingArguments));
+					Logger.LogWarning(String.Format("Invalid arguments:\n{0}", String.Join("\n", RemainingArguments)));
 				}
 			}
 		}
@@ -1051,14 +1073,15 @@ namespace EpicGames.Core
 		/// <param name="ArgumentText">The full argument text</param>
 		/// <param name="ValueText">Argument text</param>
 		/// <param name="PreviousArgumentText">The previous text used to configure this field</param>
+		/// <param name="Logger">Logger for error/warning messages</param>
 		/// <returns>True if the value was assigned to the field, false otherwise</returns>
-		private static bool ApplyArgument(object TargetObject, FieldInfo Field, string ArgumentText, string ValueText, string PreviousArgumentText)
+		private static bool ApplyArgument(object TargetObject, FieldInfo Field, string ArgumentText, string ValueText, string? PreviousArgumentText, ILogger Logger)
 		{
 			// The value type for items of this field
 			Type ValueType = Field.FieldType;
 
 			// Check if the field type implements ICollection<>. If so, we can take multiple values.
-			Type CollectionType = null;
+			Type? CollectionType = null;
 			foreach (Type InterfaceType in Field.FieldType.GetInterfaces())
 			{
 				if (InterfaceType.IsGenericType && InterfaceType.GetGenericTypeDefinition() == typeof(ICollection<>))
@@ -1070,10 +1093,10 @@ namespace EpicGames.Core
 			}
 
 			// Try to parse the value
-			object Value;
+			object? Value;
 			if(!TryParseValue(ValueType, ValueText, out Value))
 			{
-				Log.TraceWarning("Unable to parse value for argument '{0}'.", ArgumentText);
+				Logger.LogWarning("Unable to parse value for argument '{0}'.", ArgumentText);
 				return false;
 			}
 
@@ -1083,10 +1106,10 @@ namespace EpicGames.Core
 				// Check if this field has already been assigned to. Output a warning if the previous value is in conflict with the new one.
 				if(PreviousArgumentText != null)
 				{
-					object PreviousValue = Field.GetValue(TargetObject);
-					if(!PreviousValue.Equals(Value))
+					object? PreviousValue = Field.GetValue(TargetObject);
+					if(!Object.Equals(PreviousValue, Value))
 					{
-						Log.TraceWarning("Argument '{0}' conflicts with '{1}'; ignoring.", ArgumentText, PreviousArgumentText);
+						Logger.LogWarning("Argument '{0}' conflicts with '{1}'; ignoring.", ArgumentText, PreviousArgumentText);
 					}
 					return false;
 				}
@@ -1110,7 +1133,7 @@ namespace EpicGames.Core
 		/// <param name="Text">The value text</param>
 		/// <param name="Value">On success, contains the parsed object</param>
 		/// <returns>True if the text could be parsed, false otherwise</returns>
-		private static bool TryParseValue(Type FieldType, string Text, out object Value)
+		private static bool TryParseValue(Type FieldType, string Text, [NotNullWhen(true)] out object? Value)
 		{
 			if (FieldType.IsGenericType && FieldType.GetGenericTypeDefinition() == typeof(Nullable<>))
 			{
diff --git a/Engine/Source/Programs/Shared/EpicGames.Core/DirectoryReference.cs b/Engine/Source/Programs/Shared/EpicGames.Core/DirectoryReference.cs
--- a/Engine/Source/Programs/Shared/EpicGames.Core/DirectoryReference.cs
+++ b/Engine/Source/Programs/Shared/EpicGames.Core/DirectoryReference.cs
@@ -2,13 +2,12 @@
 
 using System;
 using System.Collections.Generic;
+using System.Diagnostics.CodeAnalysis;
 using System.IO;
 using System.Linq;
 using System.Text;
 using System.Threading.Tasks;
 
-#nullable disable
-
 namespace EpicGames.Core
 {
 	/// <summary>
@@ -92,7 +91,7 @@ namespace EpicGames.Core
 		/// Gets the directory containing this object
 		/// </summary>
 		/// <returns>A new directory object representing the directory containing this object</returns>
-		public DirectoryReference ParentDirectory
+		public DirectoryReference? ParentDirectory
 		{
 			get
 			{
@@ -139,7 +138,7 @@ namespace EpicGames.Core
 		/// </summary>
 		/// <param name="Folder">The folder to receive the path for</param>
 		/// <returns>Directory reference for the given folder, or null if it is not available</returns>
-		public static DirectoryReference GetSpecialFolder(Environment.SpecialFolder Folder)
+		public static DirectoryReference? GetSpecialFolder(Environment.SpecialFolder Folder)
 		{
 			string FolderPath = Environment.GetFolderPath(Folder);
 			return String.IsNullOrEmpty(FolderPath)? null : new DirectoryReference(FolderPath);
@@ -172,15 +171,15 @@ namespace EpicGames.Core
 		/// <param name="A">First object to compare.</param>
 		/// <param name="B">Second object to compare.</param>
 		/// <returns>True if the names represent the same object, false otherwise</returns>
-		public static bool operator ==(DirectoryReference A, DirectoryReference B)
+		public static bool operator ==(DirectoryReference? A, DirectoryReference? B)
 		{
-			if ((object)A == null)
+			if ((object?)A == null)
 			{
-				return (object)B == null;
+				return (object?)B == null;
 			}
 			else
 			{
-				return (object)B != null && A.FullName.Equals(B.FullName, Comparison);
+				return (object?)B != null && A.FullName.Equals(B.FullName, Comparison);
 			}
 		}
 
@@ -190,7 +189,7 @@ namespace EpicGames.Core
 		/// <param name="A">First object to compare.</param>
 		/// <param name="B">Second object to compare.</param>
 		/// <returns>False if the names represent the same object, true otherwise</returns>
-		public static bool operator !=(DirectoryReference A, DirectoryReference B)
+		public static bool operator !=(DirectoryReference? A, DirectoryReference? B)
 		{
 			return !(A == B);
 		}
@@ -200,7 +199,7 @@ namespace EpicGames.Core
 		/// </summary>
 		/// <param name="Obj">other instance to compare.</param>
 		/// <returns>True if the names represent the same object, false otherwise</returns>
-		public override bool Equals(object Obj)
+		public override bool Equals(object? Obj)
 		{
 			return (Obj is DirectoryReference) && ((DirectoryReference)Obj) == this;
 		}
@@ -210,7 +209,7 @@ namespace EpicGames.Core
 		/// </summary>
 		/// <param name="Obj">other instance to compare.</param>
 		/// <returns>True if the names represent the same object, false otherwise</returns>
-		public bool Equals(DirectoryReference Obj)
+		public bool Equals(DirectoryReference? Obj)
 		{
 			return Obj == this;
 		}
@@ -239,7 +238,8 @@ namespace EpicGames.Core
 		/// </summary>
 		/// <param name="File">The file to create a directory reference for</param>
 		/// <returns>The directory containing the file  </returns>
-		public static DirectoryReference FromFile(FileReference File)
+		[return: NotNullIfNotNull("File")]
+		public static DirectoryReference? FromFile(FileReference? File)
 		{
 			if(File == null)
 			{
@@ -256,7 +256,7 @@ namespace EpicGames.Core
 		/// </summary>
 		/// <param name="DirectoryName">Path for the new object</param>
 		/// <returns>Returns a FileReference representing the given string, or null.</returns>
-		public static DirectoryReference FromString(string DirectoryName)
+		public static DirectoryReference? FromString(string? DirectoryName)
 		{
 			if(String.IsNullOrEmpty(DirectoryName))
 			{
@@ -453,6 +453,16 @@ namespace EpicGames.Core
 		/// <param name="Reader">Binary reader to read from</param>
 		/// <returns>New DirectoryReference object</returns>
 		public static DirectoryReference ReadDirectoryReference(this BinaryReader Reader)
+		{
+			return BinaryArchiveReader.NotNull(ReadDirectoryReferenceOrNull(Reader));
+		}
+
+		/// <summary>
+		/// Manually deserialize a directory reference from a binary stream.
+		/// </summary>
+		/// <param name="Reader">Binary reader to read from</param>
+		/// <returns>New DirectoryReference object</returns>
+		public static DirectoryReference? ReadDirectoryReferenceOrNull(this BinaryReader Reader)
 		{
 			string FullName = Reader.ReadString();
 			return (FullName.Length == 0) ? null : new DirectoryReference(FullName, DirectoryReference.Sanitize.None);
@@ -465,7 +475,17 @@ namespace EpicGames.Core
 		/// <param name="Directory">The item to write</param>
 		public static void WriteDirectoryReference(this BinaryArchiveWriter Writer, DirectoryReference Directory)
 		{
-			if(Directory == null)
+			Writer.WriteNullableDirectoryReference(Directory);
+		}
+
+		/// <summary>
+		/// Writes a directory reference  to a binary archive
+		/// </summary>
+		/// <param name="Writer">The writer to output data to</param>
+		/// <param name="Directory">The item to write</param>
+		public static void WriteNullableDirectoryReference(this BinaryArchiveWriter Writer, DirectoryReference? Directory)
+		{
+			if (Directory == null)
 			{
 				Writer.WriteString(null);
 			}
@@ -481,9 +501,19 @@ namespace EpicGames.Core
 		/// <param name="Reader">Reader to serialize data from</param>
 		/// <returns>New directory reference instance</returns>
 		public static DirectoryReference ReadDirectoryReference(this BinaryArchiveReader Reader)
+		{
+			return BinaryArchiveReader.NotNull(ReadDirectoryReferenceOrNull(Reader));
+		}
+
+		/// <summary>
+		/// Reads a directory reference from a binary archive
+		/// </summary>
+		/// <param name="Reader">Reader to serialize data from</param>
+		/// <returns>New directory reference instance</returns>
+		public static DirectoryReference? ReadDirectoryReferenceOrNull(this BinaryArchiveReader Reader)
 		{
 			string FullName = Reader.ReadString();
-			if(FullName == null)
+			if (FullName == null)
 			{
 				return null;
 			}
diff --git a/Engine/Source/Programs/Shared/EpicGames.Core/ExceptionUtils.cs b/Engine/Source/Programs/Shared/EpicGames.Core/ExceptionUtils.cs
--- a/Engine/Source/Programs/Shared/EpicGames.Core/ExceptionUtils.cs
+++ b/Engine/Source/Programs/Shared/EpicGames.Core/ExceptionUtils.cs
@@ -74,11 +74,14 @@ namespace EpicGames.Core
 			StringBuilder ErrorMessage = new StringBuilder();
 			if (Ex is AggregateException)
 			{
-				Exception InnerException = ((AggregateException)Ex).InnerException!;
-				ErrorMessage.Append(InnerException.ToString());
-				foreach (string Line in GetContext(InnerException))
+				Exception? InnerException = ((AggregateException)Ex).InnerException;
+				if (InnerException != null)
 				{
-					ErrorMessage.AppendFormat("\n  {0}", Line);
+					ErrorMessage.Append(InnerException.ToString());
+					foreach (string Line in GetContext(InnerException))
+					{
+						ErrorMessage.AppendFormat("\n  {0}", Line);
+					}
 				}
 			}
 			else
@@ -115,7 +118,16 @@ namespace EpicGames.Core
 				{
 					foreach (object? Key in CurrentEx.Data.Keys)
 					{
-						object Value = CurrentEx.Data[Key!]!;
+						if (Key == null)
+						{
+							continue;
+						}
+
+						object? Value = CurrentEx.Data[Key];
+						if (Value == null)
+						{
+							continue;
+						}
 
 						string ValueString;
 						if(Value is List<string>)
@@ -124,7 +136,7 @@ namespace EpicGames.Core
 						}
 						else
 						{
-							ValueString = Value.ToString()!;
+							ValueString = Value.ToString() ?? String.Empty;
 						}
 
 						Message.AppendFormat("   data: {0} = {1}", Key, ValueString);
diff --git a/Engine/Source/Programs/Shared/EpicGames.Core/FileReference.cs b/Engine/Source/Programs/Shared/EpicGames.Core/FileReference.cs
--- a/Engine/Source/Programs/Shared/EpicGames.Core/FileReference.cs
+++ b/Engine/Source/Programs/Shared/EpicGames.Core/FileReference.cs
@@ -2,13 +2,13 @@
 
 using System;
 using System.Collections.Generic;
+using System.Diagnostics.CodeAnalysis;
 using System.IO;
 using System.Linq;
 using System.Text;
+using System.Threading;
 using System.Threading.Tasks;
 
-#nullable disable
-
 namespace EpicGames.Core
 {
 	/// <summary>
@@ -62,7 +62,8 @@ namespace EpicGames.Core
 		/// </summary>
 		/// <param name="FileName">FileName for the string</param>
 		/// <returns>Returns a FileReference representing the given string, or null.</returns>
-		public static FileReference FromString(string FileName)
+		[return: NotNullIfNotNull("FileName")]
+		public static FileReference? FromString(string? FileName)
 		{
 			if(String.IsNullOrEmpty(FileName))
 			{
@@ -125,7 +126,7 @@ namespace EpicGames.Core
 		/// </summary>
 		/// <param name="Extension">The new extension</param>
 		/// <returns>A FileReference with the same path and name, but with the new extension</returns>
-		public FileReference ChangeExtension(string Extension)
+		public FileReference ChangeExtension(string? Extension)
 		{
 			string NewFullName = Path.ChangeExtension(FullName, Extension);
 			return new FileReference(NewFullName, Sanitize.None);
@@ -186,15 +187,15 @@ namespace EpicGames.Core
 		/// <param name="A">First object to compare.</param>
 		/// <param name="B">Second object to compare.</param>
 		/// <returns>True if the names represent the same object, false otherwise</returns>
-		public static bool operator ==(FileReference A, FileReference B)
+		public static bool operator ==(FileReference? A, FileReference? B)
 		{
-			if ((object)A == null)
+			if ((object?)A == null)
 			{
-				return (object)B == null;
+				return (object?)B == null;
 			}
 			else
 			{
-				return (object)B != null && A.FullName.Equals(B.FullName, Comparison);
+				return (object?)B != null && A.FullName.Equals(B.FullName, Comparison);
 			}
 		}
 
@@ -204,7 +205,7 @@ namespace EpicGames.Core
 		/// <param name="A">First object to compare.</param>
 		/// <param name="B">Second object to compare.</param>
 		/// <returns>False if the names represent the same object, true otherwise</returns>
-		public static bool operator !=(FileReference A, FileReference B)
+		public static bool operator !=(FileReference? A, FileReference? B)
 		{
 			return !(A == B);
 		}
@@ -214,7 +215,7 @@ namespace EpicGames.Core
 		/// </summary>
 		/// <param name="Obj">other instance to compare.</param>
 		/// <returns>True if the names represent the same object, false otherwise</returns>
-		public override bool Equals(object Obj)
+		public override bool Equals(object? Obj)
 		{
 			return (Obj is FileReference) && ((FileReference)Obj) == this;
 		}
@@ -224,7 +225,7 @@ namespace EpicGames.Core
 		/// </summary>
 		/// <param name="Obj">other instance to compare.</param>
 		/// <returns>True if the names represent the same object, false otherwise</returns>
-		public bool Equals(FileReference Obj)
+		public bool Equals(FileReference? Obj)
 		{
 			return Obj == this;
 		}
@@ -364,6 +365,17 @@ namespace EpicGames.Core
 			File.Move(SourceLocation.FullName, TargetLocation.FullName);
 		}
 
+		/// <summary>
+		/// Moves a file from one location to another
+		/// </summary>
+		/// <param name="SourceLocation">Location of the source file</param>
+		/// <param name="TargetLocation">Location of the target file</param>
+		/// <param name="Overwrite">Whether to overwrite the file in the target location</param>
+		public static void Move(FileReference SourceLocation, FileReference TargetLocation, bool Overwrite)
+		{
+			File.Move(SourceLocation.FullName, TargetLocation.FullName, Overwrite);
+		}
+
 		/// <summary>
 		/// Opens a FileStream on the specified path with read/write access
 		/// </summary>
@@ -410,6 +422,16 @@ namespace EpicGames.Core
 			return File.ReadAllBytes(Location.FullName);
 		}
 
+		/// <summary>
+		/// Reads the contents of a file
+		/// </summary>
+		/// <param name="Location">Location of the file</param>
+		/// <returns>Byte array containing the contents of the file</returns>
+		public static Task<byte[]> ReadAllBytesAsync(FileReference Location, CancellationToken CancellationToken = default(CancellationToken))
+		{
+			return File.ReadAllBytesAsync(Location.FullName, CancellationToken);
+		}
+
 		/// <summary>
 		/// Reads the contents of a file
 		/// </summary>
@@ -431,6 +453,27 @@ namespace EpicGames.Core
 			return File.ReadAllText(Location.FullName, Encoding);
 		}
 
+		/// <summary>
+		/// Reads the contents of a file
+		/// </summary>
+		/// <param name="Location">Location of the file</param>
+		/// <returns>Contents of the file as a single string</returns>
+		public static Task<string> ReadAllTextAsync(FileReference Location, CancellationToken CancellationToken = default(CancellationToken))
+		{
+			return File.ReadAllTextAsync(Location.FullName, CancellationToken);
+		}
+
+		/// <summary>
+		/// Reads the contents of a file
+		/// </summary>
+		/// <param name="Location">Location of the file</param>
+		/// <param name="Encoding">Encoding of the file</param>
+		/// <returns>Contents of the file as a single string</returns>
+		public static Task<string> ReadAllTextAsync(FileReference Location, Encoding Encoding, CancellationToken CancellationToken = default(CancellationToken))
+		{
+			return File.ReadAllTextAsync(Location.FullName, Encoding, CancellationToken);
+		}
+
 		/// <summary>
 		/// Reads the contents of a file
 		/// </summary>
@@ -452,6 +495,27 @@ namespace EpicGames.Core
 			return File.ReadAllLines(Location.FullName, Encoding);
 		}
 
+		/// <summary>
+		/// Reads the contents of a file
+		/// </summary>
+		/// <param name="Location">Location of the file</param>
+		/// <returns>String array containing the contents of the file</returns>
+		public static Task<string[]> ReadAllLinesAsync(FileReference Location, CancellationToken CancellationToken = default(CancellationToken))
+		{
+			return File.ReadAllLinesAsync(Location.FullName, CancellationToken);
+		}
+
+		/// <summary>
+		/// Reads the contents of a file
+		/// </summary>
+		/// <param name="Location">Location of the file</param>
+		/// <param name="Encoding">The encoding to use when parsing the file</param>
+		/// <returns>String array containing the contents of the file</returns>
+		public static Task<string[]> ReadAllLinesAsync(FileReference Location, Encoding Encoding, CancellationToken CancellationToken = default(CancellationToken))
+		{
+			return File.ReadAllLinesAsync(Location.FullName, Encoding, CancellationToken);
+		}
+
 		/// <summary>
 		/// Sets the attributes for a file
 		/// </summary>
@@ -512,6 +576,16 @@ namespace EpicGames.Core
 			File.WriteAllBytes(Location.FullName, Contents);
 		}
 
+		/// <summary>
+		/// Writes the contents of a file
+		/// </summary>
+		/// <param name="Location">Location of the file</param>
+		/// <param name="Contents">Contents of the file</param>
+		public static Task WriteAllBytesAsync(FileReference Location, byte[] Contents, CancellationToken CancellationToken = default(CancellationToken))
+		{
+			return File.WriteAllBytesAsync(Location.FullName, Contents, CancellationToken);
+		}
+
 		/// <summary>
 		/// Writes the data to the given file, if it's different from what's there already
 		/// </summary>
@@ -572,6 +646,48 @@ namespace EpicGames.Core
 			File.WriteAllLines(Location.FullName, Contents, Encoding);
 		}
 
+		/// <summary>
+		/// Writes the contents of a file
+		/// </summary>
+		/// <param name="Location">Location of the file</param>
+		/// <param name="Contents">Contents of the file</param>
+		public static Task WriteAllLinesAsync(FileReference Location, IEnumerable<string> Contents, CancellationToken CancellationToken = default(CancellationToken))
+		{
+			return File.WriteAllLinesAsync(Location.FullName, Contents, CancellationToken);
+		}
+
+		/// <summary>
+		/// Writes the contents of a file
+		/// </summary>
+		/// <param name="Location">Location of the file</param>
+		/// <param name="Contents">Contents of the file</param>
+		/// <param name="Encoding">The encoding to use when parsing the file</param>
+		public static Task WriteAllLinesAsync(FileReference Location, IEnumerable<string> Contents, Encoding Encoding, CancellationToken CancellationToken = default(CancellationToken))
+		{
+			return File.WriteAllLinesAsync(Location.FullName, Contents, Encoding, CancellationToken);
+		}
+
+		/// <summary>
+		/// Writes the contents of a file
+		/// </summary>
+		/// <param name="Location">Location of the file</param>
+		/// <param name="Contents">Contents of the file</param>
+		public static Task WriteAllLinesAsync(FileReference Location, string[] Contents, CancellationToken CancellationToken = default(CancellationToken))
+		{
+			return File.WriteAllLinesAsync(Location.FullName, Contents, CancellationToken);
+		}
+
+		/// <summary>
+		/// Writes the contents of a file
+		/// </summary>
+		/// <param name="Location">Location of the file</param>
+		/// <param name="Contents">Contents of the file</param>
+		/// <param name="Encoding">The encoding to use when parsing the file</param>
+		public static Task WriteAllLinesAsync(FileReference Location, string[] Contents, Encoding Encoding, CancellationToken CancellationToken = default(CancellationToken))
+		{
+			return File.WriteAllLinesAsync(Location.FullName, Contents, Encoding, CancellationToken);
+		}
+
 		/// <summary>
 		/// Writes the contents of a file
 		/// </summary>
@@ -705,6 +821,16 @@ namespace EpicGames.Core
 		/// <param name="Reader">Binary reader to read from</param>
 		/// <returns>New FileReference object</returns>
 		public static FileReference ReadFileReference(this BinaryReader Reader)
+		{
+			return BinaryArchiveReader.NotNull(ReadFileReferenceOrNull(Reader));
+		}
+
+		/// <summary>
+		/// Manually deserialize a file reference from a binary stream.
+		/// </summary>
+		/// <param name="Reader">Binary reader to read from</param>
+		/// <returns>New FileReference object</returns>
+		public static FileReference? ReadFileReferenceOrNull(this BinaryReader Reader)
 		{
 			string FullName = Reader.ReadString();
 			return (FullName.Length == 0) ? null : new FileReference(FullName, FileReference.Sanitize.None);
@@ -717,6 +843,17 @@ namespace EpicGames.Core
 		/// <param name="UniqueFiles">List of previously read file references. The index into this array is used in place of subsequent ocurrences of the file.</param>
 		/// <returns>The file reference that was read</returns>
 		public static FileReference ReadFileReference(this BinaryReader Reader, List<FileReference> UniqueFiles)
+		{
+			return BinaryArchiveReader.NotNull(ReadFileReferenceOrNull(Reader, UniqueFiles));
+		}
+
+		/// <summary>
+		/// Deserializes a file reference, using a lookup table to avoid writing the same name more than once.
+		/// </summary>
+		/// <param name="Reader">The source to read from</param>
+		/// <param name="UniqueFiles">List of previously read file references. The index into this array is used in place of subsequent ocurrences of the file.</param>
+		/// <returns>The file reference that was read</returns>
+		public static FileReference? ReadFileReferenceOrNull(this BinaryReader Reader, List<FileReference> UniqueFiles)
 		{
 			int UniqueId = Reader.ReadInt32();
 			if (UniqueId == -1)
@@ -740,7 +877,7 @@ namespace EpicGames.Core
 		/// </summary>
 		/// <param name="Writer">The writer to output data to</param>
 		/// <param name="File">The file reference to write</param>
-		public static void WriteFileReference(this BinaryArchiveWriter Writer, FileReference File)
+		public static void WriteFileReference(this BinaryArchiveWriter Writer, FileReference? File)
 		{
 			if(File == null)
 			{
@@ -758,6 +895,16 @@ namespace EpicGames.Core
 		/// <param name="Reader">Reader to serialize data from</param>
 		/// <returns>New file reference instance</returns>
 		public static FileReference ReadFileReference(this BinaryArchiveReader Reader)
+		{
+			return BinaryArchiveReader.NotNull(ReadFileReferenceOrNull(Reader));
+		}
+
+		/// <summary>
+		/// Reads a FileReference from a binary archive
+		/// </summary>
+		/// <param name="Reader">Reader to serialize data from</param>
+		/// <returns>New file reference instance</returns>
+		public static FileReference? ReadFileReferenceOrNull(this BinaryArchiveReader Reader)
 		{
 			string FullName = Reader.ReadString();
 			if(FullName == null)
diff --git a/Engine/Source/Programs/Shared/EpicGames.Core/ForwardingLogger.cs b/Engine/Source/Programs/Shared/EpicGames.Core/ForwardingLogger.cs
new file mode 100644
--- /dev/null
+++ b/Engine/Source/Programs/Shared/EpicGames.Core/ForwardingLogger.cs
@@ -0,0 +1,102 @@
+// Copyright Epic Games, Inc. All Rights Reserved.
+
+using Microsoft.Extensions.Logging;
+using System;
+using System.Collections.Generic;
+using System.Linq;
+using System.Text;
+
+namespace EpicGames.Core
+{
+	/// <summary>
+	/// Logger which forwards output to several other loggers
+	/// </summary>
+	public class ForwardingLogger : ILogger
+	{
+		/// <summary>
+		/// List of IDisposable instances
+		/// </summary>
+		class DisposableList : List<IDisposable>, IDisposable
+		{
+			/// <inheritdoc/>
+			public void Dispose()
+			{
+				foreach(IDisposable Child in this)
+				{
+					Child.Dispose();
+				}
+			}
+		}
+
+		/// <summary>
+		/// The loggers that are being written to
+		/// </summary>
+		public ILogger[] Loggers
+		{
+			get;
+		}
+
+		/// <summary>
+		/// Constructor
+		/// </summary>
+		/// <param name="Loggers">Array of loggers to forward to</param>
+		public ForwardingLogger(params ILogger[] Loggers)
+		{
+			this.Loggers = Loggers;
+		}
+
+		/// <inheritdoc/>
+		public IDisposable? BeginScope<TState>(TState State)
+		{
+			IDisposable? Disposable = null;
+
+			DisposableList? Scope = null;
+			foreach (ILogger Logger in Loggers)
+			{
+				IDisposable? Next = Logger.BeginScope(State);
+				if(Next != null)
+				{
+					if (Scope != null)
+					{
+						Scope.Add(Next);
+					}
+					else if (Disposable != null)
+					{
+						Scope = new DisposableList();
+						Scope.Add(Disposable);
+						Scope.Add(Next);
+						Disposable = Scope;
+					}
+					else
+					{
+						Disposable = Next;
+					}
+				}
+			}
+
+			return Disposable;
+		}
+
+		/// <inheritdoc/>
+		public bool IsEnabled(LogLevel LogLevel)
+		{
+			foreach(ILogger Logger in Loggers)
+			{
+				if(Logger.IsEnabled(LogLevel))
+				{
+					return true;
+				}
+			}
+			return false;
+		}
+
+		/// <inheritdoc/>
+		public void Log<TState>(LogLevel LogLevel, EventId EventId, TState State, Exception Exception, Func<TState, Exception, string> Formatter)
+		{
+			foreach (ILogger Logger in Loggers)
+			{
+				Logger.Log(LogLevel, EventId, State, Exception, Formatter);
+			}
+		}
+	}
+}
diff --git a/Engine/Source/Programs/Shared/EpicGames.Core/HelpUtils.cs b/Engine/Source/Programs/Shared/EpicGames.Core/HelpUtils.cs
--- a/Engine/Source/Programs/Shared/EpicGames.Core/HelpUtils.cs
+++ b/Engine/Source/Programs/Shared/EpicGames.Core/HelpUtils.cs
@@ -1,5 +1,6 @@
 // Copyright Epic Games, Inc. All Rights Reserved.
 
+using Microsoft.Extensions.Logging;
 using System;
 using System.Collections.Generic;
 using System.ComponentModel;
@@ -45,9 +46,9 @@ namespace EpicGames.Core
 		/// Prints help for the given object type
 		/// </summary>
 		/// <param name="Type">Type to print help for</param>
-		public static void PrintHelp(string Title, Type Type)
+		public static void PrintHelp(string Title, Type Type, ILogger Logger)
 		{
-			PrintHelp(Title, GetDescription(Type), CommandLineArguments.GetParameters(Type));
+			PrintHelp(Title, GetDescription(Type), CommandLineArguments.GetParameters(Type), WindowWidth, Logger);
 		}
 
 		/// <summary>
@@ -56,14 +57,14 @@ namespace EpicGames.Core
 		/// <param name="Title">Title for the help text</param>
 		/// <param name="Description">Description for the command</param>
 		/// <param name="Parameters">List of parameters</param>
-		public static void PrintHelp(string Title, string Description, List<KeyValuePair<string, string>> Parameters)
+		/// <param name="MaxWidth">Maximum width for each line</param>
+		/// <param name="Logger">The output logger</param>
+		public static void PrintHelp(string Title, string Description, List<KeyValuePair<string, string>> Parameters, int MaxWidth, ILogger Logger)
 		{
-			int MaxWidth = WindowWidth;
-
 			bool bFirstLine = true;
 			if (!String.IsNullOrEmpty(Title))
 			{
-				PrintParagraph(Title, MaxWidth);
+				PrintParagraph(Title, MaxWidth, Logger);
 				bFirstLine = false;
 			}
 
@@ -71,9 +72,9 @@ namespace EpicGames.Core
 			{
 				if (!bFirstLine)
 				{
-					Log.TraceInformation("");
+					Logger.LogInformation("");
 				}
-				PrintParagraph(Description, MaxWidth);
+				PrintParagraph(Description, MaxWidth, Logger);
 				bFirstLine = false;
 			}
 
@@ -81,11 +82,11 @@ namespace EpicGames.Core
 			{
 				if (!bFirstLine)
 				{
-					Log.TraceInformation("");
+					Logger.LogInformation("");
 				}
 
-				Log.TraceInformation("Parameters:");
-				PrintTable(Parameters, 4, 24, MaxWidth);
+				Logger.LogInformation("Parameters:");
+				PrintTable(Parameters, 4, 24, MaxWidth, Logger);
 			}
 		}
 
@@ -112,10 +113,10 @@ namespace EpicGames.Core
 		/// Prints a paragraph of text using word wrapping
 		/// </summary>
 		/// <param name="Text">Text to print</param>
-		/// <param name="MaxWidth">Maximum width for each line</param>
-		public static void PrintParagraph(string Text)
+		/// <param name="Logger">Logger implementation to write to</param>
+		public static void PrintParagraph(string Text, ILogger Logger)
 		{
-			PrintParagraph(Text, WindowWidth - 1);
+			PrintParagraph(Text, WindowWidth - 1, Logger);
 		}
 
 		/// <summary>
@@ -123,12 +124,13 @@ namespace EpicGames.Core
 		/// </summary>
 		/// <param name="Text">Text to print</param>
 		/// <param name="MaxWidth">Maximum width for each line</param>
-		public static void PrintParagraph(string Text, int MaxWidth)
+		/// <param name="Logger">Logger implementation to write to</param>
+		public static void PrintParagraph(string Text, int MaxWidth, ILogger Logger)
 		{
 			List<string> Lines = StringUtils.WordWrap(Text, MaxWidth);
 			foreach (string Line in Lines)
 			{
-				Log.TraceInformation("{0}", Line);
+				Logger.LogInformation(Line);
 			}
 		}
 
@@ -143,11 +145,12 @@ namespace EpicGames.Core
 		/// </summary>
 		/// <param name="Items">List of parameters arranged as "-ParamName Param Description"</param>
 		/// <param name="Indent">Indent from the left hand side</param>
-		/// <param name="DefaultRightPadding">The minimum padding from the start of the param name to the start of the description (resizes with larger param names)</param>
+		/// <param name="MinFirstColumnWidth">The minimum padding from the start of the param name to the start of the description (resizes with larger param names)</param>
+		/// <param name="Logger">Logger implementation to write to</param>
 		/// <returns></returns>
-		public static void PrintTable(List<KeyValuePair<string, string>> Items, int Indent, int MinFirstColumnWidth)
+		public static void PrintTable(List<KeyValuePair<string, string>> Items, int Indent, int MinFirstColumnWidth, ILogger Logger)
 		{
-			PrintTable(Items, Indent, MinFirstColumnWidth, WindowWidth - 1);
+			PrintTable(Items, Indent, MinFirstColumnWidth, WindowWidth - 1, Logger);
 		}
 
 		/// <summary>
@@ -161,9 +164,10 @@ namespace EpicGames.Core
 		/// </summary>
 		/// <param name="Items">List of parameters arranged as "-ParamName Param Description"</param>
 		/// <param name="Indent">Indent from the left hand side</param>
-		/// <param name="DefaultRightPadding">The minimum padding from the start of the param name to the start of the description (resizes with larger param names)</param>
+		/// <param name="MinFirstColumnWidth">The minimum padding from the start of the param name to the start of the description (resizes with larger param names)</param>
+		/// <param name="Logger">Logger implementation to write to</param>
 		/// <returns></returns>
-		public static void PrintTable(List<KeyValuePair<string, string>> Items, int Indent, int MinFirstColumnWidth, int MaxWidth)
+		public static void PrintTable(List<KeyValuePair<string, string>> Items, int Indent, int MinFirstColumnWidth, int MaxWidth, ILogger Logger)
 		{
 			if(Items.Count > 0)
 			{
@@ -186,7 +190,7 @@ namespace EpicGames.Core
 					{
 						// Formatting as following:
 						// <Indent>-param<Right Padding>Description<New line>
-						Log.TraceInformation("{0}{1}", ParamString, DescriptionLine);
+						Logger.LogInformation(ParamString + DescriptionLine);
 
 						// we replace the param string on subsequent lines with white space of the same length
 						ParamString = string.Empty.PadRight(IndentString.Length + RightPadding);
diff --git a/Engine/Source/Programs/Shared/EpicGames.Core/ILoggerIndent.cs b/Engine/Source/Programs/Shared/EpicGames.Core/ILoggerIndent.cs
new file mode 100644
--- /dev/null
+++ b/Engine/Source/Programs/Shared/EpicGames.Core/ILoggerIndent.cs
@@ -0,0 +1,239 @@
+// Copyright Epic Games, Inc. All Rights Reserved.
+
+using Microsoft.Extensions.Logging;
+using System;
+using System.Collections;
+using System.Collections.Generic;
+using System.Linq;
+using System.Text;
+using System.Threading.Tasks;
+
+namespace EpicGames.Core
+{
+	/// <summary>
+	/// Class to apply a log indent for the lifetime of an object 
+	/// </summary>
+	public interface ILoggerIndent
+	{
+		/// <summary>
+		/// The indent to apply
+		/// </summary>
+		string Indent { get; }
+	}
+
+	/// <summary>
+	/// Wrapper class for ILogger classes which supports LoggerStatusScope
+	/// </summary>
+	public class DefaultLoggerIndentHandler : ILogger
+	{
+		/// <summary>
+		/// Scoped indent message
+		/// </summary>
+		class Scope : IDisposable
+		{
+			/// <summary>
+			/// Owning object
+			/// </summary>
+			DefaultLoggerIndentHandler Owner;
+
+			/// <summary>
+			/// The indent scope object
+			/// </summary>
+			public ILoggerIndent Indent { get; }
+
+			/// <summary>
+			/// Constructor
+			/// </summary>
+			public Scope(DefaultLoggerIndentHandler Owner, ILoggerIndent Indent)
+			{
+				this.Owner = Owner;
+				this.Indent = Indent;
+
+				lock (Owner.Scopes)
+				{
+					Owner.Scopes.Add(this);
+				}
+			}
+
+			/// <summary>
+			/// Remove this indent from the list
+			/// </summary>
+			public void Dispose()
+			{
+				lock (Owner.Scopes)
+				{
+					Owner.Scopes.Remove(this);
+				}
+			}
+		}
+
+		/// <summary>
+		/// Struct to wrap a formatted set of log values with applied indent
+		/// </summary>
+		/// <typeparam name="TState">Arbitrary type parameter</typeparam>
+		public struct FormattedLogValues<TState> : IEnumerable<KeyValuePair<string, object>>
+		{
+			/// <summary>
+			/// The indent to apply
+			/// </summary>
+			string Indent;
+
+			/// <summary>
+			/// The inner state
+			/// </summary>
+			TState State;
+
+			/// <summary>
+			/// Formatter for the inner state
+			/// </summary>
+			Func<TState, Exception, string> Formatter;
+
+			/// <summary>
+			/// Constructor
+			/// </summary>
+			/// <param name="Indent">The indent to apply</param>
+			/// <param name="State">The inner state</param>
+			/// <param name="Formatter">Formatter for the inner state</param>
+			public FormattedLogValues(string Indent, TState State, Func<TState, Exception, string> Formatter)
+			{
+				this.Indent = Indent;
+				this.State = State;
+				this.Formatter = Formatter;
+			}
+
+			/// <inheritdoc/>
+			public IEnumerator<KeyValuePair<string, object>> GetEnumerator()
+			{
+				IEnumerable<KeyValuePair<string, object>>? InnerEnumerable = State as IEnumerable<KeyValuePair<string, object>>;
+				if (InnerEnumerable != null)
+				{
+					foreach (KeyValuePair<string, object> Pair in InnerEnumerable)
+					{
+						if (Pair.Key.Equals("{OriginalFormat}", StringComparison.Ordinal))
+						{
+							yield return new KeyValuePair<string, object>(Pair.Key, Indent + Pair.Value.ToString());
+						}
+						else
+						{
+							yield return Pair;
+						}
+					}
+				}
+			}
+
+			/// <inheritdoc/>
+			IEnumerator IEnumerable.GetEnumerator()
+			{
+				throw new NotImplementedException();
+			}
+
+			/// <summary>
+			/// Formats an instance of this object
+			/// </summary>
+			/// <param name="Values">The object instance</param>
+			/// <param name="Exception">The exception to format</param>
+			/// <returns>The formatted string</returns>
+			public static string Format(FormattedLogValues<TState> Values, Exception Exception)
+			{
+				return Values.Indent + Values.Formatter(Values.State, Exception);
+			}
+		}
+
+		/// <summary>
+		/// The internal logger
+		/// </summary>
+		ILogger Inner;
+
+		/// <summary>
+		/// Current list of indents
+		/// </summary>
+		List<Scope> Scopes = new List<Scope>();
+
+		/// <summary>
+		/// The current indent text
+		/// </summary>
+		public string Indent
+		{
+			get;
+			private set;
+		}
+
+		/// <summary>
+		/// Constructor
+		/// </summary>
+		/// <param name="Inner">The logger to wrap</param>
+		public DefaultLoggerIndentHandler(ILogger Inner)
+		{
+			this.Inner = Inner;
+			this.Indent = "";
+		}
+
+		/// <inheritdoc/>
+		public IDisposable BeginScope<TState>(TState State)
+		{
+			ILoggerIndent? Indent = State as ILoggerIndent;
+			if (Indent != null)
+			{
+				return new Scope(this, Indent);
+			}
+
+			return Inner.BeginScope(State);
+		}
+
+		/// <inheritdoc/>
+		public bool IsEnabled(LogLevel LogLevel)
+		{
+			return Inner.IsEnabled(LogLevel);
+		}
+
+		/// <inheritdoc/>
+		public void Log<TState>(LogLevel LogLevel, EventId EventId, TState State, Exception Exception, Func<TState, Exception, string> Formatter)
+		{
+			if (Scopes.Count > 0)
+			{
+				string Indent = String.Join("", Scopes.Select(x => x.Indent.Indent));
+				Inner.Log(LogLevel, EventId, new FormattedLogValues<TState>(Indent, State, Formatter), Exception, FormattedLogValues<TState>.Format);
+				return;
+			}
+
+			Inner.Log(LogLevel, EventId, State, Exception, Formatter);
+		}
+	}
+
+	/// <summary>
+	/// Extension methods for creating an indent
+	/// </summary>
+	public static class LoggerIndentExtensions
+	{
+		/// <summary>
+		/// Class to apply a log indent for the lifetime of an object 
+		/// </summary>
+		class LoggerIndent : ILoggerIndent
+		{
+			/// <summary>
+			/// The previous indent
+			/// </summary>
+			public string Indent { get; }
+
+			/// <summary>
+			/// Constructor
+			/// </summary>
+			/// <param name="Indent">Indent to append to the existing indent</param>
+			public LoggerIndent(string Indent)
+			{
+				this.Indent = Indent;
+			}
+		}
+
+		/// <summary>
+		/// Create an indent
+		/// </summary>
+		/// <param name="Logger">Logger interface</param>
+		/// <param name="Indent">The indent to apply</param>
+		/// <returns>Disposable object</returns>
+		public static IDisposable BeginIndentScope(this ILogger Logger, string Indent)
+		{
+			return Logger.BeginScope(new LoggerIndent(Indent));
+		}
+	}
+}
diff --git a/Engine/Source/Programs/Shared/EpicGames.Core/ILoggerProgress.cs b/Engine/Source/Programs/Shared/EpicGames.Core/ILoggerProgress.cs
new file mode 100644
--- /dev/null
+++ b/Engine/Source/Programs/Shared/EpicGames.Core/ILoggerProgress.cs
@@ -0,0 +1,143 @@
+// Copyright Epic Games, Inc. All Rights Reserved.
+
+using Microsoft.Extensions.Logging;
+using System;
+using System.Collections.Generic;
+using System.Diagnostics;
+using System.Linq;
+using System.Text;
+using System.Threading.Tasks;
+
+namespace EpicGames.Core
+{
+	/// <summary>
+	/// Manages a status message for a long running operation, which can be updated with progress. Typically transient on consoles, and written to the same line.
+	/// </summary>
+	public interface ILoggerProgress : IDisposable
+	{
+		/// <summary>
+		/// Prefix message for the status
+		/// </summary>
+		string Message
+		{
+			get;
+		}
+
+		/// <summary>
+		/// The current 
+		/// </summary>
+		string Progress
+		{
+			get; set;
+		}
+	}
+
+	/// <summary>
+	/// Extension methods for status objects
+	/// </summary>
+	public static class LoggerProgressExtensions
+	{
+		/// <summary>
+		/// Concrete implementation of <see cref="ILoggerProgress"/>
+		/// </summary>
+		class LoggerProgress : ILoggerProgress
+		{
+			/// <summary>
+			/// The logger to output to
+			/// </summary>
+			private ILogger Logger;
+
+			/// <summary>
+			/// Prefix message for the status
+			/// </summary>
+			public string Message
+			{
+				get;
+			}
+
+			/// <summary>
+			/// The current 
+			/// </summary>
+			public string Progress
+			{
+				get { return ProgressInternal; }
+				set
+				{
+					ProgressInternal = value;
+
+					if (Timer.Elapsed > TimeSpan.FromSeconds(3.0))
+					{
+						LastOutput = String.Empty;
+						Flush();
+						Timer.Restart();
+					}
+				}
+			}
+
+			/// <summary>
+			/// The last string that was output
+			/// </summary>
+			string LastOutput = String.Empty;
+
+			/// <summary>
+			/// Backing storage for the Progress string
+			/// </summary>
+			string ProgressInternal = String.Empty;
+
+			/// <summary>
+			/// Timer since the last update
+			/// </summary>
+			Stopwatch Timer = Stopwatch.StartNew();
+
+			/// <summary>
+			/// Constructor
+			/// </summary>
+			/// <param name="Logger">The logger to write to</param>
+			/// <param name="Message">The base message to display</param>
+			public LoggerProgress(ILogger Logger, string Message)
+			{
+				this.Logger = Logger;
+				this.Message = Message;
+				this.ProgressInternal = String.Empty;
+
+				Logger.LogInformation(Message);
+			}
+
+			/// <summary>
+			/// Dispose of this object
+			/// </summary>
+			public void Dispose()
+			{
+				Flush();
+			}
+
+			/// <summary>
+			/// Flushes the current output to the log
+			/// </summary>
+			void Flush()
+			{
+				string Output = Message;
+				if (!String.IsNullOrEmpty(Progress))
+				{
+					Output += $" {Progress}";
+				}
+				if (!String.Equals(Output, LastOutput, StringComparison.Ordinal))
+				{
+					Logger.LogInformation(Output);
+					LastOutput = Output;
+				}
+			}
+		}
+
+		/// <summary>
+		/// Begins a new progress scope
+		/// </summary>
+		/// <param name="Logger">The logger being written to</param>
+		/// <param name="Message">The message prefix</param>
+		/// <returns>Scope object, which should be disposed when finished</returns>
+		public static ILoggerProgress BeginProgressScope(this ILogger Logger, string Message)
+		{
+			return new LoggerProgress(Logger, Message);
+		}
+	}
+}
diff --git a/Engine/Source/Programs/Shared/EpicGames.Core/MemoryReader.cs b/Engine/Source/Programs/Shared/EpicGames.Core/MemoryReader.cs
new file mode 100644
--- /dev/null
+++ b/Engine/Source/Programs/Shared/EpicGames.Core/MemoryReader.cs
@@ -0,0 +1,188 @@
+// Copyright Epic Games, Inc. All Rights Reserved.
+
+using System;
+using System.Buffers.Binary;
+using System.Collections.Generic;
+using System.Text;
+
+namespace EpicGames.Core
+{
+	/// <summary>
+	/// Reads data from a memory buffer
+	/// </summary>
+	public class MemoryReader
+	{
+		/// <summary>
+		/// The memory to read from
+		/// </summary>
+		public ReadOnlyMemory<byte> Memory
+		{
+			get;
+		}
+
+		/// <summary>
+		/// Offset to read from
+		/// </summary>
+		public int Offset
+		{
+			get; set;
+		}
+
+		/// <summary>
+		/// Length of the data
+		/// </summary>
+		public int Length => Memory.Length;
+
+		/// <summary>
+		/// Returns the memory at the current offset
+		/// </summary>
+		public ReadOnlySpan<byte> Span
+		{
+			get { return Memory.Span.Slice(Offset); }
+		}
+
+		/// <summary>
+		/// Constructor
+		/// </summary>
+		/// <param name="Memory">The memory to read from</param>
+		public MemoryReader(ReadOnlyMemory<byte> Memory)
+		{
+			this.Memory = Memory;
+			this.Offset = 0;
+		}
+
+		/// <summary>
+		/// Checks that we've used the exact buffer length
+		/// </summary>
+		/// <param name="ExpectedOffset">Expected offset within the output buffer</param>
+		public void CheckOffset(int ExpectedOffset)
+		{
+			if (Offset != ExpectedOffset)
+			{
+				throw new Exception($"Serialization is not at expected offset within the input buffer (expected {ExpectedOffset}, actual {Offset})");
+			}
+		}
+
+		/// <summary>
+		/// Read a boolean from the buffer
+		/// </summary>
+		/// <returns>The value read from the buffer</returns>
+		public bool ReadBoolean()
+		{
+			return ReadUInt8() != 0;
+		}
+
+		/// <summary>
+		/// Reads a byte from the buffer
+		/// </summary>
+		/// <returns>The value read from the buffer</returns>
+		public sbyte ReadInt8()
+		{
+			sbyte Value = (sbyte)Memory.Span[Offset];
+			Offset++;
+			return Value;
+		}
+
+		/// <summary>
+		/// Reads a byte from the buffer
+		/// </summary>
+		/// <returns>The value read from the buffer</returns>
+		public byte ReadUInt8()
+		{
+			byte Value = Memory.Span[Offset];
+			Offset++;
+			return Value;
+		}
+
+		/// <summary>
+		/// Reads an int32
+		/// </summary>
+		/// <returns>The value that was read from the buffer</returns>
+		public int ReadInt32()
+		{
+			int Value = BinaryPrimitives.ReadInt32LittleEndian(Memory.Span.Slice(Offset));
+			Offset += sizeof(int);
+			return Value;
+		}
+
+		/// <summary>
+		/// Reads a uint32
+		/// </summary>
+		/// <returns>The value that was read from the buffer</returns>
+		public uint ReadUInt32()
+		{
+			uint Value = BinaryPrimitives.ReadUInt32LittleEndian(Memory.Span.Slice(Offset));
+			Offset += sizeof(uint);
+			return Value;
+		}
+
+		/// <summary>
+		/// Reads an int64
+		/// </summary>
+		/// <returns>The value that was read from the buffer</returns>
+		public long ReadInt64()
+		{
+			long Value = BinaryPrimitives.ReadInt64LittleEndian(Memory.Span.Slice(Offset));
+			Offset += sizeof(long);
+			return Value;
+		}
+
+		/// <summary>
+		/// Reads an int64
+		/// </summary>
+		/// <returns>The value that was read from the buffer</returns>
+		public ulong ReadUInt64()
+		{
+			ulong Value = BinaryPrimitives.ReadUInt64LittleEndian(Memory.Span.Slice(Offset));
+			Offset += sizeof(ulong);
+			return Value;
+		}
+
+		/// <summary>
+		/// Reads a sequence of bytes from the buffer
+		/// </summary>
+		/// <returns>Sequence of bytes</returns>
+		public ReadOnlyMemory<byte> ReadVariableLengthBytes()
+		{
+			int Length = ReadInt32();
+			return ReadFixedLengthBytes(Length);
+		}
+
+		/// <summary>
+		/// Reads a sequence of bytes from the buffer
+		/// </summary>
+		/// <param name="Length">Number of bytes to read</param>
+		/// <returns>Sequence of bytes</returns>
+		public ReadOnlyMemory<byte> ReadFixedLengthBytes(int Length)
+		{
+			ReadOnlyMemory<byte> Bytes = Memory.Slice(Offset, Length);
+			Offset += Length;
+			return Bytes;
+		}
+
+		/// <summary>
+		/// Reads a variable length array
+		/// </summary>
+		/// <param name="ReadItem">Delegate to write an individual item</param>
+		public T[] ReadVariableLengthArray<T>(Func<T> ReadItem)
+		{
+			int Length = ReadInt32();
+			return ReadFixedLengthArray(Length, ReadItem);
+		}
+
+		/// <summary>
+		/// Writes a fixed length array
+		/// </summary>
+		/// <param name="Length">Length of the array to read</param>
+		/// <param name="ReadItem">Delegate to read an individual item</param>
+		public T[] ReadFixedLengthArray<T>(int Length, Func<T> ReadItem)
+		{
+			T[] Array = new T[Length];
+			for (int Idx = 0; Idx < Length; Idx++)
+			{
+				Array[Idx] = ReadItem();
+			}
+			return Array;
+		}
+	}
+}
diff --git a/Engine/Source/Programs/Shared/EpicGames.Core/MemoryWriter.cs b/Engine/Source/Programs/Shared/EpicGames.Core/MemoryWriter.cs
new file mode 100644
--- /dev/null
+++ b/Engine/Source/Programs/Shared/EpicGames.Core/MemoryWriter.cs
@@ -0,0 +1,180 @@
+// Copyright Epic Games, Inc. All Rights Reserved.
+
+using System;
+using System.Buffers.Binary;
+using System.Collections.Generic;
+using System.Text;
+
+namespace EpicGames.Core
+{
+	/// <summary>
+	/// Writes into a fixed size memory block
+	/// </summary>
+	public class MemoryWriter
+	{
+		/// <summary>
+		/// The memory block to write to
+		/// </summary>
+		Memory<byte> Memory;
+
+		/// <summary>
+		/// Current offset within the memory
+		/// </summary>
+		int Offset;
+
+		/// <summary>
+		/// Constructor
+		/// </summary>
+		/// <param name="Memory">Memory to write to</param>
+		public MemoryWriter(Memory<byte> Memory)
+		{
+			this.Memory = Memory;
+			this.Offset = 0;
+		}
+
+		/// <summary>
+		/// Checks that we've used the exact buffer length
+		/// </summary>
+		/// <param name="ExpectedOffset">Expected offset within the output buffer</param>
+		public void CheckOffset(int ExpectedOffset)
+		{
+			if (Offset != ExpectedOffset)
+			{
+				throw new Exception("Serialization is not at expected offset within the output buffer");
+			}
+		}
+
+		/// <summary>
+		/// Writes a boolean to memory
+		/// </summary>
+		/// <param name="Value">Value to write</param>
+		public void WriteBoolean(bool Value)
+		{
+			WriteUInt8(Value ? (byte)1 : (byte)0);
+		}
+
+		/// <summary>
+		/// Writes a byte to memory
+		/// </summary>
+		/// <param name="Value">Value to write</param>
+		public void WriteInt8(sbyte Value)
+		{
+			Memory.Span[Offset] = (byte)Value;
+			Offset++;
+		}
+
+		/// <summary>
+		/// Writes a byte to memory
+		/// </summary>
+		/// <param name="Value">Value to write</param>
+		public void WriteUInt8(byte Value)
+		{
+			Memory.Span[Offset] = Value;
+			Offset++;
+		}
+
+		/// <summary>
+		/// Writes an int16 to the memory
+		/// </summary>
+		/// <param name="Value">Value to write</param>
+		public void WriteInt16(short Value)
+		{
+			BinaryPrimitives.WriteInt16LittleEndian(Memory.Span.Slice(Offset), Value);
+			Offset += sizeof(short);
+		}
+
+		/// <summary>
+		/// Writes a uint16 to the memory
+		/// </summary>
+		/// <param name="Value">Value to write</param>
+		public void WriteUInt16(ushort Value)
+		{
+			BinaryPrimitives.WriteUInt16LittleEndian(Memory.Span.Slice(Offset), Value);
+			Offset += sizeof(ushort);
+		}
+
+		/// <summary>
+		/// Writes an int32 to the memory
+		/// </summary>
+		/// <param name="Value">Value to write</param>
+		public void WriteInt32(int Value)
+		{
+			BinaryPrimitives.WriteInt32LittleEndian(Memory.Span.Slice(Offset), Value);
+			Offset += sizeof(int);
+		}
+
+		/// <summary>
+		/// Writes a uint32 to the memory
+		/// </summary>
+		/// <param name="Value">Value to write</param>
+		public void WriteUInt32(uint Value)
+		{
+			BinaryPrimitives.WriteUInt32LittleEndian(Memory.Span.Slice(Offset), Value);
+			Offset += sizeof(uint);
+		}
+
+		/// <summary>
+		/// Writes an int64 to the memory
+		/// </summary>
+		/// <param name="Value">Value to write</param>
+		public void WriteInt64(long Value)
+		{
+			BinaryPrimitives.WriteInt64LittleEndian(Memory.Span.Slice(Offset), Value);
+			Offset += sizeof(long);
+		}
+
+		/// <summary>
+		/// Writes a uint64 to the memory
+		/// </summary>
+		/// <param name="Value">Value to write</param>
+		public void WriteUInt64(ulong Value)
+		{
+			BinaryPrimitives.WriteUInt64LittleEndian(Memory.Span.Slice(Offset), Value);
+			Offset += sizeof(ulong);
+		}
+
+		/// <summary>
+		/// Writes a variable length span of bytes
+		/// </summary>
+		/// <param name="Bytes">The bytes to write</param>
+		public void WriteVariableLengthBytes(ReadOnlySpan<byte> Bytes)
+		{
+			WriteInt32(Bytes.Length);
+			WriteFixedLengthBytes(Bytes);
+		}
+
+		/// <summary>
+		/// Write a fixed-length sequence of bytes to the buffer
+		/// </summary>
+		/// <param name="Bytes">The bytes to write</param>
+		public void WriteFixedLengthBytes(ReadOnlySpan<byte> Bytes)
+		{
+			Bytes.CopyTo(Memory.Span.Slice(Offset));
+			Offset += Bytes.Length;
+		}
+
+		/// <summary>
+		/// Writes a variable length array
+		/// </summary>
+		/// <param name="Array">The array to write</param>
+		/// <param name="WriteItem">Delegate to write an individual item</param>
+		public void WriteVariableLengthArray<T>(T[] Array, Action<T> WriteItem)
+		{
+			WriteInt32(Array.Length);
+			WriteFixedLengthArray(Array, WriteItem);
+		}
+
+		/// <summary>
+		/// Writes a fixed length array
+		/// </summary>
+		/// <param name="Array">The array to write</param>
+		/// <param name="WriteItem">Delegate to write an individual item</param>
+		public void WriteFixedLengthArray<T>(T[] Array, Action<T> WriteItem)
+		{
+			for (int Idx = 0; Idx < Array.Length; Idx++)
+			{
+				WriteItem(Array[Idx]);
+			}
+		}
+	}
+}
diff --git a/Engine/Source/Programs/Shared/EpicGames.Core/ParallelTask.cs b/Engine/Source/Programs/Shared/EpicGames.Core/ParallelTask.cs
--- a/Engine/Source/Programs/Shared/EpicGames.Core/ParallelTask.cs
+++ b/Engine/Source/Programs/Shared/EpicGames.Core/ParallelTask.cs
@@ -49,5 +49,29 @@ namespace EpicGames.Core
 
 			return Actions.Completion;
 		}
+
+		/// <summary>
+		/// Execute a large number of tasks in parallel over a collection. Parallel.ForEach() method isn't generally compatible with asynchronous programming, because any
+		/// exceptions are thrown on the 
+		/// </summary>
+		/// <typeparam name="T">The collection type</typeparam>
+		/// <param name="Collection">The collection to iterate over</param>
+		/// <param name="Action">Action to perform for each item in the collection</param>
+		/// <param name="MaxDegreeOfParallelism">Maximum degree of parallelism</param>
+		/// <returns>Async task</returns>
+		public static Task ForEachAsync<T>(IEnumerable<T> Collection, Func<T, Task> Action, int MaxDegreeOfParallelism)
+		{
+			ExecutionDataflowBlockOptions Options = new ExecutionDataflowBlockOptions();
+			Options.MaxDegreeOfParallelism = MaxDegreeOfParallelism;
+
+			ActionBlock<T> Actions = new ActionBlock<T>(Action, Options);
+			foreach (T Item in Collection)
+			{
+				Actions.Post(Item);
+			}
+			Actions.Complete();
+
+			return Actions.Completion;
+		}
 	}
 }
diff --git a/Engine/Source/Programs/Shared/EpicGames.Core/ProcessUtils.cs b/Engine/Source/Programs/Shared/EpicGames.Core/ProcessUtils.cs
new file mode 100644
--- /dev/null
+++ b/Engine/Source/Programs/Shared/EpicGames.Core/ProcessUtils.cs
@@ -0,0 +1,238 @@
+// Copyright Epic Games, Inc. All Rights Reserved.
+
+using Microsoft.Extensions.Logging;
+using Microsoft.Win32.SafeHandles;
+using System;
+using System.Collections.Generic;
+using System.Diagnostics;
+using System.Diagnostics.CodeAnalysis;
+using System.Linq;
+using System.Reflection;
+using System.Runtime.InteropServices;
+using System.Text;
+using System.Threading;
+using System.Threading.Tasks;
+
+namespace EpicGames.Core
+{
+	/// <summary>
+	/// Includes utilities for managing processes
+	/// </summary>
+	public static class ProcessUtils
+	{
+		const uint PROCESS_TERMINATE = 0x0001;
+		const uint PROCESS_QUERY_LIMITED_INFORMATION = 0x1000;
+
+		[DllImport("kernel32.dll", SetLastError = true)]
+		static extern SafeProcessHandle OpenProcess(uint dwDesiredAccess, bool bInheritHandle, uint processId);
+
+		[DllImport("kernel32.dll", SetLastError = true)]
+		[return: MarshalAs(UnmanagedType.Bool)]
+		static extern bool TerminateProcess(SafeProcessHandle hProcess, uint uExitCode);
+
+		[DllImport("kernel32.dll", SetLastError = true)]
+		static extern int QueryFullProcessImageName([In]SafeProcessHandle hProcess, [In]int dwFlags, [Out]StringBuilder lpExeName, ref int lpdwSize);
+
+		[DllImport("Psapi.dll", SetLastError = true)]
+		static extern bool EnumProcesses([MarshalAs(UnmanagedType.LPArray, ArraySubType = UnmanagedType.U4)] [In][Out] uint[] processIds, int arraySizeBytes, [MarshalAs(UnmanagedType.U4)] out int bytesCopied);
+
+		[DllImport("kernel32.dll")]
+		static extern uint WaitForMultipleObjects(int nCount, IntPtr[] lpHandles, bool bWaitAll, uint dwMilliseconds);
+
+		/// <summary>
+		/// Attempts to terminate all processes matching a predicate
+		/// </summary>
+		/// <param name="Predicate">The predicate for whether to terminate a process</param>
+		/// <param name="Logger">Logging device</param>
+		/// <param name="CancellationToken">Cancellation token to abort the search</param>
+		/// <returns>True if the process succeeded</returns>
+		public static bool TerminateProcesses(Predicate<FileReference> Predicate, ILogger Logger, CancellationToken CancellationToken)
+		{
+			if (RuntimeInformation.IsOSPlatform(OSPlatform.Windows))
+			{
+				return TerminateProcessesWin32(Predicate, Logger, CancellationToken);
+			}
+			else
+			{
+				return TerminateProcessesGenericPlatform(Predicate, Logger, CancellationToken);
+			}
+		}
+
+		/// <summary>
+		/// Attempts to terminate all processes matching a predicate
+		/// </summary>
+		/// <param name="Predicate">The predicate for whether to terminate a process</param>
+		/// <param name="Logger">Logging device</param>
+		/// <param name="CancellationToken">Cancellation token to abort the search</param>
+		/// <returns>True if the process succeeded</returns>
+		static bool TerminateProcessesWin32(Predicate<FileReference> Predicate, ILogger Logger, CancellationToken CancellationToken)
+		{
+			Dictionary<int, int> ProcessToCount = new Dictionary<int, int>();
+			for(; ;)
+			{
+				CancellationToken.ThrowIfCancellationRequested();
+
+				// Enumerate the processes
+				int NumProcessIds;
+				uint[] ProcessIds = new uint[512];
+				for (; ; )
+				{
+					int MaxBytes = ProcessIds.Length * sizeof(uint);
+					int NumBytes = 0;
+					if (!EnumProcesses(ProcessIds, MaxBytes, out NumBytes))
+					{
+						throw new Win32ExceptionWithCode("Unable to enumerate processes");
+					}
+					if (NumBytes < MaxBytes)
+					{
+						NumProcessIds = NumBytes / sizeof(uint);
+						break;
+					}
+					ProcessIds = new uint[ProcessIds.Length + 256];
+				}
+
+				// Find the processes to terminate
+				List<SafeProcessHandle> WaitHandles = new List<SafeProcessHandle>();
+				try
+				{
+					// Open each process in turn
+					foreach (uint ProcessId in ProcessIds)
+					{
+						SafeProcessHandle Handle = OpenProcess(PROCESS_QUERY_LIMITED_INFORMATION | PROCESS_TERMINATE, false, ProcessId);
+						try
+						{
+							if (!Handle.IsInvalid)
+							{
+								int Characters = 260;
+								StringBuilder Buffer = new StringBuilder(Characters);
+								if (QueryFullProcessImageName(Handle, 0, Buffer, ref Characters) > 0)
+								{
+									FileReference ImageFile = new FileReference(Buffer.ToString(0, Characters));
+									if (Predicate(ImageFile))
+									{
+										Logger.LogInformation("Terminating {ImageName} ({ProcessId})", ImageFile, ProcessId);
+										if (TerminateProcess(Handle, 9))
+										{
+											WaitHandles.Add(Handle);
+										}
+										else
+										{
+											Logger.LogInformation("Failed call to TerminateProcess ({Code})", Marshal.GetLastWin32Error());
+										}
+										Handle.SetHandleAsInvalid();
+									}
+								}
+							}
+						}
+						finally
+						{
+							Handle.Dispose();
+						}
+					}
+
+					// If there's nothing to do, exit immediately
+					if (WaitHandles.Count == 0)
+					{
+						return true;
+					}
+
+					// Wait for them all to complete
+					WaitForMultipleObjects(WaitHandles.Count, WaitHandles.Select(x => x.DangerousGetHandle()).ToArray(), true, 10 * 1000);
+				}
+				finally
+				{
+					foreach (SafeProcessHandle WaitHandle in WaitHandles)
+					{
+						WaitHandle.Close();
+					}
+				}
+			}
+		}
+
+		/// <summary>
+		/// Attempts to terminate all processes matching a predicate
+		/// </summary>
+		/// <param name="Predicate">The predicate for whether to terminate a process</param>
+		/// <param name="Logger">Logging device</param>
+		/// <param name="CancellationToken">Cancellation token to abort the search</param>
+		/// <returns>True if the process succeeded</returns>
+		static bool TerminateProcessesGenericPlatform(Predicate<FileReference> Predicate, ILogger Logger, CancellationToken CancellationToken)
+		{
+			bool Result = true;
+			Dictionary<(int, DateTime), int> ProcessToCount = new Dictionary<(int, DateTime), int>();
+			for (; ; )
+			{
+				CancellationToken.ThrowIfCancellationRequested();
+
+				bool bNoMatches = true;
+
+				// Enumerate all the processes
+				Process[] Processes = Process.GetProcesses();
+				foreach (Process Process in Processes)
+				{
+					// Attempt to get the image file. Ignore exceptions trying to fetch metadata for processes we don't have access to.
+					FileReference? ImageFile;
+					try
+					{
+						ImageFile = new FileReference(Process.MainModule.FileName);
+					}
+					catch
+					{
+						ImageFile = null;
+					}
+
+					// Test whether to terminate this process
+					if (ImageFile != null && Predicate(ImageFile))
+					{
+						// Get a unique id for the process, given that process ids are recycled
+						(int, DateTime) UniqueId;
+						try
+						{
+							UniqueId = (Process.Id, Process.StartTime);
+						}
+						catch
+						{
+							UniqueId = (Process.Id, DateTime.MinValue);
+						}
+
+						// Figure out whether to try and terminate this process
+						const int MaxCount = 5;
+						if (!ProcessToCount.TryGetValue(UniqueId, out int Count) || Count < MaxCount)
+						{
+							bNoMatches = false;
+							try
+							{
+								Logger.LogInformation("Terminating {ImageName} ({ProcessId})", ImageFile, Process.Id);
+								Process.Kill(true);
+								if (!Process.WaitForExit(5 * 1000))
+								{
+									Logger.LogInformation("Termination still pending; will retry...");
+								}
+							}
+							catch (Exception Ex)
+							{
+								Count++;
+								if (Count > 1)
+								{
+									Logger.LogInformation(Ex, "Exception while querying basic process info for pid {ProcessId}; will retry.", Process.Id);
+								}
+								else if (Count == MaxCount)
+								{
+									Logger.LogWarning(Ex, "Unable to terminate process {ImageFile} ({ProcessId}): {Message}", ImageFile, Process.Id, Ex.Message);
+									Result = false;
+								}
+								ProcessToCount[UniqueId] = Count;
+							}
+						}
+					}
+				}
+
+				// Return once we reach this point and haven't found anything else to terminate
+				if (bNoMatches)
+				{
+					return Result;
+				}
+			}
+		}
+	}
+}
diff --git a/Engine/Source/Programs/Shared/EpicGames.Core/ReadOnlyMemoryStream.cs b/Engine/Source/Programs/Shared/EpicGames.Core/ReadOnlyMemoryStream.cs
new file mode 100644
--- /dev/null
+++ b/Engine/Source/Programs/Shared/EpicGames.Core/ReadOnlyMemoryStream.cs
@@ -0,0 +1,82 @@
+using System;
+using System.Collections.Generic;
+using System.IO;
+using System.Text;
+
+namespace EpicGames.Core
+{
+	/// <summary>
+	/// Stream which reads from a <see cref="ReadOnlyMemory{byte}"/>
+	/// </summary>
+	public class ReadOnlyMemoryStream : Stream
+	{
+		/// <summary>
+		/// The buffer to read from
+		/// </summary>
+		ReadOnlyMemory<byte> Memory;
+
+		/// <summary>
+		/// Constructor
+		/// </summary>
+		/// <param name="Memory">The memory to read from</param>
+		public ReadOnlyMemoryStream(ReadOnlyMemory<byte> Memory)
+		{
+			this.Memory = Memory;
+		}
+
+		/// <inheritdoc/>
+		public override bool CanRead => true;
+
+		/// <inheritdoc/>
+		public override bool CanSeek => true;
+
+		/// <inheritdoc/>
+		public override bool CanWrite => false;
+
+		/// <inheritdoc/>
+		public override long Length => Memory.Length;
+
+		/// <inheritdoc/>
+		public override long Position { get; set; }
+
+		/// <inheritdoc/>
+		public override void Flush()
+		{
+		}
+
+		/// <inheritdoc/>
+		public override int Read(byte[] Buffer, int Offset, int Count)
+		{
+			int CopyLength = Math.Min(Count, (int)(Memory.Length - Position));
+			Memory.Slice((int)Position, CopyLength).CopyTo(Buffer.AsMemory(Offset, CopyLength));
+			Position += CopyLength;
+			return CopyLength;
+		}
+
+		/// <inheritdoc/>
+		public override long Seek(long Offset, SeekOrigin Origin)
+		{
+			switch (Origin)
+			{
+				case SeekOrigin.Begin:
+					Position = Offset;
+					break;
+				case SeekOrigin.Current:
+					Position += Offset;
+					break;
+				case SeekOrigin.End:
+					Position = Memory.Length + Offset;
+					break;
+				default:
+					throw new ArgumentException(null, nameof(Origin));
+			}
+			return Position;
+		}
+
+		/// <inheritdoc/>
+		public override void SetLength(long value) => throw new InvalidOperationException();
+
+		/// <inheritdoc/>
+		public override void Write(byte[] buffer, int offset, int count) => throw new InvalidOperationException();
+	}
+}
diff --git a/Engine/Source/Programs/Shared/EpicGames.Core/ReadOnlyUtf8String.cs b/Engine/Source/Programs/Shared/EpicGames.Core/ReadOnlyUtf8String.cs
--- a/Engine/Source/Programs/Shared/EpicGames.Core/ReadOnlyUtf8String.cs
+++ b/Engine/Source/Programs/Shared/EpicGames.Core/ReadOnlyUtf8String.cs
@@ -286,4 +286,44 @@ namespace EpicGames.Core
 		/// <inheritdoc/>
 		public abstract int GetHashCode(ReadOnlyUtf8String String);
 	}
+
+	/// <summary>
+	/// Extension methods for ReadOnlyUtf8String objects
+	/// </summary>
+	public static class MemoryWriterExtensions
+	{
+		/// <summary>
+		/// Reads a null-terminated utf8 string from the buffer
+		/// </summary>
+		/// <returns>The string data</returns>
+		public static ReadOnlyUtf8String ReadString(this MemoryReader Reader)
+		{
+			ReadOnlySpan<byte> Span = Reader.Span;
+			int Length = Span.IndexOf((byte)0);
+			ReadOnlyUtf8String Value = new ReadOnlyUtf8String(Reader.ReadFixedLengthBytes(Length));
+			Reader.ReadInt8();
+			return Value;
+		}
+
+		/// <summary>
+		/// Writes a UTF8 string into memory with a null terminator
+		/// </summary>
+		/// <param name="Writer">The memory writer to serialize to</param>
+		/// <param name="String">String to write</param>
+		public static void WriteString(this MemoryWriter Writer, ReadOnlyUtf8String String)
+		{
+			Writer.WriteFixedLengthBytes(String.Span);
+			Writer.WriteInt8(0);
+		}
+
+		/// <summary>
+		/// Determines the size of a serialized utf-8 string
+		/// </summary>
+		/// <param name="String">The string to measure</param>
+		/// <returns>Size of the serialized string</returns>
+		public static int GetSerializedSize(this ReadOnlyUtf8String String)
+		{
+			return String.Length + 1;
+		}
+	}
 }
diff --git a/Engine/Source/Programs/Shared/EpicGames.MongoDB/EpicGames.MongoDB.csproj b/Engine/Source/Programs/Shared/EpicGames.MongoDB/EpicGames.MongoDB.csproj
new file mode 100644
--- /dev/null
+++ b/Engine/Source/Programs/Shared/EpicGames.MongoDB/EpicGames.MongoDB.csproj
@@ -0,0 +1,31 @@
+<Project Sdk="Microsoft.NET.Sdk">
+  <PropertyGroup>
+    <TargetFramework>netcoreapp3.1</TargetFramework>
+    <Configuration Condition=" '$(Configuration)' == '' ">Development</Configuration>
+    <OutputType>Library</OutputType>
+    <RootNamespace>EpicGames.Core</RootNamespace>
+    <GenerateAssemblyInfo>false</GenerateAssemblyInfo>
+    <GenerateTargetFrameworkAttribute>false</GenerateTargetFrameworkAttribute>
+    <Configurations>Debug;Release;Development</Configurations>
+    <Nullable>enable</Nullable>
+  </PropertyGroup>
+  <PropertyGroup Condition=" '$(Configuration)|$(Platform)' == 'Debug|AnyCPU' ">
+    <DocumentationFile></DocumentationFile>
+  </PropertyGroup>
+  <PropertyGroup Condition="'$(Configuration)|$(Platform)' == 'Development|AnyCPU'">
+    <DocumentationFile></DocumentationFile>
+    <Optimize>true</Optimize>
+  </PropertyGroup>
+  <PropertyGroup Condition="'$(Configuration)|$(Platform)' == 'Release|AnyCPU'">
+    <DocumentationFile></DocumentationFile>
+    <Optimize>true</Optimize>
+  </PropertyGroup>
+  <ItemGroup>
+    <Compile Include="..\..\DotNETCommon\MetaData.cs">
+      <Link>Properties\MetaData.cs</Link>
+    </Compile>
+  </ItemGroup>
+  <ItemGroup>
+    <PackageReference Include="MongoDB.Driver" Version="2.11.5" />
+  </ItemGroup>
+</Project>
diff --git a/Engine/Source/Programs/Shared/EpicGames.MongoDB/MongoExtensions.cs b/Engine/Source/Programs/Shared/EpicGames.MongoDB/MongoExtensions.cs
new file mode 100644
--- /dev/null
+++ b/Engine/Source/Programs/Shared/EpicGames.MongoDB/MongoExtensions.cs
@@ -0,0 +1,140 @@
+// Copyright Epic Games, Inc. All Rights Reserved.
+
+using MongoDB.Bson;
+using MongoDB.Driver;
+using System;
+using System.Collections.Generic;
+using System.Linq;
+using System.Linq.Expressions;
+using System.Threading.Tasks;
+
+namespace EpicGames.MongoDB
+{
+	/// <summary>
+	/// Extension methods for MongoDB
+	/// </summary>
+	public static class MongoExtensions
+	{
+		/// <summary>
+		/// Filters the documents returned from a search
+		/// </summary>
+		/// <param name="Query">The query to filter</param>
+		/// <param name="Index">Index of the first document to return</param>
+		/// <param name="Count">Number of documents to return</param>
+		/// <returns>New query</returns>
+		public static IFindFluent<TDocument, TProjection> Range<TDocument, TProjection>(this IFindFluent<TDocument, TProjection> Query, int? Index, int? Count)
+		{
+			if (Index != null)
+			{
+				Query = Query.Skip(Index.Value);
+			}
+			if(Count != null)
+			{
+				Query = Query.Limit(Count.Value);
+			}
+			return Query;
+		}
+
+		/// <summary>
+		/// Filters the documents returned from a search
+		/// </summary>
+		/// <param name="Query">The query to filter</param>
+		/// <returns>New query</returns>
+		public static async Task<List<TResult>> ToListAsync<TDocument, TResult>(this IAsyncCursorSource<TDocument> Query) where TDocument : TResult
+		{
+			List<TResult> Results = new List<TResult>();
+			using (IAsyncCursor<TDocument> Cursor = await Query.ToCursorAsync())
+			{
+				while (await Cursor.MoveNextAsync())
+				{
+					foreach(TDocument Document in Cursor.Current)
+					{
+						Results.Add(Document);
+					}
+				}
+			}
+			return Results;
+		}
+
+		/// <summary>
+		/// Attempts to insert a document into a collection, handling the error case that a document with the given key already exists
+		/// </summary>
+		/// <typeparam name="TDocument"></typeparam>
+		/// <param name="Collection">Collection to insert into</param>
+		/// <param name="NewDocument">The document to insert</param>
+		/// <returns>True if the document was inserted, false if it already exists</returns>
+		public static async Task<bool> InsertOneIgnoreDuplicatesAsync<TDocument>(this IMongoCollection<TDocument> Collection, TDocument NewDocument)
+		{
+			try
+			{
+				await Collection.InsertOneAsync(NewDocument);
+				return true;
+			}
+			catch (MongoWriteException Ex)
+			{
+				if (Ex.WriteError.Category == ServerErrorCategory.DuplicateKey)
+				{
+					return false;
+				}
+				else
+				{
+					throw;
+				}
+			}
+		}
+
+		/// <summary>
+		/// Sets a field to a value, or unsets it if the value is null
+		/// </summary>
+		/// <typeparam name="TDocument">The document type</typeparam>
+		/// <typeparam name="TField">Type of the field to set</typeparam>
+		/// <param name="UpdateBuilder">Update builder</param>
+		/// <param name="Field">Expression for the field to set</param>
+		/// <param name="Value">New value to set</param>
+		/// <returns>Update defintiion</returns>
+		public static UpdateDefinition<TDocument> SetOrUnsetNull<TDocument, TField>(this UpdateDefinitionBuilder<TDocument> UpdateBuilder, Expression<Func<TDocument, TField?>> Field, TField? Value) where TField : struct
+		{
+			if (Value.HasValue)
+			{
+				return UpdateBuilder.Set(Field, Value.Value);
+			}
+			else
+			{
+				return UpdateBuilder.Unset(new ExpressionFieldDefinition<TDocument, TField?>(Field));
+			}
+		}
+
+		/// <summary>
+		/// Sets a field to a value, or unsets it if the value is null
+		/// </summary>
+		/// <typeparam name="TDocument">The document type</typeparam>
+		/// <typeparam name="TField">Type of the field to set</typeparam>
+		/// <param name="UpdateBuilder">Update builder</param>
+		/// <param name="Field">Expression for the field to set</param>
+		/// <param name="Value">New value to set</param>
+		/// <returns>Update defintiion</returns>
+		public static UpdateDefinition<TDocument> SetOrUnsetNullRef<TDocument, TField>(this UpdateDefinitionBuilder<TDocument> UpdateBuilder, Expression<Func<TDocument, TField?>> Field, TField? Value) where TField : class
+		{
+			if (Value != null)
+			{
+				return UpdateBuilder.Set(Field, Value);
+			}
+			else
+			{
+				return UpdateBuilder.Unset(new ExpressionFieldDefinition<TDocument, TField?>(Field));
+			}
+		}
+
+		/// <summary>
+		/// Creates a filter definition from a linq expression. This is not generally explicitly castable, so expose it as a FilterDefinitionBuilder method.
+		/// </summary>
+		/// <typeparam name="TDocument">The document type</typeparam>
+		/// <param name="Filter">The filter builder</param>
+		/// <param name="Expression">Expression to parse</param>
+		/// <returns>New filter definition</returns>
+		public static FilterDefinition<TDocument> Expr<TDocument>(this FilterDefinitionBuilder<TDocument> Filter, Expression<Func<TDocument, bool>> Expression)
+		{
+			return Expression;
+		}
+	}
+}
diff --git a/Engine/Source/Programs/Shared/EpicGames.MongoDB/ObjectIdHelpers.cs b/Engine/Source/Programs/Shared/EpicGames.MongoDB/ObjectIdHelpers.cs
new file mode 100644
--- /dev/null
+++ b/Engine/Source/Programs/Shared/EpicGames.MongoDB/ObjectIdHelpers.cs
@@ -0,0 +1,25 @@
+// Copyright Epic Games, Inc. All Rights Reserved.
+
+using MongoDB.Bson;
+using System;
+using System.Collections.Generic;
+using System.Linq;
+using System.Threading.Tasks;
+
+namespace EpicGames.MongoDB
+{
+	static class ObjectIdHelpers
+	{
+		public static ObjectId ToObjectId(this string Text)
+		{
+			if (Text.Length == 0)
+			{
+				return ObjectId.Empty;
+			}
+			else
+			{
+				return ObjectId.Parse(Text);
+			}
+		}
+	}
+}
diff --git a/Engine/Source/Programs/Shared/EpicGames.MongoDB/SingletonDocument.cs b/Engine/Source/Programs/Shared/EpicGames.MongoDB/SingletonDocument.cs
new file mode 100644
--- /dev/null
+++ b/Engine/Source/Programs/Shared/EpicGames.MongoDB/SingletonDocument.cs
@@ -0,0 +1,199 @@
+// Copyright Epic Games, Inc. All Rights Reserved.
+
+using MongoDB.Bson;
+using MongoDB.Bson.Serialization;
+using MongoDB.Bson.Serialization.Attributes;
+using MongoDB.Driver;
+using System;
+using System.Collections.Generic;
+using System.Linq;
+using System.Reflection;
+using System.Threading.Tasks;
+
+namespace EpicGames.MongoDB
+{
+	/// <summary>
+	/// Attribute specifying the unique id for a singleton document
+	/// </summary>
+	[AttributeUsage(AttributeTargets.Class)]
+	public class SingletonDocumentAttribute : Attribute
+	{
+		/// <summary>
+		/// Unique id for the singleton document
+		/// </summary>
+		public string Id { get; }
+
+		/// <summary>
+		/// Constructor
+		/// </summary>
+		/// <param name="Id">Unique id for the singleton document</param>
+		public SingletonDocumentAttribute(string Id)
+		{
+			this.Id = Id;
+		}
+	}
+
+	/// <summary>
+	/// Base class for singletons
+	/// </summary>
+	public class SingletonBase
+	{
+		/// <summary>
+		/// Unique id of the 
+		/// </summary>
+		public ObjectId Id { get; set; }
+
+		/// <summary>
+		/// Current revision number
+		/// </summary>
+		public int Revision { get; set; }
+
+		/// <summary>
+		/// Private constructor for serialization
+		/// </summary>
+		/// <param name="Unused"></param>
+		[BsonConstructor]
+		protected SingletonBase()
+		{
+		}
+
+		/// <summary>
+		/// Constructor
+		/// </summary>
+		/// <param name="Id">Unique id for the singleton</param>
+		public SingletonBase(ObjectId Id)
+		{
+			this.Id = Id;
+		}
+
+		/// <summary>
+		/// Base class for singleton documents
+		/// </summary>
+		class CachedId<T>
+		{
+			public static ObjectId Value { get; } = GetSingletonId();
+
+			static ObjectId GetSingletonId()
+			{
+				SingletonDocumentAttribute? Attribute = typeof(T).GetCustomAttribute<SingletonDocumentAttribute>();
+				if (Attribute == null)
+				{
+					throw new Exception($"Type {typeof(T).Name} is missing a {nameof(SingletonDocumentAttribute)} annotation");
+				}
+				return ObjectId.Parse(Attribute.Id);
+			}
+		}
+
+		/// <summary>
+		/// Gets the id for a singleton type
+		/// </summary>
+		/// <typeparam name="T"></typeparam>
+		/// <returns></returns>
+		public static ObjectId GetId<T>() where T : SingletonBase
+		{
+			return CachedId<T>.Value;
+		}
+	}
+
+	/// <summary>
+	/// Interface for the getting and setting the singleton
+	/// </summary>
+	/// <typeparam name="T">Type of document</typeparam>
+	public interface ISingletonDocument<T>
+	{
+		/// <summary>
+		/// Gets the current document
+		/// </summary>
+		/// <returns>The current document</returns>
+		Task<T> GetAsync();
+
+		/// <summary>
+		/// Attempts to update the document
+		/// </summary>
+		/// <param name="Value">New state of the document</param>
+		/// <returns>True if the document was updated, false otherwise</returns>
+		Task<bool> TryUpdateAsync(T Value);
+	}
+
+	/// <summary>
+	/// Concrete implementation of <see cref="ISingletonDocument{T}"/>
+	/// </summary>
+	/// <typeparam name="T">The document type</typeparam>
+	public class SingletonDocument<T> : ISingletonDocument<T> where T : SingletonBase, new()
+	{
+		/// <summary>
+		/// The database service instance
+		/// </summary>
+		IMongoCollection<T> Collection;
+
+		/// <summary>
+		/// Unique id for the singleton document
+		/// </summary>
+		ObjectId ObjectId;
+
+		/// <summary>
+		/// Static constructor. Registers the document using the automapper.
+		/// </summary>
+		static SingletonDocument()
+		{
+			BsonClassMap.RegisterClassMap<T>();
+		}
+
+		/// <summary>
+		/// Constructor
+		/// </summary>
+		/// <param name="DatabaseService">The database service instance</param>
+		public SingletonDocument(IMongoCollection<SingletonBase> Collection)
+		{
+			this.Collection = Collection.OfType<T>();
+
+			SingletonDocumentAttribute? Attribute = typeof(T).GetCustomAttribute<SingletonDocumentAttribute>();
+			if (Attribute == null)
+			{
+				throw new Exception($"Type {typeof(T).Name} is missing a {nameof(SingletonDocumentAttribute)} annotation");
+			}
+
+			ObjectId = new ObjectId(Attribute.Id);
+		}
+
+		/// <inheritdoc/>
+		public async Task<T> GetAsync()
+		{
+			for (; ; )
+			{
+				T? Object = await Collection.Find<T>(x => x.Id == ObjectId).FirstOrDefaultAsync();
+				if (Object != null)
+				{
+					return Object;
+				}
+
+				T NewItem = new T();
+				NewItem.Id = ObjectId;
+				await Collection.InsertOneAsync(NewItem);
+			}
+		}
+
+		/// <inheritdoc/>
+		public async Task<bool> TryUpdateAsync(T Value)
+		{
+			int PrevRevision = Value.Revision++;
+			try
+			{
+				ReplaceOneResult Result = await Collection.ReplaceOneAsync(x => x.Id == ObjectId && x.Revision == PrevRevision, Value, new ReplaceOptions { IsUpsert = true });
+				return Result.MatchedCount > 0;
+			}
+			catch (MongoWriteException Ex)
+			{
+				// Duplicate key error occurs if filter fails to match because revision is not the same.
+				if (Ex.WriteError != null && Ex.WriteError.Category == ServerErrorCategory.DuplicateKey)
+				{
+					return false;
+				}
+				else
+				{
+					throw;
+				}
+			}
+		}
+	}
+}
diff --git a/Engine/Source/Programs/Shared/EpicGames.MongoDB/TransactionBuilder.cs b/Engine/Source/Programs/Shared/EpicGames.MongoDB/TransactionBuilder.cs
new file mode 100644
--- /dev/null
+++ b/Engine/Source/Programs/Shared/EpicGames.MongoDB/TransactionBuilder.cs
@@ -0,0 +1,308 @@
+// Copyright Epic Games, Inc. All Rights Reserved.
+
+using MongoDB.Bson.Serialization;
+using MongoDB.Bson.Serialization.Serializers;
+using MongoDB.Driver;
+using System;
+using System.Collections.Generic;
+using System.Linq;
+using System.Linq.Expressions;
+using System.Reflection;
+using System.Runtime.CompilerServices;
+using System.Threading.Tasks;
+using System.Transactions;
+
+namespace EpicGames.MongoDB
+{
+	/// <summary>
+	/// Allows building a transactional MongoDB update of several fields, and applying those to an in-memory object. Useful for transactional updates.
+	/// </summary>
+	/// <typeparam name="TDocument">The type of document to update</typeparam>
+	public class TransactionBuilder<TDocument> where TDocument : class
+	{
+		/// <summary>
+		/// Information about an update to a field
+		/// </summary>
+		class FieldUpdate
+		{
+			/// <summary>
+			/// The update definition
+			/// </summary>
+			public UpdateDefinition<TDocument> Update;
+
+			/// <summary>
+			/// Applies an update to the given object
+			/// </summary>
+			public Action<object> Apply;
+
+			/// <summary>
+			/// Constructor
+			/// </summary>
+			/// <param name="Update">The update definition</param>
+			/// <param name="Apply">Applies the update to a given object</param>
+			public FieldUpdate(UpdateDefinition<TDocument> Update, Action<object> Apply)
+			{
+				this.Update = Update;
+				this.Apply = Apply;
+			}
+		}
+
+		/// <summary>
+		/// Helper class for determining the indexer for a particular type
+		/// </summary>
+		/// <typeparam name="T">The type to find the indexer for</typeparam>
+		static class Reflection<T>
+		{
+			/// <summary>
+			/// Caches the indexer for this type
+			/// </summary>
+			public static PropertyInfo Indexer { get; } = GetIndexerProperty();
+
+			/// <summary>
+			/// Finds the indexer property from a type
+			/// </summary>
+			/// <returns>The property info</returns>
+			static PropertyInfo GetIndexerProperty()
+			{
+				foreach (PropertyInfo Property in typeof(T).GetProperties())
+				{
+					if (Property.GetIndexParameters().Length > 0)
+					{
+						return Property;
+					}
+				}
+				return null!;
+			}
+		}
+
+		/// <summary>
+		/// List of field updates
+		/// </summary>
+		List<FieldUpdate> FieldUpdates = new List<FieldUpdate>();
+
+		/// <summary>
+		/// Whether the transaction is currenty empty
+		/// </summary>
+		public bool IsEmpty
+		{
+			get { return FieldUpdates.Count == 0; }
+		}
+
+		/// <summary>
+		/// Adds a setting to this transaction
+		/// </summary>
+		/// <typeparam name="TField">Type of the field to be updated</typeparam>
+		/// <param name="Expr">The expresion defining the field to update</param>
+		/// <param name="Value">New value for the field</param>
+		public void Set<TField>(Expression<Func<TDocument, TField>> Expr, TField Value)
+		{
+			UpdateDefinition<TDocument> Update = Builders<TDocument>.Update.Set(Expr, Value);
+			Action<object> Apply = Target => Assign(Target, Expr, Value);
+			FieldUpdates.Add(new FieldUpdate(Update, Apply));
+		}
+
+		/// <summary>
+		/// Adds an update to remove a field
+		/// </summary>
+		/// <param name="Expr">The expresion defining the field to update</param>
+		public void Unset(Expression<Func<TDocument, object>> Expr)
+		{
+			UpdateDefinition<TDocument> Update = Builders<TDocument>.Update.Unset(Expr);
+			Action<object> Apply = Target => Unassign(Target, Expr.Body);
+			FieldUpdates.Add(new FieldUpdate(Update, Apply));
+		}
+
+		/// <summary>
+		/// Updates a property dictionary with a set of adds and removes
+		/// </summary>
+		/// <param name="Expr">Lambda expression defining a field to update</param>
+		/// <param name="Updates">List of updates</param>
+		public void UpdateDictionary<TKey, TValue>(Expression<Func<TDocument, Dictionary<TKey, TValue>>> Expr, IEnumerable<KeyValuePair<TKey, TValue?>> Updates) where TKey : class where TValue : class
+		{
+			PropertyInfo IndexerProperty = Reflection<Dictionary<TKey, TValue>>.Indexer;
+			foreach (KeyValuePair<TKey, TValue?> Update in Updates)
+			{
+				MethodCallExpression Index = Expression.Call(Expr.Body, IndexerProperty.GetMethod, new[] { Expression.Constant(Update.Key) });
+				if (Update.Value == null)
+				{
+					Expression<Func<TDocument, object>> IndexExpression = Expression.Lambda<Func<TDocument, object>>(Index, Expr.Parameters[0]);
+					Unset(IndexExpression);
+				}
+				else
+				{
+					Expression<Func<TDocument, TValue>> IndexExpression = Expression.Lambda<Func<TDocument, TValue>>(Index, Expr.Parameters[0]);
+					Set(IndexExpression, Update.Value);
+				}
+			}
+		}
+
+		/// <summary>
+		/// Applies this transaction to the given update
+		/// </summary>
+		/// <param name="Target">The object to update</param>
+		public void ApplyTo(TDocument Target)
+		{
+			foreach (FieldUpdate Setting in FieldUpdates)
+			{
+				Setting.Apply(Target);
+			}
+		}
+
+		/// <summary>
+		/// Assign a value to a field
+		/// </summary>
+		/// <param name="Target">The target object to be updated</param>
+		/// <param name="FieldExpression">Lambda indicating the field to update</param>
+		/// <param name="Value">New value for the field</param>
+		static void Assign(object? Target, LambdaExpression FieldExpression, object? Value)
+		{
+			Expression Body = FieldExpression.Body;
+			if (Body.NodeType == ExpressionType.MemberAccess)
+			{
+				MemberExpression MemberExpression = (MemberExpression)Body;
+				object? Object = Evaluate(MemberExpression.Expression, Target);
+				PropertyInfo PropertyInfo = (PropertyInfo)MemberExpression.Member;
+				PropertyInfo.SetValue(Object, Value);
+			}
+			else if (Body.NodeType == ExpressionType.ArrayIndex)
+			{
+				BinaryExpression BinaryExpression = (BinaryExpression)Body;
+				System.Collections.IList List = (System.Collections.IList)Evaluate(BinaryExpression.Left, Target)!;
+				int Index = (int)Evaluate(BinaryExpression.Right, Target)!;
+				List[Index] = Value;
+			}
+			else if (Body.NodeType == ExpressionType.Call)
+			{
+				MethodCallExpression CallExpression = (MethodCallExpression)Body;
+				System.Collections.IDictionary Dictionary = (System.Collections.IDictionary)Evaluate(CallExpression.Object, Target)!;
+				object? Key = Evaluate(CallExpression.Arguments[0], null);
+				Dictionary[Key!] = Value;
+			}
+			else if (Body.NodeType == ExpressionType.Index)
+			{
+				IndexExpression IndexExpression = (IndexExpression)Body;
+				object? Object = Evaluate(IndexExpression.Object, Target);
+
+				object?[] Arguments = new object?[IndexExpression.Arguments.Count];
+				for (int Idx = 0; Idx < IndexExpression.Arguments.Count; Idx++)
+				{
+					Arguments[Idx] = Evaluate(IndexExpression.Arguments[Idx], null);
+				}
+
+				IndexExpression.Indexer.SetValue(Object, Value, Arguments);
+			}
+			else
+			{
+				throw new NotImplementedException();
+			}
+		}
+
+		/// <summary>
+		/// Removes an entry from a dictionary
+		/// </summary>
+		/// <param name="Target">The target object to be updated</param>
+		/// <param name="Body">Body of the expression indicating the field to update</param>
+		static void Unassign(object? Target, Expression Body)
+		{
+			if (Body.NodeType == ExpressionType.MemberAccess)
+			{
+				MemberExpression MemberExpression = (MemberExpression)Body;
+				switch (MemberExpression.Member.MemberType)
+				{
+					case MemberTypes.Field:
+						((FieldInfo)MemberExpression.Member).SetValue(Target, null);
+						break;
+					case MemberTypes.Property:
+						((PropertyInfo)MemberExpression.Member).SetValue(Target, null);
+						break;
+					default:
+						throw new NotImplementedException();
+				}
+			}
+			else if (Body.NodeType == ExpressionType.Call)
+			{
+				MethodCallExpression CallExpression = (MethodCallExpression)Body;
+				System.Collections.IDictionary? Dictionary = (System.Collections.IDictionary?)Evaluate(CallExpression.Object, Target);
+				object? Key = Evaluate(CallExpression.Arguments[0], null);
+				Dictionary!.Remove(Key!);
+			}
+			else if (Body.NodeType == ExpressionType.Index)
+			{
+				IndexExpression IndexExpression = (IndexExpression)Body;
+				System.Collections.IDictionary? Dictionary = (System.Collections.IDictionary?)Evaluate(IndexExpression.Object, Target);
+				object? Key = Evaluate(IndexExpression.Arguments[0], null);
+				Dictionary!.Remove(Key!);
+			}
+			else if (Body.NodeType == ExpressionType.Convert)
+			{
+				UnaryExpression UnaryExpression = (UnaryExpression)Body;
+				Unassign(Target, UnaryExpression.Operand);
+			}
+			else
+			{
+				throw new NotImplementedException();
+			}
+		}
+
+		/// <summary>
+		/// Evaluates an expression
+		/// </summary>
+		/// <param name="Expression">The expression to evaluate</param>
+		/// <param name="Parameter">Parameter to the unary lambda expression</param>
+		/// <returns>Value of the expression</returns>
+		static object? Evaluate(Expression Expression, object? Parameter)
+		{
+			if (Expression.NodeType == ExpressionType.Call)
+			{
+				MethodCallExpression CallExpression = (MethodCallExpression)Expression;
+				object? Object = Evaluate(CallExpression.Object, Parameter);
+				object?[] Arguments = CallExpression.Arguments.Select(x => Evaluate(x, Parameter)).ToArray();
+				return CallExpression.Method.Invoke(Object, Arguments);
+			}
+			else if (Expression.NodeType == ExpressionType.Constant)
+			{
+				ConstantExpression ConstantExpression = (ConstantExpression)Expression;
+				return ConstantExpression.Value;
+			}
+			else if (Expression.NodeType == ExpressionType.MemberAccess)
+			{
+				MemberExpression MemberExpression = (MemberExpression)Expression;
+				object? Target = Evaluate(MemberExpression.Expression, Parameter);
+
+				MemberInfo Member = MemberExpression.Member;
+				switch (Member.MemberType)
+				{
+					case MemberTypes.Property:
+						return ((PropertyInfo)Member).GetValue(Target);
+					case MemberTypes.Field:
+						return ((FieldInfo)Member).GetValue(Target);
+					default:
+						throw new NotImplementedException("Unsupported expression type");
+				}
+			}
+			else if (Expression.NodeType == ExpressionType.ArrayIndex)
+			{
+				BinaryExpression BinaryExpression = (BinaryExpression)Expression;
+				System.Collections.IList List = (System.Collections.IList)Evaluate(BinaryExpression.Left, Parameter)!;
+				int Index = (int)Evaluate(BinaryExpression.Right, null)!;
+				return List[Index];
+			}
+			else if (Expression.NodeType == ExpressionType.Parameter)
+			{
+				return Parameter;
+			}
+			else
+			{
+				throw new NotImplementedException();
+			}
+		}
+
+		/// <summary>
+		/// Converts this transaction to an update definition
+		/// </summary>
+		public UpdateDefinition<TDocument> ToUpdateDefinition()
+		{
+			return Builders<TDocument>.Update.Combine(FieldUpdates.Select(x => x.Update));
+		}
+	}
+}
diff --git a/Engine/Source/Programs/Shared/EpicGames.Perforce/PerforceConnection.cs b/Engine/Source/Programs/Shared/EpicGames.Perforce/PerforceConnection.cs
--- a/Engine/Source/Programs/Shared/EpicGames.Perforce/PerforceConnection.cs
+++ b/Engine/Source/Programs/Shared/EpicGames.Perforce/PerforceConnection.cs
@@ -2025,7 +2025,9 @@ namespace EpicGames.Perforce
 		/// <returns>Response from the server</returns>
 		public async Task<List<ReconcileRecord>> ReconcileAsync(int ChangeNumber, ReconcileOptions Options, string[] FileSpecs, CancellationToken CancellationToken)
 		{
-			return (await TryReconcileAsync(ChangeNumber, Options, FileSpecs, CancellationToken)).Data;
+			PerforceResponseList<ReconcileRecord> Records = await TryReconcileAsync(ChangeNumber, Options, FileSpecs, CancellationToken);
+			Records.RemoveAll(x => x.Info != null);
+			return Records.Data;
 		}
 
 		/// <summary>
@@ -2648,7 +2650,9 @@ namespace EpicGames.Perforce
 		/// <returns>Response from the server</returns>
 		public async Task<List<SyncSummaryRecord>> SyncQuietAsync(SyncOptions Options, int MaxFiles, string[] FileSpecs, CancellationToken CancellationToken)
 		{
-			return (await TrySyncQuietAsync(Options, MaxFiles, FileSpecs, CancellationToken)).Data;
+			PerforceResponseList<SyncSummaryRecord> Records = await TrySyncQuietAsync(Options, MaxFiles, FileSpecs, CancellationToken);
+			Records.RemoveAll(x => (x.Error != null && x.Error.Generic == PerforceGenericCode.Empty) || x.Info != null);
+			return Records.Data;
 		}
 
 		/// <summary>
@@ -2883,7 +2887,7 @@ namespace EpicGames.Perforce
 		/// <returns>Response from the server</returns>
 		public async Task<UserRecord> GetUserAsync(string UserName, CancellationToken CancellationToken)
 		{
-			return (await TryGetUserAsync(UserName, CancellationToken)).Data;
+			return (await TryGetUserAsync(UserName, CancellationToken))[0].Data;
 		}
 
 		/// <summary>
@@ -2892,11 +2896,11 @@ namespace EpicGames.Perforce
 		/// <param name="UserName">Name of the user to fetch information for</param>
 		/// <param name="CancellationToken">Token used to cancel the operation</param>
 		/// <returns>Response from the server</returns>
-		public Task<PerforceResponse<UserRecord>> TryGetUserAsync(string UserName, CancellationToken CancellationToken)
+		public Task<PerforceResponseList<UserRecord>> TryGetUserAsync(string UserName, CancellationToken CancellationToken)
 		{
 			StringBuilder Arguments = new StringBuilder("user");
 			Arguments.AppendFormat(" -o \"{0}\"", UserName);
-			return SingleResponseCommandAsync<UserRecord>(Arguments.ToString(), null, CancellationToken);
+			return CommandAsync<UserRecord>(Arguments.ToString(), null, CancellationToken);
 		}
 
 		#endregion
diff --git a/Engine/Source/Programs/Shared/EpicGames.Perforce/PerforceResponse.cs b/Engine/Source/Programs/Shared/EpicGames.Perforce/PerforceResponse.cs
--- a/Engine/Source/Programs/Shared/EpicGames.Perforce/PerforceResponse.cs
+++ b/Engine/Source/Programs/Shared/EpicGames.Perforce/PerforceResponse.cs
@@ -60,11 +60,7 @@ namespace EpicGames.Perforce
 		/// </summary>
 		public PerforceInfo? Info
 		{
-			get
-			{
-				EnsureSuccess();
-				return InternalData as PerforceInfo;
-			}
+			get { return InternalData as PerforceInfo; }
 		}
 
 		/// <summary>
diff --git a/Engine/Source/Programs/Shared/EpicGames.Perforce/Records/DescribeFileRecord.cs b/Engine/Source/Programs/Shared/EpicGames.Perforce/Records/DescribeFileRecord.cs
--- a/Engine/Source/Programs/Shared/EpicGames.Perforce/Records/DescribeFileRecord.cs
+++ b/Engine/Source/Programs/Shared/EpicGames.Perforce/Records/DescribeFileRecord.cs
@@ -41,7 +41,7 @@ namespace EpicGames.Perforce
 		/// Size of the file, or -1 if not specified
 		/// </summary>
 		[PerforceTag("fileSize", Optional = true)]
-		public readonly int FileSize;
+		public readonly long FileSize;
 
 		/// <summary>
 		/// Digest of the file, or null if not specified
diff --git a/Engine/Source/Programs/Shared/EpicGames.Perforce/Records/RevisionRecord.cs b/Engine/Source/Programs/Shared/EpicGames.Perforce/Records/RevisionRecord.cs
--- a/Engine/Source/Programs/Shared/EpicGames.Perforce/Records/RevisionRecord.cs
+++ b/Engine/Source/Programs/Shared/EpicGames.Perforce/Records/RevisionRecord.cs
@@ -59,7 +59,7 @@ namespace EpicGames.Perforce
 		/// Size of the file, or -1 if not specified
 		/// </summary>
 		[PerforceTag("fileSize", Optional = true)]
-		public readonly int FileSize;
+		public readonly long FileSize;
 
 		/// <summary>
 		/// Digest of the file, or null if not specified
diff --git a/Engine/Source/Programs/Shared/EpicGames.Perforce/Records/UserRecord.cs b/Engine/Source/Programs/Shared/EpicGames.Perforce/Records/UserRecord.cs
--- a/Engine/Source/Programs/Shared/EpicGames.Perforce/Records/UserRecord.cs
+++ b/Engine/Source/Programs/Shared/EpicGames.Perforce/Records/UserRecord.cs
@@ -28,13 +28,13 @@ namespace EpicGames.Perforce
 		/// <summary>
 		/// Last time the user's information was updated
 		/// </summary>
-		[PerforceTag("Update")]
+		[PerforceTag("Update", Optional = true)]
 		public DateTime Update;
 
 		/// <summary>
 		/// Last time the user's information was accessed
 		/// </summary>
-		[PerforceTag("Access")]
+		[PerforceTag("Access", Optional = true)]
 		public DateTime Access;
 
 		/// <summary>
@@ -52,8 +52,8 @@ namespace EpicGames.Perforce
 		/// <summary>
 		/// The type of user
 		/// </summary>
-		[PerforceTag("Type")]
-		public string Type;
+		[PerforceTag("Type", Optional = true)]
+		public string? Type;
 
 		/// <summary>
 		/// Method used to authenticate
diff --git a/Engine/Source/Programs/UnrealBuildTool/Configuration/TargetDescriptor.cs b/Engine/Source/Programs/UnrealBuildTool/Configuration/TargetDescriptor.cs
--- a/Engine/Source/Programs/UnrealBuildTool/Configuration/TargetDescriptor.cs
+++ b/Engine/Source/Programs/UnrealBuildTool/Configuration/TargetDescriptor.cs
@@ -396,7 +396,7 @@ namespace UnrealBuildTool
 		/// <returns>True if the project file was parsed, false otherwise</returns>
 		public static bool TryParseProjectFileArgument(CommandLineArguments Arguments, out FileReference? ProjectFile)
 		{
-			FileReference ExplicitProjectFile;
+			FileReference? ExplicitProjectFile;
 			if(Arguments.TryGetValue("-Project=", out ExplicitProjectFile))
 			{
 				ProjectFile = ExplicitProjectFile;
diff --git a/Engine/Source/Programs/UnrealBuildTool/Configuration/UEBuildBinary.cs b/Engine/Source/Programs/UnrealBuildTool/Configuration/UEBuildBinary.cs
--- a/Engine/Source/Programs/UnrealBuildTool/Configuration/UEBuildBinary.cs
+++ b/Engine/Source/Programs/UnrealBuildTool/Configuration/UEBuildBinary.cs
@@ -583,7 +583,7 @@ namespace UnrealBuildTool
 			DirectoryReference Directory = BinaryPath.Directory;
 			if (Directory.FullName.EndsWith(".app/Contents/MacOS"))
 			{
-				Directory = Directory.ParentDirectory.ParentDirectory.ParentDirectory;
+				Directory = Directory.ParentDirectory!.ParentDirectory!.ParentDirectory!;
 			}
 
 			return FileReference.Combine(Directory, BinaryPath.GetFileNameWithoutExtension() + "-Cmd" + BinaryPath.GetExtension());
diff --git a/Engine/Source/Programs/UnrealBuildTool/Configuration/UEBuildModule.cs b/Engine/Source/Programs/UnrealBuildTool/Configuration/UEBuildModule.cs
--- a/Engine/Source/Programs/UnrealBuildTool/Configuration/UEBuildModule.cs
+++ b/Engine/Source/Programs/UnrealBuildTool/Configuration/UEBuildModule.cs
@@ -511,7 +511,7 @@ namespace UnrealBuildTool
 			)
 		{
 			// Add the module's parent directory to the include path, so we can root #includes from generated source files to it
-			IncludePaths.Add(ModuleDirectory.ParentDirectory);
+			IncludePaths.Add(ModuleDirectory.ParentDirectory!);
 
 			// Add this module's public include paths and definitions.
 			AddIncludePaths(IncludePaths, PublicIncludePaths);
diff --git a/Engine/Source/Programs/UnrealBuildTool/Configuration/UEBuildModuleCPP.cs b/Engine/Source/Programs/UnrealBuildTool/Configuration/UEBuildModuleCPP.cs
--- a/Engine/Source/Programs/UnrealBuildTool/Configuration/UEBuildModuleCPP.cs
+++ b/Engine/Source/Programs/UnrealBuildTool/Configuration/UEBuildModuleCPP.cs
@@ -222,7 +222,7 @@ namespace UnrealBuildTool
 			// Add the module's parent directory to the public include paths, so other modules may include headers from it explicitly.
 			foreach (DirectoryReference ModuleDir in ModuleDirectories)
 			{
-				PublicIncludePaths.Add(ModuleDir.ParentDirectory);
+				PublicIncludePaths.Add(ModuleDir.ParentDirectory!);
 
 				// Add the base directory to the legacy include paths.
 				LegacyPublicIncludePaths.Add(ModuleDir);
@@ -625,7 +625,7 @@ namespace UnrealBuildTool
 		public PrecompiledHeaderTemplate CreateSharedPCHTemplate(UEBuildTarget Target, CppCompileEnvironment BaseCompileEnvironment)
 		{
 			CppCompileEnvironment CompileEnvironment = CreateSharedPCHCompileEnvironment(Target, BaseCompileEnvironment);
-			FileItem HeaderFile = FileItem.GetItemByFileReference(FileReference.Combine(ModuleDirectory, Rules.SharedPCHHeaderFile));
+			FileItem HeaderFile = FileItem.GetItemByFileReference(FileReference.Combine(ModuleDirectory, Rules.SharedPCHHeaderFile!));
 
 			DirectoryReference PrecompiledHeaderDir;
 			if(Rules.bUsePrecompiled)
diff --git a/Engine/Source/Programs/UnrealBuildTool/Configuration/UEBuildTarget.cs b/Engine/Source/Programs/UnrealBuildTool/Configuration/UEBuildTarget.cs
--- a/Engine/Source/Programs/UnrealBuildTool/Configuration/UEBuildTarget.cs
+++ b/Engine/Source/Programs/UnrealBuildTool/Configuration/UEBuildTarget.cs
@@ -1201,11 +1201,11 @@ namespace UnrealBuildTool
 				{
 					if (Directory.Exists(Resource.ResourcePath))
 					{
-						Files.UnionWith(DirectoryReference.EnumerateFiles(new DirectoryReference(Resource.ResourcePath), "*", SearchOption.AllDirectories));
+						Files.UnionWith(DirectoryReference.EnumerateFiles(new DirectoryReference(Resource.ResourcePath!), "*", SearchOption.AllDirectories));
 					}
 					else
 					{
-						Files.Add(new FileReference(Resource.ResourcePath));
+						Files.Add(new FileReference(Resource.ResourcePath!));
 					}
 				}
 
@@ -2024,7 +2024,7 @@ namespace UnrealBuildTool
 			DirectoryReference ExeDir = Binaries[0].OutputDir;
 			if (Platform == UnrealTargetPlatform.Mac && ExeDir.FullName.EndsWith(".app/Contents/MacOS"))
 			{
-				ExeDir = ExeDir.ParentDirectory.ParentDirectory.ParentDirectory;
+				ExeDir = ExeDir.ParentDirectory!.ParentDirectory!.ParentDirectory!;
 			}
 			return ExeDir;
 		}
diff --git a/Engine/Source/Programs/UnrealBuildTool/Modes/JsonExportMode.cs b/Engine/Source/Programs/UnrealBuildTool/Modes/JsonExportMode.cs
--- a/Engine/Source/Programs/UnrealBuildTool/Modes/JsonExportMode.cs
+++ b/Engine/Source/Programs/UnrealBuildTool/Modes/JsonExportMode.cs
@@ -37,7 +37,7 @@ namespace UnrealBuildTool
 				UEBuildTarget Target = UEBuildTarget.Create(TargetDescriptor, false, false);
 
 				// Get the output file
-				FileReference OutputFile = TargetDescriptor.AdditionalArguments.GetFileReferenceOrDefault("-OutputFile=", null);
+				FileReference? OutputFile = TargetDescriptor.AdditionalArguments.GetFileReferenceOrDefault("-OutputFile=", null);
 				if(OutputFile == null)
 				{
 					OutputFile = Target.ReceiptFileName.ChangeExtension(".json");
diff --git a/Engine/Source/Programs/UnrealBuildTool/Platform/Windows/UEBuildWindows.cs b/Engine/Source/Programs/UnrealBuildTool/Platform/Windows/UEBuildWindows.cs
--- a/Engine/Source/Programs/UnrealBuildTool/Platform/Windows/UEBuildWindows.cs
+++ b/Engine/Source/Programs/UnrealBuildTool/Platform/Windows/UEBuildWindows.cs
@@ -1153,7 +1153,7 @@ namespace UnrealBuildTool
 					if(Compiler == WindowsCompiler.Clang)
 					{
 						// Check for a manual installation to the default directory
-						DirectoryReference ManualInstallDir = DirectoryReference.Combine(DirectoryReference.GetSpecialFolder(Environment.SpecialFolder.ProgramFiles), "LLVM");
+						DirectoryReference ManualInstallDir = DirectoryReference.Combine(DirectoryReference.GetSpecialFolder(Environment.SpecialFolder.ProgramFiles)!, "LLVM");
 						AddClangToolChain(ManualInstallDir, ToolChains);
 
 						// Check for a manual installation to a custom directory
@@ -1186,7 +1186,7 @@ namespace UnrealBuildTool
 					else if(Compiler == WindowsCompiler.Intel)
 					{
 						// Just check for a manual installation
-						DirectoryReference InstallDir = DirectoryReference.Combine(DirectoryReference.GetSpecialFolder(Environment.SpecialFolder.ProgramFilesX86), "IntelSWTools", "compilers_and_libraries", "windows");
+						DirectoryReference InstallDir = DirectoryReference.Combine(DirectoryReference.GetSpecialFolder(Environment.SpecialFolder.ProgramFilesX86)!, "IntelSWTools", "compilers_and_libraries", "windows");
 						if(DirectoryReference.Exists(InstallDir))
 						{
 							FileReference IclPath = FileReference.Combine(InstallDir, "bin", "intel64", "icl.exe");
@@ -1580,7 +1580,7 @@ namespace UnrealBuildTool
 			}
 
 			// Try to get the MSBuild 14.0 path directly (see https://msdn.microsoft.com/en-us/library/hh162058(v=vs.120).aspx)
-			FileReference? ToolPath = FileReference.Combine(DirectoryReference.GetSpecialFolder(Environment.SpecialFolder.ProgramFilesX86), "MSBuild", "14.0", "bin", "MSBuild.exe");
+			FileReference? ToolPath = FileReference.Combine(DirectoryReference.GetSpecialFolder(Environment.SpecialFolder.ProgramFilesX86)!, "MSBuild", "14.0", "bin", "MSBuild.exe");
 			if(FileReference.Exists(ToolPath))
 			{
 				OutLocation = ToolPath;
@@ -2061,7 +2061,7 @@ namespace UnrealBuildTool
 				//CompileEnvironment.bAllowLTCG = true;
 				//LinkEnvironment.bAllowLTCG = true;
 
-				CompileEnvironment.PGODirectory = Path.Combine(DirectoryReference.FromFile(Target.ProjectFile).FullName, "Platforms", "Windows", "Build", "PGO");
+				CompileEnvironment.PGODirectory = Path.Combine(DirectoryReference.FromFile(Target.ProjectFile!).FullName, "Platforms", "Windows", "Build", "PGO");
 				CompileEnvironment.PGOFilenamePrefix = string.Format("{0}-{1}-{2}", Target.Name, Target.Platform, Target.Configuration);
 
 				LinkEnvironment.PGODirectory = CompileEnvironment.PGODirectory;
diff --git a/Engine/Source/Programs/UnrealBuildTool/Platform/Windows/VCEnvironment.cs b/Engine/Source/Programs/UnrealBuildTool/Platform/Windows/VCEnvironment.cs
--- a/Engine/Source/Programs/UnrealBuildTool/Platform/Windows/VCEnvironment.cs
+++ b/Engine/Source/Programs/UnrealBuildTool/Platform/Windows/VCEnvironment.cs
@@ -246,13 +246,13 @@ namespace UnrealBuildTool
 			// If we were asked to use Clang, then we'll redirect the path to the compiler to the LLVM installation directory
 			if (Compiler == WindowsCompiler.Clang && WindowsPlatform.bAllowClangLinker)
 			{
-				FileReference LinkerPath = FileReference.Combine(DirectoryReference.GetSpecialFolder(Environment.SpecialFolder.ProgramFiles), "LLVM", "bin", "lld-link.exe");
+				FileReference LinkerPath = FileReference.Combine(DirectoryReference.GetSpecialFolder(Environment.SpecialFolder.ProgramFiles)!, "LLVM", "bin", "lld-link.exe");
 				if (FileReference.Exists(LinkerPath))
 				{
 					return LinkerPath;
 				}
 
-				FileReference LinkerPathX86 = FileReference.Combine(DirectoryReference.GetSpecialFolder(Environment.SpecialFolder.ProgramFilesX86), "LLVM", "bin", "lld-link.exe");
+				FileReference LinkerPathX86 = FileReference.Combine(DirectoryReference.GetSpecialFolder(Environment.SpecialFolder.ProgramFilesX86)!, "LLVM", "bin", "lld-link.exe");
 				if (FileReference.Exists(LinkerPathX86))
 				{
 					return LinkerPathX86;
@@ -262,7 +262,7 @@ namespace UnrealBuildTool
 			}
 			else if(Compiler == WindowsCompiler.Intel && WindowsPlatform.bAllowICLLinker)
 			{
-				FileReference LinkerPath = FileReference.Combine(DirectoryReference.GetSpecialFolder(Environment.SpecialFolder.ProgramFilesX86), "IntelSWTools", "compilers_and_libraries", "windows", "bin", "intel64", "xilink.exe");
+				FileReference LinkerPath = FileReference.Combine(DirectoryReference.GetSpecialFolder(Environment.SpecialFolder.ProgramFilesX86)!, "IntelSWTools", "compilers_and_libraries", "windows", "bin", "intel64", "xilink.exe");
 				if (FileReference.Exists(LinkerPath))
 				{
 					return LinkerPath;
@@ -284,7 +284,7 @@ namespace UnrealBuildTool
 			// Regardless of the target, if we're linking on a 64 bit machine, we want to use the 64 bit linker (it's faster than the 32 bit linker)
 			if (Compiler == WindowsCompiler.Intel && WindowsPlatform.bAllowICLLinker)
 			{
-				FileReference LibPath = FileReference.Combine(DirectoryReference.GetSpecialFolder(Environment.SpecialFolder.ProgramFilesX86), "IntelSWTools", "compilers_and_libraries", "windows", "bin", "intel64", "xilib.exe");
+				FileReference LibPath = FileReference.Combine(DirectoryReference.GetSpecialFolder(Environment.SpecialFolder.ProgramFilesX86)!, "IntelSWTools", "compilers_and_libraries", "windows", "bin", "intel64", "xilib.exe");
 				if (FileReference.Exists(LibPath))
 				{
 					return LibPath;
diff --git a/Engine/Source/Programs/UnrealBuildTool/System/BuildHostPlatform.cs b/Engine/Source/Programs/UnrealBuildTool/System/BuildHostPlatform.cs
--- a/Engine/Source/Programs/UnrealBuildTool/System/BuildHostPlatform.cs
+++ b/Engine/Source/Programs/UnrealBuildTool/System/BuildHostPlatform.cs
@@ -253,7 +253,7 @@ namespace UnrealBuildTool
 
 		public override FileReference Shell
 		{
-			get { return new FileReference(Environment.GetEnvironmentVariable("COMSPEC")); }
+			get { return new FileReference(Environment.GetEnvironmentVariable("COMSPEC")!); }
 		}
 
 		public override ShellType ShellType
diff --git a/Engine/Source/Programs/UnrealBuildTool/System/ConfigHierarchy.cs b/Engine/Source/Programs/UnrealBuildTool/System/ConfigHierarchy.cs
--- a/Engine/Source/Programs/UnrealBuildTool/System/ConfigHierarchy.cs
+++ b/Engine/Source/Programs/UnrealBuildTool/System/ConfigHierarchy.cs
@@ -1144,12 +1144,13 @@ namespace UnrealBuildTool
 		{
 			if(UnrealBuildTool.IsEngineInstalled())
 			{
-				return DirectoryReference.Combine(Utils.GetUserSettingDirectory(), "UnrealEngine", String.Format("{0}.{1}", ReadOnlyBuildVersion.Current.MajorVersion, ReadOnlyBuildVersion.Current.MinorVersion), "Saved");
-			}
-			else
-			{
-				return DirectoryReference.Combine(UnrealBuildTool.EngineDirectory, "Saved");
+				DirectoryReference? UserSettingDir = Utils.GetUserSettingDirectory();
+				if (UserSettingDir != null)
+				{
+					return DirectoryReference.Combine(UserSettingDir, "UnrealEngine", String.Format("{0}.{1}", ReadOnlyBuildVersion.Current.MajorVersion, ReadOnlyBuildVersion.Current.MinorVersion), "Saved");
+				}
 			}
+			return DirectoryReference.Combine(UnrealBuildTool.EngineDirectory, "Saved");
 		}
 
 		/// <summary>
diff --git a/Engine/Source/Programs/UnrealBuildTool/System/ExternalExecution.cs b/Engine/Source/Programs/UnrealBuildTool/System/ExternalExecution.cs
--- a/Engine/Source/Programs/UnrealBuildTool/System/ExternalExecution.cs
+++ b/Engine/Source/Programs/UnrealBuildTool/System/ExternalExecution.cs
@@ -349,7 +349,7 @@ namespace UnrealBuildTool
 				ModuleType = Info.ModuleType;
 				OverrideModuleType = Info.OverrideModuleType;
 				BaseDirectory = Info.ModuleDirectories[0].FullName;
-				IncludeBase = Info.ModuleDirectories[0].ParentDirectory.FullName;
+				IncludeBase = Info.ModuleDirectories[0].ParentDirectory!.FullName;
 				OutputDirectory = Path.GetDirectoryName(Info.GeneratedCPPFilenameBase)!;
 				ClassesHeaders = Info.PublicUObjectClassesHeaders.Select((Header) => Header.AbsolutePath).ToList();
 				PublicHeaders = Info.PublicUObjectHeaders.Select((Header) => Header.AbsolutePath).ToList();
diff --git a/Engine/Source/Programs/UnrealBuildTool/System/SourceFileWorkingSet.cs b/Engine/Source/Programs/UnrealBuildTool/System/SourceFileWorkingSet.cs
--- a/Engine/Source/Programs/UnrealBuildTool/System/SourceFileWorkingSet.cs
+++ b/Engine/Source/Programs/UnrealBuildTool/System/SourceFileWorkingSet.cs
@@ -371,9 +371,9 @@ namespace UnrealBuildTool
 					{
 						WorkingSet = new GitSourceFileWorkingSet(GitPath, ProjectDir, WorkingSet);
 					}
-					else if (DirectoryReference.Exists(DirectoryReference.Combine(ProjectDir.ParentDirectory, ".git")))
+					else if (DirectoryReference.Exists(DirectoryReference.Combine(ProjectDir.ParentDirectory!, ".git")))
 					{
-						WorkingSet = new GitSourceFileWorkingSet(GitPath, ProjectDir.ParentDirectory, WorkingSet);
+						WorkingSet = new GitSourceFileWorkingSet(GitPath, ProjectDir.ParentDirectory!, WorkingSet);
 					}
 				}
 			}
diff --git a/Engine/Source/Programs/UnrealBuildTool/System/XmlConfigData.cs b/Engine/Source/Programs/UnrealBuildTool/System/XmlConfigData.cs
--- a/Engine/Source/Programs/UnrealBuildTool/System/XmlConfigData.cs
+++ b/Engine/Source/Programs/UnrealBuildTool/System/XmlConfigData.cs
@@ -104,7 +104,7 @@ namespace UnrealBuildTool
 						}
 
 						// Try to parse the value and add it to the output array
-						object Value = Reader.ReadObject(Field.FieldType);
+						object Value = Reader.ReadObject(Field.FieldType)!;
 						Values[ValueIdx] = new KeyValuePair<FieldInfo, object>(Field, Value);
 					}
 
diff --git a/Engine/Source/Programs/UnrealBuildTool/UnrealBuildTool.cs b/Engine/Source/Programs/UnrealBuildTool/UnrealBuildTool.cs
--- a/Engine/Source/Programs/UnrealBuildTool/UnrealBuildTool.cs
+++ b/Engine/Source/Programs/UnrealBuildTool/UnrealBuildTool.cs
@@ -522,7 +522,7 @@ namespace UnrealBuildTool
 						string XmlConfigMutexName = SingleInstanceMutex.GetUniqueMutexForPath("UnrealBuildTool_Mutex_XmlConfig", Assembly.GetExecutingAssembly().CodeBase!);
 						using(SingleInstanceMutex XmlConfigMutex = new SingleInstanceMutex(XmlConfigMutexName, true))
 						{
-							FileReference XmlConfigCache = Arguments.GetFileReferenceOrDefault("-XmlConfigCache=", null);
+							FileReference? XmlConfigCache = Arguments.GetFileReferenceOrDefault("-XmlConfigCache=", null);
 							XmlConfig.ReadConfigFiles(XmlConfigCache);
 						}
 					}
