Allow using different phrasing when joining lists of strings.

diff --git a/Engine/Source/Programs/Shared/EpicGames.Core/StringUtils.cs b/Engine/Source/Programs/Shared/EpicGames.Core/StringUtils.cs
--- a/Engine/Source/Programs/Shared/EpicGames.Core/StringUtils.cs
+++ b/Engine/Source/Programs/Shared/EpicGames.Core/StringUtils.cs
@@ -197,18 +197,19 @@ namespace EpicGames.Core
 		/// Formats a list of strings in the style "1, 2, 3 and 4"
 		/// </summary>
 		/// <param name="Arguments">List of strings to format</param>
+		/// <param name="Conjunction">Conjunction to use between the last two items in the list (eg. "and" or "or")</param>
 		/// <returns>Formatted list of strings</returns>
-		public static string FormatList(string[] Arguments)
+		public static string FormatList(string[] Arguments, string Conjunction = "and")
 		{
 			StringBuilder Result = new StringBuilder();
-			if(Arguments.Length > 0)
+			if (Arguments.Length > 0)
 			{
 				Result.Append(Arguments[0]);
-				for(int Idx = 1; Idx < Arguments.Length; Idx++)
+				for (int Idx = 1; Idx < Arguments.Length; Idx++)
 				{
-					if(Idx == Arguments.Length - 1)
+					if (Idx == Arguments.Length - 1)
 					{
-						Result.Append(" and ");
+						Result.AppendFormat(" {0} ", Conjunction);
 					}
 					else
 					{
@@ -224,12 +225,14 @@ namespace EpicGames.Core
 		/// Formats a list of strings in the style "1, 2, 3 and 4"
 		/// </summary>
 		/// <param name="Arguments">List of strings to format</param>
+		/// <param name="Conjunction">Conjunction to use between the last two items in the list (eg. "and" or "or")</param>
 		/// <returns>Formatted list of strings</returns>
-		public static string FormatList(IEnumerable<string> Arguments)
+		public static string FormatList(IEnumerable<string> Arguments, string Conjunction = "and")
 		{
-			return FormatList(Arguments.ToArray());
+			return FormatList(Arguments.ToArray(), Conjunction);
 		}
 
+
 		/// <summary>
 		/// Formats a list of items
 		/// </summary>
