Add new CopyToAsync() overloads.

diff --git a/Engine/Source/Programs/Shared/EpicGames.Core/ManagedProcess.cs b/Engine/Source/Programs/Shared/EpicGames.Core/ManagedProcess.cs
--- a/Engine/Source/Programs/Shared/EpicGames.Core/ManagedProcess.cs
+++ b/Engine/Source/Programs/Shared/EpicGames.Core/ManagedProcess.cs
@@ -713,7 +713,7 @@ namespace EpicGames.Core
 		/// <param name="OutputStream">The output stream</param>
 		/// <param name="CancellationToken">Cancellation token</param>
 		/// <returns></returns>
-		public Task CopyTo(Stream OutputStream, CancellationToken CancellationToken)
+		public Task CopyToAsync(Stream OutputStream, CancellationToken CancellationToken)
 		{
 			if(FrameworkProcess == null)
 			{
@@ -725,6 +725,47 @@ namespace EpicGames.Core
 			}
 		}
 
+		/// <summary>
+		/// Copy the process output to the given stream
+		/// </summary>
+		/// <param name="WriteOutput">The output stream</param>
+		/// <param name="CancellationToken">Cancellation token</param>
+		/// <returns></returns>
+		public Task CopyToAsync(Action<byte[], int, int> WriteOutput, int BufferSize, CancellationToken CancellationToken)
+		{
+			Func<byte[], int, int, CancellationToken, Task> WriteOutputAsync = (Buffer, Offset, Length, CancellationToken) => { WriteOutput(Buffer, Offset, Length); return Task.CompletedTask; };
+			return CopyToAsync(WriteOutputAsync, BufferSize, CancellationToken);
+		}
+
+		/// <summary>
+		/// Copy the process output to the given stream
+		/// </summary>
+		/// <param name="WriteOutputAsync">The output stream</param>
+		/// <param name="CancellationToken">Cancellation token</param>
+		/// <returns></returns>
+		public async Task CopyToAsync(Func<byte[], int, int, CancellationToken, Task> WriteOutputAsync, int BufferSize, CancellationToken CancellationToken)
+		{
+			TaskCompletionSource<bool> TaskCompletionSource = new TaskCompletionSource<bool>();
+			using (CancellationTokenRegistration Registration = CancellationToken.Register(() => TaskCompletionSource.SetResult(false)))
+			{
+				byte[] Buffer = new byte[BufferSize];
+				for (; ; )
+				{
+					Task<int> ReadTask = StdOut.ReadAsync(Buffer, 0, BufferSize, CancellationToken);
+
+					Task CompletedTask = await Task.WhenAny(ReadTask, TaskCompletionSource.Task);
+					CancellationToken.ThrowIfCancellationRequested();
+
+					int Bytes = await ReadTask;
+					if (Bytes == 0)
+					{
+						break;
+					}
+					await WriteOutputAsync(Buffer, 0, Bytes, CancellationToken);
+				}
+			}
+		}
+
 		/// <summary>
 		/// Read all the output from the process. Does not return until the process terminates.
 		/// </summary>
