Missing files.

diff --git a/Engine/Source/Programs/Shared/EpicGames.Perforce/PerforceChildProcess.cs b/Engine/Source/Programs/Shared/EpicGames.Perforce/PerforceChildProcess.cs
--- a/Engine/Source/Programs/Shared/EpicGames.Perforce/PerforceChildProcess.cs
+++ b/Engine/Source/Programs/Shared/EpicGames.Perforce/PerforceChildProcess.cs
@@ -169,7 +169,7 @@ namespace EpicGames.Perforce
 		public async Task<Tuple<bool, string>> TryReadToEndAsync(CancellationToken CancellationToken)
 		{
 			MemoryStream Stream = new MemoryStream();
-			await ChildProcess!.CopyTo(Stream, CancellationToken);
+			await ChildProcess!.CopyToAsync(Stream, CancellationToken);
 			ChildProcess.WaitForExit();
 			string Data = Encoding.Default.GetString(Stream.ToArray());
 			return Tuple.Create(ChildProcess.ExitCode == 0, Data);
diff --git a/Engine/Source/Programs/Shared/EpicGames.Perforce/PerforceConnection.cs b/Engine/Source/Programs/Shared/EpicGames.Perforce/PerforceConnection.cs
--- a/Engine/Source/Programs/Shared/EpicGames.Perforce/PerforceConnection.cs
+++ b/Engine/Source/Programs/Shared/EpicGames.Perforce/PerforceConnection.cs
@@ -2549,6 +2549,31 @@ namespace EpicGames.Perforce
 			return (await CommandAsync<SubmitRecord>(Arguments.ToString(), null, CancellationToken))[0];
 		}
 
+		/// <summary>
+		/// Submits a shelved changelist
+		/// </summary>
+		/// <param name="ChangeNumber">The changelist to submit</param>
+		/// <param name="CancellationToken">Token used to cancel the operation</param>
+		/// <returns>Response from the server</returns>
+		public async Task<SubmitRecord> SubmitShelvedAsync(int ChangeNumber, CancellationToken CancellationToken)
+		{
+			return (await TrySubmitShelvedAsync(ChangeNumber, CancellationToken)).Data;
+		}
+
+		/// <summary>
+		/// Submits a pending changelist
+		/// </summary>
+		/// <param name="ChangeNumber">The changelist to submit</param>
+		/// <param name="CancellationToken">Token used to cancel the operation</param>
+		/// <returns>Response from the server</returns>
+		public async Task<PerforceResponse<SubmitRecord>> TrySubmitShelvedAsync(int ChangeNumber, CancellationToken CancellationToken)
+		{
+			StringBuilder Arguments = new StringBuilder("submit");
+			Arguments.AppendFormat(" -e {0}", ChangeNumber);
+
+			return (await CommandAsync<SubmitRecord>(Arguments.ToString(), null, CancellationToken))[0];
+		}
+
 		#endregion
 
 		#region p4 sync
