UBT: Fix error trying to read makefile, due to value being null.

diff --git a/Engine/Source/Programs/Shared/EpicGames.Core/DirectoryReference.cs b/Engine/Source/Programs/Shared/EpicGames.Core/DirectoryReference.cs
--- a/Engine/Source/Programs/Shared/EpicGames.Core/DirectoryReference.cs
+++ b/Engine/Source/Programs/Shared/EpicGames.Core/DirectoryReference.cs
@@ -452,9 +452,10 @@ namespace EpicGames.Core
 		/// </summary>
 		/// <param name="Reader">Binary reader to read from</param>
 		/// <returns>New DirectoryReference object</returns>
-		public static DirectoryReference ReadDirectoryReference(this BinaryReader Reader)
+		public static DirectoryReference? ReadDirectoryReference(this BinaryReader Reader)
 		{
-			return BinaryArchiveReader.NotNull(ReadDirectoryReferenceOrNull(Reader));
+			string FullName = Reader.ReadString();
+			return (FullName.Length == 0) ? null : new DirectoryReference(FullName, DirectoryReference.Sanitize.None);
 		}
 
 		/// <summary>
@@ -462,10 +463,9 @@ namespace EpicGames.Core
 		/// </summary>
 		/// <param name="Reader">Binary reader to read from</param>
 		/// <returns>New DirectoryReference object</returns>
-		public static DirectoryReference? ReadDirectoryReferenceOrNull(this BinaryReader Reader)
+		public static DirectoryReference ReadDirectoryReferenceNotNull(this BinaryReader Reader)
 		{
-			string FullName = Reader.ReadString();
-			return (FullName.Length == 0) ? null : new DirectoryReference(FullName, DirectoryReference.Sanitize.None);
+			return BinaryArchiveReader.NotNull(ReadDirectoryReference(Reader));
 		}
 
 		/// <summary>
@@ -500,17 +500,7 @@ namespace EpicGames.Core
 		/// </summary>
 		/// <param name="Reader">Reader to serialize data from</param>
 		/// <returns>New directory reference instance</returns>
-		public static DirectoryReference ReadDirectoryReference(this BinaryArchiveReader Reader)
-		{
-			return BinaryArchiveReader.NotNull(ReadDirectoryReferenceOrNull(Reader));
-		}
-
-		/// <summary>
-		/// Reads a directory reference from a binary archive
-		/// </summary>
-		/// <param name="Reader">Reader to serialize data from</param>
-		/// <returns>New directory reference instance</returns>
-		public static DirectoryReference? ReadDirectoryReferenceOrNull(this BinaryArchiveReader Reader)
+		public static DirectoryReference? ReadDirectoryReference(this BinaryArchiveReader Reader)
 		{
 			string FullName = Reader.ReadString();
 			if (FullName == null)
@@ -522,5 +512,15 @@ namespace EpicGames.Core
 				return new DirectoryReference(FullName, DirectoryReference.Sanitize.None);
 			}
 		}
+
+		/// <summary>
+		/// Reads a directory reference from a binary archive
+		/// </summary>
+		/// <param name="Reader">Reader to serialize data from</param>
+		/// <returns>New directory reference instance</returns>
+		public static DirectoryReference ReadDirectoryReferenceNotNull(this BinaryArchiveReader Reader)
+		{
+			return BinaryArchiveReader.NotNull(ReadDirectoryReference(Reader));
+		}
 	}
 }
diff --git a/Engine/Source/Programs/UnrealBuildTool/System/Action.cs b/Engine/Source/Programs/UnrealBuildTool/System/Action.cs
--- a/Engine/Source/Programs/UnrealBuildTool/System/Action.cs
+++ b/Engine/Source/Programs/UnrealBuildTool/System/Action.cs
@@ -249,7 +249,7 @@ namespace UnrealBuildTool
 		public Action(BinaryArchiveReader Reader)
 		{
 			ActionType = (ActionType)Reader.ReadByte();
-			WorkingDirectory = Reader.ReadDirectoryReference();
+			WorkingDirectory = Reader.ReadDirectoryReferenceNotNull();
 			CommandPath = Reader.ReadFileReference();
 			CommandArguments = Reader.ReadString();
 			CommandVersion = Reader.ReadString();
diff --git a/Engine/Source/Programs/UnrealBuildTool/System/ConfigValueTracker.cs b/Engine/Source/Programs/UnrealBuildTool/System/ConfigValueTracker.cs
--- a/Engine/Source/Programs/UnrealBuildTool/System/ConfigValueTracker.cs
+++ b/Engine/Source/Programs/UnrealBuildTool/System/ConfigValueTracker.cs
@@ -65,7 +65,7 @@ namespace UnrealBuildTool
 		public ConfigDependencyKey(BinaryArchiveReader Reader)
 		{
 			Type = (ConfigHierarchyType)Reader.ReadInt();
-			ProjectDir = Reader.ReadDirectoryReference();
+			ProjectDir = Reader.ReadDirectoryReferenceNotNull();
 			Platform = Reader.ReadUnrealTargetPlatform();
 			SectionName = Reader.ReadString();
 			KeyName = Reader.ReadString();
diff --git a/Engine/Source/Programs/UnrealBuildTool/System/DirectoryItem.cs b/Engine/Source/Programs/UnrealBuildTool/System/DirectoryItem.cs
--- a/Engine/Source/Programs/UnrealBuildTool/System/DirectoryItem.cs
+++ b/Engine/Source/Programs/UnrealBuildTool/System/DirectoryItem.cs
@@ -343,7 +343,7 @@ namespace UnrealBuildTool
 		/// <returns>Instance of the serialized directory item</returns>
 		public static DirectoryItem ReadDirectoryItem(this BinaryArchiveReader Reader)
 		{
-			return Reader.ReadObjectReference<DirectoryItem>(() => DirectoryItem.GetItemByDirectoryReference(Reader.ReadDirectoryReference()));
+			return Reader.ReadObjectReference<DirectoryItem>(() => DirectoryItem.GetItemByDirectoryReference(Reader.ReadDirectoryReferenceNotNull()));
 		}
 
 		/// <summary>
diff --git a/Engine/Source/Programs/UnrealBuildTool/System/ExternalExecution.cs b/Engine/Source/Programs/UnrealBuildTool/System/ExternalExecution.cs
--- a/Engine/Source/Programs/UnrealBuildTool/System/ExternalExecution.cs
+++ b/Engine/Source/Programs/UnrealBuildTool/System/ExternalExecution.cs
@@ -264,7 +264,7 @@ namespace UnrealBuildTool
 		{
 			ModuleName = Reader.ReadString();
 			ModuleRulesFile = Reader.ReadFileReference();
-			ModuleDirectories = Reader.ReadArray<DirectoryReference>(Reader.ReadDirectoryReference);
+			ModuleDirectories = Reader.ReadArray<DirectoryReference>(Reader.ReadDirectoryReferenceNotNull);
 			ModuleType = Reader.ReadString();
 			OverrideModuleType = Reader.ReadString();
 			PublicUObjectClassesHeaders = Reader.ReadList(() => Reader.ReadFileItem());
diff --git a/Engine/Source/Programs/UnrealBuildTool/System/TargetMakefile.cs b/Engine/Source/Programs/UnrealBuildTool/System/TargetMakefile.cs
--- a/Engine/Source/Programs/UnrealBuildTool/System/TargetMakefile.cs
+++ b/Engine/Source/Programs/UnrealBuildTool/System/TargetMakefile.cs
@@ -221,7 +221,7 @@ namespace UnrealBuildTool
 			ExternalMetadata = Reader.ReadString();
 			ExecutableFile = Reader.ReadFileReference();
 			ReceiptFile = Reader.ReadFileReference();
-			ProjectIntermediateDirectory = Reader.ReadDirectoryReference();
+			ProjectIntermediateDirectory = Reader.ReadDirectoryReferenceNotNull();
 			TargetType = (TargetType)Reader.ReadInt();
 			ConfigValueTracker = new ConfigValueTracker(Reader);
 			bDeployAfterCompile = Reader.ReadBool();
