UGS: Enable deterministic output.

diff --git a/Engine/Source/Programs/Shared/EpicGames.Core/EpicGames.Core.csproj b/Engine/Source/Programs/Shared/EpicGames.Core/EpicGames.Core.csproj
--- a/Engine/Source/Programs/Shared/EpicGames.Core/EpicGames.Core.csproj
+++ b/Engine/Source/Programs/Shared/EpicGames.Core/EpicGames.Core.csproj
@@ -3,6 +3,7 @@
     <TargetFramework>netcoreapp3.1</TargetFramework>
     <Configuration Condition=" '$(Configuration)' == '' ">Development</Configuration>
     <OutputType>Library</OutputType>
+    <Deterministic>true</Deterministic>
     <RootNamespace>EpicGames.Core</RootNamespace>
     <GenerateAssemblyInfo>false</GenerateAssemblyInfo>
     <GenerateTargetFrameworkAttribute>false</GenerateTargetFrameworkAttribute>
diff --git a/Engine/Source/Programs/Shared/EpicGames.Jupiter/EpicGames.Jupiter.csproj b/Engine/Source/Programs/Shared/EpicGames.Jupiter/EpicGames.Jupiter.csproj
--- a/Engine/Source/Programs/Shared/EpicGames.Jupiter/EpicGames.Jupiter.csproj
+++ b/Engine/Source/Programs/Shared/EpicGames.Jupiter/EpicGames.Jupiter.csproj
@@ -3,6 +3,7 @@
     <TargetFramework>netcoreapp3.1</TargetFramework>
     <Configuration Condition=" '$(Configuration)' == '' ">Development</Configuration>
     <OutputType>Library</OutputType>
+    <Deterministic>true</Deterministic>
     <RootNamespace>EpicGames.Jupiter</RootNamespace>
     <GenerateAssemblyInfo>false</GenerateAssemblyInfo>
     <GenerateTargetFrameworkAttribute>false</GenerateTargetFrameworkAttribute>
diff --git a/Engine/Source/Programs/UnrealGameSync/UnrealGameSync/UnrealGameSync.csproj b/Engine/Source/Programs/UnrealGameSync/UnrealGameSync/UnrealGameSync.csproj
--- a/Engine/Source/Programs/UnrealGameSync/UnrealGameSync/UnrealGameSync.csproj
+++ b/Engine/Source/Programs/UnrealGameSync/UnrealGameSync/UnrealGameSync.csproj
@@ -4,6 +4,7 @@
     <TargetFramework>netcoreapp3.1</TargetFramework>
     <UseWindowsForms>true</UseWindowsForms>
     <GenerateAssemblyInfo>false</GenerateAssemblyInfo>
+    <Deterministic>true</Deterministic>
     <RootNamespace>UnrealGameSync</RootNamespace>
     <AssemblyName>UnrealGameSync</AssemblyName>
     <StartupObject>UnrealGameSync.Program</StartupObject>
