Fix ManagedProcess not correctly saving stdout/stderr handles on non-Windows platforms.

diff --git a/Engine/Source/Programs/Shared/EpicGames.Core/ManagedProcess.cs b/Engine/Source/Programs/Shared/EpicGames.Core/ManagedProcess.cs
--- a/Engine/Source/Programs/Shared/EpicGames.Core/ManagedProcess.cs
+++ b/Engine/Source/Programs/Shared/EpicGames.Core/ManagedProcess.cs
@@ -285,12 +285,12 @@ namespace EpicGames.Core
 		/// <summary>
 		/// Input stream for the child process.
 		/// </summary>
-		public FileStream StdIn { get; private set; } = null!;
+		public Stream StdIn { get; private set; } = null!;
 
 		/// <summary>
 		/// Output stream for the child process.
 		/// </summary>
-		public FileStream StdOut { get; private set; } = null!;
+		public Stream StdOut { get; private set; } = null!;
 
 		/// <summary>
 		/// Reader for the process' output stream.
@@ -300,7 +300,7 @@ namespace EpicGames.Core
 		/// <summary>
 		/// Output stream for the child process.
 		/// </summary>
-		public FileStream StdErr { get; private set; } = null!;
+		public Stream StdErr { get; private set; } = null!;
 
 		/// <summary>
 		/// Reader for the process' output stream.
@@ -619,6 +619,9 @@ namespace EpicGames.Core
 			{
 
 			}
+
+			StdIn = FrameworkProcess.StandardInput.BaseStream;
+			StdOut = FrameworkProcess.StandardOutput.BaseStream;
 		}
 
 		/// <summary>
