Fix force-sync in managed P4 library not correctly writing out temp file.

diff --git a/Engine/Source/Programs/Shared/EpicGames.Perforce/PerforceConnection.cs b/Engine/Source/Programs/Shared/EpicGames.Perforce/PerforceConnection.cs
--- a/Engine/Source/Programs/Shared/EpicGames.Perforce/PerforceConnection.cs
+++ b/Engine/Source/Programs/Shared/EpicGames.Perforce/PerforceConnection.cs
@@ -298,6 +298,10 @@ namespace EpicGames.Perforce
 			List<PerforceResponse> Responses = await CommandAsync(Arguments, InputData, StatRecordType, CancellationToken);
 			if (Responses.Count != 1)
 			{
+				for (int Idx = 0; Idx < Responses.Count; Idx++)
+				{
+					Logger.LogDebug("Unexpected response {Idx}: {Text}", Idx, Responses[Idx].ToString());
+				}
 				throw new PerforceException("Expected one result from 'p4 {0}', got {1}", Arguments, Responses.Count);
 			}
 			return Responses[0];
@@ -2623,7 +2627,7 @@ namespace EpicGames.Perforce
 		/// <returns>Response from the server</returns>
 		public Task<PerforceResponseList<SyncRecord>> TrySyncAsync(SyncOptions Options, int MaxFiles, string[] FileSpecs, CancellationToken CancellationToken)
 		{
-			return TrySyncAsync(Options, -1, -1, -1, -1, -1, -1, FileSpecs, CancellationToken);
+			return TrySyncAsync(Options, MaxFiles, -1, -1, -1, -1, -1, FileSpecs, CancellationToken);
 		}
 
 		/// <summary>
@@ -2821,6 +2825,7 @@ namespace EpicGames.Perforce
 				string TempFileName = Path.GetTempFileName();
 				try
 				{
+					await File.WriteAllLinesAsync(TempFileName, FileSpecs);
 					Arguments.Insert(0, $"-x\"{TempFileName}\" ");
 					return await CommandAsync<T>(Arguments.ToString(), null, CancellationToken);
 				}
