Fix P4 fstat calls for a particular changelist including an invalid last record.

diff --git a/Engine/Source/Programs/Shared/EpicGames.Perforce/PerforceConnection.cs b/Engine/Source/Programs/Shared/EpicGames.Perforce/PerforceConnection.cs
--- a/Engine/Source/Programs/Shared/EpicGames.Perforce/PerforceConnection.cs
+++ b/Engine/Source/Programs/Shared/EpicGames.Perforce/PerforceConnection.cs
@@ -1573,7 +1573,7 @@ namespace EpicGames.Perforce
 		/// <param name="FileSpecs">List of file specifications to query</param>
 		/// <param name="CancellationToken">Token used to cancel the operation</param>
 		/// <returns>List of responses from the server</returns>
-		public Task<PerforceResponseList<FStatRecord>> TryFStatAsync(int AfterChangeNumber, int OnlyChangeNumber, string? Filter, int MaxFiles, FStatOptions Options, string[] FileSpecs, CancellationToken CancellationToken)
+		public async Task<PerforceResponseList<FStatRecord>> TryFStatAsync(int AfterChangeNumber, int OnlyChangeNumber, string? Filter, int MaxFiles, FStatOptions Options, string[] FileSpecs, CancellationToken CancellationToken)
 		{
 			// Build the argument list
 			StringBuilder Arguments = new StringBuilder("fstat");
@@ -1671,7 +1671,12 @@ namespace EpicGames.Perforce
 			}
 
 			// Execute the command
-			return BatchedCommandAsync<FStatRecord>(Arguments.ToString(), FileSpecs, null, CancellationToken);
+			PerforceResponseList<FStatRecord> Records = await BatchedCommandAsync<FStatRecord>(Arguments.ToString(), FileSpecs, null, CancellationToken);
+			if (OnlyChangeNumber != -1 && Records.Count > 0 && Records[Records.Count - 1].Succeeded && Records[Records.Count - 1].Data.Description != null)
+			{
+				Records.RemoveAt(Records.Count - 1);
+			}
+			return Records;
 		}
 
 		#endregion
