Merging ContentHash changes from Horde repo.

diff --git a/Engine/Source/Programs/Shared/EpicGames.Core/ContentHash.cs b/Engine/Source/Programs/Shared/EpicGames.Core/ContentHash.cs
--- a/Engine/Source/Programs/Shared/EpicGames.Core/ContentHash.cs
+++ b/Engine/Source/Programs/Shared/EpicGames.Core/ContentHash.cs
@@ -7,6 +7,8 @@ using System.IO;
 using System.Linq;
 using System.Security.Cryptography;
 using System.Text;
+using System.Text.Json;
+using System.Text.Json.Serialization;
 using System.Threading.Tasks;
 
 namespace EpicGames.Core
@@ -14,11 +16,17 @@ namespace EpicGames.Core
 	/// <summary>
 	/// Stores the hash value for a piece of content as a byte array, allowing it to be used as a dictionary key
 	/// </summary>
+	[JsonConverter(typeof(ContentHashJsonConverter))]
 	public class ContentHash : IEquatable<ContentHash>
 	{
 		public const int LengthMD5 = 16;
 		public const int LengthSHA1 = 20;
 
+		/// <summary>
+		/// Retrieves an empty content hash
+		/// </summary>
+		public static ContentHash Empty { get; } = new ContentHash(Array.Empty<byte>());
+
 		/// <summary>
 		/// The bytes compromising this hash
 		/// </summary>
@@ -54,22 +62,14 @@ namespace EpicGames.Core
 		/// <returns>True if the hashes are equal, false otherwise</returns>
 		public bool Equals(ContentHash? Other)
 		{
-			if((object?)Other == null)
+			if (ReferenceEquals(Other, null))
 			{
 				return false;
 			}
-			if(Bytes.Length != Other.Bytes.Length)
-			{
-				return false;
-			}
-			for(int Idx = 0; Idx < Bytes.Length; Idx++)
+			else
 			{
-				if(Bytes[Idx] != Other.Bytes[Idx])
-				{
-					return false;
-				}
+				return Bytes.SequenceEqual(Other.Bytes);
 			}
-			return true;
 		}
 
 		/// <summary>
@@ -78,11 +78,11 @@ namespace EpicGames.Core
 		/// <param name="A">The first hash to compare</param>
 		/// <param name="B">The second has to compare</param>
 		/// <returns>True if the objects are equal, false otherwise</returns>
-		public static bool operator==(ContentHash A, ContentHash B)
+		public static bool operator ==(ContentHash? A, ContentHash? B)
 		{
-			if((object)A == null)
+			if (ReferenceEquals(A, null))
 			{
-				return ((object)B == null);
+				return ReferenceEquals(B, null);
 			}
 			else
 			{
@@ -96,7 +96,7 @@ namespace EpicGames.Core
 		/// <param name="A">The first hash to compare</param>
 		/// <param name="B">The second has to compare</param>
 		/// <returns>True if the objects are not equal, false otherwise</returns>
-		public static bool operator!=(ContentHash A, ContentHash B)
+		public static bool operator !=(ContentHash? A, ContentHash? B)
 		{
 			return !(A == B);
 		}
@@ -131,7 +131,7 @@ namespace EpicGames.Core
 		/// <returns>New content hash instance containing the hash of the file</returns>
 		public static ContentHash Compute(FileReference Location, HashAlgorithm Algorithm)
 		{
-			using(FileStream Stream = FileReference.Open(Location, FileMode.Open, FileAccess.Read, FileShare.Read))
+			using (FileStream Stream = FileReference.Open(Location, FileMode.Open, FileAccess.Read, FileShare.Read))
 			{
 				return new ContentHash(Algorithm.ComputeHash(Stream));
 			}
@@ -144,7 +144,7 @@ namespace EpicGames.Core
 		/// <returns>New content hash instance containing the hash of the data</returns>
 		public static ContentHash MD5(byte[] Data)
 		{
-			using(MD5 Algorithm = System.Security.Cryptography.MD5.Create())
+			using (MD5 Algorithm = System.Security.Cryptography.MD5.Create())
 			{
 				return Compute(Data, Algorithm);
 			}
@@ -157,7 +157,7 @@ namespace EpicGames.Core
 		/// <returns>New content hash instance containing the hash of the text</returns>
 		public static ContentHash MD5(string Text)
 		{
-			using(MD5 Algorithm = System.Security.Cryptography.MD5.Create())
+			using (MD5 Algorithm = System.Security.Cryptography.MD5.Create())
 			{
 				return Compute(Text, Algorithm);
 			}
@@ -170,7 +170,7 @@ namespace EpicGames.Core
 		/// <returns>New content hash instance containing the hash of the file</returns>
 		public static ContentHash MD5(FileReference Location)
 		{
-			using(MD5 Algorithm = System.Security.Cryptography.MD5.Create())
+			using (MD5 Algorithm = System.Security.Cryptography.MD5.Create())
 			{
 				return Compute(Location, Algorithm);
 			}
@@ -183,7 +183,7 @@ namespace EpicGames.Core
 		/// <returns>New content hash instance containing the hash of the data</returns>
 		public static ContentHash SHA1(byte[] Data)
 		{
-			using(SHA1 Algorithm = System.Security.Cryptography.SHA1.Create())
+			using (SHA1 Algorithm = System.Security.Cryptography.SHA1.Create())
 			{
 				return Compute(Data, Algorithm);
 			}
@@ -196,7 +196,7 @@ namespace EpicGames.Core
 		/// <returns>New content hash instance containing the hash of the text</returns>
 		public static ContentHash SHA1(string Text)
 		{
-			using(SHA1 Algorithm = System.Security.Cryptography.SHA1.Create())
+			using (SHA1 Algorithm = System.Security.Cryptography.SHA1.Create())
 			{
 				return Compute(Text, Algorithm);
 			}
@@ -209,7 +209,7 @@ namespace EpicGames.Core
 		/// <returns>New content hash instance containing the hash of the file</returns>
 		public static ContentHash SHA1(FileReference Location)
 		{
-			using(SHA1 Algorithm = System.Security.Cryptography.SHA1.Create())
+			using (SHA1 Algorithm = System.Security.Cryptography.SHA1.Create())
 			{
 				return Compute(Location, Algorithm);
 			}
@@ -223,7 +223,7 @@ namespace EpicGames.Core
 		public static ContentHash Parse(string Text)
 		{
 			ContentHash? Hash;
-			if(!TryParse(Text, out Hash))
+			if (!TryParse(Text, out Hash))
 			{
 				throw new ArgumentException(String.Format("'{0}' is not a valid content hash", Text));
 			}
@@ -238,7 +238,7 @@ namespace EpicGames.Core
 		public static bool TryParse(string Text, [NotNullWhen(true)] out ContentHash? Hash)
 		{
 			byte[]? Bytes;
-			if(StringUtils.TryParseHexString(Text, out Bytes))
+			if (StringUtils.TryParseHexString(Text, out Bytes))
 			{
 				Hash = new ContentHash(Bytes);
 				return true;
@@ -257,7 +257,7 @@ namespace EpicGames.Core
 		public override int GetHashCode()
 		{
 			int HashCode = Bytes[0];
-			for(int Idx = 1; Idx < Bytes.Length; Idx++)
+			for (int Idx = 1; Idx < Bytes.Length; Idx++)
 			{
 				HashCode = (HashCode * 31) + Bytes[Idx];
 			}
@@ -274,6 +274,24 @@ namespace EpicGames.Core
 		}
 	}
 
+	/// <summary>
+	/// Converts <see cref="ContentHash"/> values to and from JSON
+	/// </summary>
+	public class ContentHashJsonConverter : JsonConverter<ContentHash>
+	{
+		/// <inheritdoc/>
+		public override ContentHash Read(ref Utf8JsonReader Reader, Type TypeToConvert, JsonSerializerOptions Options)
+		{
+			return ContentHash.Parse(Reader.GetString());
+		}
+
+		/// <inheritdoc/>
+		public override void Write(Utf8JsonWriter Writer, ContentHash Value, JsonSerializerOptions Options)
+		{
+			Writer.WriteStringValue(Value.ToString());
+		}
+	}
+
 	/// <summary>
 	/// Utility methods for serializing ContentHash objects
 	/// </summary>
@@ -286,7 +304,7 @@ namespace EpicGames.Core
 		/// <param name="Hash">The hash to write</param>
 		public static void WriteContentHash(this BinaryArchiveWriter Writer, ContentHash? Hash)
 		{
-			if(ReferenceEquals(Hash, null))
+			if (ReferenceEquals(Hash, null))
 			{
 				Writer.WriteByteArray(null);
 			}
@@ -304,7 +322,7 @@ namespace EpicGames.Core
 		public static ContentHash? ReadContentHash(this BinaryArchiveReader Reader)
 		{
 			byte[]? Data = Reader.ReadByteArray();
-			if(Data == null)
+			if (Data == null)
 			{
 				return null;
 			}
