adding ParseBytesString / TryParseBytesString for converting "1.2KB" into 1229 etc.

#rnx

diff --git a/Engine/Source/Programs/Shared/EpicGames.Core/StringUtils.cs b/Engine/Source/Programs/Shared/EpicGames.Core/StringUtils.cs
--- a/Engine/Source/Programs/Shared/EpicGames.Core/StringUtils.cs
+++ b/Engine/Source/Programs/Shared/EpicGames.Core/StringUtils.cs
@@ -434,5 +434,53 @@ namespace EpicGames.Core
 			double Value = Math.Round(BytesAbs / Math.Pow(1024, Power), DecimalPlaces);
 			return $"{(Math.Sign(Bytes) * Value)} {ByteSizes[Power]}";
 		}
+
+		/// <summary>
+		/// Converts a bytes string into bytes. E.g 1.2KB -> 1229
+		/// </summary>
+		/// <param name="BytesString"></param>
+		/// <returns></returns>
+		public static long ParseBytesString( string BytesString )
+		{
+			BytesString = BytesString.Trim();
+
+			int Power = ByteSizes.FindIndex( s => (s != ByteSizes[0]) && BytesString.EndsWith(s, StringComparison.InvariantCultureIgnoreCase ) ); // need to handle 'B' suffix separately
+			if (Power == -1 && BytesString.EndsWith(ByteSizes[0]))
+			{
+				Power = 0;
+			}
+			if (Power != -1)
+			{
+				BytesString = BytesString.Substring(0, BytesString.Length - ByteSizes[Power].Length );
+			}
+
+			double Value = double.Parse(BytesString);
+			if (Power > 0 )
+			{
+				Value *= Math.Pow(1024, Power);
+			}
+
+			return (long)Math.Round(Value);
+		}
+
+		/// <summary>
+		/// Converts a bytes string into bytes. E.g 1.5KB -> 1536
+		/// </summary>
+		/// <param name="BytesString"></param>
+		/// <returns></returns>
+		public static bool TryParseBytesString( string BytesString, out long? Bytes )
+		{
+			try
+			{
+				Bytes = ParseBytesString(BytesString);
+				return true;
+			}
+			catch(Exception)
+			{
+			}
+
+			Bytes = null;
+			return false;
+		}
 	}
 }
