Horde: Throw an error when a preflight changelist does not contain any files.

diff --git a/Engine/Source/Programs/Shared/EpicGames.Perforce.Managed/ManagedWorkspace.cs b/Engine/Source/Programs/Shared/EpicGames.Perforce.Managed/ManagedWorkspace.cs
--- a/Engine/Source/Programs/Shared/EpicGames.Perforce.Managed/ManagedWorkspace.cs
+++ b/Engine/Source/Programs/Shared/EpicGames.Perforce.Managed/ManagedWorkspace.cs
@@ -881,6 +881,10 @@ namespace EpicGames.Perforce.Managed
 				throw new PerforceException($"Changelist {UnshelveChangelist} is not shelved");
 			}
 			DescribeRecord LastRecord = Records[0];
+			if (LastRecord.Files.Count == 0)
+			{
+				throw new PerforceException($"Changelist {UnshelveChangelist} does not contain any shelved files");
+			}
 
 			// query the location of each file
 			PerforceResponseList<WhereRecord> WhereRecords = await Perforce.TryWhereAsync(LastRecord.Files.Select(x => x.DepotFile).ToArray(), CancellationToken);
