Implement IComparer<ReadOnlyUtf8String> from ReadOnlyUtf8StringComparer.

diff --git a/Engine/Source/Programs/Shared/EpicGames.Core/ReadOnlyUtf8String.cs b/Engine/Source/Programs/Shared/EpicGames.Core/ReadOnlyUtf8String.cs
--- a/Engine/Source/Programs/Shared/EpicGames.Core/ReadOnlyUtf8String.cs
+++ b/Engine/Source/Programs/Shared/EpicGames.Core/ReadOnlyUtf8String.cs
@@ -2,6 +2,7 @@
 
 using System;
 using System.Collections.Generic;
+using System.Diagnostics.CodeAnalysis;
 using System.Runtime.CompilerServices;
 using System.Text;
 
@@ -204,7 +205,7 @@ namespace EpicGames.Core
 	/// <summary>
 	/// Comparison classes for utf8 strings
 	/// </summary>
-	public abstract class ReadOnlyUtf8StringComparer : IEqualityComparer<ReadOnlyUtf8String>
+	public abstract class ReadOnlyUtf8StringComparer : IEqualityComparer<ReadOnlyUtf8String>, IComparer<ReadOnlyUtf8String>
 	{
 		/// <summary>
 		/// Ordinal comparer for utf8 strings
@@ -222,6 +223,11 @@ namespace EpicGames.Core
 			{
 				return String.GetHashCode();
 			}
+
+			public override int Compare(ReadOnlyUtf8String StrA, ReadOnlyUtf8String StrB)
+			{
+				return StrA.Span.SequenceCompareTo(StrB.Span);
+			}
 		}
 
 		/// <summary>
@@ -259,6 +265,29 @@ namespace EpicGames.Core
 				return HashCode.ToHashCode();
 			}
 
+			/// <inheritdoc/>
+			public override int Compare(ReadOnlyUtf8String StrA, ReadOnlyUtf8String StrB)
+			{
+				ReadOnlySpan<byte> SpanA = StrA.Span;
+				ReadOnlySpan<byte> SpanB = StrB.Span;
+
+				int Length = Math.Min(StrA.Length, StrB.Length);
+				for (int Idx = 0; Idx < Length; Idx++)
+				{
+					if (SpanA[Idx] != SpanB[Idx])
+					{
+						int UpperA = ToUpper(SpanA[Idx]);
+						int UpperB = ToUpper(SpanB[Idx]);
+						if (UpperA != UpperB)
+						{
+							return UpperA - UpperB;
+						}
+					}
+				}
+
+				return StrA.Length - StrB.Length;
+			}
+
 			/// <summary>
 			/// Convert a character to uppercase
 			/// </summary>
@@ -285,6 +314,9 @@ namespace EpicGames.Core
 
 		/// <inheritdoc/>
 		public abstract int GetHashCode(ReadOnlyUtf8String String);
+
+		/// <inheritdoc/>
+		public abstract int Compare(ReadOnlyUtf8String StrA, ReadOnlyUtf8String StrB);
 	}
 
 	/// <summary>
