Add weakly and strongly typed digest structs for C#, supporting MD5, SHA1, SHA256 (for now).

diff --git a/Engine/Source/Programs/Horde/HordeServer/Rpc/ActionCacheService.cs b/Engine/Source/Programs/Horde/HordeServer/Rpc/ActionCacheService.cs
--- a/Engine/Source/Programs/Horde/HordeServer/Rpc/ActionCacheService.cs
+++ b/Engine/Source/Programs/Horde/HordeServer/Rpc/ActionCacheService.cs
@@ -13,6 +13,8 @@ using System.Threading.Tasks;
 using HordeServer.Storage.Primitives;
 using HordeServer.Utility;
 
+using Digest = Build.Bazel.Remote.Execution.V2.Digest;
+
 namespace HordeServer.Rpc
 {
 	/// <inheritdoc cref="ActionCache"/>
diff --git a/Engine/Source/Programs/Horde/HordeServer/Rpc/ContentStorageService.cs b/Engine/Source/Programs/Horde/HordeServer/Rpc/ContentStorageService.cs
--- a/Engine/Source/Programs/Horde/HordeServer/Rpc/ContentStorageService.cs
+++ b/Engine/Source/Programs/Horde/HordeServer/Rpc/ContentStorageService.cs
@@ -16,6 +16,7 @@ using System.Reflection;
 using System.Threading;
 using System.Threading.Tasks;
 
+using Digest = Build.Bazel.Remote.Execution.V2.Digest;
 using Status = Google.Rpc.Status;
 
 namespace HordeServer.Rpc
diff --git a/Engine/Source/Programs/Shared/EpicGames.Core/Digest.cs b/Engine/Source/Programs/Shared/EpicGames.Core/Digest.cs
new file mode 100644
--- /dev/null
+++ b/Engine/Source/Programs/Shared/EpicGames.Core/Digest.cs
@@ -0,0 +1,354 @@
+// Copyright Epic Games, Inc. All Rights Reserved.
+
+using System;
+using System.Collections.Generic;
+using System.Buffers;
+using System.Security.Cryptography;
+using System.Text;
+using System.Buffers.Binary;
+
+namespace EpicGames.Core
+{
+	/// <summary>
+	/// Struct representing a weakly typed hash value. Counterpart to <see cref="Digest{T}"/> - a strongly typed digest.
+	/// </summary>
+	public struct Digest
+	{
+		/// <summary>
+		/// Memory storing the digest data
+		/// </summary>
+		public ReadOnlyMemory<byte> Memory;
+
+		/// <summary>
+		/// Accessor for the span of memory storing the data
+		/// </summary>
+		public ReadOnlySpan<byte> Span => Memory.Span;
+
+		/// <summary>
+		/// Constructor
+		/// </summary>
+		/// <param name="Memory">Memory to construct from</param>
+		public Digest(ReadOnlyMemory<byte> Memory)
+		{
+			this.Memory = Memory;
+		}
+
+		/// <summary>
+		/// Creates a content hash for a block of data, using a given algorithm.
+		/// </summary>
+		/// <param name="Data">Data to compute the hash for</param>
+		/// <returns>New content hash instance containing the hash of the data</returns>
+		public static Digest<T> Compute<T>(byte[] Data) where T : DigestTraits, new()
+		{
+			using HashAlgorithm Algorithm = Digest<T>.Traits.CreateAlgorithm();
+			return Algorithm.ComputeHash(Data);
+		}
+
+		/// <summary>
+		/// Creates a content hash for a block of data, using a given algorithm.
+		/// </summary>
+		/// <param name="Text">Text to compute the hash for</param>
+		/// <returns>New content hash instance containing the hash of the data</returns>
+		public static Digest<T> Compute<T>(string Text) where T : DigestTraits, new()
+		{
+			return Compute<T>(Encoding.UTF8.GetBytes(Text));
+		}
+
+		/// <summary>
+		/// Creates a content hash for a block of data, using a given algorithm.
+		/// </summary>
+		/// <param name="Data">Data to compute the hash for</param>
+		/// <returns>New content hash instance containing the hash of the data</returns>
+		public static Digest<T> Compute<T>(ReadOnlySpan<byte> Data) where T : DigestTraits, new()
+		{
+			byte[] Value = new byte[Digest<T>.Traits.Length];
+			using (HashAlgorithm Algorithm = Digest<T>.Traits.CreateAlgorithm())
+			{
+				if (!Algorithm.TryComputeHash(Data, Value, out int Written) || Written != Value.Length)
+				{
+					throw new InvalidOperationException("Unable to compute hash for buffer");
+				}
+			}
+			return new Digest<T>(Value);
+		}
+
+		/// <summary>
+		/// Creates a content hash for a string, using a given algorithm.
+		/// </summary>
+		/// <param name="Text">Text to compute a hash for</param>
+		/// <param name="Algorithm">Algorithm to use to create the hash</param>
+		/// <returns>New content hash instance containing the hash of the text</returns>
+		public static ContentHash Compute(string Text, HashAlgorithm Algorithm)
+		{
+			return new ContentHash(Algorithm.ComputeHash(Encoding.Unicode.GetBytes(Text)));
+		}
+
+
+		/// <inheritdoc/>
+		public override bool Equals(object? Obj) => (Obj is Digest Digest) && Digest.Span.SequenceEqual(Span);
+
+		/// <inheritdoc/>
+		public override int GetHashCode() => BinaryPrimitives.ReadInt32LittleEndian(Span);
+
+		/// <inheritdoc/>
+		public override string ToString() => StringUtils.FormatHexString(Memory.Span);
+
+		/// <summary>
+		/// Test two hash values for equality
+		/// </summary>
+		/// <param name="A"></param>
+		/// <param name="B"></param>
+		/// <returns></returns>
+		public static bool operator ==(Digest A, Digest B)
+		{
+			return A.Memory.Span.SequenceEqual(B.Memory.Span);
+		}
+
+		/// <summary>
+		/// Test two hash values for equality
+		/// </summary>
+		/// <param name="A"></param>
+		/// <param name="B"></param>
+		/// <returns></returns>
+		public static bool operator !=(Digest A, Digest B)
+		{
+			return !(A == B);
+		}
+
+		/// <summary>
+		/// Implicit conversion operator from memory objects
+		/// </summary>
+		/// <param name="Memory"></param>
+		public static implicit operator Digest(ReadOnlyMemory<byte> Memory)
+		{
+			return new Digest(Memory);
+		}
+
+		/// <summary>
+		/// Implicit conversion operator from byte arrays
+		/// </summary>
+		/// <param name="Memory"></param>
+		public static implicit operator Digest(byte[] Memory)
+		{
+			return new Digest(Memory);
+		}
+	}
+
+	/// <summary>
+	/// Traits for a hashing algorithm
+	/// </summary>
+	public abstract class DigestTraits
+	{
+		/// <summary>
+		/// Length of the produced hash
+		/// </summary>
+		public int Length { get; }
+
+		/// <summary>
+		/// Constructor
+		/// </summary>
+		/// <param name="Length"></param>
+		public DigestTraits(int Length)
+		{
+			this.Length = Length;
+		}
+
+		/// <summary>
+		/// Creates a HashAlgorithm object
+		/// </summary>
+		/// <returns></returns>
+		public abstract HashAlgorithm CreateAlgorithm();
+	}
+
+	/// <summary>
+	/// Traits for the MD5 hash algorithm
+	/// </summary>
+	public class Md5 : DigestTraits
+	{
+		/// <summary>
+		/// Length of the produced digest
+		/// </summary>
+		public new const int Length = 16;
+
+		/// <summary>
+		/// Constructor
+		/// </summary>
+		public Md5()
+			: base(Length)
+		{
+		}
+
+		/// <inheritdoc/>
+		public override HashAlgorithm CreateAlgorithm() => MD5.Create();
+	}
+
+	/// <summary>
+	/// Traits for the SHA1 hash algorithm
+	/// </summary>
+	public class Sha1 : DigestTraits
+	{
+		/// <summary>
+		/// Length of the produced digest
+		/// </summary>
+		public new const int Length = 20;
+
+		/// <summary>
+		/// Constructor
+		/// </summary>
+		public Sha1()
+			: base(Length)
+		{
+		}
+
+		/// <inheritdoc/>
+		public override HashAlgorithm CreateAlgorithm() => SHA1.Create();
+	}
+
+	/// <summary>
+	/// Traits for the SHA1 hash algorithm
+	/// </summary>
+	public class Sha256 : DigestTraits
+	{
+		/// <summary>
+		/// Length of the produced digest
+		/// </summary>
+		public new const int Length = 32;
+
+		/// <summary>
+		/// Constructor
+		/// </summary>
+		public Sha256()
+			: base(Length)
+		{
+		}
+
+		/// <inheritdoc/>
+		public override HashAlgorithm CreateAlgorithm() => SHA256.Create();
+	}
+
+	/// <summary>
+	/// Generic HashValue implementation
+	/// </summary>
+	public struct Digest<T> where T : DigestTraits, new()
+	{
+		/// <summary>
+		/// Traits instance
+		/// </summary>
+		public static T Traits { get; } = new T();
+
+		/// <summary>
+		/// Length of a hash value
+		/// </summary>
+		public int Length => Traits.Length;
+
+		/// <summary>
+		/// Memory storing the digest data
+		/// </summary>
+		public ReadOnlyMemory<byte> Memory;
+
+		/// <summary>
+		/// Accessor for the span of memory storing the data
+		/// </summary>
+		public ReadOnlySpan<byte> Span => Memory.Span;
+
+		/// <summary>
+		/// Constructor
+		/// </summary>
+		/// <param name="Memory">Memory to construct from</param>
+		public Digest(ReadOnlyMemory<byte> Memory)
+		{
+			this.Memory = Memory;
+		}
+
+		/// <inheritdoc/>
+		public override bool Equals(object? Obj) => (Obj is Digest<T> Hash) && Hash.Span.SequenceEqual(Span);
+
+		/// <inheritdoc/>
+		public override int GetHashCode() => BinaryPrimitives.ReadInt32LittleEndian(Span);
+
+		/// <inheritdoc/>
+		public override string ToString() => StringUtils.FormatHexString(Memory.Span);
+
+		/// <summary>
+		/// Test two hash values for equality
+		/// </summary>
+		/// <param name="A"></param>
+		/// <param name="B"></param>
+		/// <returns></returns>
+		public static bool operator ==(Digest<T> A, Digest<T> B) => A.Span.SequenceEqual(B.Span);
+
+		/// <summary>
+		/// Test two hash values for equality
+		/// </summary>
+		/// <param name="A"></param>
+		/// <param name="B"></param>
+		/// <returns></returns>
+		public static bool operator !=(Digest<T> A, Digest<T> B) => !A.Span.SequenceEqual(B.Span);
+
+		/// <summary>
+		/// Implicit conversion operator from memory objects
+		/// </summary>
+		/// <param name="Memory"></param>
+		public static implicit operator Digest<T>(ReadOnlyMemory<byte> Memory)
+		{
+			return new Digest<T>(Memory);
+		}
+
+		/// <summary>
+		/// Implicit conversion operator from byte arrays
+		/// </summary>
+		/// <param name="Memory"></param>
+		public static implicit operator Digest<T>(byte[] Memory)
+		{
+			return new Digest<T>(Memory);
+		}
+	}
+
+	/// <summary>
+	/// Extension methods for dealing with digests
+	/// </summary>
+	public static class DigestExtensions
+	{
+		/// <summary>
+		/// Read a digest from a memory reader
+		/// </summary>
+		/// <param name="Reader"></param>
+		/// <returns></returns>
+		public static Digest ReadDigest(this MemoryReader Reader)
+		{
+			return new Digest(Reader.ReadVariableLengthBytes());
+		}
+
+		/// <summary>
+		/// Read a strongly-typed digest from a memory reader
+		/// </summary>
+		/// <typeparam name="T"></typeparam>
+		/// <param name="Reader"></param>
+		/// <returns></returns>
+		public static Digest<T> ReadDigest<T>(this MemoryReader Reader) where T : DigestTraits, new()
+		{
+			return new Digest<T>(Reader.ReadFixedLengthBytes(Digest<T>.Traits.Length));
+		}
+
+		/// <summary>
+		/// Write a digest to a memory writer
+		/// </summary>
+		/// <param name="Writer"></param>
+		/// <param name="Digest"></param>
+		public static void WriteDigest(this MemoryWriter Writer, Digest Digest)
+		{
+			Writer.WriteVariableLengthBytes(Digest.Span);
+		}
+
+		/// <summary>
+		/// Write a strongly typed digest to a memory writer
+		/// </summary>
+		/// <typeparam name="T"></typeparam>
+		/// <param name="Writer"></param>
+		/// <param name="Digest"></param>
+		public static void WriteDigest<T>(this MemoryWriter Writer, Digest<T> Digest) where T : DigestTraits, new()
+		{
+			Writer.WriteFixedLengthBytes(Digest.Span);
+		}
+	}
+}
