Add a compact encoding of Perforce filetypes.

diff --git a/Engine/Source/Programs/Shared/EpicGames.Perforce/Records/FileType.cs b/Engine/Source/Programs/Shared/EpicGames.Perforce/Records/FileType.cs
new file mode 100644
--- /dev/null
+++ b/Engine/Source/Programs/Shared/EpicGames.Perforce/Records/FileType.cs
@@ -0,0 +1,432 @@
+// Copyright Epic Games, Inc. All Rights Reserved.
+
+using EpicGames.Core;
+using System;
+using System.Collections.Generic;
+using System.Text;
+
+namespace EpicGames.Perforce.Managed
+{
+	/// <summary>
+	/// Base file type
+	/// </summary>
+	public enum FileTypeBase : uint
+	{
+		/// <summary>
+		/// Text file
+		/// </summary>
+		Text = 0,
+
+		/// <summary>
+		/// Non-text file
+		/// </summary>
+		Binary = 1,
+
+		/// <summary>
+		/// Symbolic link
+		/// </summary>
+		Symlink = 2,
+
+		/// <summary>
+		/// Multi-forked Macintosh file
+		/// </summary>
+		Apple = 3,
+
+		/// <summary>
+		/// Macintosh resource fork
+		/// </summary>
+		Resource = 4,
+
+		/// <summary>
+		/// Unicode file
+		/// </summary>
+		Unicode = 5,
+
+		/// <summary>
+		/// Unicode file, encoded as UTF-16
+		/// </summary>
+		Utf16 = 6,
+	}
+
+	/// <summary>
+	/// Modifiers on top of the base filetype
+	/// </summary>
+	[Flags]
+	public enum FileTypeModifiers : uint
+	{
+		/// <summary>
+		/// File is always writable on the client
+		/// </summary>
+		AlwaysWritable = 0x8, // +w
+
+		/// <summary>
+		/// Executable bit set on client
+		/// </summary>
+		Executable = 0x10, // +x
+
+		/// <summary>
+		/// RCS keyword expansion
+		/// </summary>
+		KeywordExpansion = 0x20, // +k
+
+		/// <summary>
+		/// Old-style keyword expansion
+		/// </summary>
+		KeywordExpansionOld = 0x40, // +ko
+
+		/// <summary>
+		/// Exclusive open (locking)
+		/// </summary>
+		ExclusiveOpen = 0x80, // +l
+
+		/// <summary>
+		/// Server stores the full compressed version of each file revision
+		/// </summary>
+		StoreFullRevisionsCompressed = 0x100, // +C
+
+		/// <summary>
+		/// Server stores deltas in RCS format
+		/// </summary>
+		StoreDeltas = 0x200, // +D
+
+		/// <summary>
+		/// Server stores full file per revision, uncompressed
+		/// </summary>
+		StoreFullRevisionsUncompressed = 0x400, // +F
+
+		/// <summary>
+		/// Preserve original modtime
+		/// </summary>
+		Modtime = 0x800, // +m
+
+		/// <summary>
+		/// Archive trigger required
+		/// </summary>
+		ArchiveTriggerRequired = 0x1000, // +X
+	}
+
+	/// <summary>
+	/// Indicates the type of a file
+	/// </summary>
+	public struct FileType
+	{
+		/// <summary>
+		/// Size of the object when serialized to memory
+		/// </summary>
+		public const int SerializedSize = sizeof(ushort);
+
+		/// <summary>
+		/// Mask for the base bits
+		/// </summary>
+		public const uint BaseMask = 7;
+
+		/// <summary>
+		/// Mask for the modifiers
+		/// </summary>
+		public const uint ModifiersMask = 0xff & ~BaseMask;
+
+		/// <summary>
+		/// Array of file type bases, with order matching <see cref="FileTypeBase"/>
+		/// </summary>
+		static ReadOnlyUtf8String[] BaseNames =
+		{
+			"text",
+			"binary",
+			"symlink",
+			"apple",
+			"resource",
+			"unicode",
+			"utf16"
+		};
+
+		/// <summary>
+		/// Encoded value of the filetype
+		/// </summary>
+		public uint Encoded { get; }
+
+		/// <summary>
+		/// Base file type
+		/// </summary>
+		public FileTypeBase Base => (FileTypeBase)(Encoded & BaseMask);
+
+		/// <summary>
+		/// File type modifiers
+		/// </summary>
+		public FileTypeModifiers Modifiers => (FileTypeModifiers)(Encoded & ModifiersMask);
+
+		/// <summary>
+		/// Number of revisions to be stored
+		/// </summary>
+		public int NumRevisions => (int)(Encoded >> 16);
+
+		/// <summary>
+		/// Constructor
+		/// </summary>
+		/// <param name="Encoded">The encoded value</param>
+		public FileType(uint Encoded)
+		{
+			this.Encoded = Encoded;
+		}
+
+		/// <summary>
+		/// Constructor
+		/// </summary>
+		/// <param name="Base">Base file type</param>
+		/// <param name="Modifiers">File type modifiers</param>
+		/// <param name="NumRevisions">Number of revisions to store</param>
+		public FileType(FileTypeBase Base, FileTypeModifiers Modifiers = 0, int NumRevisions = 0)
+		{
+			Encoded = (uint)Base | (uint)Modifiers | (uint)(NumRevisions << 16);
+		}
+
+		/// <summary>
+		/// Parse a string as a filetype
+		/// </summary>
+		/// <param name="Text"></param>
+		/// <returns></returns>
+		public static FileType Parse(ReadOnlySpan<byte> Text)
+		{
+			if (TryParse(Text, out FileType Type))
+			{
+				return Type;
+			}
+			else
+			{
+				throw new InvalidCastException($"Cannot parse text ('{Encoding.UTF8.GetString(Text)}') as FileType");
+			}
+		}
+
+		/// <summary>
+		/// Try to parse a utf8 string as a filetype
+		/// </summary>
+		/// <param name="Text"></param>
+		/// <param name="Type"></param>
+		/// <returns></returns>
+		public static bool TryParse(ReadOnlySpan<byte> Text, out FileType Type)
+		{
+			int PlusIdx = Text.IndexOf((byte)'+');
+			if (PlusIdx == -1)
+			{
+				if (TryParseBase(Text, out FileTypeBase Base))
+				{
+					Type = new FileType(Base);
+					return true;
+				}
+			}
+			else
+			{
+				if (TryParseBase(Text.Slice(0, PlusIdx), out FileTypeBase Base) && TryParseModifiers(Text.Slice(PlusIdx + 1), out (FileTypeModifiers, int) Modifiers))
+				{
+					Type = new FileType(Base, Modifiers.Item1, Modifiers.Item2);
+					return true;
+				}
+			}
+
+			Type = default;
+			return false;
+		}
+
+		/// <summary>
+		/// Try to parse a utf8 string as a filetype base
+		/// </summary>
+		/// <param name="Text"></param>
+		/// <param name="Base"></param>
+		/// <returns></returns>
+		public static bool TryParseBase(ReadOnlySpan<byte> Text, out FileTypeBase Base)
+		{
+			for (int Idx = 0; Idx < BaseNames.Length; Idx++)
+			{
+				if (Text.SequenceEqual(BaseNames[Idx].Span))
+				{
+					Base = (FileTypeBase)Idx;
+					return true;
+				}
+			}
+
+			Base = 0;
+			return false;
+		}
+
+		/// <summary>
+		/// Try to parse modifiers from a utf8 string
+		/// </summary>
+		/// <param name="Text"></param>
+		/// <param name="Modifiers"></param>
+		/// <returns></returns>
+		public static bool TryParseModifiers(ReadOnlySpan<byte> Text, out (FileTypeModifiers, int) Result)
+		{
+			FileTypeModifiers Modifiers = 0;
+			int NumRevisions = 0;
+
+			for (int Idx = 0; Idx < Text.Length; Idx++)
+			{
+				switch(Text[Idx])
+				{
+					case (byte)'w':
+						Modifiers |= FileTypeModifiers.AlwaysWritable;
+						break;
+					case (byte)'x':
+						Modifiers |= FileTypeModifiers.Executable;
+						break;
+					case (byte)'k':
+						Modifiers |= FileTypeModifiers.KeywordExpansion;
+						break;
+					case (byte)'o':
+						Modifiers |= FileTypeModifiers.KeywordExpansionOld;
+						break;
+					case (byte)'l':
+						Modifiers |= FileTypeModifiers.ExclusiveOpen;
+						break;
+					case (byte)'C':
+						Modifiers |= FileTypeModifiers.StoreFullRevisionsCompressed;
+						break;
+					case (byte)'D':
+						Modifiers |= FileTypeModifiers.StoreDeltas;
+						break;
+					case (byte)'F':
+						Modifiers |= FileTypeModifiers.StoreFullRevisionsUncompressed;
+						break;
+					case (byte)'S':
+						while (Idx + 1 < Text.Length && (Text[Idx + 1] >= '0' && Text[Idx + 1] <= '9'))
+						{
+							NumRevisions = (NumRevisions * 10) + (Text[++Idx] - '0');
+						}
+						if(NumRevisions == 0)
+						{
+							NumRevisions = 1;
+						}
+						break;
+					case (byte)'m':
+						Modifiers |= FileTypeModifiers.Modtime;
+						break;
+					default:
+						Result = (0, 0);
+						return false;
+				}
+			}
+
+			Result = (Modifiers, NumRevisions);
+			return true;
+		}
+
+		/// <summary>
+		/// Compares two filetypes for equality
+		/// </summary>
+		/// <param name="A"></param>
+		/// <param name="B"></param>
+		/// <returns></returns>
+		public static bool operator ==(FileType A, FileType B)
+		{
+			return A.Encoded == B.Encoded;
+		}
+
+		/// <summary>
+		/// Compares two filetypes for equality
+		/// </summary>
+		/// <param name="A"></param>
+		/// <param name="B"></param>
+		/// <returns></returns>
+		public static bool operator !=(FileType A, FileType B)
+		{
+			return A.Encoded != B.Encoded;
+		}
+
+		/// <inheritdoc/>
+		public override bool Equals(object? obj)
+		{
+			return (obj is FileType Type) && Type.Encoded == Encoded;
+		}
+
+		/// <inheritdoc/>
+		public override int GetHashCode()
+		{
+			return (int)Encoded;
+		}
+
+		/// <inheritdoc/>
+		public override string ToString()
+		{
+			StringBuilder Type = new StringBuilder(BaseNames[(int)Base].ToString());
+
+			int RemainingModifiers = (int)Modifiers;
+			if (RemainingModifiers != 0 || NumRevisions > 0)
+			{
+				Type.Append('+');
+				while (RemainingModifiers != 0)
+				{
+					FileTypeModifiers Modifier = (FileTypeModifiers)(RemainingModifiers & ~(RemainingModifiers - 1));
+					switch (Modifier)
+					{
+						case FileTypeModifiers.AlwaysWritable:
+							Type.Append('w');
+							break;
+						case FileTypeModifiers.Executable:
+							Type.Append('x');
+							break;
+						case FileTypeModifiers.KeywordExpansion:
+							Type.Append('k');
+							break;
+						case FileTypeModifiers.KeywordExpansionOld:
+							Type.Append('o');
+							break;
+						case FileTypeModifiers.ExclusiveOpen:
+							Type.Append('l');
+							break;
+						case FileTypeModifiers.StoreFullRevisionsCompressed:
+							Type.Append('C');
+							break;
+						case FileTypeModifiers.StoreDeltas:
+							Type.Append('D');
+							break;
+						case FileTypeModifiers.StoreFullRevisionsUncompressed:
+							Type.Append('F');
+							break;
+						case FileTypeModifiers.Modtime:
+							Type.Append('m');
+							break;
+						case FileTypeModifiers.ArchiveTriggerRequired:
+							Type.Append('X');
+							break;
+					}
+					RemainingModifiers ^= (int)Modifier;
+				}
+
+				if (NumRevisions == 1)
+				{
+					Type.Append('S');
+				}
+				else if (NumRevisions > 1)
+				{
+					Type.Append($"S{NumRevisions}");
+				}
+			}
+
+			return Type.ToString();
+		}
+	}
+
+	/// <summary>
+	/// Extension methods for file types
+	/// </summary>
+	public static class FileTypeExtensions
+	{
+		/// <summary>
+		/// Constructor for reading a file info from disk
+		/// </summary>
+		/// <param name="Reader">Binary reader to read data from</param>
+		public static FileType ReadFileType(this MemoryReader Reader)
+		{
+			return new FileType(Reader.ReadUInt32());
+		}
+
+		/// <summary>
+		/// Save the file info to disk
+		/// </summary>
+		/// <param name="Writer">Writer to output to</param>
+		/// <param name="FileType">The file type to serialize</param>
+		public static void WriteFileType(this MemoryWriter Writer, FileType FileType)
+		{
+			Writer.WriteUInt32(FileType.Encoded);
+		}
+	}
+}
