Add overloads for reading/writing 16 bit integers.

diff --git a/Engine/Source/Programs/Shared/EpicGames.Core/MemoryReader.cs b/Engine/Source/Programs/Shared/EpicGames.Core/MemoryReader.cs
--- a/Engine/Source/Programs/Shared/EpicGames.Core/MemoryReader.cs
+++ b/Engine/Source/Programs/Shared/EpicGames.Core/MemoryReader.cs
@@ -94,6 +94,28 @@ namespace EpicGames.Core
 			return Value;
 		}
 
+		/// <summary>
+		/// Reads an int16
+		/// </summary>
+		/// <returns>The value read from the buffer</returns>
+		public short ReadInt16()
+		{
+			short Value = BinaryPrimitives.ReadInt16LittleEndian(Memory.Span.Slice(Offset));
+			Offset += sizeof(short);
+			return Value;
+		}
+
+		/// <summary>
+		/// Reads a uint16
+		/// </summary>
+		/// <returns>The value read from the buffer</returns>
+		public ushort ReadUInt16()
+		{
+			ushort Value = BinaryPrimitives.ReadUInt16LittleEndian(Memory.Span.Slice(Offset));
+			Offset += sizeof(ushort);
+			return Value;
+		}
+
 		/// <summary>
 		/// Reads an int32
 		/// </summary>
