Misisng files.

diff --git a/Engine/Source/Programs/Shared/EpicGames.Perforce.Managed/StreamSnapshot.cs b/Engine/Source/Programs/Shared/EpicGames.Perforce.Managed/StreamSnapshot.cs
--- a/Engine/Source/Programs/Shared/EpicGames.Perforce.Managed/StreamSnapshot.cs
+++ b/Engine/Source/Programs/Shared/EpicGames.Perforce.Managed/StreamSnapshot.cs
@@ -1,3 +1,5 @@
+// Copyright Epic Games, Inc. All Rights Reserved.
+
 using EpicGames.Core;
 using System;
 using System.Collections.Generic;
diff --git a/Engine/Source/Programs/Shared/EpicGames.Perforce.Managed/StreamSnapshotBuilder.cs b/Engine/Source/Programs/Shared/EpicGames.Perforce.Managed/StreamSnapshotBuilder.cs
--- a/Engine/Source/Programs/Shared/EpicGames.Perforce.Managed/StreamSnapshotBuilder.cs
+++ b/Engine/Source/Programs/Shared/EpicGames.Perforce.Managed/StreamSnapshotBuilder.cs
@@ -1,3 +1,5 @@
+// Copyright Epic Games, Inc. All Rights Reserved.
+
 using EpicGames.Core;
 using System;
 using System.Collections.Generic;
diff --git a/Engine/Source/Programs/Shared/EpicGames.Perforce/Records/FileType.cs b/Engine/Source/Programs/Shared/EpicGames.Perforce/Records/FileType.cs
--- a/Engine/Source/Programs/Shared/EpicGames.Perforce/Records/FileType.cs
+++ b/Engine/Source/Programs/Shared/EpicGames.Perforce/Records/FileType.cs
@@ -46,6 +46,11 @@ namespace EpicGames.Perforce.Managed
 		/// Unicode file, encoded as UTF-16
 		/// </summary>
 		Utf16 = 6,
+
+		/// <summary>
+		/// Unicode file, encoded as UTF-8
+		/// </summary>
+		Utf8 = 7,
 	}
 
 	/// <summary>
@@ -113,7 +118,7 @@ namespace EpicGames.Perforce.Managed
 		/// <summary>
 		/// Size of the object when serialized to memory
 		/// </summary>
-		public const int SerializedSize = sizeof(ushort);
+		public const int SerializedSize = sizeof(uint);
 
 		/// <summary>
 		/// Mask for the base bits
@@ -136,7 +141,8 @@ namespace EpicGames.Perforce.Managed
 			"apple",
 			"resource",
 			"unicode",
-			"utf16"
+			"utf16",
+			"utf8"
 		};
 
 		/// <summary>
