Horde: Change back to using strings for storing file types in FileContentId objects, to avoid invalidating existing workspace caches.

diff --git a/Engine/Source/Programs/Shared/EpicGames.Perforce.Managed/CachedFileInfo.cs b/Engine/Source/Programs/Shared/EpicGames.Perforce.Managed/CachedFileInfo.cs
--- a/Engine/Source/Programs/Shared/EpicGames.Perforce.Managed/CachedFileInfo.cs
+++ b/Engine/Source/Programs/Shared/EpicGames.Perforce.Managed/CachedFileInfo.cs
@@ -107,7 +107,7 @@ namespace EpicGames.Perforce.Managed
 
 		public static int GetSerializedSize(this CachedFileInfo FileInfo)
 		{
-			return FileContentId.SerializedSize + sizeof(ulong) + sizeof(long) + sizeof(long) + sizeof(byte) + sizeof(uint);
+			return FileInfo.ContentId.GetSerializedSize() + sizeof(ulong) + sizeof(long) + sizeof(long) + sizeof(byte) + sizeof(uint);
 		}
 	}
 }
diff --git a/Engine/Source/Programs/Shared/EpicGames.Perforce.Managed/FileContentId.cs b/Engine/Source/Programs/Shared/EpicGames.Perforce.Managed/FileContentId.cs
--- a/Engine/Source/Programs/Shared/EpicGames.Perforce.Managed/FileContentId.cs
+++ b/Engine/Source/Programs/Shared/EpicGames.Perforce.Managed/FileContentId.cs
@@ -14,19 +14,17 @@ namespace EpicGames.Perforce.Managed
 	[DebuggerDisplay("{Digest} ({Type})")]
 	class FileContentId
 	{
-		public static int SerializedSize => Digest<Md5>.Length + FileType.SerializedSize;
-
 		public Digest<Md5> Digest
 		{
 			get;
 		}
 
-		public FileType Type
+		public ReadOnlyUtf8String Type
 		{
 			get;
 		}
 
-		public FileContentId(Digest<Md5> Digest, FileType Type)
+		public FileContentId(Digest<Md5> Digest, ReadOnlyUtf8String Type)
 		{
 			this.Digest = Digest;
 			this.Type = Type;
@@ -48,14 +46,19 @@ namespace EpicGames.Perforce.Managed
 		public static FileContentId ReadFileContentId(this MemoryReader Reader)
 		{
 			Digest<Md5> Digest = Reader.ReadDigest<Md5>();
-			FileType Type = Reader.ReadFileType();
+			ReadOnlyUtf8String Type = Reader.ReadString();
 			return new FileContentId(Digest, Type);
 		}
 
 		public static void WriteFileContentId(this MemoryWriter Writer, FileContentId FileContentId)
 		{
 			Writer.WriteDigest<Md5>(FileContentId.Digest);
-			Writer.WriteFileType(FileContentId.Type);
+			Writer.WriteString(FileContentId.Type);
+		}
+
+		public static int GetSerializedSize(this FileContentId FileContentId)
+		{
+			return Digest<Md5>.Length + FileContentId.Type.GetSerializedSize();
 		}
 	}
 }
diff --git a/Engine/Source/Programs/Shared/EpicGames.Perforce.Managed/ManagedWorkspace.cs b/Engine/Source/Programs/Shared/EpicGames.Perforce.Managed/ManagedWorkspace.cs
--- a/Engine/Source/Programs/Shared/EpicGames.Perforce.Managed/ManagedWorkspace.cs
+++ b/Engine/Source/Programs/Shared/EpicGames.Perforce.Managed/ManagedWorkspace.cs
@@ -1470,7 +1470,7 @@ namespace EpicGames.Perforce.Managed
 					}
 
 					byte[] DigestBytes = StringUtils.ParseHexString(Record.Digest.Span);
-					FileContentId ContentId = new FileContentId(DigestBytes, FileType.Parse(Record.HeadType.Span));
+					FileContentId ContentId = new FileContentId(DigestBytes, Record.HeadType);
 
 					// Get the depot file and revision
 					byte[] DepotFileAndRevisionData = new byte[Record.DepotFile.Length + 1 + Record.HaveRev.Length];
diff --git a/Engine/Source/Programs/Shared/EpicGames.Perforce.Managed/StreamFileInfo.cs b/Engine/Source/Programs/Shared/EpicGames.Perforce.Managed/StreamFileInfo.cs
--- a/Engine/Source/Programs/Shared/EpicGames.Perforce.Managed/StreamFileInfo.cs
+++ b/Engine/Source/Programs/Shared/EpicGames.Perforce.Managed/StreamFileInfo.cs
@@ -92,7 +92,7 @@ namespace EpicGames.Perforce.Managed
 		/// </summary>
 		public static int GetSerializedSize(this StreamFileInfo FileInfo)
 		{
-			return FileInfo.Name.GetSerializedSize() + sizeof(long) + FileContentId.SerializedSize + FileInfo.DepotFileAndRevision.GetSerializedSize();
+			return FileInfo.Name.GetSerializedSize() + sizeof(long) + FileInfo.ContentId.GetSerializedSize() + FileInfo.DepotFileAndRevision.GetSerializedSize();
 		}
 	}
 }
diff --git a/Engine/Source/Programs/Shared/EpicGames.Perforce.Managed/WorkspaceFileInfo.cs b/Engine/Source/Programs/Shared/EpicGames.Perforce.Managed/WorkspaceFileInfo.cs
--- a/Engine/Source/Programs/Shared/EpicGames.Perforce.Managed/WorkspaceFileInfo.cs
+++ b/Engine/Source/Programs/Shared/EpicGames.Perforce.Managed/WorkspaceFileInfo.cs
@@ -121,7 +121,7 @@ namespace EpicGames.Perforce.Managed
 
 		public static int GetSerializedSize(this WorkspaceFileInfo FileInfo)
 		{
-			return FileInfo.Name.GetSerializedSize() + sizeof(long) + sizeof(long) + sizeof(byte) + FileContentId.SerializedSize;
+			return FileInfo.Name.GetSerializedSize() + sizeof(long) + sizeof(long) + sizeof(byte) + FileInfo.ContentId.GetSerializedSize();
 		}
 	}
 }
