Add Parse() methods to Digest class.

diff --git a/Engine/Source/Programs/Shared/EpicGames.Core/Digest.cs b/Engine/Source/Programs/Shared/EpicGames.Core/Digest.cs
--- a/Engine/Source/Programs/Shared/EpicGames.Core/Digest.cs
+++ b/Engine/Source/Programs/Shared/EpicGames.Core/Digest.cs
@@ -73,16 +73,24 @@ namespace EpicGames.Core
 		}
 
 		/// <summary>
-		/// Creates a content hash for a string, using a given algorithm.
+		/// Parses a digest from the given hex string
 		/// </summary>
-		/// <param name="Text">Text to compute a hash for</param>
-		/// <param name="Algorithm">Algorithm to use to create the hash</param>
-		/// <returns>New content hash instance containing the hash of the text</returns>
-		public static ContentHash Compute(string Text, HashAlgorithm Algorithm)
+		/// <param name="Text"></param>
+		/// <returns></returns>
+		public static Digest Parse(string Text)
 		{
-			return new ContentHash(Algorithm.ComputeHash(Encoding.Unicode.GetBytes(Text)));
+			return new Digest(StringUtils.ParseHexString(Text));
 		}
 
+		/// <summary>
+		/// Parses a digest from the given hex string
+		/// </summary>
+		/// <param name="Text"></param>
+		/// <returns></returns>
+		public static Digest<T> Parse<T>(string Text) where T : DigestTraits, new()
+		{
+			return new Digest<T>(StringUtils.ParseHexString(Text));
+		}
 
 		/// <inheritdoc/>
 		public override bool Equals(object? Obj) => (Obj is Digest Digest) && Digest.Span.SequenceEqual(Span);
