Horde: Make ManagedWorkspace classes used to represent snapshots public, so we can generate them from external tools.

diff --git a/Engine/Source/Programs/Shared/EpicGames.Perforce.Managed/FileContentId.cs b/Engine/Source/Programs/Shared/EpicGames.Perforce.Managed/FileContentId.cs
--- a/Engine/Source/Programs/Shared/EpicGames.Perforce.Managed/FileContentId.cs
+++ b/Engine/Source/Programs/Shared/EpicGames.Perforce.Managed/FileContentId.cs
@@ -12,7 +12,7 @@ using EpicGames.Core;
 namespace EpicGames.Perforce.Managed
 {
 	[DebuggerDisplay("{Digest} ({Type})")]
-	class FileContentId
+	public class FileContentId
 	{
 		public Digest<Md5> Digest
 		{
diff --git a/Engine/Source/Programs/Shared/EpicGames.Perforce.Managed/StreamDirectoryInfo.cs b/Engine/Source/Programs/Shared/EpicGames.Perforce.Managed/StreamDirectoryInfo.cs
--- a/Engine/Source/Programs/Shared/EpicGames.Perforce.Managed/StreamDirectoryInfo.cs
+++ b/Engine/Source/Programs/Shared/EpicGames.Perforce.Managed/StreamDirectoryInfo.cs
@@ -14,7 +14,7 @@ namespace EpicGames.Perforce.Managed
 	/// <summary>
 	/// Information about a directory within a stream
 	/// </summary>
-	class StreamDirectoryInfo
+	public class StreamDirectoryInfo
 	{
 		/// <summary>
 		/// Map of name to file within the directory
diff --git a/Engine/Source/Programs/Shared/EpicGames.Perforce.Managed/StreamFileInfo.cs b/Engine/Source/Programs/Shared/EpicGames.Perforce.Managed/StreamFileInfo.cs
--- a/Engine/Source/Programs/Shared/EpicGames.Perforce.Managed/StreamFileInfo.cs
+++ b/Engine/Source/Programs/Shared/EpicGames.Perforce.Managed/StreamFileInfo.cs
@@ -22,7 +22,7 @@ namespace EpicGames.Perforce.Managed
 	/// Metadata for a Perforce file
 	/// </summary>
 	[DebuggerDisplay("{Name}")]
-	class StreamFileInfo
+	public class StreamFileInfo
 	{
 		/// <summary>
 		/// Name of this file
diff --git a/Engine/Source/Programs/Shared/EpicGames.Perforce.Managed/StreamSnapshot.cs b/Engine/Source/Programs/Shared/EpicGames.Perforce.Managed/StreamSnapshot.cs
--- a/Engine/Source/Programs/Shared/EpicGames.Perforce.Managed/StreamSnapshot.cs
+++ b/Engine/Source/Programs/Shared/EpicGames.Perforce.Managed/StreamSnapshot.cs
@@ -14,7 +14,7 @@ namespace EpicGames.Perforce.Managed
 	/// <summary>
 	/// Interface for a stream snapshot
 	/// </summary>
-	abstract class StreamSnapshot
+	public abstract class StreamSnapshot
 	{
 		/// <summary>
 		/// Empty snapshot instance
diff --git a/Engine/Source/Programs/Shared/EpicGames.Perforce.Managed/StreamSnapshotBuilder.cs b/Engine/Source/Programs/Shared/EpicGames.Perforce.Managed/StreamSnapshotBuilder.cs
--- a/Engine/Source/Programs/Shared/EpicGames.Perforce.Managed/StreamSnapshotBuilder.cs
+++ b/Engine/Source/Programs/Shared/EpicGames.Perforce.Managed/StreamSnapshotBuilder.cs
@@ -10,7 +10,7 @@ namespace EpicGames.Perforce.Managed
 	/// <summary>
 	/// Utility class to build a stream snapshot in memory
 	/// </summary>
-	class StreamSnapshotBuilder
+	public class StreamSnapshotBuilder
 	{
 		/// <summary>
 		/// Map of name to file within the directory
diff --git a/Engine/Source/Programs/Shared/EpicGames.Perforce.Managed/StreamSnapshotFromMemory.cs b/Engine/Source/Programs/Shared/EpicGames.Perforce.Managed/StreamSnapshotFromMemory.cs
--- a/Engine/Source/Programs/Shared/EpicGames.Perforce.Managed/StreamSnapshotFromMemory.cs
+++ b/Engine/Source/Programs/Shared/EpicGames.Perforce.Managed/StreamSnapshotFromMemory.cs
@@ -11,7 +11,10 @@ using System.Threading.Tasks;
 
 namespace EpicGames.Perforce.Managed
 {
-	class StreamSnapshotFromMemory : StreamSnapshot
+	/// <summary>
+	/// Stores the contents of a stream in memory
+	/// </summary>
+	public class StreamSnapshotFromMemory : StreamSnapshot
 	{
 		/// <summary>
 		/// The current signature for saved directory objects
