Horde: Fix false positives being logged in code to track down StatsV2 sync issue.

diff --git a/Engine/Source/Programs/Shared/EpicGames.Perforce.Managed/ManagedWorkspace.cs b/Engine/Source/Programs/Shared/EpicGames.Perforce.Managed/ManagedWorkspace.cs
--- a/Engine/Source/Programs/Shared/EpicGames.Perforce.Managed/ManagedWorkspace.cs
+++ b/Engine/Source/Programs/Shared/EpicGames.Perforce.Managed/ManagedWorkspace.cs
@@ -1862,7 +1862,7 @@ namespace EpicGames.Perforce.Managed
 					}
 				}
 
-				WorkspaceFileToSync? StatsFile = FilesToSync.FirstOrDefault(x => x.StreamFile.Name == StatsFileName);
+				WorkspaceFileToSync? StatsFile = FilesToSync[BeginIdx..EndIdx].FirstOrDefault(x => x.StreamFile.Name == StatsFileName);
 				if(StatsFile != null)
 				{
 					for (int Idx = BeginIdx; Idx < EndIdx; Idx++)
