Horde: Store the stream directory hash on the workspace directory info, and invalidate it whenever a file is found to be modified. This allows short-circuiting operations to modify directory structure once the initial state is cached.

diff --git a/Engine/Source/Programs/Shared/EpicGames.Perforce.Managed/ManagedWorkspace.cs b/Engine/Source/Programs/Shared/EpicGames.Perforce.Managed/ManagedWorkspace.cs
--- a/Engine/Source/Programs/Shared/EpicGames.Perforce.Managed/ManagedWorkspace.cs
+++ b/Engine/Source/Programs/Shared/EpicGames.Perforce.Managed/ManagedWorkspace.cs
@@ -69,6 +69,22 @@ namespace EpicGames.Perforce.Managed
 		}
 	}
 
+	/// <summary>
+	/// Version number for managed workspace cache files
+	/// </summary>
+	enum ManagedWorkspaceVersion
+	{
+		/// <summary>
+		/// Initial version number
+		/// </summary>
+		Initial = 2,
+
+		/// <summary>
+		/// Including stream directory digests in workspace directories
+		/// </summary>
+		AddDigest = 3,
+	}
+
 	/// <summary>
 	/// Represents a repository of streams and cached data
 	/// </summary>
@@ -89,9 +105,9 @@ namespace EpicGames.Perforce.Managed
 		const int CurrentSignature = ('W' << 24) | ('T' << 16) | 2;
 
 		/// <summary>
-		/// The current revision number.
+		/// The current revision number for cache archives.
 		/// </summary>
-		const int CurrentRevision = 2;
+		static int CurrentVersion { get; } = Enum.GetValues(typeof(ManagedWorkspaceVersion)).Cast<int>().Max();
 
 		/// <summary>
 		/// Maximum number of threads to sync in parallel
@@ -331,10 +347,10 @@ namespace EpicGames.Perforce.Managed
 			byte[] Data = await FileReference.ReadAllBytesAsync(DataFile, CancellationToken);
 			MemoryReader Reader = new MemoryReader(Data.AsMemory());
 
-			int Revision = Reader.ReadInt32();
-			if(Revision != CurrentRevision)
+			int Version = Reader.ReadInt32();
+			if(Version > CurrentVersion)
 			{
-				throw new FatalErrorException("Unsupported data format (revision {0}, expected {1})", Revision, CurrentRevision);
+				throw new FatalErrorException("Unsupported data format (version {0}, current {1})", Version, CurrentVersion);
 			}
 
 			bool bRequiresRepair = Reader.ReadBoolean();
@@ -351,7 +367,7 @@ namespace EpicGames.Perforce.Managed
 				Repo.CacheEntries.Add(TrackedFile.CacheId);
 			}
 
-			Reader.ReadWorkspaceDirectoryInfo(Repo.Workspace);
+			Reader.ReadWorkspaceDirectoryInfo(Repo.Workspace, (ManagedWorkspaceVersion)Version);
 
 			await Repo.RunOptionalRepairAsync(CancellationToken);
 			return Repo;
@@ -368,7 +384,7 @@ namespace EpicGames.Perforce.Managed
 
 			// Write the data to memory
 			MemoryWriter Writer = new MemoryWriter(Buffer.AsMemory());
-			Writer.WriteInt32(CurrentRevision);
+			Writer.WriteInt32(CurrentVersion);
 			Writer.WriteBoolean(bRequiresRepair || (State != TransactionState.Clean));
 			Writer.WriteUInt32(NextSequenceNumber);
 			Writer.WriteInt32(ContentIdToTrackedFile.Count);
diff --git a/Engine/Source/Programs/Shared/EpicGames.Perforce.Managed/Transactions/AddTransaction.cs b/Engine/Source/Programs/Shared/EpicGames.Perforce.Managed/Transactions/AddTransaction.cs
--- a/Engine/Source/Programs/Shared/EpicGames.Perforce.Managed/Transactions/AddTransaction.cs
+++ b/Engine/Source/Programs/Shared/EpicGames.Perforce.Managed/Transactions/AddTransaction.cs
@@ -136,14 +136,14 @@ namespace EpicGames.Perforce.Managed
 
 		void MergeSubDirectory(ReadOnlyUtf8String Name, WorkspaceDirectoryInfo WorkspaceSubDir, Digest<Sha1> StreamSubDirHash, WorkspaceDirectoryInfo NewWorkspaceDir, ThreadPoolWorkQueue Queue)
 		{
-			WorkspaceDirectoryInfo NewWorkspaceSubDir = new WorkspaceDirectoryInfo(NewWorkspaceDir, Name);
+			WorkspaceDirectoryInfo NewWorkspaceSubDir = new WorkspaceDirectoryInfo(NewWorkspaceDir, Name, StreamSubDirHash);
 			NewWorkspaceDir.NameToSubDirectory.Add(Name, NewWorkspaceSubDir);
 			Queue.Enqueue(() => MergeDirectory(WorkspaceSubDir, NewWorkspaceSubDir, StreamSubDirHash, Queue));
 		}
 
 		void AddSubDirectory(ReadOnlyUtf8String Name, WorkspaceDirectoryInfo NewWorkspaceDir, Digest<Sha1> StreamSubDirHash, ThreadPoolWorkQueue Queue)
 		{
-			WorkspaceDirectoryInfo NewWorkspaceSubDir = new WorkspaceDirectoryInfo(NewWorkspaceDir, Name);
+			WorkspaceDirectoryInfo NewWorkspaceSubDir = new WorkspaceDirectoryInfo(NewWorkspaceDir, Name, StreamSubDirHash);
 			NewWorkspaceDir.NameToSubDirectory.Add(Name, NewWorkspaceSubDir);
 			Queue.Enqueue(() => AddDirectory(NewWorkspaceSubDir, StreamSubDirHash, Queue));
 		}
diff --git a/Engine/Source/Programs/Shared/EpicGames.Perforce.Managed/Transactions/RemoveTransaction.cs b/Engine/Source/Programs/Shared/EpicGames.Perforce.Managed/Transactions/RemoveTransaction.cs
--- a/Engine/Source/Programs/Shared/EpicGames.Perforce.Managed/Transactions/RemoveTransaction.cs
+++ b/Engine/Source/Programs/Shared/EpicGames.Perforce.Managed/Transactions/RemoveTransaction.cs
@@ -67,7 +67,7 @@ namespace EpicGames.Perforce.Managed
 
 		void MergeSubDirectory(ReadOnlyUtf8String Name, WorkspaceDirectoryInfo NewWorkspaceDir, WorkspaceDirectoryInfo WorkspaceSubDir, Digest<Sha1> StreamSubDirHash, ThreadPoolWorkQueue Queue)
 		{
-			WorkspaceDirectoryInfo NewWorkspaceSubDir = new WorkspaceDirectoryInfo(NewWorkspaceDir, Name);
+			WorkspaceDirectoryInfo NewWorkspaceSubDir = new WorkspaceDirectoryInfo(NewWorkspaceDir, Name, StreamSubDirHash);
 			NewWorkspaceDir.NameToSubDirectory.Add(NewWorkspaceSubDir.Name, NewWorkspaceSubDir);
 			Queue.Enqueue(() => Merge(WorkspaceSubDir, NewWorkspaceSubDir, StreamSubDirHash, Queue));
 		}
diff --git a/Engine/Source/Programs/Shared/EpicGames.Perforce.Managed/WorkspaceDirectoryInfo.cs b/Engine/Source/Programs/Shared/EpicGames.Perforce.Managed/WorkspaceDirectoryInfo.cs
--- a/Engine/Source/Programs/Shared/EpicGames.Perforce.Managed/WorkspaceDirectoryInfo.cs
+++ b/Engine/Source/Programs/Shared/EpicGames.Perforce.Managed/WorkspaceDirectoryInfo.cs
@@ -11,30 +11,71 @@ using System.Text;
 namespace EpicGames.Perforce.Managed
 {
 	/// <summary>
-	/// 
+	/// Stores the state of a directory in the workspace
 	/// </summary>
 	class WorkspaceDirectoryInfo
 	{
-		public readonly WorkspaceDirectoryInfo? ParentDirectory;
-		public readonly ReadOnlyUtf8String Name;
-		public Dictionary<ReadOnlyUtf8String, WorkspaceFileInfo> NameToFile;
-		public Dictionary<ReadOnlyUtf8String, WorkspaceDirectoryInfo> NameToSubDirectory;
+		/// <summary>
+		/// The parent directory
+		/// </summary>
+		public WorkspaceDirectoryInfo? ParentDirectory { get; }
 
+		/// <summary>
+		/// Name of this directory
+		/// </summary>
+		public ReadOnlyUtf8String Name { get; }
+
+		/// <summary>
+		/// Digest of the matching stream directory info. This should be set to zero if the workspace is modified.
+		/// </summary>
+		public Digest<Sha1> StreamDirectoryDigest { get; private set; }
+
+		/// <summary>
+		/// Map of name to file
+		/// </summary>
+		public Dictionary<ReadOnlyUtf8String, WorkspaceFileInfo> NameToFile { get; set; }
+
+		/// <summary>
+		/// Map of name to subdirectory
+		/// </summary>
+		public Dictionary<ReadOnlyUtf8String, WorkspaceDirectoryInfo> NameToSubDirectory { get; set; }
+
+		/// <summary>
+		/// Constructor
+		/// </summary>
+		/// <param name="RootDir"></param>
 		public WorkspaceDirectoryInfo(DirectoryReference RootDir)
-			: this(null, RootDir.FullName)
+			: this(null, RootDir.FullName, Digest<Sha1>.Zero)
 		{
 		}
 
-		public WorkspaceDirectoryInfo(WorkspaceDirectoryInfo? ParentDirectory, ReadOnlyUtf8String Name)
+		/// <summary>
+		/// Constructor
+		/// </summary>
+		/// <param name="ParentDirectory">The parent directory</param>
+		/// <param name="Name">Name of this directory</param>
+		/// <param name="Digest">The corresponding stream digest</param>
+		public WorkspaceDirectoryInfo(WorkspaceDirectoryInfo? ParentDirectory, ReadOnlyUtf8String Name, Digest<Sha1> Digest)
 		{
 			this.ParentDirectory = ParentDirectory;
 			this.Name = Name;
+			this.StreamDirectoryDigest = Digest;
 			this.NameToFile = new Dictionary<ReadOnlyUtf8String, WorkspaceFileInfo>(ReadOnlyUtf8StringComparer.Ordinal);
 			this.NameToSubDirectory = new Dictionary<ReadOnlyUtf8String, WorkspaceDirectoryInfo>(FileUtils.PlatformPathComparerUtf8);
 		}
 
+		/// <summary>
+		/// Adds a file to the workspace
+		/// </summary>
+		/// <param name="Path">Relative path to the file, using forward slashes, and without a leading slash</param>
+		/// <param name="Length">Length of the file on disk</param>
+		/// <param name="LastModifiedTicks">Last modified time of the file</param>
+		/// <param name="bReadOnly">Whether the file is read only</param>
+		/// <param name="ContentId">Unique identifier for the server content</param>
 		public void AddFile(ReadOnlyUtf8String Path, long Length, long LastModifiedTicks, bool bReadOnly, FileContentId ContentId)
 		{
+			StreamDirectoryDigest = Digest<Sha1>.Zero;
+
 			int Idx = Path.Span.IndexOf((byte)'/');
 			if (Idx == -1)
 			{
@@ -47,7 +88,7 @@ namespace EpicGames.Perforce.Managed
 				WorkspaceDirectoryInfo? SubDirectory;
 				if (!NameToSubDirectory.TryGetValue(Name, out SubDirectory))
 				{
-					SubDirectory = new WorkspaceDirectoryInfo(this, Name);
+					SubDirectory = new WorkspaceDirectoryInfo(this, Name, Digest<Sha1>.Zero);
 					NameToSubDirectory[Name] = SubDirectory;
 				}
 
@@ -55,6 +96,10 @@ namespace EpicGames.Perforce.Managed
 			}
 		}
 
+		/// <summary>
+		/// Create a flat list of files in this workspace
+		/// </summary>
+		/// <returns>List of files</returns>
 		public List<WorkspaceFileInfo> GetFiles()
 		{
 			List<WorkspaceFileInfo> Files = new List<WorkspaceFileInfo>();
@@ -62,6 +107,10 @@ namespace EpicGames.Perforce.Managed
 			return Files;
 		}
 
+		/// <summary>
+		/// Internal helper method for recursing through the tree to build a file list
+		/// </summary>
+		/// <param name="Files"></param>
 		private void GetFilesInternal(List<WorkspaceFileInfo> Files)
 		{
 			Files.AddRange(NameToFile.Values);
@@ -72,6 +121,12 @@ namespace EpicGames.Perforce.Managed
 			}
 		}
 
+		/// <summary>
+		/// Refresh the state of the workspace on disk
+		/// </summary>
+		/// <param name="bRemoveUntracked">Whether to remove files that are not part of the stream</param>
+		/// <param name="FilesToDelete">Receives an array of files to delete</param>
+		/// <param name="DirectoriesToDelete">Recevies an array of directories to delete</param>
 		public void Refresh(bool bRemoveUntracked, out FileInfo[] FilesToDelete, out DirectoryInfo[] DirectoriesToDelete)
 		{
 			ConcurrentBag<FileInfo> ConcurrentFilesToDelete = new ConcurrentBag<FileInfo>();
@@ -84,7 +139,15 @@ namespace EpicGames.Perforce.Managed
 			FilesToDelete = ConcurrentFilesToDelete.ToArray();
 		}
 
-		private void Refresh(DirectoryInfo Info, bool bRemoveUntracked, ConcurrentBag<FileInfo> FilesToDelete, ConcurrentBag<DirectoryInfo> DirectoriesToDelete, ThreadPoolWorkQueue Queue)
+		/// <summary>
+		/// Recursive method for querying the workspace state
+		/// </summary>
+		/// <param name="Info"></param>
+		/// <param name="bRemoveUntracked"></param>
+		/// <param name="FilesToDelete"></param>
+		/// <param name="DirectoriesToDelete"></param>
+		/// <param name="Queue"></param>
+		void Refresh(DirectoryInfo Info, bool bRemoveUntracked, ConcurrentBag<FileInfo> FilesToDelete, ConcurrentBag<DirectoryInfo> DirectoriesToDelete, ThreadPoolWorkQueue Queue)
 		{
 			// Recurse through subdirectories
 			Dictionary<ReadOnlyUtf8String, WorkspaceDirectoryInfo> NewNameToSubDirectory = new Dictionary<ReadOnlyUtf8String, WorkspaceDirectoryInfo>(NameToSubDirectory.Count, NameToSubDirectory.Comparer);
@@ -127,9 +190,24 @@ namespace EpicGames.Perforce.Managed
 					}
 				}
 			}
+
+			// If the file state has changed, clear the directory hashes
+			if (NameToFile.Count != NewNameToFile.Count)
+			{
+				for (WorkspaceDirectoryInfo? Directory = this; Directory != null && Directory.StreamDirectoryDigest != Digest<Sha1>.Zero; Directory = Directory.ParentDirectory)
+				{
+					Directory.StreamDirectoryDigest = Digest<Sha1>.Zero;
+				}
+			}
+
+			// Update the new file list
 			NameToFile = NewNameToFile;
 		}
 
+		/// <summary>
+		/// Builds a list of differences from the working directory
+		/// </summary>
+		/// <returns></returns>
 		public string[] FindDifferences()
 		{
 			ConcurrentBag<string> Paths = new ConcurrentBag<string>();
@@ -140,7 +218,14 @@ namespace EpicGames.Perforce.Managed
 			return Paths.OrderBy(x => x).ToArray();
 		}
 
-		private void FindDifferences(DirectoryInfo Directory, string Path, ConcurrentBag<string> Paths, ThreadPoolWorkQueue Queue)
+		/// <summary>
+		/// Helper method for finding differences from the working directory
+		/// </summary>
+		/// <param name="Directory"></param>
+		/// <param name="Path"></param>
+		/// <param name="Paths"></param>
+		/// <param name="Queue"></param>
+		void FindDifferences(DirectoryInfo Directory, string Path, ConcurrentBag<string> Paths, ThreadPoolWorkQueue Queue)
 		{
 			// Recurse through subdirectories
 			HashSet<ReadOnlyUtf8String> RemainingSubDirectoryNames = new HashSet<ReadOnlyUtf8String>(NameToSubDirectory.Keys);
@@ -185,6 +270,10 @@ namespace EpicGames.Perforce.Managed
 			}
 		}
 
+		/// <summary>
+		/// Get the full path to this directory
+		/// </summary>
+		/// <returns></returns>
 		public string GetFullName()
 		{
 			StringBuilder Builder = new StringBuilder();
@@ -192,11 +281,19 @@ namespace EpicGames.Perforce.Managed
 			return Builder.ToString();
 		}
 
+		/// <summary>
+		/// Get the path to this directory
+		/// </summary>
+		/// <returns></returns>
 		public DirectoryReference GetLocation()
 		{
 			return new DirectoryReference(GetFullName());
 		}
 
+		/// <summary>
+		/// Append the client path, using native directory separators, to the given string builder
+		/// </summary>
+		/// <param name="Builder"></param>
 		public void AppendClientPath(StringBuilder Builder)
 		{
 			if (ParentDirectory != null)
@@ -207,6 +304,10 @@ namespace EpicGames.Perforce.Managed
 			}
 		}
 
+		/// <summary>
+		/// Append the path for this directory to the given string builder
+		/// </summary>
+		/// <param name="Builder"></param>
 		public void AppendFullPath(StringBuilder Builder)
 		{
 			if (ParentDirectory != null)
@@ -217,15 +318,19 @@ namespace EpicGames.Perforce.Managed
 			Builder.Append(Name);
 		}
 
+		/// <inheritdoc/>
 		public override string ToString()
 		{
 			return GetFullName();
 		}
 	}
 
+	/// <summary>
+	/// Extension methods for WorkspaceDirectoryInfo
+	/// </summary>
 	static class WorkspaceDirectoryInfoExtensions
 	{
-		public static void ReadWorkspaceDirectoryInfo(this MemoryReader Reader, WorkspaceDirectoryInfo DirectoryInfo)
+		public static void ReadWorkspaceDirectoryInfo(this MemoryReader Reader, WorkspaceDirectoryInfo DirectoryInfo, ManagedWorkspaceVersion Version)
 		{
 			int NumFiles = Reader.ReadInt32();
 			for (int Idx = 0; Idx < NumFiles; Idx++)
@@ -237,8 +342,20 @@ namespace EpicGames.Perforce.Managed
 			int NumSubDirectories = Reader.ReadInt32();
 			for (int Idx = 0; Idx < NumSubDirectories; Idx++)
 			{
-				WorkspaceDirectoryInfo SubDirectory = new WorkspaceDirectoryInfo(DirectoryInfo, Reader.ReadString());
-				Reader.ReadWorkspaceDirectoryInfo(SubDirectory);
+				ReadOnlyUtf8String Name = Reader.ReadString();
+
+				Digest<Sha1> Digest;
+				if (Version < ManagedWorkspaceVersion.AddDigest)
+				{
+					Digest = Digest<Sha1>.Zero;
+				}
+				else
+				{
+					Digest = Reader.ReadDigest<Sha1>();
+				}
+
+				WorkspaceDirectoryInfo SubDirectory = new WorkspaceDirectoryInfo(DirectoryInfo, Name, Digest);
+				Reader.ReadWorkspaceDirectoryInfo(SubDirectory, Version);
 				DirectoryInfo.NameToSubDirectory[SubDirectory.Name] = SubDirectory;
 			}
 		}
@@ -255,6 +372,7 @@ namespace EpicGames.Perforce.Managed
 			foreach (WorkspaceDirectoryInfo SubDirectory in DirectoryInfo.NameToSubDirectory.Values)
 			{
 				Writer.WriteString(SubDirectory.Name);
+				Writer.WriteDigest<Sha1>(SubDirectory.StreamDirectoryDigest);
 				Writer.WriteWorkspaceDirectoryInfo(SubDirectory);
 			}
 		}
