Horde: Add support for managed workspaces reading from a Git repository.

Perforce metadata is stored in an incredibly inefficient 'flat' format, as evidenced by the painfully long metadata query and flush operations when running at scale. By storing metadata in a Git repository, we can take advantage of a more efficient temporal encoding of workspace snapshots via Merkle trees, and leverage existing Git transfer/negotiation algorithms over standard protocols for fetching updates.

Storing metadata in a Git repository also allows us to execute queries efficiently locally, while still getting server-centric storage/file locking/etc... from Perforce.

diff --git a/Engine/Source/Programs/Shared/EpicGames.Perforce.Managed/EpicGames.Perforce.Managed.csproj b/Engine/Source/Programs/Shared/EpicGames.Perforce.Managed/EpicGames.Perforce.Managed.csproj
--- a/Engine/Source/Programs/Shared/EpicGames.Perforce.Managed/EpicGames.Perforce.Managed.csproj
+++ b/Engine/Source/Programs/Shared/EpicGames.Perforce.Managed/EpicGames.Perforce.Managed.csproj
@@ -27,5 +27,8 @@
       <Link>Properties\MetaData.cs</Link>
     </Compile>
   </ItemGroup>
+  <ItemGroup>
+    <PackageReference Include="LibGit2Sharp" Version="0.26.2" />
+  </ItemGroup>
 
 </Project>
diff --git a/Engine/Source/Programs/Shared/EpicGames.Perforce.Managed/StreamSnapshotFromGit.cs b/Engine/Source/Programs/Shared/EpicGames.Perforce.Managed/StreamSnapshotFromGit.cs
new file mode 100644
--- /dev/null
+++ b/Engine/Source/Programs/Shared/EpicGames.Perforce.Managed/StreamSnapshotFromGit.cs
@@ -0,0 +1,132 @@
+// Copyright Epic Games, Inc. All Rights Reserved.
+
+using EpicGames.Core;
+using System;
+using System.Collections.Generic;
+using System.IO;
+using System.Linq;
+using System.Text;
+using System.Threading;
+using System.Threading.Tasks;
+using LibGit2Sharp;
+using System.Collections.Concurrent;
+
+namespace EpicGames.Perforce.Managed
+{
+	/// <summary>
+	/// Provides the contents of a Perforce stream via a Git repository.
+	/// 
+	/// Perforce metadata is stored in an incredibly inefficient 'flat' format, as evidenced by the painfully long metadata query and flush operations when 
+	/// running at scale. By storing metadata in a Git repository, we can take advantage of a more efficient temporal encoding of workspace snapshots via
+	/// Merkle trees, and leverage existing transfer/negotiation protocols for fetching updates.
+	/// 
+	/// Storing metadata in a Git repository also allows us to execute queries efficiently locally, while still getting server-centric storage/file locking/etc...
+	/// from Perforce.
+	/// </summary>
+	public class StreamSnapshotFromGit : StreamSnapshot
+	{
+		/// <summary>
+		/// The Git repository
+		/// </summary>
+		Repository Repository { get; }
+
+		/// <summary>
+		/// The root SHA
+		/// </summary>
+		public override Digest<Sha1> Root { get; }
+
+		/// <summary>
+		/// Cache of directory objects
+		/// </summary>
+		ConcurrentDictionary<Digest<Sha1>, StreamDirectoryInfo> CachedDirectories = new ConcurrentDictionary<Digest<Sha1>, StreamDirectoryInfo>();
+
+		/// <summary>
+		/// Constructor
+		/// </summary>
+		/// <param name="RootTreeId"></param>
+		public StreamSnapshotFromGit(Repository Repository, ObjectId RootTreeId)
+		{
+			this.Repository = Repository;
+			this.Root = new Digest<Sha1>(RootTreeId.RawId);
+		}
+
+		/// <inheritdoc/>
+		public override StreamDirectoryInfo Lookup(Digest<Sha1> Hash)
+		{
+			StreamDirectoryInfo? DirectoryInfo;
+			if (!CachedDirectories.TryGetValue(Hash, out DirectoryInfo))
+			{
+				DirectoryInfo = new StreamDirectoryInfo();
+
+				Tree Tree = Repository.Lookup<Tree>(new ObjectId(Hash.Memory.ToArray()));
+				foreach (TreeEntry TreeEntry in Tree)
+				{
+					if (TreeEntry.TargetType == TreeEntryTargetType.Tree)
+					{
+						DirectoryInfo.NameToSubDirectory[new ReadOnlyUtf8String(TreeEntry.Name)] = new Digest<Sha1>(TreeEntry.Target.Id.RawId);
+					}
+					else if (TreeEntry.TargetType == TreeEntryTargetType.Blob)
+					{
+						ReadFileList((Blob)TreeEntry.Target, DirectoryInfo.NameToFile);
+					}
+				}
+
+				CachedDirectories.TryAdd(Hash, DirectoryInfo);
+			}
+			return DirectoryInfo;
+		}
+
+		/// <summary>
+		/// Reads a list of files from a Git blob
+		/// </summary>
+		/// <param name="Blob">The blob to read from</param>
+		/// <param name="Files">Map of name to file</param>
+		public static void ReadFileList(Blob Blob, Dictionary<ReadOnlyUtf8String, StreamFileInfo> Files)
+		{
+			byte[] Data = new byte[Blob.Size];
+			using (Stream Stream = Blob.GetContentStream())
+			{
+				Stream.Read(Data);
+			}
+			ReadFileList(Data, Files);
+		}
+
+		/// <summary>
+		/// Reads a list of files from a block of memory
+		/// </summary>
+		/// <param name="Blob">The blob to read from</param>
+		/// <param name="Files">Map of name to file</param>
+		public static void ReadFileList(byte[] Data, Dictionary<ReadOnlyUtf8String, StreamFileInfo> Files)
+		{
+			MemoryReader Reader = new MemoryReader(Data);
+
+			int NumFiles = Reader.ReadInt32();
+			for (int Idx = 0; Idx < NumFiles; Idx++)
+			{
+				StreamFileInfo File = Reader.ReadStreamFileInfo();
+				Files[File.Name] = File;
+			}
+		}
+
+		/// <summary>
+		/// Serializes a list of files to a byte array
+		/// </summary>
+		/// <param name="Files"></param>
+		/// <returns></returns>
+		public static byte[] WriteFileList(Dictionary<ReadOnlyUtf8String, StreamFileInfo> Files)
+		{
+			int Length = sizeof(int) + Files.Values.Sum(x => x.GetSerializedSize());
+			byte[] Data = new byte[Length];
+
+			MemoryWriter Writer = new MemoryWriter(Data);
+			Writer.WriteInt32(Files.Count);
+
+			foreach (StreamFileInfo File in Files.Values)
+			{
+				Writer.WriteStreamFileInfo(File);
+			}
+
+			return Data;
+		}
+	}
+}
