Horde: Fix memory overrun when deserializing WorkspaceDirectoryInfo.

diff --git a/Engine/Source/Programs/Shared/EpicGames.Perforce.Managed/WorkspaceDirectoryInfo.cs b/Engine/Source/Programs/Shared/EpicGames.Perforce.Managed/WorkspaceDirectoryInfo.cs
--- a/Engine/Source/Programs/Shared/EpicGames.Perforce.Managed/WorkspaceDirectoryInfo.cs
+++ b/Engine/Source/Programs/Shared/EpicGames.Perforce.Managed/WorkspaceDirectoryInfo.cs
@@ -28,7 +28,7 @@ namespace EpicGames.Perforce.Managed
 		/// <summary>
 		/// Digest of the matching stream directory info. This should be set to zero if the workspace is modified.
 		/// </summary>
-		public Digest<Sha1> StreamDirectoryDigest { get; private set; }
+		public Digest<Sha1> StreamDirectoryDigest { get; set; }
 
 		/// <summary>
 		/// Map of name to file
@@ -332,6 +332,15 @@ namespace EpicGames.Perforce.Managed
 	{
 		public static void ReadWorkspaceDirectoryInfo(this MemoryReader Reader, WorkspaceDirectoryInfo DirectoryInfo, ManagedWorkspaceVersion Version)
 		{
+			if (Version < ManagedWorkspaceVersion.AddDigest)
+			{
+				DirectoryInfo.StreamDirectoryDigest = Digest<Sha1>.Zero;
+			}
+			else
+			{
+				DirectoryInfo.StreamDirectoryDigest = Reader.ReadDigest<Sha1>();
+			}
+
 			int NumFiles = Reader.ReadInt32();
 			for (int Idx = 0; Idx < NumFiles; Idx++)
 			{
@@ -344,17 +353,7 @@ namespace EpicGames.Perforce.Managed
 			{
 				ReadOnlyUtf8String Name = Reader.ReadString();
 
-				Digest<Sha1> Digest;
-				if (Version < ManagedWorkspaceVersion.AddDigest)
-				{
-					Digest = Digest<Sha1>.Zero;
-				}
-				else
-				{
-					Digest = Reader.ReadDigest<Sha1>();
-				}
-
-				WorkspaceDirectoryInfo SubDirectory = new WorkspaceDirectoryInfo(DirectoryInfo, Name, Digest);
+				WorkspaceDirectoryInfo SubDirectory = new WorkspaceDirectoryInfo(DirectoryInfo, Name, Digest<Sha1>.Zero);
 				Reader.ReadWorkspaceDirectoryInfo(SubDirectory, Version);
 				DirectoryInfo.NameToSubDirectory[SubDirectory.Name] = SubDirectory;
 			}
@@ -362,6 +361,8 @@ namespace EpicGames.Perforce.Managed
 
 		public static void WriteWorkspaceDirectoryInfo(this MemoryWriter Writer, WorkspaceDirectoryInfo DirectoryInfo)
 		{
+			Writer.WriteDigest<Sha1>(DirectoryInfo.StreamDirectoryDigest);
+
 			Writer.WriteInt32(DirectoryInfo.NameToFile.Count);
 			foreach (WorkspaceFileInfo File in DirectoryInfo.NameToFile.Values)
 			{
@@ -372,14 +373,13 @@ namespace EpicGames.Perforce.Managed
 			foreach (WorkspaceDirectoryInfo SubDirectory in DirectoryInfo.NameToSubDirectory.Values)
 			{
 				Writer.WriteString(SubDirectory.Name);
-				Writer.WriteDigest<Sha1>(SubDirectory.StreamDirectoryDigest);
 				Writer.WriteWorkspaceDirectoryInfo(SubDirectory);
 			}
 		}
 
 		public static int GetSerializedSize(this WorkspaceDirectoryInfo DirectoryInfo)
 		{
-			return sizeof(int) + DirectoryInfo.NameToFile.Values.Sum(x => x.GetSerializedSize()) + sizeof(int) + DirectoryInfo.NameToSubDirectory.Values.Sum(x => x.Name.GetSerializedSize() + x.GetSerializedSize());
+			return Digest<Sha1>.Length + sizeof(int) + DirectoryInfo.NameToFile.Values.Sum(x => x.GetSerializedSize()) + sizeof(int) + DirectoryInfo.NameToSubDirectory.Values.Sum(x => x.Name.GetSerializedSize() + x.GetSerializedSize());
 		}
 	}
 }
