Add a Boxed<T> utility class to explicitly box a struct type, allowing a mutable reference to it to be passed to async methods.

diff --git a/Engine/Source/Programs/Shared/EpicGames.Core/Boxed.cs b/Engine/Source/Programs/Shared/EpicGames.Core/Boxed.cs
new file mode 100644
--- /dev/null
+++ b/Engine/Source/Programs/Shared/EpicGames.Core/Boxed.cs
@@ -0,0 +1,70 @@
+using System;
+using System.Collections.Generic;
+using System.Text;
+
+namespace EpicGames.Core
+{
+	/// <summary>
+	/// Explicitly boxes a value type, allowing it to be passed by reference to (and modified in) an async method.
+	/// </summary>
+	/// <typeparam name="T"></typeparam>
+	public class Boxed<T> where T : struct
+	{
+		/// <summary>
+		/// The struct value
+		/// </summary>
+		public T Value { get; set; }
+
+		/// <summary>
+		/// Default constructor
+		/// </summary>
+		public Boxed()
+		{
+		}
+
+		/// <summary>
+		/// Default constructor
+		/// </summary>
+		/// <param name="Value">The value to copy</param>
+		public Boxed(T Value)
+		{
+			this.Value = Value;
+		}
+
+		/// <inheritdoc/>
+		public override bool Equals(object? Obj)
+		{
+			return Obj is Boxed<T> Boxed && Value.Equals(Boxed.Value);
+		}
+
+		/// <inheritdoc/>
+		public override int GetHashCode()
+		{
+			return Value.GetHashCode();
+		}
+
+		/// <inheritdoc/>
+		public override string? ToString()
+		{
+			return Value.ToString();
+		}
+
+		/// <summary>
+		/// Implicit conversion of a boxed value to value type
+		/// </summary>
+		/// <param name="Boxed"></param>
+		public static implicit operator T(Boxed<T> Boxed)
+		{
+			return Boxed.Value;
+		}
+
+		/// <summary>
+		/// Implicit conversion of a boxed value to value type
+		/// </summary>
+		/// <param name="Value"></param>
+		public static implicit operator Boxed<T>(T Value)
+		{
+			return new Boxed<T>(Value);
+		}
+	}
+}
