Add a Boxed<T> class to allow wrapping a generic type in a class, so it can be passed into an async method as a mutable value.

diff --git a/Engine/Source/Programs/Shared/EpicGames.Core/Boxed.cs b/Engine/Source/Programs/Shared/EpicGames.Core/Boxed.cs
--- a/Engine/Source/Programs/Shared/EpicGames.Core/Boxed.cs
+++ b/Engine/Source/Programs/Shared/EpicGames.Core/Boxed.cs
@@ -8,20 +8,13 @@ namespace EpicGames.Core
 	/// Explicitly boxes a value type, allowing it to be passed by reference to (and modified in) an async method.
 	/// </summary>
 	/// <typeparam name="T"></typeparam>
-	public class Boxed<T> where T : struct
+	public class Boxed<T>
 	{
 		/// <summary>
 		/// The struct value
 		/// </summary>
 		public T Value { get; set; }
 
-		/// <summary>
-		/// Default constructor
-		/// </summary>
-		public Boxed()
-		{
-		}
-
 		/// <summary>
 		/// Default constructor
 		/// </summary>
@@ -31,22 +24,10 @@ namespace EpicGames.Core
 			this.Value = Value;
 		}
 
-		/// <inheritdoc/>
-		public override bool Equals(object? Obj)
-		{
-			return Obj is Boxed<T> Boxed && Value.Equals(Boxed.Value);
-		}
-
-		/// <inheritdoc/>
-		public override int GetHashCode()
-		{
-			return Value.GetHashCode();
-		}
-
 		/// <inheritdoc/>
 		public override string? ToString()
 		{
-			return Value.ToString();
+			return Value?.ToString();
 		}
 
 		/// <summary>
