Add FileReference.WriteAllTextAsync() implementations.

diff --git a/Engine/Source/Programs/Shared/EpicGames.Core/FileReference.cs b/Engine/Source/Programs/Shared/EpicGames.Core/FileReference.cs
--- a/Engine/Source/Programs/Shared/EpicGames.Core/FileReference.cs
+++ b/Engine/Source/Programs/Shared/EpicGames.Core/FileReference.cs
@@ -709,6 +709,28 @@ namespace EpicGames.Core
 			File.WriteAllText(Location.FullName, Contents, Encoding);
 		}
 
+
+		/// <summary>
+		/// Writes the contents of a file
+		/// </summary>
+		/// <param name="Location">Location of the file</param>
+		/// <param name="Contents">Contents of the file</param>
+		public static Task WriteAllTextAsync(FileReference Location, string Contents)
+		{
+			return File.WriteAllTextAsync(Location.FullName, Contents);
+		}
+
+		/// <summary>
+		/// Writes the contents of a file
+		/// </summary>
+		/// <param name="Location">Location of the file</param>
+		/// <param name="Contents">Contents of the file</param>
+		/// <param name="Encoding">The encoding to use when parsing the file</param>
+		public static Task WriteAllTextAsync(FileReference Location, string Contents, Encoding Encoding)
+		{
+			return File.WriteAllTextAsync(Location.FullName, Contents, Encoding);
+		}
+
 		/// <summary>
 		/// Appends the contents to a file
 		/// </summary>
