Horde: Catch exceptions thrown when waiting for tasks to complete.

diff --git a/Engine/Source/Programs/Horde/HordeServer/Services/LogFileService.cs b/Engine/Source/Programs/Horde/HordeServer/Services/LogFileService.cs
--- a/Engine/Source/Programs/Horde/HordeServer/Services/LogFileService.cs
+++ b/Engine/Source/Programs/Horde/HordeServer/Services/LogFileService.cs
@@ -991,7 +991,14 @@ namespace HordeServer.Services
 		{
 			lock (WriteLock)
 			{
-				WriteTasks.RemoveCompleteTasks();
+				try
+				{
+					WriteTasks.RemoveCompleteTasks();
+				}
+				catch (Exception Ex)
+				{
+					Logger.LogError(Ex, "Exception while waiting for write tasks to complete");
+				}
 			}
 			await IncrementalFlush();
 		}
diff --git a/Engine/Source/Programs/Shared/EpicGames.Core/AsyncUtils.cs b/Engine/Source/Programs/Shared/EpicGames.Core/AsyncUtils.cs
--- a/Engine/Source/Programs/Shared/EpicGames.Core/AsyncUtils.cs
+++ b/Engine/Source/Programs/Shared/EpicGames.Core/AsyncUtils.cs
@@ -18,12 +18,18 @@ namespace EpicGames.Core
 		/// <param name="Tasks">List of tasks to remove tasks from</param>
 		public static void RemoveCompleteTasks(this List<Task> Tasks)
 		{
+			List<Exception> Exceptions = new List<Exception>();
+
 			int OutIdx = 0;
 			for (int Idx = 0; Idx < Tasks.Count; Idx++)
 			{
 				if (Tasks[Idx].IsCompleted)
 				{
-					Tasks[Idx].Wait();
+					AggregateException? Exception = Tasks[Idx].Exception;
+					if(Exception != null)
+					{
+						Exceptions.AddRange(Exception.InnerExceptions);
+					}
 				}
 				else
 				{
@@ -35,6 +41,11 @@ namespace EpicGames.Core
 				}
 			}
 			Tasks.RemoveRange(OutIdx, Tasks.Count - OutIdx);
+
+			if(Exceptions.Count > 0)
+			{
+				throw new AggregateException(Exceptions);
+			}
 		}
 
 		/// <summary>
