[P4VUtils] - add undo command

diff --git a/Engine/Source/Programs/Shared/EpicGames.Perforce/PerforceConnection.cs b/Engine/Source/Programs/Shared/EpicGames.Perforce/PerforceConnection.cs
--- a/Engine/Source/Programs/Shared/EpicGames.Perforce/PerforceConnection.cs
+++ b/Engine/Source/Programs/Shared/EpicGames.Perforce/PerforceConnection.cs
@@ -2980,5 +2980,35 @@ namespace EpicGames.Perforce
 		}
 
 		#endregion
+
+		#region p4 undo
+		/// <summary>
+		/// perform undo on a changelist (p4 undo -c <targetCL> //...@undoCL)
+		/// </summary>
+		/// <param name="ChangeNumberToUndo">Changelist number to undo</param>
+		/// <param name="ChangeNumber">Changelist number to receive the changes</param>
+		/// <param name="CancellationToken">Token used to cancel the operation</param>
+		/// <returns>Response from the server</returns>
+		public async Task UndoChangeAsync(int ChangeNumberToUndo, int ChangeNumber, CancellationToken CancellationToken)
+		{
+			(await TryUndoChangeAsync(ChangeNumberToUndo, ChangeNumber, CancellationToken))[0].EnsureSuccess();
+		}
+
+		/// <summary>
+		/// perform undo on a changelist (p4 undo -c <targetCL> //...@undoCL)
+		/// </summary>
+		/// <param name="ChangeNumberToUndo">Changelist number to undo</param>
+		/// <param name="ChangeNumber">Changelist number to receive the changes</param>
+		/// <param name="CancellationToken">Token used to cancel the operation</param>
+		/// <returns>Response from the server</returns>
+		public Task<PerforceResponseList<UndoRecord>> TryUndoChangeAsync(int ChangeNumberToUndo, int ChangeNumber, CancellationToken CancellationToken)
+		{
+			StringBuilder Arguments = new StringBuilder("undo -c");
+			
+			Arguments.AppendFormat(" {0} //...@{1}", ChangeNumber, ChangeNumberToUndo);
+
+			return CommandAsync<UndoRecord>(Arguments.ToString(), null, CancellationToken);
+		}
+		#endregion
 	}
 }
diff --git a/Engine/Source/Programs/Shared/EpicGames.Perforce/Records/UndoRecord.cs b/Engine/Source/Programs/Shared/EpicGames.Perforce/Records/UndoRecord.cs
new file mode 100644
--- /dev/null
+++ b/Engine/Source/Programs/Shared/EpicGames.Perforce/Records/UndoRecord.cs
@@ -0,0 +1,66 @@
+// Copyright Epic Games, Inc. All Rights Reserved.
+
+using System;
+using System.Collections.Generic;
+using System.Linq;
+using System.Text;
+using System.Threading.Tasks;
+
+namespace EpicGames.Perforce
+{
+	/// <summary>
+	/// Information about a file opened for Undo
+	/// </summary>
+	public class UndoRecord
+	{
+		/// <summary>
+		/// Path to the file in the depot
+		/// </summary>
+		[PerforceTag("depotFile")]
+		public string DepotFile;
+
+		/// <summary>
+		/// Path to the file in the workspace
+		/// </summary>
+		[PerforceTag("clientFile")]
+		public string ClientFile;
+
+		/// <summary>
+		/// The working revision number of the file that was synced
+		/// </summary>
+		[PerforceTag("workRev")]
+		public int WorkingRevision;
+
+		/// <summary>
+		/// Action taken when syncing the file
+		/// </summary>
+		[PerforceTag("action")]
+		public string Action;
+
+		/// <summary>
+		/// Type of the file
+		/// </summary>
+		[PerforceTag("type")]
+		public string Type;
+
+		/// <summary>
+		/// Private constructor for serialization
+		/// </summary>
+		private UndoRecord()
+		{
+			DepotFile = null!;
+			ClientFile = null!;
+			Action = null!;
+			Type = null!;
+		}
+
+		/// <summary>
+		/// Format this record for display in the debugger
+		/// </summary>
+		/// <returns>Summary of this revision</returns>
+		public override string? ToString()
+		{
+			return DepotFile;
+		}
+	}
+}
