Add an IoHash library, allowing IoHash values to be computed from C#.

diff --git a/Engine/Source/Programs/Shared/EpicGames.IoHash/EpicGames.IoHash.csproj b/Engine/Source/Programs/Shared/EpicGames.IoHash/EpicGames.IoHash.csproj
new file mode 100644
--- /dev/null
+++ b/Engine/Source/Programs/Shared/EpicGames.IoHash/EpicGames.IoHash.csproj
@@ -0,0 +1,36 @@
+﻿<Project Sdk="Microsoft.NET.Sdk">
+  <PropertyGroup>
+    <TargetFramework>netcoreapp3.1</TargetFramework>
+    <Configuration Condition=" '$(Configuration)' == '' ">Development</Configuration>
+    <OutputType>Library</OutputType>
+    <Deterministic>true</Deterministic>
+    <RootNamespace>EpicGames.Core</RootNamespace>
+    <GenerateAssemblyInfo>false</GenerateAssemblyInfo>
+    <GenerateTargetFrameworkAttribute>false</GenerateTargetFrameworkAttribute>
+    <Configurations>Debug;Release;Development</Configurations>
+    <Nullable>enable</Nullable>
+  </PropertyGroup>
+  <PropertyGroup Condition=" '$(Configuration)|$(Platform)' == 'Debug|AnyCPU' ">
+    <DocumentationFile></DocumentationFile>
+  </PropertyGroup>
+  <PropertyGroup Condition="'$(Configuration)|$(Platform)' == 'Development|AnyCPU'">
+    <DocumentationFile></DocumentationFile>
+    <Optimize>true</Optimize>
+  </PropertyGroup>
+  <PropertyGroup Condition="'$(Configuration)|$(Platform)' == 'Release|AnyCPU'">
+    <DocumentationFile></DocumentationFile>
+    <Optimize>true</Optimize>
+  </PropertyGroup>
+  <ItemGroup>
+    <Compile Include="..\..\DotNETCommon\MetaData.cs">
+      <Link>Properties\MetaData.cs</Link>
+    </Compile>
+  </ItemGroup>
+  <ItemGroup>
+    <PackageReference Include="Blake3" Version="0.3.0" />
+    <PackageReference Include="System.Memory" Version="4.5.4" />
+  </ItemGroup>
+  <ItemGroup>
+    <ProjectReference Include="..\EpicGames.Core\EpicGames.Core.csproj" />
+  </ItemGroup>
+</Project>
\ No newline at end of file
diff --git a/Engine/Source/Programs/Shared/EpicGames.IoHash/IoHash.cs b/Engine/Source/Programs/Shared/EpicGames.IoHash/IoHash.cs
new file mode 100644
--- /dev/null
+++ b/Engine/Source/Programs/Shared/EpicGames.IoHash/IoHash.cs
@@ -0,0 +1,138 @@
+// Copyright Epic Games, Inc. All Rights Reserved.
+
+using System;
+using System.Collections.Generic;
+using System.Buffers;
+using System.Security.Cryptography;
+using System.Text;
+using System.Buffers.Binary;
+using System.Diagnostics.CodeAnalysis;
+
+namespace EpicGames.Core
+{
+	/// <summary>
+	/// Struct representing a strongly typed IoHash value (a 20-byte Blake3 hash).
+	/// </summary>
+	public struct IoHash : IEquatable<IoHash>, IComparable<IoHash>
+	{
+		/// <summary>
+		/// Length of an IoHash
+		/// </summary>
+		public const int NumBytes = 20;
+
+		/// <summary>
+		/// Length of the hash in bits
+		/// </summary>
+		public const int NumBits = NumBytes * 8;
+
+		/// <summary>
+		/// Memory storing the digest data
+		/// </summary>
+		public ReadOnlyMemory<byte> Memory;
+
+		/// <summary>
+		/// Span for the underlying memory
+		/// </summary>
+		public ReadOnlySpan<byte> Span => Memory.Span;
+
+		/// <summary>
+		/// Hash consisting of zeroes
+		/// </summary>
+		public static IoHash Zero { get; } = new IoHash(new byte[NumBytes]);
+
+		/// <summary>
+		/// Constructor
+		/// </summary>
+		/// <param name="Memory">Memory to construct from</param>
+		public IoHash(ReadOnlyMemory<byte> Memory)
+		{
+			if (Memory.Length != NumBytes)
+			{
+				throw new ArgumentException($"IoHash must be {NumBytes} bytes long");
+			}
+
+			this.Memory = Memory;
+		}
+
+		/// <summary>
+		/// Creates a content hash for a block of data, using a given algorithm.
+		/// </summary>
+		/// <param name="Data">Data to compute the hash for</param>
+		/// <returns>New content hash instance containing the hash of the data</returns>
+		public static IoHash Compute(ReadOnlySpan<byte> Data)
+		{
+			byte[] Output = new byte[32];
+			Blake3.Hasher.Hash(Data, Output);
+			return new IoHash(Output.AsMemory(0, 20));
+		}
+
+		/// <summary>
+		/// Parses a digest from the given hex string
+		/// </summary>
+		/// <param name="Text"></param>
+		/// <returns></returns>
+		public static IoHash Parse(string Text)
+		{
+			return new IoHash(StringUtils.ParseHexString(Text));
+		}
+
+		/// <inheritdoc cref="IComparable{T}.CompareTo(T)"/>
+		public int CompareTo(IoHash Other)
+		{
+			ReadOnlySpan<byte> A = Span;
+			ReadOnlySpan<byte> B = Other.Span;
+
+			for (int Idx = 0; Idx < A.Length && Idx < B.Length; Idx++)
+			{
+				int Compare = A[Idx] - B[Idx];
+				if (Compare != 0)
+				{
+					return Compare;
+				}
+			}
+			return A.Length - B.Length;
+		}
+
+		/// <inheritdoc/>
+		public bool Equals(IoHash Other) => Span.SequenceEqual(Other.Span);
+
+		/// <inheritdoc/>
+		public override bool Equals(object? Obj) => (Obj is IoHash Hash) && Hash.Span.SequenceEqual(Span);
+
+		/// <inheritdoc/>
+		public override int GetHashCode() => BinaryPrimitives.ReadInt32LittleEndian(Span);
+
+		/// <inheritdoc/>
+		public override string ToString() => StringUtils.FormatHexString(Memory.Span);
+
+		/// <summary>
+		/// Test two hash values for equality
+		/// </summary>
+		public static bool operator ==(IoHash A, IoHash B) => A.Span.SequenceEqual(B.Span);
+
+		/// <summary>
+		/// Test two hash values for equality
+		/// </summary>
+		public static bool operator !=(IoHash A, IoHash B) => !(A == B);
+
+		/// <summary>
+		/// Tests whether A > B
+		/// </summary>
+		public static bool operator >(IoHash A, IoHash B) => A.CompareTo(B) > 0;
+
+		/// <summary>
+		/// Tests whether A is less than B
+		/// </summary>
+		public static bool operator <(IoHash A, IoHash B) => A.CompareTo(B) < 0;
+
+		/// <summary>
+		/// Tests whether A is greater than or equal to B
+		/// </summary>
+		public static bool operator >=(IoHash A, IoHash B) => A.CompareTo(B) >= 0;
+
+		/// <summary>
+		/// Tests whether A is less than or equal to B
+		/// </summary>
+		public static bool operator <=(IoHash A, IoHash B) => A.CompareTo(B) <= 0;
+	}
+}
