Change EpicGames.Perforce.Managed to use IoHash, simplifying the migration to using Jupiter for storage.

diff --git a/Engine/Source/Programs/Shared/EpicGames.IoHash/IoHash.cs b/Engine/Source/Programs/Shared/EpicGames.IoHash/IoHash.cs
--- a/Engine/Source/Programs/Shared/EpicGames.IoHash/IoHash.cs
+++ b/Engine/Source/Programs/Shared/EpicGames.IoHash/IoHash.cs
@@ -135,4 +135,30 @@ namespace EpicGames.Core
 		/// </summary>
 		public static bool operator <=(IoHash A, IoHash B) => A.CompareTo(B) <= 0;
 	}
+
+	/// <summary>
+	/// Extension methods for dealing with IoHash values
+	/// </summary>
+	public static class IoHashExtensions
+	{
+		/// <summary>
+		/// Read an <see cref="IoHash"/> from a memory reader
+		/// </summary>
+		/// <param name="Reader"></param>
+		/// <returns></returns>
+		public static IoHash ReadIoHash(this MemoryReader Reader)
+		{
+			return new IoHash(Reader.ReadFixedLengthBytes(IoHash.NumBytes));
+		}
+
+		/// <summary>
+		/// Write an <see cref="IoHash"/> to a memory writer
+		/// </summary>
+		/// <param name="Writer"></param>
+		/// <param name="Hash"></param>
+		public static void WriteIoHash(this MemoryWriter Writer, IoHash Hash)
+		{
+			Writer.WriteFixedLengthBytes(Hash.Span);
+		}
+	}
 }
diff --git a/Engine/Source/Programs/Shared/EpicGames.Perforce.Managed/EpicGames.Perforce.Managed.csproj b/Engine/Source/Programs/Shared/EpicGames.Perforce.Managed/EpicGames.Perforce.Managed.csproj
--- a/Engine/Source/Programs/Shared/EpicGames.Perforce.Managed/EpicGames.Perforce.Managed.csproj
+++ b/Engine/Source/Programs/Shared/EpicGames.Perforce.Managed/EpicGames.Perforce.Managed.csproj
@@ -1,4 +1,4 @@
-<Project Sdk="Microsoft.NET.Sdk">
+﻿<Project Sdk="Microsoft.NET.Sdk">
 
   <PropertyGroup>
     <TargetFramework>netcoreapp3.1</TargetFramework>
@@ -20,6 +20,7 @@
   </PropertyGroup>
   <ItemGroup>
     <ProjectReference Include="..\EpicGames.Core\EpicGames.Core.csproj" />
+    <ProjectReference Include="..\EpicGames.IoHash\EpicGames.IoHash.csproj" />
     <ProjectReference Include="..\EpicGames.Perforce\EpicGames.Perforce.csproj" />
   </ItemGroup>
   <ItemGroup>
diff --git a/Engine/Source/Programs/Shared/EpicGames.Perforce.Managed/ManagedWorkspace.cs b/Engine/Source/Programs/Shared/EpicGames.Perforce.Managed/ManagedWorkspace.cs
--- a/Engine/Source/Programs/Shared/EpicGames.Perforce.Managed/ManagedWorkspace.cs
+++ b/Engine/Source/Programs/Shared/EpicGames.Perforce.Managed/ManagedWorkspace.cs
@@ -83,6 +83,11 @@ namespace EpicGames.Perforce.Managed
 		/// Including stream directory digests in workspace directories
 		/// </summary>
 		AddDigest = 3,
+
+		/// <summary>
+		/// Changing hash algorithm from SHA1 to IoHash
+		/// </summary>
+		AddDigestIoHash = 4,
 	}
 
 	/// <summary>
diff --git a/Engine/Source/Programs/Shared/EpicGames.Perforce.Managed/StreamDirectoryInfo.cs b/Engine/Source/Programs/Shared/EpicGames.Perforce.Managed/StreamDirectoryInfo.cs
--- a/Engine/Source/Programs/Shared/EpicGames.Perforce.Managed/StreamDirectoryInfo.cs
+++ b/Engine/Source/Programs/Shared/EpicGames.Perforce.Managed/StreamDirectoryInfo.cs
@@ -24,7 +24,7 @@ namespace EpicGames.Perforce.Managed
 		/// <summary>
 		/// Map of name to subdirectory
 		/// </summary>
-		public Dictionary<ReadOnlyUtf8String, Digest<Sha1>> NameToSubDirectory = new Dictionary<ReadOnlyUtf8String, Digest<Sha1>>(FileUtils.PlatformPathComparerUtf8);
+		public Dictionary<ReadOnlyUtf8String, IoHash> NameToSubDirectory = new Dictionary<ReadOnlyUtf8String, IoHash>(FileUtils.PlatformPathComparerUtf8);
 	}
 
 	/// <summary>
@@ -51,7 +51,7 @@ namespace EpicGames.Perforce.Managed
 			for (int Idx = 0; Idx < NumSubDirectories; Idx++)
 			{
 				ReadOnlyUtf8String Name = Reader.ReadString();
-				DirectoryInfo.NameToSubDirectory[Name] = Reader.ReadDigest<Sha1>();
+				DirectoryInfo.NameToSubDirectory[Name] = Reader.ReadIoHash();
 			}
 
 			return DirectoryInfo;
@@ -71,10 +71,10 @@ namespace EpicGames.Perforce.Managed
 			}
 
 			Writer.WriteInt32(DirectoryInfo.NameToSubDirectory.Count);
-			foreach ((ReadOnlyUtf8String Name, Digest<Sha1> Digest) in DirectoryInfo.NameToSubDirectory)
+			foreach ((ReadOnlyUtf8String Name, IoHash Hash) in DirectoryInfo.NameToSubDirectory)
 			{
 				Writer.WriteString(Name);
-				Writer.WriteDigest<Sha1>(Digest);
+				Writer.WriteIoHash(Hash);
 			}
 		}
 
@@ -84,7 +84,7 @@ namespace EpicGames.Perforce.Managed
 		/// <returns>The serialized size of this object</returns>
 		public static int GetSerializedSize(this StreamDirectoryInfo DirectoryInfo)
 		{
-			return sizeof(int) + DirectoryInfo.NameToFile.Values.Sum(x => x.GetSerializedSize()) + sizeof(int) + DirectoryInfo.NameToSubDirectory.Sum(x => x.Key.GetSerializedSize() + Digest<Sha1>.Length);
+			return sizeof(int) + DirectoryInfo.NameToFile.Values.Sum(x => x.GetSerializedSize()) + sizeof(int) + DirectoryInfo.NameToSubDirectory.Sum(x => x.Key.GetSerializedSize() + IoHash.NumBytes);
 		}
 	}
 }
diff --git a/Engine/Source/Programs/Shared/EpicGames.Perforce.Managed/StreamSnapshot.cs b/Engine/Source/Programs/Shared/EpicGames.Perforce.Managed/StreamSnapshot.cs
--- a/Engine/Source/Programs/Shared/EpicGames.Perforce.Managed/StreamSnapshot.cs
+++ b/Engine/Source/Programs/Shared/EpicGames.Perforce.Managed/StreamSnapshot.cs
@@ -24,14 +24,14 @@ namespace EpicGames.Perforce.Managed
 		/// <summary>
 		/// The root digest
 		/// </summary>
-		public abstract Digest<Sha1> Root { get; }
+		public abstract IoHash Root { get; }
 
 		/// <summary>
 		/// Lookup a directory by digest
 		/// </summary>
-		/// <param name="Digest">The digest value</param>
+		/// <param name="Hash">The hash value</param>
 		/// <returns></returns>
-		public abstract StreamDirectoryInfo Lookup(Digest<Sha1> Digest);
+		public abstract StreamDirectoryInfo Lookup(IoHash Hash);
 	}
 
 	/// <summary>
@@ -54,10 +54,10 @@ namespace EpicGames.Perforce.Managed
 		/// Append the contents of this directory and subdirectories to a list
 		/// </summary>
 		/// <param name="Files">List to append to</param>
-		static void AppendFiles(StreamSnapshot Snapshot, Digest<Sha1> DirHash, List<StreamFileInfo> Files)
+		static void AppendFiles(StreamSnapshot Snapshot, IoHash DirHash, List<StreamFileInfo> Files)
 		{
 			StreamDirectoryInfo DirectoryInfo = Snapshot.Lookup(DirHash);
-			foreach (Digest<Sha1> SubDirHash in DirectoryInfo.NameToSubDirectory.Values)
+			foreach (IoHash SubDirHash in DirectoryInfo.NameToSubDirectory.Values)
 			{
 				AppendFiles(Snapshot, SubDirHash, Files);
 			}
diff --git a/Engine/Source/Programs/Shared/EpicGames.Perforce.Managed/StreamSnapshotBuilder.cs b/Engine/Source/Programs/Shared/EpicGames.Perforce.Managed/StreamSnapshotBuilder.cs
--- a/Engine/Source/Programs/Shared/EpicGames.Perforce.Managed/StreamSnapshotBuilder.cs
+++ b/Engine/Source/Programs/Shared/EpicGames.Perforce.Managed/StreamSnapshotBuilder.cs
@@ -43,14 +43,14 @@ namespace EpicGames.Perforce.Managed
 		/// </summary>
 		/// <param name="Snapshot"></param>
 		/// <param name="Hash"></param>
-		public StreamSnapshotBuilder(StreamSnapshot Snapshot, Digest<Sha1> Hash)
+		public StreamSnapshotBuilder(StreamSnapshot Snapshot, IoHash Hash)
 		{
 			StreamDirectoryInfo OldDirectory = Snapshot.Lookup(Hash);
 			foreach ((ReadOnlyUtf8String Name, StreamFileInfo File) in OldDirectory.NameToFile)
 			{
 				NameToFile[Name] = File;
 			}
-			foreach ((ReadOnlyUtf8String Name, Digest<Sha1> SubDirHash) in OldDirectory.NameToSubDirectory)
+			foreach ((ReadOnlyUtf8String Name, IoHash SubDirHash) in OldDirectory.NameToSubDirectory)
 			{
 				NameToSubDirectory[Name] = new StreamSnapshotBuilder(Snapshot, SubDirHash);
 			}
@@ -61,7 +61,7 @@ namespace EpicGames.Perforce.Managed
 		/// </summary>
 		/// <param name="HashToDirectory"></param>
 		/// <returns></returns>
-		public Digest<Sha1> Encode(Dictionary<Digest<Sha1>, StreamDirectoryInfo> HashToDirectory)
+		public IoHash Encode(Dictionary<IoHash, StreamDirectoryInfo> HashToDirectory)
 		{
 			// Create the stream directory object
 			StreamDirectoryInfo Directory = new StreamDirectoryInfo();
@@ -78,7 +78,7 @@ namespace EpicGames.Perforce.Managed
 			Writer.WriteStreamDirectoryInfo(Directory);
 			Writer.CheckOffset(Data.Length);
 
-			Digest<Sha1> Hash = Digest.Compute<Sha1>(Data);
+			IoHash Hash = IoHash.Compute(Data);
 			HashToDirectory[Hash] = Directory;
 
 			return Hash;
diff --git a/Engine/Source/Programs/Shared/EpicGames.Perforce.Managed/StreamSnapshotFromGit.cs b/Engine/Source/Programs/Shared/EpicGames.Perforce.Managed/StreamSnapshotFromGit.cs
deleted file mode 100644
diff --git a/Engine/Source/Programs/Shared/EpicGames.Perforce.Managed/StreamSnapshotFromMemory.cs b/Engine/Source/Programs/Shared/EpicGames.Perforce.Managed/StreamSnapshotFromMemory.cs
--- a/Engine/Source/Programs/Shared/EpicGames.Perforce.Managed/StreamSnapshotFromMemory.cs
+++ b/Engine/Source/Programs/Shared/EpicGames.Perforce.Managed/StreamSnapshotFromMemory.cs
@@ -24,19 +24,19 @@ namespace EpicGames.Perforce.Managed
 		/// <summary>
 		/// The root digest
 		/// </summary>
-		public override Digest<Sha1> Root { get; }
+		public override IoHash Root { get; }
 
 		/// <summary>
 		/// Map of digest to directory
 		/// </summary>
-		public IReadOnlyDictionary<Digest<Sha1>, StreamDirectoryInfo> HashToDirectory { get; }
+		public IReadOnlyDictionary<IoHash, StreamDirectoryInfo> HashToDirectory { get; }
 
 		/// <summary>
 		/// Constructor
 		/// </summary>
 		/// <param name="Root"></param>
 		/// <param name="HashToDirectory"></param>
-		public StreamSnapshotFromMemory(Digest<Sha1> Root, Dictionary<Digest<Sha1>, StreamDirectoryInfo> HashToDirectory)
+		public StreamSnapshotFromMemory(IoHash Root, Dictionary<IoHash, StreamDirectoryInfo> HashToDirectory)
 		{
 			this.Root = Root;
 			this.HashToDirectory = HashToDirectory;
@@ -48,13 +48,13 @@ namespace EpicGames.Perforce.Managed
 		/// <param name="Builder"></param>
 		public StreamSnapshotFromMemory(StreamSnapshotBuilder Builder)
 		{
-			Dictionary<Digest<Sha1>, StreamDirectoryInfo> HashToDirectory = new Dictionary<Digest<Sha1>, StreamDirectoryInfo>();
+			Dictionary<IoHash, StreamDirectoryInfo> HashToDirectory = new Dictionary<IoHash, StreamDirectoryInfo>();
 			this.Root = Builder.Encode(HashToDirectory);
 			this.HashToDirectory = HashToDirectory;
 		}
 
 		/// <inheritdoc/>
-		public override StreamDirectoryInfo Lookup(Digest<Sha1> Hash)
+		public override StreamDirectoryInfo Lookup(IoHash Hash)
 		{
 			return HashToDirectory[Hash];
 		}
@@ -75,14 +75,14 @@ namespace EpicGames.Perforce.Managed
 
 			MemoryReader Reader = new MemoryReader(Data.AsMemory(CurrentSignature.Length));
 
-			Digest<Sha1> Root = Reader.ReadDigest<Sha1>();
+			IoHash Root = Reader.ReadIoHash();
 
 			int NumItems = Reader.ReadInt32();
 
-			Dictionary<Digest<Sha1>, StreamDirectoryInfo> HashToDirectory = new Dictionary<Digest<Sha1>, StreamDirectoryInfo>(NumItems);
+			Dictionary<IoHash, StreamDirectoryInfo> HashToDirectory = new Dictionary<IoHash, StreamDirectoryInfo>(NumItems);
 			for (int Idx = 0; Idx < NumItems; Idx++)
 			{
-				Digest<Sha1> Hash = Reader.ReadDigest<Sha1>();
+				IoHash Hash = Reader.ReadIoHash();
 				StreamDirectoryInfo Info = Reader.ReadStreamDirectoryInfo();
 				HashToDirectory[Hash] = Info;
 			}
@@ -100,11 +100,11 @@ namespace EpicGames.Perforce.Managed
 			byte[] Data = new byte[Length];
 
 			MemoryWriter Writer = new MemoryWriter(Data.AsMemory());
-			Writer.WriteDigest<Sha1>(Root);
+			Writer.WriteIoHash(Root);
 			Writer.WriteInt32(HashToDirectory.Count);
-			foreach ((Digest<Sha1> Hash, StreamDirectoryInfo Info) in HashToDirectory)
+			foreach ((IoHash Hash, StreamDirectoryInfo Info) in HashToDirectory)
 			{
-				Writer.WriteDigest<Sha1>(Hash);
+				Writer.WriteIoHash(Hash);
 				Writer.WriteStreamDirectoryInfo(Info);
 			}
 			Writer.CheckOffset(Data.Length);
diff --git a/Engine/Source/Programs/Shared/EpicGames.Perforce.Managed/Transactions/AddTransaction.cs b/Engine/Source/Programs/Shared/EpicGames.Perforce.Managed/Transactions/AddTransaction.cs
--- a/Engine/Source/Programs/Shared/EpicGames.Perforce.Managed/Transactions/AddTransaction.cs
+++ b/Engine/Source/Programs/Shared/EpicGames.Perforce.Managed/Transactions/AddTransaction.cs
@@ -81,14 +81,14 @@ namespace EpicGames.Perforce.Managed
 			}
 		}
 
-		void MergeDirectory(WorkspaceDirectoryInfo WorkspaceDir, WorkspaceDirectoryInfo NewWorkspaceDir, Digest<Sha1> StreamDirHash, ThreadPoolWorkQueue Queue)
+		void MergeDirectory(WorkspaceDirectoryInfo WorkspaceDir, WorkspaceDirectoryInfo NewWorkspaceDir, IoHash StreamDirHash, ThreadPoolWorkQueue Queue)
 		{
 			// Make sure the directory exists
 			Directory.CreateDirectory(WorkspaceDir.GetFullName());
 
 			// Update all the subdirectories
 			StreamDirectoryInfo StreamDir = StreamSnapshot.Lookup(StreamDirHash);
-			foreach((ReadOnlyUtf8String SubDirName, Digest<Sha1> SubDirHash) in StreamDir.NameToSubDirectory)
+			foreach((ReadOnlyUtf8String SubDirName, IoHash SubDirHash) in StreamDir.NameToSubDirectory)
 			{
 				WorkspaceDirectoryInfo? WorkspaceSubDir;
 				if(WorkspaceDir.NameToSubDirectory.TryGetValue(SubDirName, out WorkspaceSubDir))
@@ -116,7 +116,7 @@ namespace EpicGames.Perforce.Managed
 			}
 		}
 
-		void AddDirectory(WorkspaceDirectoryInfo NewWorkspaceDir, Digest<Sha1> StreamDirHash, ThreadPoolWorkQueue Queue)
+		void AddDirectory(WorkspaceDirectoryInfo NewWorkspaceDir, IoHash StreamDirHash, ThreadPoolWorkQueue Queue)
 		{
 			StreamDirectoryInfo StreamDir = StreamSnapshot.Lookup(StreamDirHash);
 
@@ -124,7 +124,7 @@ namespace EpicGames.Perforce.Managed
 			Directory.CreateDirectory(NewWorkspaceDir.GetFullName());
 
 			// Add all the sub directories and files
-			foreach((ReadOnlyUtf8String SubDirName, Digest<Sha1> SubDirHash) in StreamDir.NameToSubDirectory)
+			foreach((ReadOnlyUtf8String SubDirName, IoHash SubDirHash) in StreamDir.NameToSubDirectory)
 			{
 				AddSubDirectory(SubDirName, NewWorkspaceDir, SubDirHash, Queue);
 			}
@@ -134,14 +134,14 @@ namespace EpicGames.Perforce.Managed
 			}
 		}
 
-		void MergeSubDirectory(ReadOnlyUtf8String Name, WorkspaceDirectoryInfo WorkspaceSubDir, Digest<Sha1> StreamSubDirHash, WorkspaceDirectoryInfo NewWorkspaceDir, ThreadPoolWorkQueue Queue)
+		void MergeSubDirectory(ReadOnlyUtf8String Name, WorkspaceDirectoryInfo WorkspaceSubDir, IoHash StreamSubDirHash, WorkspaceDirectoryInfo NewWorkspaceDir, ThreadPoolWorkQueue Queue)
 		{
 			WorkspaceDirectoryInfo NewWorkspaceSubDir = new WorkspaceDirectoryInfo(NewWorkspaceDir, Name, StreamSubDirHash);
 			NewWorkspaceDir.NameToSubDirectory.Add(Name, NewWorkspaceSubDir);
 			Queue.Enqueue(() => MergeDirectory(WorkspaceSubDir, NewWorkspaceSubDir, StreamSubDirHash, Queue));
 		}
 
-		void AddSubDirectory(ReadOnlyUtf8String Name, WorkspaceDirectoryInfo NewWorkspaceDir, Digest<Sha1> StreamSubDirHash, ThreadPoolWorkQueue Queue)
+		void AddSubDirectory(ReadOnlyUtf8String Name, WorkspaceDirectoryInfo NewWorkspaceDir, IoHash StreamSubDirHash, ThreadPoolWorkQueue Queue)
 		{
 			WorkspaceDirectoryInfo NewWorkspaceSubDir = new WorkspaceDirectoryInfo(NewWorkspaceDir, Name, StreamSubDirHash);
 			NewWorkspaceDir.NameToSubDirectory.Add(Name, NewWorkspaceSubDir);
diff --git a/Engine/Source/Programs/Shared/EpicGames.Perforce.Managed/Transactions/RemoveTransaction.cs b/Engine/Source/Programs/Shared/EpicGames.Perforce.Managed/Transactions/RemoveTransaction.cs
--- a/Engine/Source/Programs/Shared/EpicGames.Perforce.Managed/Transactions/RemoveTransaction.cs
+++ b/Engine/Source/Programs/Shared/EpicGames.Perforce.Managed/Transactions/RemoveTransaction.cs
@@ -32,14 +32,14 @@ namespace EpicGames.Perforce.Managed
 			}
 		}
 
-		void Merge(WorkspaceDirectoryInfo WorkspaceDir, WorkspaceDirectoryInfo NewWorkspaceDir, Digest<Sha1> StreamDirHash, ThreadPoolWorkQueue Queue)
+		void Merge(WorkspaceDirectoryInfo WorkspaceDir, WorkspaceDirectoryInfo NewWorkspaceDir, IoHash StreamDirHash, ThreadPoolWorkQueue Queue)
 		{
 			StreamDirectoryInfo StreamDir = StreamSnapshot.Lookup(StreamDirHash);
 
 			// Update all the subdirectories
 			foreach(WorkspaceDirectoryInfo WorkspaceSubDir in WorkspaceDir.NameToSubDirectory.Values)
 			{
-				Digest<Sha1> StreamSubDirHash;
+				IoHash StreamSubDirHash;
 				if(StreamDir.NameToSubDirectory.TryGetValue(WorkspaceSubDir.Name, out StreamSubDirHash))
 				{
 					MergeSubDirectory(WorkspaceSubDir.Name, NewWorkspaceDir, WorkspaceSubDir, StreamSubDirHash, Queue);
@@ -65,7 +65,7 @@ namespace EpicGames.Perforce.Managed
 			}
 		}
 
-		void MergeSubDirectory(ReadOnlyUtf8String Name, WorkspaceDirectoryInfo NewWorkspaceDir, WorkspaceDirectoryInfo WorkspaceSubDir, Digest<Sha1> StreamSubDirHash, ThreadPoolWorkQueue Queue)
+		void MergeSubDirectory(ReadOnlyUtf8String Name, WorkspaceDirectoryInfo NewWorkspaceDir, WorkspaceDirectoryInfo WorkspaceSubDir, IoHash StreamSubDirHash, ThreadPoolWorkQueue Queue)
 		{
 			WorkspaceDirectoryInfo NewWorkspaceSubDir = new WorkspaceDirectoryInfo(NewWorkspaceDir, Name, StreamSubDirHash);
 			NewWorkspaceDir.NameToSubDirectory.Add(NewWorkspaceSubDir.Name, NewWorkspaceSubDir);
diff --git a/Engine/Source/Programs/Shared/EpicGames.Perforce.Managed/WorkspaceDirectoryInfo.cs b/Engine/Source/Programs/Shared/EpicGames.Perforce.Managed/WorkspaceDirectoryInfo.cs
--- a/Engine/Source/Programs/Shared/EpicGames.Perforce.Managed/WorkspaceDirectoryInfo.cs
+++ b/Engine/Source/Programs/Shared/EpicGames.Perforce.Managed/WorkspaceDirectoryInfo.cs
@@ -28,7 +28,7 @@ namespace EpicGames.Perforce.Managed
 		/// <summary>
 		/// Digest of the matching stream directory info. This should be set to zero if the workspace is modified.
 		/// </summary>
-		public Digest<Sha1> StreamDirectoryDigest { get; set; }
+		public IoHash StreamDirectoryDigest { get; set; }
 
 		/// <summary>
 		/// Map of name to file
@@ -45,7 +45,7 @@ namespace EpicGames.Perforce.Managed
 		/// </summary>
 		/// <param name="RootDir"></param>
 		public WorkspaceDirectoryInfo(DirectoryReference RootDir)
-			: this(null, RootDir.FullName, Digest<Sha1>.Zero)
+			: this(null, RootDir.FullName, IoHash.Zero)
 		{
 		}
 
@@ -55,7 +55,7 @@ namespace EpicGames.Perforce.Managed
 		/// <param name="ParentDirectory">The parent directory</param>
 		/// <param name="Name">Name of this directory</param>
 		/// <param name="Digest">The corresponding stream digest</param>
-		public WorkspaceDirectoryInfo(WorkspaceDirectoryInfo? ParentDirectory, ReadOnlyUtf8String Name, Digest<Sha1> Digest)
+		public WorkspaceDirectoryInfo(WorkspaceDirectoryInfo? ParentDirectory, ReadOnlyUtf8String Name, IoHash Digest)
 		{
 			this.ParentDirectory = ParentDirectory;
 			this.Name = Name;
@@ -74,7 +74,7 @@ namespace EpicGames.Perforce.Managed
 		/// <param name="ContentId">Unique identifier for the server content</param>
 		public void AddFile(ReadOnlyUtf8String Path, long Length, long LastModifiedTicks, bool bReadOnly, FileContentId ContentId)
 		{
-			StreamDirectoryDigest = Digest<Sha1>.Zero;
+			StreamDirectoryDigest = IoHash.Zero;
 
 			int Idx = Path.Span.IndexOf((byte)'/');
 			if (Idx == -1)
@@ -88,7 +88,7 @@ namespace EpicGames.Perforce.Managed
 				WorkspaceDirectoryInfo? SubDirectory;
 				if (!NameToSubDirectory.TryGetValue(Name, out SubDirectory))
 				{
-					SubDirectory = new WorkspaceDirectoryInfo(this, Name, Digest<Sha1>.Zero);
+					SubDirectory = new WorkspaceDirectoryInfo(this, Name, IoHash.Zero);
 					NameToSubDirectory[Name] = SubDirectory;
 				}
 
@@ -194,9 +194,9 @@ namespace EpicGames.Perforce.Managed
 			// If the file state has changed, clear the directory hashes
 			if (NameToFile.Count != NewNameToFile.Count)
 			{
-				for (WorkspaceDirectoryInfo? Directory = this; Directory != null && Directory.StreamDirectoryDigest != Digest<Sha1>.Zero; Directory = Directory.ParentDirectory)
+				for (WorkspaceDirectoryInfo? Directory = this; Directory != null && Directory.StreamDirectoryDigest != IoHash.Zero; Directory = Directory.ParentDirectory)
 				{
-					Directory.StreamDirectoryDigest = Digest<Sha1>.Zero;
+					Directory.StreamDirectoryDigest = IoHash.Zero;
 				}
 			}
 
@@ -334,11 +334,16 @@ namespace EpicGames.Perforce.Managed
 		{
 			if (Version < ManagedWorkspaceVersion.AddDigest)
 			{
-				DirectoryInfo.StreamDirectoryDigest = Digest<Sha1>.Zero;
+				DirectoryInfo.StreamDirectoryDigest = IoHash.Zero;
+			}
+			else if (Version < ManagedWorkspaceVersion.AddDigestIoHash)
+			{
+				Reader.ReadFixedLengthBytes(Sha1.Length);
+				DirectoryInfo.StreamDirectoryDigest = IoHash.Zero;
 			}
 			else
 			{
-				DirectoryInfo.StreamDirectoryDigest = Reader.ReadDigest<Sha1>();
+				DirectoryInfo.StreamDirectoryDigest = Reader.ReadIoHash();
 			}
 
 			int NumFiles = Reader.ReadInt32();
@@ -353,7 +358,7 @@ namespace EpicGames.Perforce.Managed
 			{
 				ReadOnlyUtf8String Name = Reader.ReadString();
 
-				WorkspaceDirectoryInfo SubDirectory = new WorkspaceDirectoryInfo(DirectoryInfo, Name, Digest<Sha1>.Zero);
+				WorkspaceDirectoryInfo SubDirectory = new WorkspaceDirectoryInfo(DirectoryInfo, Name, IoHash.Zero);
 				Reader.ReadWorkspaceDirectoryInfo(SubDirectory, Version);
 				DirectoryInfo.NameToSubDirectory[SubDirectory.Name] = SubDirectory;
 			}
@@ -361,7 +366,7 @@ namespace EpicGames.Perforce.Managed
 
 		public static void WriteWorkspaceDirectoryInfo(this MemoryWriter Writer, WorkspaceDirectoryInfo DirectoryInfo)
 		{
-			Writer.WriteDigest<Sha1>(DirectoryInfo.StreamDirectoryDigest);
+			Writer.WriteIoHash(DirectoryInfo.StreamDirectoryDigest);
 
 			Writer.WriteInt32(DirectoryInfo.NameToFile.Count);
 			foreach (WorkspaceFileInfo File in DirectoryInfo.NameToFile.Values)
