Horde: Add a custom Json schema library that supports polymorphic types and System.Text.Json.

diff --git a/Engine/Source/Programs/Horde/HordeServer/Controllers/SchemaController.cs b/Engine/Source/Programs/Horde/HordeServer/Controllers/SchemaController.cs
--- a/Engine/Source/Programs/Horde/HordeServer/Controllers/SchemaController.cs
+++ b/Engine/Source/Programs/Horde/HordeServer/Controllers/SchemaController.cs
@@ -1,10 +1,8 @@
 // Copyright Epic Games, Inc. All Rights Reserved.
 
+using EpicGames.Core;
 using HordeServer.Models;
 using HordeServer.Utilities;
-using Json.Schema;
-using Json.Schema.Generation;
-
 using Microsoft.AspNetCore.Mvc;
 using Microsoft.Extensions.Primitives;
 using System;
@@ -96,9 +94,7 @@ namespace HordeServer.Controllers
 					JsonSchema Schema = Schemas.CreateSchema(SchemaType);
 
 					using MemoryStream Stream = new MemoryStream();
-					using Utf8JsonWriter Writer = new Utf8JsonWriter(Stream);
-					JsonSerializer.Serialize(Writer, Schema, new JsonSerializerOptions());
-					Writer.Flush();
+					Schema.Write(Stream);
 
 					return new FileContentResult(Stream.ToArray(), "application/json");
 				}
diff --git a/Engine/Source/Programs/Horde/HordeServer/HordeServer.csproj b/Engine/Source/Programs/Horde/HordeServer/HordeServer.csproj
--- a/Engine/Source/Programs/Horde/HordeServer/HordeServer.csproj
+++ b/Engine/Source/Programs/Horde/HordeServer/HordeServer.csproj
@@ -24,8 +24,6 @@
       <PrivateAssets>all</PrivateAssets>
       <IncludeAssets>runtime; build; native; contentfiles; analyzers; buildtransitive</IncludeAssets>
     </PackageReference>
-    <PackageReference Include="JsonSchema.Net" Version="1.10.2" />
-    <PackageReference Include="JsonSchema.Net.Generation" Version="1.5.0" />
     <PackageReference Include="Microsoft.AspNetCore.Authentication.Cookies" Version="2.2.0" />
     <PackageReference Include="Microsoft.AspNetCore.Authentication.JwtBearer" Version="3.1.5" />
     <PackageReference Include="Microsoft.AspNetCore.Authentication.OpenIdConnect" Version="3.1.5" />
diff --git a/Engine/Source/Programs/Horde/HordeServer/Models/Config.cs b/Engine/Source/Programs/Horde/HordeServer/Models/Config.cs
--- a/Engine/Source/Programs/Horde/HordeServer/Models/Config.cs
+++ b/Engine/Source/Programs/Horde/HordeServer/Models/Config.cs
@@ -1,5 +1,6 @@
 // Copyright Epic Games, Inc. All Rights Reserved.
 
+using EpicGames.Core;
 using HordeServer.Api;
 using HordeServer.Utilities;
 using System;
diff --git a/Engine/Source/Programs/Horde/HordeServer/Utilities/Schemas.cs b/Engine/Source/Programs/Horde/HordeServer/Utilities/Schemas.cs
--- a/Engine/Source/Programs/Horde/HordeServer/Utilities/Schemas.cs
+++ b/Engine/Source/Programs/Horde/HordeServer/Utilities/Schemas.cs
@@ -1,12 +1,11 @@
 // Copyright Epic Games, Inc. All Rights Reserved.
 
 using EpicGames.Core;
-using Json.Schema;
-using Json.Schema.Generation;
-using Json.Schema.Generation.Intents;
 using System;
 using System.Collections.Concurrent;
 using System.Collections.Generic;
+using System.Diagnostics.CodeAnalysis;
+using System.IO;
 using System.Linq;
 using System.Reflection;
 using System.Text.Json;
@@ -15,304 +14,16 @@ using System.Xml;
 
 namespace HordeServer.Utilities
 {
-	/// <summary>
-	/// Attribute setting the identifier for a schema
-	/// </summary>
-	[AttributeUsage(AttributeTargets.Class)]
-	class JsonSchemaAttribute : Attribute
-	{
-		/// <summary>
-		/// The schema identifier
-		/// </summary>
-		public string Id { get; }
-
-		/// <summary>
-		/// Constructor
-		/// </summary>
-		/// <param name="Id"></param>
-		public JsonSchemaAttribute(string Id)
-		{
-			this.Id = Id;
-		}
-	}
-
-	/// <summary>
-	/// Attribute setting catalog entries for a schema
-	/// </summary>
-	[AttributeUsage(AttributeTargets.Class)]
-	public class JsonSchemaCatalogAttribute : Attribute
-	{
-		/// <summary>
-		/// The schema name
-		/// </summary>
-		public string Name { get; set; }
-
-		/// <summary>
-		/// Description of the schema
-		/// </summary>
-		public string Description { get; set; }
-
-		/// <summary>
-		/// File patterns to match
-		/// </summary>
-		public string[]? FileMatch { get; set; }
-
-		/// <summary>
-		/// Constructor
-		/// </summary>
-		/// <param name="Name"></param>
-		/// <param name="Description"></param>
-		/// <param name="FileMatch">File patterns to match</param>
-		public JsonSchemaCatalogAttribute(string Name, string Description, string? FileMatch)
-			: this(Name, Description, (FileMatch == null) ? (string[]?)null : new[] { FileMatch })
-		{
-		}
-
-		/// <summary>
-		/// Constructor
-		/// </summary>
-		/// <param name="Name"></param>
-		/// <param name="Description"></param>
-		/// <param name="FileMatch">File patterns to match</param>
-		public JsonSchemaCatalogAttribute(string Name, string Description, string[]? FileMatch)
-		{
-			this.Name = Name;
-			this.Description = Description;
-			this.FileMatch = FileMatch;
-		}
-	}
-
-	/// <summary>
-	/// Information about a type to generate a schema for
-	/// </summary>
-	class SchemaInfo
-	{
-		/// <summary>
-		/// Identifier for the schema
-		/// </summary>
-		public string Id { get; set; }
-
-		/// <summary>
-		/// Name of the type
-		/// </summary>
-		public string Name { get; set; }
-
-		/// <summary>
-		/// Description of the schema
-		/// </summary>
-		public string Description { get; set; }
-
-		/// <summary>
-		/// The type to generate a schema for
-		/// </summary>
-		public Type Type { get; set; }
-
-		/// <summary>
-		/// Constructor
-		/// </summary>
-		/// <param name="Id"></param>
-		/// <param name="Name"></param>
-		/// <param name="Description"></param>
-		/// <param name="Type"></param>
-		public SchemaInfo(string Id, string Name, string Description, Type Type)
-		{
-			this.Id = Id;
-			this.Name = Name;
-			this.Description = Description;
-			this.Type = Type;
-		}
-	}
-
 	/// <summary>
 	/// Utility functions for generating schemas
 	/// </summary>
 	static class Schemas
 	{
-		class NoAdditionalPropertiesIntent : ISchemaKeywordIntent
-		{
-			public void Apply(JsonSchemaBuilder builder)
-			{
-				builder.AdditionalProperties(false);
-			}
-		}
-
-		class NoAdditionalPropertiesRefiner : ISchemaRefiner
-		{
-			public void Run(SchemaGeneratorContext Context)
-			{
-				Context.Intents.Add(new NoAdditionalPropertiesIntent());
-			}
-
-			public bool ShouldRun(SchemaGeneratorContext Context)
-			{
-				return Context.Intents.OfType<PropertiesIntent>().Any();
-			}
-		}
-
-		class CamelCaseRefiner : ISchemaRefiner
-		{
-			public void Run(SchemaGeneratorContext Context)
-			{
-				PropertiesIntent? Intent = Context.Intents.OfType<PropertiesIntent>().FirstOrDefault();
-				if (Intent != null)
-				{
-					Dictionary<string, SchemaGeneratorContext> Properties = new Dictionary<string, SchemaGeneratorContext>(Intent.Properties);
-					foreach ((string Name, SchemaGeneratorContext PropertyContext) in Properties)
-					{
-						string NewName = Char.ToLowerInvariant(Name[0]) + Name.Substring(1);
-						Intent.Properties.Remove(Name);
-						Intent.Properties[NewName] = PropertyContext;
-					}
-				}
-
-			}
-
-			public bool ShouldRun(SchemaGeneratorContext Context)
-			{
-				return Context.Intents.OfType<PropertiesIntent>().Any();
-			}
-		}
-
-		class DescriptionIntent : ISchemaKeywordIntent
-		{
-			string Description;
-
-			public DescriptionIntent(string Description)
-			{
-				this.Description = Description;
-			}
-
-			public void Apply(JsonSchemaBuilder Builder)
-			{
-				Builder.Description(Description);
-			}
-
-			public override bool Equals(object? Obj)
-			{
-				return (Obj is DescriptionIntent Other) && Other.Description.Equals(Description, StringComparison.Ordinal);
-			}
-
-			public override int GetHashCode()
-			{
-				return Description.GetHashCode(StringComparison.Ordinal);
-			}
-		}
-	
-		[AttributeUsage(AttributeTargets.Class)]
-		class XmlDocAttribute : Attribute, IAttributeHandler
-		{
-			public string Selector;
-			public string Description;
-
-			public XmlDocAttribute(string Selector, string Description)
-			{
-				this.Selector = Selector;
-				this.Description = Description;
-			}
-
-			public void AddConstraints(SchemaGeneratorContext Context)
-			{
-				int Index = Math.Min(1, Context.Intents.Count);
-				Context.Intents.Insert(Index, new DescriptionIntent(Description));
-			}
-		}
-
-		class XmlDocRefiner : ISchemaRefiner
-		{
-			XmlDocument XmlDoc;
-
-			public XmlDocRefiner(XmlDocument XmlDoc)
-			{
-				this.XmlDoc = XmlDoc;
-			}
-
-			public void Run(SchemaGeneratorContext Context)
-			{
-				PropertiesIntent? Intent = Context.Intents.OfType<PropertiesIntent>().FirstOrDefault();
-				if (Intent != null)
-				{
-					KeyValuePair<string, SchemaGeneratorContext>[] Properties = Intent.Properties.ToArray();
-					foreach ((string PropertyName, SchemaGeneratorContext PropertyContext) in Properties)
-					{
-						string Selector = $"//member[@name='P:{Context.Type.FullName}.{PropertyName}']/summary";
-
-						XmlNode Node = XmlDoc.SelectSingleNode(Selector);
-						if (Node != null)
-						{
-							string Description = Node.InnerText.Trim().Replace("\r\n", "\n", StringComparison.Ordinal);
-
-							List<Attribute> Attributes = new List<Attribute>(PropertyContext.Attributes);
-							Attributes.Add(new XmlDocAttribute(Selector, Description));
-
-							SchemaGeneratorContext NewContext = SchemaGenerationContextCache.Get(PropertyContext.Type, Attributes, PropertyContext.Configuration);
-							Intent.Properties[PropertyName] = NewContext;
-						}
-					}
-				}
-			}
-
-			public bool ShouldRun(SchemaGeneratorContext Context)
-			{
-				return Context.Intents.OfType<PropertiesIntent>().Any();
-			}
-		}
-
-		class LegacyIdKeywordJsonConverter : JsonConverter<LegacyIdKeyword>
-		{
-			public override LegacyIdKeyword Read(ref Utf8JsonReader reader, Type typeToConvert, JsonSerializerOptions options)
-			{
-				throw new NotImplementedException();
-			}
-
-			public override void Write(Utf8JsonWriter writer, LegacyIdKeyword value, JsonSerializerOptions options)
-			{
-				writer.WriteString("id", value.Id);
-			}
-		}
-
-		[SchemaKeyword("id")]
-		[JsonConverter(typeof(LegacyIdKeywordJsonConverter))]
-		class LegacyIdKeyword : IJsonSchemaKeyword, IEquatable<LegacyIdKeyword>
-		{
-			public string Id { get; set; }
-
-			public LegacyIdKeyword(string Id)
-			{
-				this.Id = Id;
-			}
-
-			public bool Equals(LegacyIdKeyword? other)
-			{
-				return other != null && other.Id.Equals(Id, StringComparison.Ordinal);
-			}
-
-			public override bool Equals(object? obj)
-			{
-				return Equals(obj as LegacyIdKeyword);
-			}
-
-			public override int GetHashCode()
-			{
-				return Id.GetHashCode(StringComparison.Ordinal);
-			}
-
-			public void Validate(ValidationContext context)
-			{
-			}
-		}
-
-		/// <summary>
-		/// Cache of generated schemas
-		/// </summary>
+		static ConcurrentDictionary<Assembly, XmlDocument?> CachedDocumentation = new ConcurrentDictionary<Assembly, XmlDocument?>();
 		static ConcurrentDictionary<Type, JsonSchema> CachedSchemas = new ConcurrentDictionary<Type, JsonSchema>();
 
 		/// <summary>
-		/// Xml documentation for this assembly
-		/// </summary>
-		static XmlDocument? XmlDoc { get; } = ReadXmlDocumentation();
-
-		/// <summary>
-		/// Creates a schema for the given type
+		/// Create a Json schema (or retrieve a cached schema)
 		/// </summary>
 		/// <param name="Type"></param>
 		/// <returns></returns>
@@ -321,61 +32,22 @@ namespace HordeServer.Utilities
 			JsonSchema? Schema;
 			if (!CachedSchemas.TryGetValue(Type, out Schema))
 			{
-				JsonSchemaAttribute SchemaAttribute = Type.GetCustomAttribute<JsonSchemaAttribute>()!;
-
-				SchemaGeneratorConfiguration Config = new SchemaGeneratorConfiguration();
-				Config.PropertyOrder = PropertyOrder.AsDeclared;
-				Config.Generators.Add(new StringIdSchemaGenerator());
-				Config.Refiners.Add(new NoAdditionalPropertiesRefiner());
-				if (XmlDoc != null)
+				XmlDocument? Documentation;
+				if (!CachedDocumentation.TryGetValue(Type.Assembly, out Documentation))
 				{
-					Config.Refiners.Add(new XmlDocRefiner(XmlDoc));
+					FileReference InputDocumentationFile = new FileReference(Type.Assembly.Location).ChangeExtension(".xml");
+					if (FileReference.Exists(InputDocumentationFile))
+					{
+						Documentation = new XmlDocument();
+						Documentation.Load(InputDocumentationFile.FullName);
+					}
+					CachedDocumentation.TryAdd(Type.Assembly, Documentation);
 				}
-				Config.Refiners.Add(new CamelCaseRefiner());
-
-				JsonSchemaBuilder Builder = new JsonSchemaBuilder()
-					.Schema("http://json-schema.org/draft-04/schema#"); // VS2019 only supports draft-4 :(
-				Builder.Add(new LegacyIdKeyword(SchemaAttribute.Id));
-				Schema = Builder
-					.FromType(Type, Config)
-					.Build();
 
+				Schema = JsonSchema.FromType(Type, Documentation);
 				CachedSchemas.TryAdd(Type, Schema);
 			}
 			return Schema;
 		}
-
-		static XmlDocument? ReadXmlDocumentation()
-		{
-			// Get the path to the XML documentation
-			FileReference InputDocumentationFile = new FileReference(Assembly.GetExecutingAssembly().Location).ChangeExtension(".xml");
-			if (!FileReference.Exists(InputDocumentationFile))
-			{
-				return null;
-			}
-
-			// Read the documentation
-			XmlDocument InputDocumentation = new XmlDocument();
-			InputDocumentation.Load(InputDocumentationFile.FullName);
-			return InputDocumentation;
-		}
-
-		/// <summary>
-		/// Writes a schema to a source file
-		/// </summary>
-		/// <param name="Type"></param>
-		/// <param name="OutputFile"></param>
-		public static void WriteSchema(Type Type, FileReference OutputFile)
-		{
-			JsonSchema Schema = CreateSchema(Type);
-
-			JsonSerializerOptions Options = new JsonSerializerOptions();
-			Options.WriteIndented = true;
-			Options.PropertyNamingPolicy = JsonNamingPolicy.CamelCase;
-			Options.DictionaryKeyPolicy = JsonNamingPolicy.CamelCase;
-
-			byte[] Result = JsonSerializer.SerializeToUtf8Bytes(Schema, Options);
-			FileReference.WriteAllBytes(OutputFile, Result);
-		}
 	}
 }
diff --git a/Engine/Source/Programs/Horde/HordeServer/Utilities/StringId.cs b/Engine/Source/Programs/Horde/HordeServer/Utilities/StringId.cs
--- a/Engine/Source/Programs/Horde/HordeServer/Utilities/StringId.cs
+++ b/Engine/Source/Programs/Horde/HordeServer/Utilities/StringId.cs
@@ -1,9 +1,7 @@
 // Copyright Epic Games, Inc. All Rights Reserved.
 
+using EpicGames.Core;
 using HordeServer.Models;
-using Json.Schema;
-using Json.Schema.Generation;
-using Json.Schema.Generation.Generators;
 using MongoDB.Bson.Serialization;
 using MongoDB.Bson.Serialization.Attributes;
 using MongoDB.Bson.Serialization.Serializers;
@@ -20,28 +18,10 @@ using System.Threading.Tasks;
 
 namespace HordeServer.Utilities
 {
-	class StringIdSchemaIntent : ISchemaKeywordIntent
-	{
-		public void Apply(JsonSchemaBuilder Builder) => Builder.Type(SchemaValueType.String);
-	}
-
-
-	internal class StringIdSchemaGenerator : ISchemaGenerator
-	{
-		public bool Handles(Type Type)
-		{
-			return Type.IsGenericType && Type.GetGenericTypeDefinition() == typeof(StringId<>);
-		}
-
-		public void AddConstraints(SchemaGeneratorContext context)
-		{
-			context.Intents.Add(new StringIdSchemaIntent());
-		}
-	}
-
 	/// <summary>
 	/// Normalized string identifier for a resource
 	/// </summary>
+	[JsonSchemaString]
 	[TypeConverter(typeof(StringIdTypeConverter))]
 	public struct StringId<T> : IEquatable<StringId<T>>
 	{
diff --git a/Engine/Source/Programs/Shared/EpicGames.Core/JsonSchema.cs b/Engine/Source/Programs/Shared/EpicGames.Core/JsonSchema.cs
new file mode 100644
--- /dev/null
+++ b/Engine/Source/Programs/Shared/EpicGames.Core/JsonSchema.cs
@@ -0,0 +1,466 @@
+// Copyright Epic Games, Inc. All Rights Reserved.
+
+using System;
+using System.Collections.Generic;
+using System.IO;
+using System.Linq;
+using System.Reflection;
+using System.Text;
+using System.Text.Json;
+using System.Threading.Tasks;
+using System.Xml;
+
+namespace EpicGames.Core
+{
+	/// <summary>
+	/// Interface used to write JSON schema types. This is abstracted to allow multiple passes over the document structure, in order to optimize multiple references to the same type definition.
+	/// </summary>
+	public interface IJsonSchemaWriter
+	{
+		/// <inheritdoc cref="Utf8JsonWriter.WriteStartObject()"/>
+		void WriteStartObject();
+
+		/// <inheritdoc cref="Utf8JsonWriter.WriteStartObject(string)"/>
+		void WriteStartObject(string Name);
+
+		/// <inheritdoc cref="Utf8JsonWriter.WriteEndObject()"/>
+		void WriteEndObject();
+
+		/// <inheritdoc cref="Utf8JsonWriter.WriteStartArray()"/>
+		void WriteStartArray();
+
+		/// <inheritdoc cref="Utf8JsonWriter.WriteStartArray(string)"/>
+		void WriteStartArray(string Name);
+
+		/// <inheritdoc cref="Utf8JsonWriter.WriteEndObject()"/>
+		void WriteEndArray();
+
+		/// <inheritdoc cref="Utf8JsonWriter.WriteBoolean(string, bool)"/>
+		void WriteBoolean(string Name, bool Value);
+
+		/// <inheritdoc cref="Utf8JsonWriter.WriteString(string, string)"/>
+		void WriteString(string Key, string Value);
+
+		/// <inheritdoc cref="Utf8JsonWriter.WriteStringValue(string)"/>
+		void WriteStringValue(string Name);
+
+		/// <summary>
+		/// Serialize a type to JSON
+		/// </summary>
+		/// <param name="Type"></param>
+		void WriteType(JsonSchemaType Type);
+	}
+
+	/// <summary>
+	/// Implementation of a JSON schema. Implements draft 04 (latest supported by Visual Studio 2019).
+	/// </summary>
+	public class JsonSchema
+	{
+		class JsonTypeRefCollector : IJsonSchemaWriter
+		{
+			/// <summary>
+			/// Reference counts for each type (max of 2)
+			/// </summary>
+			public Dictionary<JsonSchemaType, int> TypeRefCount { get; } = new Dictionary<JsonSchemaType, int>();
+
+			/// <inheritdoc/>
+			public void WriteBoolean(string Name, bool Value) { }
+
+			/// <inheritdoc/>
+			public void WriteString(string Key, string Value) { }
+
+			/// <inheritdoc/>
+			public void WriteStringValue(string Name) { }
+
+			/// <inheritdoc/>
+			public void WriteStartObject() { }
+
+			/// <inheritdoc/>
+			public void WriteStartObject(string Name) { }
+
+			/// <inheritdoc/>
+			public void WriteEndObject() { }
+
+			/// <inheritdoc/>
+			public void WriteStartArray() { }
+
+			/// <inheritdoc/>
+			public void WriteStartArray(string Name) { }
+
+			/// <inheritdoc/>
+			public void WriteEndArray() { }
+
+			/// <inheritdoc/>
+			public void WriteType(JsonSchemaType Type)
+			{
+				if (!(Type is JsonSchemaPrimitiveType))
+				{
+					TypeRefCount.TryGetValue(Type, out int RefCount);
+					if (RefCount < 2)
+					{
+						TypeRefCount[Type] = ++RefCount;
+					}
+					if (RefCount < 2)
+					{
+						Type.Write(this);
+					}
+				}
+			}
+		}
+
+		/// <summary>
+		/// Implementation of <see cref="IJsonSchemaWriter"/>
+		/// </summary>
+		class JsonSchemaWriter : IJsonSchemaWriter
+		{
+			/// <summary>
+			/// Raw Json output
+			/// </summary>
+			Utf8JsonWriter JsonWriter;
+
+			/// <summary>
+			/// Mapping of type to definition name
+			/// </summary>
+			Dictionary<JsonSchemaType, string> TypeToDefinition { get; }
+
+			/// <summary>
+			/// Constructor
+			/// </summary>
+			/// <param name="Writer"></param>
+			/// <param name="TypeToDefinition"></param>
+			public JsonSchemaWriter(Utf8JsonWriter Writer, Dictionary<JsonSchemaType, string> TypeToDefinition)
+			{
+				this.JsonWriter = Writer;
+				this.TypeToDefinition = TypeToDefinition;
+			}
+
+			/// <inheritdoc/>
+			public void WriteBoolean(string Name, bool Value) => JsonWriter.WriteBoolean(Name, Value);
+
+			/// <inheritdoc/>
+			public void WriteString(string Key, string Value) => JsonWriter.WriteString(Key, Value);
+
+			/// <inheritdoc/>
+			public void WriteStringValue(string Name) => JsonWriter.WriteStringValue(Name);
+
+			/// <inheritdoc/>
+			public void WriteStartObject() => JsonWriter.WriteStartObject();
+
+			/// <inheritdoc/>
+			public void WriteStartObject(string Name) => JsonWriter.WriteStartObject(Name);
+
+			/// <inheritdoc/>
+			public void WriteEndObject() => JsonWriter.WriteEndObject();
+
+			/// <inheritdoc/>
+			public void WriteStartArray() => JsonWriter.WriteStartArray();
+
+			/// <inheritdoc/>
+			public void WriteStartArray(string Name) => JsonWriter.WriteStartArray(Name);
+
+			/// <inheritdoc/>
+			public void WriteEndArray() => JsonWriter.WriteEndArray();
+
+			/// <summary>
+			/// Writes a type, either inline or as a reference to a definition elsewhere
+			/// </summary>
+			/// <param name="Type"></param>
+			public void WriteType(JsonSchemaType Type)
+			{
+				if (TypeToDefinition.TryGetValue(Type, out string? Definition))
+				{
+					JsonWriter.WriteString("$ref", $"#/definitions/{Definition}");
+				}
+				else
+				{
+					Type.Write(this);
+				}
+			}
+		}
+
+		/// <summary>
+		/// Identifier for the schema
+		/// </summary>
+		public string? Id { get; set; }
+
+		/// <summary>
+		/// The root schema type
+		/// </summary>
+		public JsonSchemaType RootType { get; set; }
+
+		/// <summary>
+		/// Constructor
+		/// </summary>
+		/// <param name="Id">Id for the schema</param>
+		/// <param name="RootType"></param>
+		public JsonSchema(string? Id, JsonSchemaType RootType)
+		{
+			this.Id = Id;
+			this.RootType = RootType;
+		}
+
+		/// <summary>
+		/// Write this schema to a byte array
+		/// </summary>
+		/// <param name="Writer"></param>
+		public void Write(Utf8JsonWriter Writer)
+		{
+			// Determine reference counts for each type. Any type referenced at least twice will be split off into a separate definition.
+			JsonTypeRefCollector RefCollector = new JsonTypeRefCollector();
+			RefCollector.WriteType(RootType);
+
+			// Assign names to each type definition
+			HashSet<string> DefinitionNames = new HashSet<string>();
+			Dictionary<JsonSchemaType, string> TypeToDefinition = new Dictionary<JsonSchemaType, string>();
+			foreach ((JsonSchemaType Type, int RefCount) in RefCollector.TypeRefCount)
+			{
+				if (RefCount > 1)
+				{
+					string BaseName = Type.Name ?? "unnamed";
+
+					string Name = BaseName;
+					for (int Idx = 1; !DefinitionNames.Add(Name); Idx++)
+					{
+						Name = $"{BaseName}{Idx}";
+					}
+
+					TypeToDefinition[Type] = Name;
+				}
+			}
+
+			// Write the schema
+			Writer.WriteStartObject();
+			Writer.WriteString("$schema", "http://json-schema.org/draft-04/schema#");
+			if (Id != null)
+			{
+				Writer.WriteString("$id", Id);
+			}
+
+			JsonSchemaWriter SchemaWriter = new JsonSchemaWriter(Writer, TypeToDefinition);
+			RootType.Write(SchemaWriter);
+
+			if (TypeToDefinition.Count > 0)
+			{
+				Writer.WriteStartObject("definitions");
+				foreach ((JsonSchemaType Type, string RefName) in TypeToDefinition)
+				{
+					Writer.WriteStartObject(RefName);
+					Type.Write(SchemaWriter);
+					Writer.WriteEndObject();
+				}
+				Writer.WriteEndObject();
+			}
+
+			Writer.WriteEndObject();
+		}
+
+		/// <summary>
+		/// Write this schema to a stream
+		/// </summary>
+		/// <param name="Stream">The output stream</param>
+		public void Write(Stream Stream)
+		{
+			using (Utf8JsonWriter Writer = new Utf8JsonWriter(Stream, new JsonWriterOptions { Indented = true }))
+			{
+				Write(Writer);
+			}
+		}
+
+		/// <summary>
+		/// Writes this schema to a file
+		/// </summary>
+		/// <param name="File">The output file</param>
+		public void Write(FileReference File)
+		{
+			using (FileStream Stream = FileReference.Open(File, FileMode.Create))
+			{
+				Write(Stream);
+			}
+		}
+
+		/// <summary>
+		/// Constructs a Json schema from a type
+		/// </summary>
+		/// <param name="Type">The type to construct from</param>
+		/// <param name="XmlDoc">C# Xml documentation file, to use for property descriptions</param>
+		/// <returns>New schema object</returns>
+		public static JsonSchema FromType(Type Type, XmlDocument? XmlDoc)
+		{
+			JsonSchemaAttribute? SchemaAttribute = Type.GetCustomAttribute<JsonSchemaAttribute>();
+			return new JsonSchema(SchemaAttribute?.Id, CreateSchemaType(Type, new Dictionary<Type, JsonSchemaType>(), XmlDoc));
+		}
+
+		/// <summary>
+		/// Gets the description for a type
+		/// </summary>
+		/// <param name="Type"></param>
+		/// <param name="XmlDoc"></param>
+		/// <returns></returns>
+		static string? GetTypeDescription(Type Type, XmlDocument? XmlDoc)
+		{
+			if (XmlDoc == null)
+			{
+				return null;
+			}
+
+			string Selector = $"//member[@name='T:{Type.FullName}']/summary";
+
+			XmlNode Node = XmlDoc.SelectSingleNode(Selector);
+			if (Node == null)
+			{
+				return null;
+			}
+
+			return Node.InnerText.Trim().Replace("\r\n", "\n", StringComparison.Ordinal);
+		}
+
+		/// <summary>
+		/// Gets a property description from Xml documentation file
+		/// </summary>
+		/// <param name="Type"></param>
+		/// <param name="Name"></param>
+		/// <param name="XmlDoc"></param>
+		/// <returns></returns>
+		static string? GetPropertyDescription(Type Type, string Name, XmlDocument? XmlDoc)
+		{
+			if (XmlDoc == null)
+			{
+				return null;
+			}
+
+			string Selector = $"//member[@name='P:{Type.FullName}.{Name}']/summary";
+
+			XmlNode Node = XmlDoc.SelectSingleNode(Selector);
+			if (Node == null)
+			{
+				return null;
+			}
+
+			return Node.InnerText.Trim().Replace("\r\n", "\n", StringComparison.Ordinal);
+		}
+
+		/// <summary>
+		/// Constructs a schema type from the given type object
+		/// </summary>
+		/// <param name="Type"></param>
+		/// <param name="XmlDoc"></param>
+		/// <returns></returns>
+		static JsonSchemaType CreateSchemaType(Type Type, Dictionary<Type, JsonSchemaType> TypeCache, XmlDocument? XmlDoc)
+		{
+			switch (Type.GetTypeCode(Type))
+			{
+				case TypeCode.Boolean:
+					return new JsonSchemaBoolean();
+				case TypeCode.Byte:
+				case TypeCode.SByte:
+				case TypeCode.Int16:
+				case TypeCode.UInt16:
+				case TypeCode.Int32:
+				case TypeCode.UInt32:
+				case TypeCode.Int64:
+				case TypeCode.UInt64:
+					return new JsonSchemaInteger();
+				case TypeCode.Single:
+				case TypeCode.Double:
+					return new JsonSchemaNumber();
+				case TypeCode.String:
+					return new JsonSchemaString();
+			}
+
+			JsonSchemaTypeAttribute? Attribute = Type.GetCustomAttribute<JsonSchemaTypeAttribute>();
+			switch (Attribute)
+			{
+				case JsonSchemaStringAttribute String:
+					return new JsonSchemaString(String.Format);
+			}
+
+			if (Type.IsEnum)
+			{
+				return new JsonSchemaEnum(Enum.GetNames(Type)) { Name = Type.Name };
+			}
+			if (Type.IsGenericType && Type.GetGenericTypeDefinition() == typeof(Nullable<>))
+			{
+				return CreateSchemaType(Type.GetGenericArguments()[0], TypeCache, XmlDoc);
+			}
+			if (Type == typeof(DateTime) || Type == typeof(DateTimeOffset))
+			{
+				return new JsonSchemaString(JsonSchemaStringFormat.DateTime);
+			}
+
+			Type[] Interfaces = Type.GetInterfaces();
+			foreach (Type Interface in Interfaces)
+			{
+				if (Interface.IsGenericType)
+				{
+					Type[] Arguments = Interface.GetGenericArguments();
+					if (Interface.GetGenericTypeDefinition() == typeof(IList<>))
+					{
+						return new JsonSchemaArray(CreateSchemaType(Arguments[0], TypeCache, XmlDoc));
+					}
+					if (Interface.GetGenericTypeDefinition() == typeof(IDictionary<,>))
+					{
+						JsonSchemaObject Object = new JsonSchemaObject();
+						Object.AdditionalProperties = CreateSchemaType(Arguments[1], TypeCache, XmlDoc);
+						return Object;
+					}
+				}
+			}
+
+			if (Type.IsClass)
+			{
+				if (TypeCache.TryGetValue(Type, out JsonSchemaType? SchemaType))
+				{
+					return SchemaType;
+				}
+
+				JsonKnownTypesAttribute? KnownTypes = Type.GetCustomAttribute<JsonKnownTypesAttribute>(false);
+				if (KnownTypes != null)
+				{
+					JsonSchemaOneOf Object = new JsonSchemaOneOf();
+					TypeCache[Type] = Object;
+					SetOneOfProperties(Object, Type, KnownTypes.Types, TypeCache, XmlDoc);
+					return Object;
+				}
+				else
+				{
+					JsonSchemaObject Object = new JsonSchemaObject();
+					TypeCache[Type] = Object;
+					SetObjectProperties(Object, Type, TypeCache, XmlDoc);
+					return Object;
+				}
+			}
+
+			throw new Exception($"Unknown type for schema generation: {Type}");
+		}
+
+		static void SetOneOfProperties(JsonSchemaOneOf Object, Type Type, Type[] KnownTypes, Dictionary<Type, JsonSchemaType> TypeCache, XmlDocument? XmlDoc)
+		{
+			Object.Name = Type.Name;
+
+			foreach (Type KnownType in KnownTypes)
+			{
+				JsonDiscriminatorAttribute? Attribute = KnownType.GetCustomAttribute<JsonDiscriminatorAttribute>();
+				if (Attribute != null)
+				{
+					JsonSchemaObject KnownObject = new JsonSchemaObject();
+					KnownObject.Properties.Add(new JsonSchemaProperty("type", "Type discriminator", new JsonSchemaEnum(new[] { Attribute.Name })));
+					SetObjectProperties(KnownObject, KnownType, TypeCache, XmlDoc);
+					Object.Types.Add(KnownObject);
+				}
+			}
+		}
+
+		static void SetObjectProperties(JsonSchemaObject Object, Type Type, Dictionary<Type, JsonSchemaType> TypeCache, XmlDocument? XmlDoc)
+		{
+			Object.Name = Type.Name;
+
+			PropertyInfo[] Properties = Type.GetProperties(BindingFlags.Instance | BindingFlags.Public);
+			foreach (PropertyInfo Property in Properties)
+			{
+				string? Description = GetPropertyDescription(Type, Property.Name, XmlDoc);
+				JsonSchemaType PropertyType = CreateSchemaType(Property.PropertyType, TypeCache, XmlDoc);
+				Object.Properties.Add(new JsonSchemaProperty(Property.Name, Description, PropertyType));
+			}
+		}
+	}
+}
diff --git a/Engine/Source/Programs/Shared/EpicGames.Core/JsonSchemaAttributes.cs b/Engine/Source/Programs/Shared/EpicGames.Core/JsonSchemaAttributes.cs
new file mode 100644
--- /dev/null
+++ b/Engine/Source/Programs/Shared/EpicGames.Core/JsonSchemaAttributes.cs
@@ -0,0 +1,94 @@
+// Copyright Epic Games, Inc. All Rights Reserved.
+
+using System;
+using System.Collections.Generic;
+using System.Text;
+
+namespace EpicGames.Core
+{
+	/// <summary>
+	/// Attribute setting the identifier for a schema
+	/// </summary>
+	[AttributeUsage(AttributeTargets.Class)]
+	public class JsonSchemaAttribute : Attribute
+	{
+		/// <summary>
+		/// The schema identifier
+		/// </summary>
+		public string Id { get; }
+
+		/// <summary>
+		/// Constructor
+		/// </summary>
+		/// <param name="Id"></param>
+		public JsonSchemaAttribute(string Id)
+		{
+			this.Id = Id;
+		}
+	}
+
+	/// <summary>
+	/// Attribute setting catalog entries for a schema
+	/// </summary>
+	[AttributeUsage(AttributeTargets.Class)]
+	public class JsonSchemaCatalogAttribute : Attribute
+	{
+		/// <summary>
+		/// The schema name
+		/// </summary>
+		public string Name { get; set; }
+
+		/// <summary>
+		/// Description of the schema
+		/// </summary>
+		public string Description { get; set; }
+
+		/// <summary>
+		/// File patterns to match
+		/// </summary>
+		public string[]? FileMatch { get; set; }
+
+		/// <summary>
+		/// Constructor
+		/// </summary>
+		/// <param name="Name"></param>
+		/// <param name="Description"></param>
+		/// <param name="FileMatch">File patterns to match</param>
+		public JsonSchemaCatalogAttribute(string Name, string Description, string? FileMatch)
+			: this(Name, Description, (FileMatch == null) ? (string[]?)null : new[] { FileMatch })
+		{
+		}
+
+		/// <summary>
+		/// Constructor
+		/// </summary>
+		/// <param name="Name"></param>
+		/// <param name="Description"></param>
+		/// <param name="FileMatch">File patterns to match</param>
+		public JsonSchemaCatalogAttribute(string Name, string Description, string[]? FileMatch)
+		{
+			this.Name = Name;
+			this.Description = Description;
+			this.FileMatch = FileMatch;
+		}
+	}
+
+	/// <summary>
+	/// Attribute setting properties for a type to be serialized as a string
+	/// </summary>
+	public abstract class JsonSchemaTypeAttribute : Attribute
+	{
+	}
+
+	/// <summary>
+	/// Attribute setting properties for a type to be serialized as a string
+	/// </summary>
+	[AttributeUsage(AttributeTargets.Class | AttributeTargets.Struct | AttributeTargets.Property)]
+	public class JsonSchemaStringAttribute : JsonSchemaTypeAttribute
+	{
+		/// <summary>
+		/// Format of the string
+		/// </summary>
+		public JsonSchemaStringFormat? Format { get; set; }
+	}
+}
diff --git a/Engine/Source/Programs/Shared/EpicGames.Core/JsonSchemaType.cs b/Engine/Source/Programs/Shared/EpicGames.Core/JsonSchemaType.cs
new file mode 100644
--- /dev/null
+++ b/Engine/Source/Programs/Shared/EpicGames.Core/JsonSchemaType.cs
@@ -0,0 +1,319 @@
+// Copyright Epic Games, Inc. All Rights Reserved.
+
+using System;
+using System.Collections.Generic;
+using System.IO;
+using System.Linq;
+using System.Text;
+using System.Text.Json;
+using System.Threading.Tasks;
+
+namespace EpicGames.Core
+{
+	/// <summary>
+	/// Base class for Json schema types
+	/// </summary>
+	public abstract class JsonSchemaType
+	{
+		/// <summary>
+		/// Name for this type, if stored in a standalone definition
+		/// </summary>
+		public string? Name { get; set; }
+
+		/// <summary>
+		/// Write this type to the given archive
+		/// </summary>
+		/// <param name="Writer"></param>
+		public abstract void Write(IJsonSchemaWriter Writer);
+	}
+
+	/// <summary>
+	/// Base class for types that cannot contain references to other types
+	/// </summary>
+	public abstract class JsonSchemaPrimitiveType : JsonSchemaType
+	{
+	}
+
+	/// <summary>
+	/// Represents a boolean in a Json schema
+	/// </summary>
+	public class JsonSchemaBoolean : JsonSchemaPrimitiveType
+	{
+		/// <inheritdoc/>
+		public override void Write(IJsonSchemaWriter Writer)
+		{
+			Writer.WriteString("type", "boolean");
+		}
+	}
+
+	/// <summary>
+	/// Represents an integer in a Json schema
+	/// </summary>
+	public class JsonSchemaInteger : JsonSchemaPrimitiveType
+	{
+		/// <inheritdoc/>
+		public override void Write(IJsonSchemaWriter Writer)
+		{
+			Writer.WriteString("type", "integer");
+		}
+	}
+
+	/// <summary>
+	/// Represents a number of any type in a Json schema
+	/// </summary>
+	public class JsonSchemaNumber : JsonSchemaPrimitiveType
+	{
+		/// <inheritdoc/>
+		public override void Write(IJsonSchemaWriter Writer)
+		{
+			Writer.WriteString("type", "integer");
+		}
+	}
+
+	/// <summary>
+	/// Format for a string element schema
+	/// </summary>
+	public enum JsonSchemaStringFormat
+	{
+		/// <summary>
+		/// A date/time value
+		/// </summary>
+		DateTime,
+
+		/// <summary>
+		/// Internet email address, see RFC 5322, section 3.4.1.
+		/// </summary>
+		Email,
+
+		/// <summary>
+		/// Internet host name, see RFC 1034, section 3.1.
+		/// </summary>
+		HostName,
+
+		/// <summary>
+		/// IPv4 address, according to dotted-quad ABNF syntax as defined in RFC 2673, section 3.2.
+		/// </summary>
+		Ipv4,
+
+		/// <summary>
+		/// IPv6 address, as defined in RFC 2373, section 2.2.
+		/// </summary>
+		Ipv6,
+
+		/// <summary>
+		/// A universal resource identifier (URI), according to RFC3986.
+		/// </summary>
+		Uri
+	}
+
+	/// <summary>
+	/// Represents a string in a Json schema
+	/// </summary>
+	public class JsonSchemaString : JsonSchemaPrimitiveType
+	{
+		/// <summary>
+		/// Optional string describing the formatting of this type
+		/// </summary>
+		public JsonSchemaStringFormat? Format { get; set; }
+
+		/// <summary>
+		/// Pattern for matched strings
+		/// </summary>
+		public string? Pattern { get; set; } 
+
+		/// <summary>
+		/// Constructor
+		/// </summary>
+		/// <param name="Format"></param>
+		public JsonSchemaString(JsonSchemaStringFormat? Format = null, string? Pattern = null)
+		{
+			this.Format = Format;
+			this.Pattern = Pattern;
+		}
+
+		/// <inheritdoc/>
+		public override void Write(IJsonSchemaWriter Writer)
+		{
+			Writer.WriteString("type", "string");
+			if (Format != null)
+			{
+				Writer.WriteString("format", Format.Value.ToString().ToLowerInvariant());
+			}
+			if (Pattern != null)
+			{
+				Writer.WriteString("pattern", Pattern);
+			}
+		}
+	}
+
+	/// <summary>
+	/// Represents an enum in a Json schema
+	/// </summary>
+	public class JsonSchemaEnum : JsonSchemaType
+	{
+		/// <summary>
+		/// Values for the enum
+		/// </summary>
+		public List<string> Values { get; } = new List<string>();
+
+		/// <summary>
+		/// Constructor
+		/// </summary>
+		public JsonSchemaEnum(IEnumerable<string> Values)
+		{
+			this.Values.AddRange(Values);
+		}
+
+		/// <inheritdoc/>
+		public override void Write(IJsonSchemaWriter Writer)
+		{
+			Writer.WriteStartArray("enum");
+			foreach (string Value in Values)
+			{
+				Writer.WriteStringValue(Value);
+			}
+			Writer.WriteEndArray();
+		}
+	}
+
+	/// <summary>
+	/// Represents an array in a Json schema
+	/// </summary>
+	public class JsonSchemaArray : JsonSchemaType
+	{
+		/// <summary>
+		/// The type of each element
+		/// </summary>
+		public JsonSchemaType ItemType { get; }
+
+		/// <summary>
+		/// Constructor
+		/// </summary>
+		public JsonSchemaArray(JsonSchemaType ItemType)
+		{
+			this.ItemType = ItemType;
+		}
+
+		/// <inheritdoc/>
+		public override void Write(IJsonSchemaWriter Writer)
+		{
+			Writer.WriteString("type", "array");
+			Writer.WriteStartObject("items");
+			Writer.WriteType(ItemType);
+			Writer.WriteEndObject();
+		}
+	}
+
+	/// <summary>
+	/// Property within a <see cref="JsonSchemaObject"/>
+	/// </summary>
+	public class JsonSchemaProperty
+	{
+		/// <summary>
+		/// Name of this property
+		/// </summary>
+		public string Name { get; set; }
+
+		/// <summary>
+		/// Description of the property
+		/// </summary>
+		public string? Description { get; set; }
+
+		/// <summary>
+		/// The property type
+		/// </summary>
+		public JsonSchemaType Type { get; set; }
+
+		/// <summary>
+		/// Constructor
+		/// </summary>
+		public JsonSchemaProperty(string Name, string? Description, JsonSchemaType Type)
+		{
+			this.Name = Name;
+			this.Description = Description;
+			this.Type = Type;
+		}
+
+		/// <summary>
+		/// Writes a property to a schema
+		/// </summary>
+		/// <param name="Writer"></param>
+		public void Write(IJsonSchemaWriter Writer)
+		{
+			string CamelCaseName = Char.ToLowerInvariant(Name[0]) + Name.Substring(1);
+			Writer.WriteStartObject(CamelCaseName);
+			Writer.WriteType(Type);
+			if (Description != null)
+			{
+				Writer.WriteString("description", Description);
+			}
+			Writer.WriteEndObject();
+		}
+	}
+
+	/// <summary>
+	/// Represents an object in a Json schema
+	/// </summary>
+	public class JsonSchemaObject : JsonSchemaType
+	{
+		/// <summary>
+		/// Properties to allow in this object
+		/// </summary>
+		public List<JsonSchemaProperty> Properties { get; set; } = new List<JsonSchemaProperty>();
+
+		/// <summary>
+		/// Type for additional properties
+		/// </summary>
+		public JsonSchemaType? AdditionalProperties { get; set; }
+
+		/// <inheritdoc/>
+		public override void Write(IJsonSchemaWriter Writer)
+		{
+			Writer.WriteString("type", "object");
+			if (Properties.Count > 0)
+			{
+				Writer.WriteStartObject("properties");
+				foreach (JsonSchemaProperty Property in Properties)
+				{
+					Property.Write(Writer);
+				}
+				Writer.WriteEndObject();
+			}
+
+			if (AdditionalProperties == null)
+			{
+				Writer.WriteBoolean("additionalProperties", false);
+			}
+			else
+			{
+				Writer.WriteStartObject("additionalProperties");
+				AdditionalProperties.Write(Writer);
+				Writer.WriteEndObject();
+			}
+		}
+	}
+
+	/// <summary>
+	/// Represents one of a set of object types in a Json schema
+	/// </summary>
+	public class JsonSchemaOneOf : JsonSchemaType
+	{
+		/// <summary>
+		/// List of valid types
+		/// </summary>
+		public List<JsonSchemaType> Types { get; } = new List<JsonSchemaType>();
+
+		/// <inheritdoc/>
+		public override void Write(IJsonSchemaWriter Writer)
+		{
+			Writer.WriteStartArray("oneOf");
+			foreach (JsonSchemaType Type in Types)
+			{
+				Writer.WriteStartObject();
+				Type.Write(Writer);
+				Writer.WriteEndObject();
+			}
+			Writer.WriteEndArray();
+		}
+	}
+}
