Manually copy stdout/stderr from managed framework processes asynchronously, in order to prevent garbled output when converting to/from strings.

diff --git a/Engine/Source/Programs/Shared/EpicGames.Core/ManagedProcess.cs b/Engine/Source/Programs/Shared/EpicGames.Core/ManagedProcess.cs
--- a/Engine/Source/Programs/Shared/EpicGames.Core/ManagedProcess.cs
+++ b/Engine/Source/Programs/Shared/EpicGames.Core/ManagedProcess.cs
@@ -317,6 +317,16 @@ namespace EpicGames.Core
 		/// </summary>
 		Process? FrameworkProcess;
 
+		/// <summary>
+		/// Thread to read from stdin
+		/// </summary>
+		Thread? FrameworkStdOutThread;
+
+		/// <summary>
+		/// Thread to read from stderr
+		/// </summary>
+		Thread? FrameworkStdErrThread;
+
 		/// <summary>
 		/// Merged output & error write stream for the framework child process.
 		/// </summary>
@@ -649,15 +659,6 @@ namespace EpicGames.Core
 				FrameworkMergedStdWriter = new AnonymousPipeServerStream(PipeDirection.Out);
 				StdOut = new AnonymousPipeClientStream(PipeDirection.In, FrameworkMergedStdWriter.ClientSafePipeHandle);
 				StdOutText = new StreamReader(StdOut, Console.OutputEncoding);
-				DataReceivedEventHandler OutputEventHandler = (Sender, Args) =>
-				{
-					if (FrameworkMergedStdWriter != null && Args.Data != null)
-					{
-						FrameworkMergedStdWriter.Write(Console.OutputEncoding.GetBytes(Args.Data + System.Environment.NewLine));
-					}
-				};
-				FrameworkProcess.OutputDataReceived += OutputEventHandler;
-				FrameworkProcess.ErrorDataReceived += OutputEventHandler;
 
 				StdErr = StdOut;
 				StdErrText = StdOutText;
@@ -666,8 +667,11 @@ namespace EpicGames.Core
 			FrameworkProcess.Start();
 			if ((ManagedFlags & ManagedProcessFlags.MergeOutputPipes) != 0)
 			{
-				FrameworkProcess.BeginOutputReadLine();
-				FrameworkProcess.BeginErrorReadLine();
+				FrameworkStdOutThread = new Thread(() => FrameworkProcess.StandardOutput.BaseStream.CopyTo(FrameworkMergedStdWriter!));
+				FrameworkStdOutThread.Start();
+
+				FrameworkStdErrThread = new Thread(() => FrameworkProcess.StandardError.BaseStream.CopyTo(FrameworkMergedStdWriter!));
+				FrameworkStdErrThread.Start();
 			}
 
 			try
@@ -702,7 +706,7 @@ namespace EpicGames.Core
 				ProcessHandle.Dispose();
 				ProcessHandle = null;
 			}
-			if(StdInWrite != null)
+			if (StdInWrite != null)
 			{
 				StdInWrite.Dispose();
 				StdInWrite = null;
@@ -725,6 +729,18 @@ namespace EpicGames.Core
 					{
 					}
 				}
+
+				if (FrameworkStdOutThread != null)
+				{
+					FrameworkStdOutThread.Join();
+					FrameworkStdOutThread = null;
+				}
+				if (FrameworkStdErrThread != null)
+				{
+					FrameworkStdErrThread.Join();
+					FrameworkStdOutThread = null;
+				}
+
 				FrameworkProcess.Dispose();
 				FrameworkProcess = null;
 			}
