Swallow exceptions when copying output from stdout if a process is killed.

diff --git a/Engine/Source/Programs/Shared/EpicGames.Core/ManagedProcess.cs b/Engine/Source/Programs/Shared/EpicGames.Core/ManagedProcess.cs
--- a/Engine/Source/Programs/Shared/EpicGames.Core/ManagedProcess.cs
+++ b/Engine/Source/Programs/Shared/EpicGames.Core/ManagedProcess.cs
@@ -667,10 +667,10 @@ namespace EpicGames.Core
 			FrameworkProcess.Start();
 			if ((ManagedFlags & ManagedProcessFlags.MergeOutputPipes) != 0)
 			{
-				FrameworkStdOutThread = new Thread(() => FrameworkProcess.StandardOutput.BaseStream.CopyTo(FrameworkMergedStdWriter!));
+				FrameworkStdOutThread = new Thread(() => CopyPipe(FrameworkProcess.StandardOutput.BaseStream, FrameworkMergedStdWriter!));
 				FrameworkStdOutThread.Start();
 
-				FrameworkStdErrThread = new Thread(() => FrameworkProcess.StandardError.BaseStream.CopyTo(FrameworkMergedStdWriter!));
+				FrameworkStdErrThread = new Thread(() => CopyPipe(FrameworkProcess.StandardError.BaseStream, FrameworkMergedStdWriter!));
 				FrameworkStdErrThread.Start();
 			}
 
@@ -693,6 +693,22 @@ namespace EpicGames.Core
 			}
 		}
 
+		/// <summary>
+		/// Copy data from one pipe to another.
+		/// </summary>
+		/// <param name="Source"></param>
+		/// <param name="Target"></param>
+		void CopyPipe(Stream Source, Stream Target)
+		{
+			try
+			{
+				Source.CopyTo(Target);
+			}
+			catch
+			{
+			}
+		}
+
 		/// <summary>
 		/// Free the managed resources for this process
 		/// </summary>
@@ -706,7 +722,7 @@ namespace EpicGames.Core
 				ProcessHandle.Dispose();
 				ProcessHandle = null;
 			}
-			if (StdInWrite != null)
+			if(StdInWrite != null)
 			{
 				StdInWrite.Dispose();
 				StdInWrite = null;
