Join stdout/stderr threads before closing the buffering thread.

diff --git a/Engine/Source/Programs/Shared/EpicGames.Core/ManagedProcess.cs b/Engine/Source/Programs/Shared/EpicGames.Core/ManagedProcess.cs
--- a/Engine/Source/Programs/Shared/EpicGames.Core/ManagedProcess.cs
+++ b/Engine/Source/Programs/Shared/EpicGames.Core/ManagedProcess.cs
@@ -647,6 +647,16 @@ namespace EpicGames.Core
 				FrameworkProcess.EnableRaisingEvents = true;
 				FrameworkProcess.Exited += (Sender, Args) =>
 				{
+					if (FrameworkStdOutThread != null)
+					{
+						FrameworkStdOutThread.Join();
+						FrameworkStdOutThread = null;
+					}
+					if (FrameworkStdErrThread != null)
+					{
+						FrameworkStdErrThread.Join();
+						FrameworkStdErrThread = null;
+					}
 					if (FrameworkMergedStdWriter != null)
 					{
 						FrameworkMergedStdWriter.Flush();
@@ -746,17 +756,6 @@ namespace EpicGames.Core
 					}
 				}
 
-				if (FrameworkStdOutThread != null)
-				{
-					FrameworkStdOutThread.Join();
-					FrameworkStdOutThread = null;
-				}
-				if (FrameworkStdErrThread != null)
-				{
-					FrameworkStdErrThread.Join();
-					FrameworkStdOutThread = null;
-				}
-
 				FrameworkProcess.Dispose();
 				FrameworkProcess = null;
 			}
