EpicGames.Core: Ensure ManagedProcess threads start before Process.Exited callback runs, to handle cases where the spawned process completes extremely quickly

diff --git a/Engine/Source/Programs/Shared/EpicGames.Core/ManagedProcess.cs b/Engine/Source/Programs/Shared/EpicGames.Core/ManagedProcess.cs
--- a/Engine/Source/Programs/Shared/EpicGames.Core/ManagedProcess.cs
+++ b/Engine/Source/Programs/Shared/EpicGames.Core/ManagedProcess.cs
@@ -647,21 +647,24 @@ namespace EpicGames.Core
 				FrameworkProcess.EnableRaisingEvents = true;
 				FrameworkProcess.Exited += (Sender, Args) =>
 				{
-					if (FrameworkStdOutThread != null)
+					lock (this)
 					{
-						FrameworkStdOutThread.Join();
-						FrameworkStdOutThread = null;
-					}
-					if (FrameworkStdErrThread != null)
-					{
-						FrameworkStdErrThread.Join();
-						FrameworkStdErrThread = null;
-					}
-					if (FrameworkMergedStdWriter != null)
-					{
-						FrameworkMergedStdWriter.Flush();
-						FrameworkMergedStdWriter.Dispose();
-						FrameworkMergedStdWriter = null;
+						if (FrameworkStdOutThread != null)
+						{
+							FrameworkStdOutThread.Join();
+							FrameworkStdOutThread = null;
+						}
+						if (FrameworkStdErrThread != null)
+						{
+							FrameworkStdErrThread.Join();
+							FrameworkStdErrThread = null;
+						}
+						if (FrameworkMergedStdWriter != null)
+						{
+							FrameworkMergedStdWriter.Flush();
+							FrameworkMergedStdWriter.Dispose();
+							FrameworkMergedStdWriter = null;
+						}
 					}
 				};
 
@@ -674,14 +677,18 @@ namespace EpicGames.Core
 				StdErrText = StdOutText;
 			}
 
-			FrameworkProcess.Start();
-			if ((ManagedFlags & ManagedProcessFlags.MergeOutputPipes) != 0)
+			// Ensure threads start before Exited is called, if the process terminates instantly
+			lock (this)
 			{
-				FrameworkStdOutThread = new Thread(() => CopyPipe(FrameworkProcess.StandardOutput.BaseStream, FrameworkMergedStdWriter!));
-				FrameworkStdOutThread.Start();
+				FrameworkProcess.Start();
+				if ((ManagedFlags & ManagedProcessFlags.MergeOutputPipes) != 0)
+				{
+					FrameworkStdOutThread = new Thread(() => CopyPipe(FrameworkProcess.StandardOutput.BaseStream, FrameworkMergedStdWriter!));
+					FrameworkStdOutThread.Start();
 
-				FrameworkStdErrThread = new Thread(() => CopyPipe(FrameworkProcess.StandardError.BaseStream, FrameworkMergedStdWriter!));
-				FrameworkStdErrThread.Start();
+					FrameworkStdErrThread = new Thread(() => CopyPipe(FrameworkProcess.StandardError.BaseStream, FrameworkMergedStdWriter!));
+					FrameworkStdErrThread.Start();
+				}
 			}
 
 			try
