RuntimePlatform:

Add an enum and accessor for the current runtime platform

#trivial

diff --git a/Engine/Source/Programs/Shared/EpicGames.Core/RuntimePlatform.cs b/Engine/Source/Programs/Shared/EpicGames.Core/RuntimePlatform.cs
--- a/Engine/Source/Programs/Shared/EpicGames.Core/RuntimePlatform.cs
+++ b/Engine/Source/Programs/Shared/EpicGames.Core/RuntimePlatform.cs
@@ -23,5 +23,15 @@ namespace EpicGames.Core
 		/// Whether we are currently running a Windows platform.
 		/// </summary>
 		public static readonly bool IsWindows = RuntimeInformation.IsOSPlatform(OSPlatform.Windows);
+
+		public enum Type
+		{
+			Windows, Linux, Mac
+		};
+
+		/// <summary>
+		/// The current runtime platform
+		/// </summary>
+		public static readonly Type Current = IsWindows ? Type.Windows : IsMac ? Type.Mac : Type.Linux;
 	}
 }
