Add C# compact binary library. Includes almost all functionality from C++ library, and test suite.

diff --git a/Engine/Source/Programs/Shared/EpicGames.Serialization.Tests/EngineTests.cs b/Engine/Source/Programs/Shared/EpicGames.Serialization.Tests/EngineTests.cs
new file mode 100644
--- /dev/null
+++ b/Engine/Source/Programs/Shared/EpicGames.Serialization.Tests/EngineTests.cs
@@ -0,0 +1,1672 @@
+// Copyright Epic Games, Inc. All Rights Reserved.
+
+using EpicGames.Core;
+using Microsoft.VisualStudio.TestTools.UnitTesting;
+using System;
+using System.Buffers.Binary;
+using System.Collections.Generic;
+using System.Text;
+
+namespace EpicGames.Serialization.Tests
+{
+	[TestClass]
+	public class EngineTests
+	{
+		class CbFieldAccessors
+		{
+			public object Default;
+			public Func<CbField, bool> IsType;
+			public Func<CbField, object> AsType;
+			public Func<CbField, object, object> AsTypeWithDefault;
+			public Func<object, object, bool> Comparer = (x, y) => x.Equals(y);
+
+			public CbFieldAccessors(object Default, Func<CbField, bool> IsType, Func<CbField, object> AsType, Func<CbField, object, object> AsTypeWithDefault)
+			{
+				this.Default = Default;
+				this.IsType = IsType;
+				this.AsType = AsType;
+				this.AsTypeWithDefault = AsTypeWithDefault;
+			}
+
+			public static CbFieldAccessors FromStruct<T>(Func<CbField, bool> IsType, Func<CbField, T, object> AsTypeWithDefault) where T : struct
+			{
+				return new CbFieldAccessors(new T(), IsType, x => AsTypeWithDefault(x, new T()), (x, y) => AsTypeWithDefault(x, (T)y));
+			}
+
+			public static CbFieldAccessors FromStruct<T>(Func<CbField, bool> IsType, Func<CbField, T, object> AsTypeWithDefault, Func<T, T, bool> Comparer) where T : struct
+			{
+				return new CbFieldAccessors(new T(), IsType, x => AsTypeWithDefault(x, new T()), (x, y) => AsTypeWithDefault(x, (T)y)) { Comparer = (x, y) => Comparer((T)x, (T)y) };
+			}
+		}
+
+		Dictionary<CbFieldType, CbFieldAccessors> TypeAccessors = new Dictionary<CbFieldType, CbFieldAccessors>
+		{
+			[CbFieldType.Object] = new CbFieldAccessors(CbObject.Empty, x => x.IsObject(), x => x.AsObjectView(), (x, y) => x.AsObjectView()),
+			[CbFieldType.UniformObject] = new CbFieldAccessors(CbObject.Empty, x => x.IsObject(), x => x.AsObjectView(), (x, y) => x.AsObjectView()),
+			[CbFieldType.Array] = new CbFieldAccessors(CbArray.Empty, x => x.IsArray(), x => x.AsArrayView(), (x, y) => x.AsArrayView()),
+			[CbFieldType.UniformArray] = new CbFieldAccessors(CbArray.Empty, x => x.IsArray(), x => x.AsArrayView(), (x, y) => x.AsArrayView()),
+			[CbFieldType.Binary] = CbFieldAccessors.FromStruct<ReadOnlyMemory<byte>>(x => x.IsBinary(), (x, y) => x.AsBinaryView(y), (x, y) => x.Span.SequenceEqual(y.Span)),
+			[CbFieldType.String] = new CbFieldAccessors(ReadOnlyUtf8String.Empty, x => x.IsString(), x => x.AsString(), (x, Default) => x.AsString((ReadOnlyUtf8String)Default)),
+			[CbFieldType.IntegerPositive] = CbFieldAccessors.FromStruct<ulong>(x => x.IsInteger(), (x, y) => x.AsUInt64(y)),
+			[CbFieldType.IntegerNegative] = CbFieldAccessors.FromStruct<long>(x => x.IsInteger(), (x, y) => x.AsInt64(y)),
+			[CbFieldType.Float32] = CbFieldAccessors.FromStruct<float>(x => x.IsFloat(), (x, y) => x.AsFloat(y)),
+			[CbFieldType.Float64] = CbFieldAccessors.FromStruct<double>(x => x.IsFloat(), (x, y) => x.AsDouble(y)),
+			[CbFieldType.BoolTrue] = CbFieldAccessors.FromStruct<bool>(x => x.IsBool(), (x, y) => x.AsBool(y)),
+			[CbFieldType.BoolFalse] = CbFieldAccessors.FromStruct<bool>(x => x.IsBool(), (x, y) => x.AsBool(y)),
+			[CbFieldType.ObjectAttachment] = new CbFieldAccessors(IoHash.Zero, x => x.IsObjectAttachment(), x => x.AsObjectAttachment(), (x, Default) => x.AsObjectAttachment((IoHash)Default)),
+			[CbFieldType.BinaryAttachment] = new CbFieldAccessors(IoHash.Zero, x => x.IsBinaryAttachment(), x => x.AsBinaryAttachment(), (x, Default) => x.AsBinaryAttachment((IoHash)Default)),
+			[CbFieldType.Hash] = new CbFieldAccessors(IoHash.Zero, x => x.IsHash(), x => x.AsHash(), (x, Default) => x.AsHash((IoHash)Default)),
+			[CbFieldType.Uuid] = CbFieldAccessors.FromStruct<Guid>(x => x.IsUuid(), (x, y) => x.AsUuid(y)),
+			[CbFieldType.DateTime] = CbFieldAccessors.FromStruct<DateTime>(x => x.IsDateTime(), (x, y) => x.AsDateTime(y)),
+			[CbFieldType.TimeSpan] = CbFieldAccessors.FromStruct<TimeSpan>(x => x.IsTimeSpan(), (x, y) => x.AsTimeSpan(y)),
+		};
+
+		void TestField(CbFieldType FieldType, CbField Field, object? ExpectedValue = null, object? DefaultValue = null, CbFieldError ExpectedError = CbFieldError.None, CbFieldAccessors? Accessors = null)
+		{
+			Accessors ??= TypeAccessors[FieldType];
+			ExpectedValue ??= Accessors.Default;
+			DefaultValue ??= Accessors.Default;
+
+			Assert.AreEqual(Accessors.IsType(Field), ExpectedError != CbFieldError.TypeError);
+			if (ExpectedError == CbFieldError.None && !Field.IsBool())
+			{
+				Assert.IsFalse(Field.AsBool());
+				Assert.IsTrue(Field.HasError());
+				Assert.AreEqual(Field.GetError(), CbFieldError.TypeError);
+			}
+
+			object Value = Accessors.AsTypeWithDefault(Field, DefaultValue);
+			Assert.IsTrue(Accessors.Comparer(Accessors.AsTypeWithDefault(Field, DefaultValue), ExpectedValue));
+			Assert.AreEqual(Field.HasError(), ExpectedError != CbFieldError.None);
+			Assert.AreEqual(Field.GetError(), ExpectedError);
+		}
+
+		void TestField(CbFieldType FieldType, byte[] Payload, object? ExpectedValue = null, object? DefaultValue = null, CbFieldError ExpectedError = CbFieldError.None, CbFieldAccessors? Accessors = null)
+		{
+			CbField Field = new CbField(Payload, FieldType);
+			Assert.AreEqual(Field.GetSize(), Payload.Length + (CbFieldUtils.HasFieldType(FieldType) ? 0 : 1));
+			Assert.IsTrue(Field.HasValue());
+			Assert.IsFalse(Field.HasError());
+			Assert.AreEqual(Field.GetError(), CbFieldError.None);
+			TestField(FieldType, Field, ExpectedValue, DefaultValue, ExpectedError, Accessors);
+		}
+
+		void TestFieldError(CbFieldType FieldType, CbField Field, CbFieldError ExpectedError, object? ExpectedValue = null, CbFieldAccessors? Accessors = null)
+		{
+			TestField(FieldType, Field, ExpectedValue, ExpectedValue, ExpectedError, Accessors);
+		}
+
+		void TestFieldError(CbFieldType FieldType, ReadOnlyMemory<byte> Payload, CbFieldError ExpectedError, object? ExpectedValue = null, CbFieldAccessors? Accessors = null)
+		{
+			CbField Field = new CbField(Payload, FieldType);
+			TestFieldError(FieldType, Field, ExpectedError, ExpectedValue, Accessors);
+		}
+
+		[TestMethod]
+		public void CbFieldNoneTest()
+		{
+			// Test CbField()
+			{
+				CbField DefaultField = new CbField();
+				Assert.IsFalse(DefaultField.HasName());
+				Assert.IsFalse(DefaultField.HasValue());
+				Assert.IsFalse(DefaultField.HasError());
+				Assert.IsTrue(DefaultField.GetError() == CbFieldError.None);
+				Assert.AreEqual(DefaultField.GetSize(), 1);
+				Assert.AreEqual(DefaultField.GetName().Length, 0);
+				Assert.IsFalse(DefaultField.HasName());
+				Assert.IsFalse(DefaultField.HasValue());
+				Assert.IsFalse(DefaultField.HasError());
+				Assert.AreEqual(DefaultField.GetError(), CbFieldError.None);
+				Assert.AreEqual(DefaultField.GetHash(), Blake3Hash.Compute(new byte[] { (byte)CbFieldType.None }));
+				ReadOnlyMemory<byte> View;
+				Assert.IsFalse(DefaultField.TryGetView(out View));
+			}
+
+			// Test CbField(None)
+			{
+				CbField NoneField = new CbField(ReadOnlyMemory<byte>.Empty, CbFieldType.None);
+				Assert.AreEqual(NoneField.GetSize(), 1);
+				Assert.AreEqual(NoneField.GetName().Length, 0);
+				Assert.IsFalse(NoneField.HasName());
+				Assert.IsFalse(NoneField.HasValue());
+				Assert.IsFalse(NoneField.HasError());
+				Assert.AreEqual(NoneField.GetError(), CbFieldError.None);
+				Assert.AreEqual(NoneField.GetHash(), new CbField().GetHash());
+				ReadOnlyMemory<byte> View;
+				Assert.IsFalse(NoneField.TryGetView(out View));
+			}
+
+			// Test CbField(None|Type|Name)
+			{
+				CbFieldType FieldType = CbFieldType.None | CbFieldType.HasFieldName;
+				byte[] NoneBytes = { (byte)FieldType, 4, (byte)'N', (byte)'a', (byte)'m', (byte)'e' };
+				CbField NoneField = new CbField(NoneBytes);
+				Assert.AreEqual(NoneField.GetSize(), NoneBytes.Length);
+				Assert.AreEqual(NoneField.GetName(), "Name");
+				Assert.IsTrue(NoneField.HasName());
+				Assert.IsFalse(NoneField.HasValue());
+				Assert.AreEqual(NoneField.GetHash(), Blake3Hash.Compute(NoneBytes));
+				ReadOnlyMemory<byte> View;
+				Assert.IsTrue(NoneField.TryGetView(out View) && View.Span.SequenceEqual(NoneBytes));
+
+				byte[] CopyBytes = new byte[NoneBytes.Length];
+				NoneField.CopyTo(CopyBytes);
+				Assert.IsTrue(NoneBytes.AsSpan().SequenceEqual(CopyBytes));
+			}
+
+			// Test CbField(None|Type)
+			{
+				CbFieldType FieldType = CbFieldType.None;
+				byte[] NoneBytes = { (byte)FieldType };
+				CbField NoneField = new CbField(NoneBytes);
+				Assert.AreEqual(NoneField.GetSize(), NoneBytes.Length);
+				Assert.AreEqual(NoneField.GetName().Length, 0);
+				Assert.IsFalse(NoneField.HasName());
+				Assert.IsFalse(NoneField.HasValue());
+				Assert.AreEqual(NoneField.GetHash(), new CbField().GetHash());
+				ReadOnlyMemory<byte> View;
+				Assert.IsTrue(NoneField.TryGetView(out View) && View.Span.SequenceEqual(NoneBytes));
+			}
+
+			// Test CbField(None|Name)
+			{
+				CbFieldType FieldType = CbFieldType.None | CbFieldType.HasFieldName;
+				byte[] NoneBytes = { (byte)FieldType, 4, (byte)'N', (byte)'a', (byte)'m', (byte)'e' };
+				CbField NoneField = new CbField(NoneBytes.AsMemory(1), FieldType);
+				Assert.AreEqual(NoneField.GetSize(), NoneBytes.Length);
+				Assert.AreEqual(NoneField.GetName(), "Name");
+				Assert.IsTrue(NoneField.HasName());
+				Assert.IsFalse(NoneField.HasValue());
+				Assert.AreEqual(NoneField.GetHash(), Blake3Hash.Compute(NoneBytes));
+				ReadOnlyMemory<byte> View;
+				Assert.IsFalse(NoneField.TryGetView(out View));
+
+				byte[] CopyBytes = new byte[NoneBytes.Length];
+				NoneField.CopyTo(CopyBytes);
+				Assert.IsTrue(NoneBytes.AsSpan().SequenceEqual(CopyBytes));
+			}
+
+			// Test CbField(None|EmptyName)
+			{
+				CbFieldType FieldType = CbFieldType.None | CbFieldType.HasFieldName;
+				byte[] NoneBytes = { (byte)FieldType, 0 };
+				CbField NoneField = new CbField(NoneBytes.AsMemory(1), FieldType);
+				Assert.AreEqual(NoneField.GetSize(), NoneBytes.Length);
+				Assert.AreEqual(NoneField.GetName(), "");
+				Assert.IsTrue(NoneField.HasName());
+				Assert.IsFalse(NoneField.HasValue());
+				Assert.AreEqual(NoneField.GetHash(), Blake3Hash.Compute(NoneBytes));
+				ReadOnlyMemory<byte> View;
+				Assert.IsFalse(NoneField.TryGetView(out View));
+			}
+		}
+
+		[TestMethod]
+		public void CbFieldNullTest()
+		{
+			// Test CbField(Null)
+			{
+				CbField NullField = new CbField(ReadOnlyMemory<byte>.Empty, CbFieldType.Null);
+				Assert.AreEqual(NullField.GetSize(), 1);
+				Assert.IsTrue(NullField.IsNull());
+				Assert.IsTrue(NullField.HasValue());
+				Assert.IsFalse(NullField.HasError());
+				Assert.AreEqual(NullField.GetError(), CbFieldError.None);
+				Assert.AreEqual(NullField.GetHash(), Blake3Hash.Compute(new byte[] { (byte)CbFieldType.Null }));
+			}
+
+			// Test CbField(None) as Null
+			{
+				CbField Field = new CbField();
+				Assert.IsFalse(Field.IsNull());
+			}
+		}
+
+		[TestMethod]
+		public void CbFieldObjectTest()
+		{
+			Action<CbObject, int, int> TestIntObject = (CbObject Object, int ExpectedNum, int ExpectedPayloadSize) =>
+			{
+				Assert.AreEqual(Object.GetSize(), ExpectedPayloadSize + sizeof(CbFieldType));
+
+				int ActualNum = 0;
+				foreach (CbField Field in Object)
+				{
+					++ActualNum;
+					Assert.AreNotEqual(Field.GetName().Length, 0);
+					Assert.AreEqual(Field.AsInt32(), ActualNum);
+				}
+				Assert.AreEqual(ActualNum, ExpectedNum);
+			};
+
+			// Test CbField(Object, Empty)
+			TestField(CbFieldType.Object, new byte[1]);
+
+			// Test CbField(Object, Empty)
+			{
+				CbObject Object = CbObject.Empty;
+				TestIntObject(Object, 0, 1);
+
+				// Find fields that do not exist.
+				Assert.IsFalse(Object.FindView("Field").HasValue());
+				Assert.IsFalse(Object.FindViewIgnoreCase("Field").HasValue());
+				Assert.IsFalse(Object["Field"].HasValue());
+
+				// Advance an iterator past the last field.
+				CbFieldIterator It = Object.CreateViewIterator();
+				Assert.IsFalse((bool)It);
+				Assert.IsTrue(!It);
+				for (int Count = 16; Count > 0; --Count)
+				{
+					++It;
+					It.Current.AsInt32();
+				}
+				Assert.IsFalse((bool)It);
+				Assert.IsTrue(!It);
+			}
+
+			// Test CbField(Object, NotEmpty)
+			{
+				byte IntType = (byte)(CbFieldType.HasFieldName | CbFieldType.IntegerPositive);
+				byte[] Payload = { 12, IntType, 1, (byte)'A', 1, IntType, 1, (byte)'B', 2, IntType, 1, (byte)'C', 3 };
+				CbField Field = new CbField(Payload, CbFieldType.Object);
+				TestField(CbFieldType.Object, Field, new CbObject(Payload, CbFieldType.Object));
+				CbObject Object = CbObject.Clone(Field.AsObjectView());
+				TestIntObject(Object, 3, Payload.Length);
+				TestIntObject(Field.AsObjectView(), 3, Payload.Length);
+				Assert.IsTrue(Object.Equals(Field.AsObjectView()));
+				Assert.AreEqual(Object.FindView("B").AsInt32(), 2);
+				Assert.AreEqual(Object.FindView("b").AsInt32(4), 4);
+				Assert.AreEqual(Object.FindViewIgnoreCase("B").AsInt32(), 2);
+				Assert.AreEqual(Object.FindViewIgnoreCase("b").AsInt32(), 2);
+				Assert.AreEqual(Object["B"].AsInt32(), 2);
+				Assert.AreEqual(Object["b"].AsInt32(4), 4);
+			}
+
+			// Test CbField(UniformObject, NotEmpty)
+			{
+				byte IntType = (byte)(CbFieldType.HasFieldName | CbFieldType.IntegerPositive);
+				byte[] Payload = { 10, IntType, 1, (byte)'A', 1, 1, (byte)'B', 2, 1, (byte)'C', 3 };
+				CbField Field = new CbField(Payload, CbFieldType.UniformObject);
+				TestField(CbFieldType.UniformObject, Field, new CbObject(Payload, CbFieldType.UniformObject));
+				CbObject Object = CbObject.Clone(Field.AsObjectView());
+				TestIntObject(Object, 3, Payload.Length);
+				TestIntObject(Field.AsObject(), 3, Payload.Length);
+				Assert.IsTrue(Object.Equals(Field.AsObject()));
+				Assert.AreEqual(Object.FindView("B").AsInt32(), 2);
+				Assert.AreEqual(Object.Find("B").AsInt32(), 2);
+				Assert.AreEqual(Object.FindView("b").AsInt32(4), 4);
+				Assert.AreEqual(Object.Find("b").AsInt32(4), 4);
+				Assert.AreEqual(Object.FindViewIgnoreCase("B").AsInt32(), 2);
+				Assert.AreEqual(Object.FindIgnoreCase("B").AsInt32(), 2);
+				Assert.AreEqual(Object.FindViewIgnoreCase("b").AsInt32(), 2);
+				Assert.AreEqual(Object.FindIgnoreCase("b").AsInt32(), 2);
+				Assert.AreEqual(Object["B"].AsInt32(), 2);
+				Assert.AreEqual(Object["b"].AsInt32(4), 4);
+
+				// Equals
+				byte[] NamedPayload = { 1, (byte)'O', 10, IntType, 1, (byte)'A', 1, 1, (byte)'B', 2, 1, (byte)'C', 3 };
+				CbField NamedField = new CbField(NamedPayload, CbFieldType.UniformObject | CbFieldType.HasFieldName);
+				Assert.IsTrue(Field.AsObject().Equals(NamedField.AsObject()));
+
+				// CopyTo
+				byte[] CopyBytes = new byte[Payload.Length + 1];
+				Field.AsObjectView().CopyTo(CopyBytes);
+				Assert.IsTrue(Payload.AsSpan().SequenceEqual(CopyBytes.AsSpan(1)));
+				NamedField.AsObjectView().CopyTo(CopyBytes);
+				Assert.IsTrue(Payload.AsSpan().SequenceEqual(CopyBytes.AsSpan(1)));
+
+				// TryGetView
+				ReadOnlyMemory<byte> View;
+				Assert.IsFalse(Field.AsObjectView().TryGetView(out View));
+				Assert.IsFalse(NamedField.AsObjectView().TryGetView(out View));
+			}
+
+			// Test CbField(None) as Object
+			{
+				CbField Field = CbField.Empty;
+				TestFieldError(CbFieldType.Object, Field, CbFieldError.TypeError);
+				CbField.MakeView(Field).AsObject();
+			}
+
+			// Test FCbObjectView(ObjectWithName) and CreateIterator
+			{
+				byte ObjectType = (byte)(CbFieldType.Object | CbFieldType.HasFieldName);
+				byte[] Buffer = { ObjectType, 3, (byte)'K', (byte)'e', (byte)'y', 4, (byte)(CbFieldType.HasFieldName | CbFieldType.IntegerPositive), 1, (byte)'F', 8 };
+				CbObject Object = new CbObject(Buffer);
+				Assert.AreEqual(Object.GetSize(), 6);
+				CbObject ObjectClone = CbObject.Clone(Object);
+				Assert.AreEqual(ObjectClone.GetSize(), 6);
+				Assert.IsTrue(Object.Equals(ObjectClone));
+				Assert.AreEqual(ObjectClone.GetHash(), Object.GetHash());
+				for (CbFieldIterator It = ObjectClone.CreateIterator(); It; ++It)
+				{
+					CbField Field = It.Current;
+					Assert.AreEqual(Field.GetName(), "F");
+					Assert.AreEqual(Field.AsInt32(), 8);
+				}
+				for (CbFieldIterator It = ObjectClone.CreateIterator(), End = new CbFieldIterator(); It != End; ++It)
+				{
+				}
+				foreach (CbField Field in ObjectClone)
+				{
+				}
+			}
+
+			// Test FCbObjectView as CbFieldIterator
+			{
+				int Count = 0;
+				CbObject Object = CbObject.Empty;
+				for (CbFieldIterator It = CbFieldIterator.MakeSingle(Object.AsFieldView()); It; ++It)
+				{
+					CbField Field = It.Current;
+					Assert.IsTrue(Field.IsObject());
+					++Count;
+				}
+				Assert.AreEqual(Count, 1);
+			}
+		}
+
+		public void CbFieldArrayTest()
+		{
+			Action<CbArray, int, int> TestIntArray = (CbArray Array, int ExpectedNum, int ExpectedPayloadSize) =>
+			{
+				Assert.AreEqual(Array.GetSize(), ExpectedPayloadSize + sizeof(CbFieldType));
+				Assert.AreEqual(Array.Num(), ExpectedNum);
+
+				int ActualNum = 0;
+				for (CbFieldIterator It = Array.CreateViewIterator(); It; ++It)
+				{
+					++ActualNum;
+					Assert.AreEqual(It.Current.AsInt32(), ActualNum);
+				}
+				Assert.AreEqual(ActualNum, ExpectedNum);
+
+				ActualNum = 0;
+				foreach (CbField Field in Array)
+				{
+					++ActualNum;
+					Assert.AreEqual(Field.AsInt32(), ActualNum);
+				}
+				Assert.AreEqual(ActualNum, ExpectedNum);
+
+				ActualNum = 0;
+				foreach (CbField Field in Array.AsFieldView())
+				{
+					++ActualNum;
+					Assert.AreEqual(Field.AsInt32(), ActualNum);
+				}
+				Assert.AreEqual(ActualNum, ExpectedNum);
+			};
+
+			// Test CbField(Array, Empty)
+			TestField(CbFieldType.Array, new byte[]{1, 0});
+
+			// Test CbField(Array, Empty)
+			{
+				CbArray Array = new CbArray();
+				TestIntArray(Array, 0, 2);
+
+				// Advance an iterator past the last field.
+				CbFieldIterator It = Array.CreateViewIterator();
+				Assert.IsFalse((bool)It);
+				Assert.IsTrue(!It);
+				for (int Count = 16; Count > 0; --Count)
+				{
+					++It;
+					It.Current.AsInt32();
+				}
+				Assert.IsFalse((bool)It);
+				Assert.IsTrue(!It);
+			}
+
+			// Test CbField(Array, NotEmpty)
+			{
+				byte IntType = (byte)CbFieldType.IntegerPositive;
+				byte[] Payload = new byte[]{ 7, 3, IntType, 1, IntType, 2, IntType, 3 };
+				CbField Field = new CbField(Payload, CbFieldType.Array);
+				TestField(CbFieldType.Array, Field, new CbArray(Payload, CbFieldType.Array));
+				CbArray Array = CbArray.Clone(Field.AsArrayView());
+				TestIntArray(Array, 3, Payload.Length);
+				TestIntArray(Field.AsArrayView(), 3, Payload.Length);
+				Assert.IsTrue(Array.Equals(Field.AsArrayView()));
+			}
+
+			// Test CbField(UniformArray)
+			{
+				byte IntType = (byte)(CbFieldType.IntegerPositive);
+				byte[] Payload = new byte[]{ 5, 3, IntType, 1, 2, 3 };
+				CbField Field = new CbField(Payload, CbFieldType.UniformArray);
+				TestField(CbFieldType.UniformArray, Field, new CbArray(Payload, CbFieldType.UniformArray));
+				CbArray Array = CbArray.Clone(Field.AsArrayView());
+				TestIntArray(Array, 3, Payload.Length);
+				TestIntArray(Field.AsArrayView(), 3, Payload.Length);
+				Assert.IsTrue(Array.Equals(Field.AsArrayView()));
+
+//				Assert.IsTrue(Array.GetOuterBuffer() == Array.AsField().AsArray().GetOuterBuffer());
+
+				// Equals
+				byte[] NamedPayload = new byte[]{ 1, (byte)'A', 5, 3, IntType, 1, 2, 3 };
+				CbField NamedField = new CbField(NamedPayload, CbFieldType.UniformArray | CbFieldType.HasFieldName);
+				Assert.IsTrue(Field.AsArrayView().Equals(NamedField.AsArrayView()));
+				Assert.IsTrue(Field.Equals(Field.AsArrayView().AsFieldView()));
+				Assert.IsTrue(NamedField.Equals(NamedField.AsArrayView().AsFieldView()));
+
+				// CopyTo
+				byte[] CopyBytes = new byte[Payload.Length + 1];
+				Field.AsArrayView().CopyTo(CopyBytes);
+				Assert.IsTrue(Payload.AsSpan().SequenceEqual(CopyBytes.AsSpan(1)));
+				NamedField.AsArrayView().CopyTo(CopyBytes);
+				Assert.IsTrue(Payload.AsSpan().SequenceEqual(CopyBytes.AsSpan(1)));
+
+				// TryGetView
+				ReadOnlyMemory<byte> View;
+//				Assert.IsTrue(Array.TryGetView(out View) && View == Array.GetOuterBuffer().GetView());
+				Assert.IsFalse(Field.AsArrayView().TryGetView(out View));
+				Assert.IsFalse(NamedField.AsArrayView().TryGetView(out View));
+
+//				// GetBuffer
+//				Assert.IsTrue(Array.GetBuffer().Flatten().GetView() == Array.GetOuterBuffer().GetView());
+//				Assert.IsTrue(CbField.MakeView(Field).AsArray().GetBuffer().Flatten().GetView().Span.SequenceEqual(Array.GetOuterBuffer().GetView()));
+//				Assert.IsTrue(CbField.MakeView(NamedField).AsArray().GetBuffer().Flatten().GetView().Span.SequenceEqual(Array.GetOuterBuffer().GetView()));
+			}
+
+			// Test CbField(None) as Array
+			{
+				CbField Field = new CbField();
+//				TestFieldError(CbFieldType.Array, Field, CbFieldError.TypeError);
+				CbField.MakeView(Field).AsArray();
+			}
+
+			// Test CbArray(ArrayWithName) and CreateIterator
+			{
+				byte ArrayType = (byte)(CbFieldType.Array | CbFieldType.HasFieldName);
+				byte[] Buffer = new byte[] { ArrayType, 3, (byte)'K', (byte)'e', (byte)'y', 3, 1, (byte)(CbFieldType.IntegerPositive), 8 };
+				CbArray Array = new CbArray(Buffer);
+				Assert.AreEqual(Array.GetSize(), 5);
+				CbArray ArrayClone = CbArray.Clone(Array);
+				Assert.AreEqual(ArrayClone.GetSize(), 5);
+				Assert.IsTrue(Array.Equals(ArrayClone));
+				Assert.AreEqual(ArrayClone.GetHash(), Array.GetHash());
+				for (CbFieldIterator It = ArrayClone.CreateIterator(); It; ++It)
+				{
+					CbField Field = It.Current;
+					Assert.AreEqual(Field.AsInt32(), 8);
+//					Assert.IsTrue(Field.IsOwned());
+				}
+				for (CbFieldIterator It = ArrayClone.CreateIterator(), End = new CbFieldIterator(); It != End; ++It)
+				{
+				}
+				foreach (CbField Field in ArrayClone)
+				{
+				}
+
+				// CopyTo
+				byte[] CopyBytes = new byte[5];
+				Array.CopyTo(CopyBytes);
+//				Assert.IsTrue(ArrayClone.GetOuterBuffer().GetView().Span.SequenceEqual(CopyBytes));
+				ArrayClone.CopyTo(CopyBytes);
+//				Assert.IsTrue(ArrayClone.GetOuterBuffer().GetView().Span.SequenceEqual(CopyBytes));
+
+//				// GetBuffer
+//				Assert.IsTrue(CbField(FSharedBuffer.MakeView(Buffer)).GetBuffer().Flatten().GetView().Span.SequenceEqual(Buffer));
+//				Assert.IsTrue(TEXT("CbField(ArrayWithNameNoType).GetBuffer()"),
+//					CbField(CbField(Buffer + 1, CbFieldType(ArrayType)), FSharedBuffer.MakeView(Buffer)).GetBuffer().Flatten().GetView().Span.SequenceEqual(Buffer));
+			}
+
+			// Test CbArray as CbFieldIterator
+			{
+				uint Count = 0;
+				CbArray Array = new CbArray();
+				for (CbFieldIterator Iter = CbFieldIterator.MakeSingle(Array.AsFieldView()); Iter; ++Iter)
+				{
+					CbField Field = Iter.Current;
+					Assert.IsTrue(Field.IsArray());
+					++Count;
+				}
+				Assert.AreEqual(Count, 1u);
+			}
+
+			// Test CbArray as CbFieldIterator
+			{
+				uint Count = 0;
+				CbArray Array = new CbArray();
+//				Array.MakeOwned();
+				for(CbFieldIterator Iter = CbFieldIterator.MakeSingle(Array.AsField()); Iter; ++Iter)
+				{
+					CbField Field = Iter.Current;
+					Assert.IsTrue(Field.IsArray());
+					++Count;
+				}
+				Assert.AreEqual(Count, 1u);
+			}
+		}
+
+		[TestMethod]
+		public void CbFieldBinaryTest()
+		{
+			// Test CbField(Binary, Empty)
+			TestField(CbFieldType.Binary, new byte[]{0});
+
+			// Test CbField(Binary, Value)
+			{
+				byte[] Payload = { 3, 4, 5, 6 }; // Size: 3, Data: 4/5/6
+				CbField FieldView = new CbField(Payload, CbFieldType.Binary);
+				TestField(CbFieldType.Binary, FieldView, (ReadOnlyMemory<byte>)Payload.AsMemory(1, 3));
+
+				CbField Field = CbField.Clone(FieldView);
+				Field.AsBinary();
+//				Assert.IsFalse(Field.GetOuterBuffer().IsNull());
+//				MoveTemp(Field).AsBinary();
+//				Assert.IsTrue(Field.GetOuterBuffer().IsNull());
+			}
+
+			// Test CbField(None) as Binary
+			{
+				CbField FieldView = new CbField();
+				byte[] Default = { 1, 2, 3 };
+				TestFieldError(CbFieldType.Binary, FieldView, CbFieldError.TypeError, (ReadOnlyMemory<byte>)Default);
+
+//				CbField Field = CbField.Clone(FieldView);
+//				TestFieldError(CbFieldType.Binary, FSharedBuffer, Field, CbFieldError.TypeError, FSharedBuffer.MakeView(Default), FCbBinaryAccessors());
+			}
+		}
+
+		[TestMethod]
+		public void CbFieldStringTest()
+		{
+			// Test CbField(String, Empty)
+			TestField(CbFieldType.String, new byte[]{0});
+
+			// Test CbField(String, Value)
+			{
+				byte[] Payload = { 3, (byte)'A', (byte)'B', (byte)'C' }; // Size: 3, Data: ABC
+				TestField(CbFieldType.String, Payload, new ReadOnlyUtf8String(Payload.AsMemory(1, 3)));
+			}
+
+			// Test CbField(String, OutOfRangeSize)
+			{
+				byte[] Payload = new byte[9];
+				BitUtils.WriteVarUInt((ulong)(1) << 31, Payload);
+				TestFieldError(CbFieldType.String, Payload, CbFieldError.RangeError, new ReadOnlyUtf8String("ABC"));
+			}
+
+			// Test CbField(None) as String
+			{
+				CbField Field = new CbField();
+				TestFieldError(CbFieldType.String, Field, CbFieldError.TypeError, new ReadOnlyUtf8String("ABC"));
+			}
+		}
+
+		[Flags]
+		enum EIntType : byte
+		{
+			None   = 0x00,
+			Int8   = 0x01,
+			Int16  = 0x02,
+			Int32  = 0x04,
+			Int64  = 0x08,
+			UInt8  = 0x10,
+			UInt16 = 0x20,
+			UInt32 = 0x40,
+			UInt64 = 0x80,
+			// Masks for positive values requiring the specified number of bits.
+			Pos64 = UInt64,
+			Pos63 = Pos64 |  Int64,
+			Pos32 = Pos63 | UInt32,
+			Pos31 = Pos32 |  Int32,
+			Pos16 = Pos31 | UInt16,
+			Pos15 = Pos16 |  Int16,
+			Pos8  = Pos15 | UInt8,
+			Pos7  = Pos8  |  Int8,
+			// Masks for negative values requiring the specified number of bits.
+			Neg63 = Int64,
+			Neg31 = Neg63 | Int32,
+			Neg15 = Neg31 | Int16,
+			Neg7  = Neg15 | Int8,
+		};
+
+		void TestIntegerField(CbFieldType FieldType, EIntType ExpectedMask, ulong Magnitude)
+		{
+			byte[] Payload = new byte[9];
+			ulong Negative = (ulong)((byte)FieldType & 1);
+			BitUtils.WriteVarUInt(Magnitude - Negative, Payload);
+			ulong DefaultValue = 8;
+			ulong ExpectedValue = (Negative != 0)? (ulong)(-(long)(Magnitude)) : Magnitude;
+			CbField Field = new CbField(Payload, FieldType);
+
+			TestField(CbFieldType.IntegerNegative, Field, (sbyte)(((ExpectedMask & EIntType.Int8) != 0) ? ExpectedValue : DefaultValue),
+				(sbyte)(DefaultValue), ((ExpectedMask & EIntType.Int8) != 0)? CbFieldError.None : CbFieldError.RangeError, CbFieldAccessors.FromStruct<sbyte>(x => x.IsInteger(), (x, y) => x.AsInt8(y)));
+
+			TestField(CbFieldType.IntegerNegative, Field, (short)(((ExpectedMask & EIntType.Int16) != 0)? ExpectedValue : DefaultValue),
+				(short)(DefaultValue), ((ExpectedMask & EIntType.Int16) != 0)? CbFieldError.None : CbFieldError.RangeError, CbFieldAccessors.FromStruct<short>(x => x.IsInteger(), (x, y) => x.AsInt16(y)));
+
+			TestField(CbFieldType.IntegerNegative, Field, (int)(((ExpectedMask & EIntType.Int32) != 0)? ExpectedValue : DefaultValue),
+				(int)(DefaultValue), ((ExpectedMask & EIntType.Int32) != 0)? CbFieldError.None : CbFieldError.RangeError, CbFieldAccessors.FromStruct<int>(x => x.IsInteger(), (x, y) => x.AsInt32(y)));
+
+			TestField(CbFieldType.IntegerNegative, Field, (long)(((ExpectedMask & EIntType.Int64) != 0)? ExpectedValue : DefaultValue),
+				(long)(DefaultValue), ((ExpectedMask & EIntType.Int64) != 0)? CbFieldError.None : CbFieldError.RangeError, CbFieldAccessors.FromStruct<long>(x => x.IsInteger(), (x, y) => x.AsInt64(y)));
+
+			TestField(CbFieldType.IntegerPositive, Field, (byte)(((ExpectedMask & EIntType.UInt8) != 0) ? ExpectedValue : DefaultValue),
+				(byte)(DefaultValue), ((ExpectedMask & EIntType.UInt8) != 0) ? CbFieldError.None : CbFieldError.RangeError, CbFieldAccessors.FromStruct<byte>(x => x.IsInteger(), (x, y) => x.AsUInt8(y)));
+
+			TestField(CbFieldType.IntegerPositive, Field, (ushort)(((ExpectedMask & EIntType.UInt16) != 0)? ExpectedValue : DefaultValue),
+				(ushort)(DefaultValue), ((ExpectedMask & EIntType.UInt16) != 0) ? CbFieldError.None : CbFieldError.RangeError, CbFieldAccessors.FromStruct<ushort>(x => x.IsInteger(), (x, y) => x.AsUInt16(y)));
+
+			TestField(CbFieldType.IntegerPositive, Field, (uint)(((ExpectedMask & EIntType.UInt32) != 0) ? ExpectedValue : DefaultValue),
+				(uint)(DefaultValue), ((ExpectedMask & EIntType.UInt32) != 0) ? CbFieldError.None : CbFieldError.RangeError, CbFieldAccessors.FromStruct<uint>(x => x.IsInteger(), (x, y) => x.AsUInt32(y)));
+			TestField(CbFieldType.IntegerPositive, Field, (ulong)(((ExpectedMask & EIntType.UInt64) != 0) ? ExpectedValue : DefaultValue),
+				(ulong)(DefaultValue), ((ExpectedMask & EIntType.UInt64) != 0)? CbFieldError.None : CbFieldError.RangeError, CbFieldAccessors.FromStruct<ulong>(x => x.IsInteger(), (x, y) => x.AsUInt64(y)));
+		}
+
+		void CbFieldIntegerTest()
+		{
+			// Test CbField(IntegerPositive)
+			TestIntegerField(CbFieldType.IntegerPositive, EIntType.Pos7,  0x00);
+			TestIntegerField(CbFieldType.IntegerPositive, EIntType.Pos7,  0x7f);
+			TestIntegerField(CbFieldType.IntegerPositive, EIntType.Pos8,  0x80);
+			TestIntegerField(CbFieldType.IntegerPositive, EIntType.Pos8,  0xff);
+			TestIntegerField(CbFieldType.IntegerPositive, EIntType.Pos15, 0x0100);
+			TestIntegerField(CbFieldType.IntegerPositive, EIntType.Pos15, 0x7fff);
+			TestIntegerField(CbFieldType.IntegerPositive, EIntType.Pos16, 0x8000);
+			TestIntegerField(CbFieldType.IntegerPositive, EIntType.Pos16, 0xffff);
+			TestIntegerField(CbFieldType.IntegerPositive, EIntType.Pos31, 0x0001_0000);
+			TestIntegerField(CbFieldType.IntegerPositive, EIntType.Pos31, 0x7fff_ffff);
+			TestIntegerField(CbFieldType.IntegerPositive, EIntType.Pos32, 0x8000_0000);
+			TestIntegerField(CbFieldType.IntegerPositive, EIntType.Pos32, 0xffff_ffff);
+			TestIntegerField(CbFieldType.IntegerPositive, EIntType.Pos63, 0x0000_0001_0000_0000);
+			TestIntegerField(CbFieldType.IntegerPositive, EIntType.Pos63, 0x7fff_ffff_ffff_ffff);
+			TestIntegerField(CbFieldType.IntegerPositive, EIntType.Pos64, 0x8000_0000_0000_0000);
+			TestIntegerField(CbFieldType.IntegerPositive, EIntType.Pos64, 0xffff_ffff_ffff_ffff);
+
+			// Test CbField(IntegerNegative)
+			TestIntegerField(CbFieldType.IntegerNegative, EIntType.Neg7,  0x01);
+			TestIntegerField(CbFieldType.IntegerNegative, EIntType.Neg7,  0x80);
+			TestIntegerField(CbFieldType.IntegerNegative, EIntType.Neg15, 0x81);
+			TestIntegerField(CbFieldType.IntegerNegative, EIntType.Neg15, 0x8000);
+			TestIntegerField(CbFieldType.IntegerNegative, EIntType.Neg31, 0x8001);
+			TestIntegerField(CbFieldType.IntegerNegative, EIntType.Neg31, 0x8000_0000);
+			TestIntegerField(CbFieldType.IntegerNegative, EIntType.Neg63, 0x8000_0001);
+			TestIntegerField(CbFieldType.IntegerNegative, EIntType.Neg63, 0x8000_0000_0000_0000);
+			TestIntegerField(CbFieldType.IntegerNegative, EIntType.None,  0x8000_0000_0000_0001);
+			TestIntegerField(CbFieldType.IntegerNegative, EIntType.None,  0xffff_ffff_ffff_ffff);
+
+			// Test CbField(None) as Integer
+			{
+				CbField Field = new CbField();
+				TestFieldError(CbFieldType.IntegerPositive, Field, CbFieldError.TypeError, (ulong)(8));
+				TestFieldError(CbFieldType.IntegerNegative, Field, CbFieldError.TypeError, (long)(8));
+			}
+		}
+
+		[TestMethod]
+		public void CbFieldFloatTest()
+		{
+			// Test CbField(Float, 32-bit)
+			{
+				byte[] Payload = new byte[]{ 0xc0, 0x12, 0x34, 0x56 }; // -2.28444433f
+				TestField(CbFieldType.Float32, Payload, -2.28444433f);
+
+				CbField Field = new CbField(Payload, CbFieldType.Float32);
+				TestField(CbFieldType.Float64, Field, (double)-2.28444433f);
+			}
+
+			// Test CbField(Float, 64-bit)
+			{
+				byte[] Payload = new byte[]{ 0xc1, 0x23, 0x45, 0x67, 0x89, 0xab, 0xcd, 0xef }; // -631475.76888888876
+				TestField(CbFieldType.Float64, Payload, -631475.76888888876);
+
+				CbField Field = new CbField(Payload, CbFieldType.Float64);
+				TestFieldError(CbFieldType.Float32, Field, CbFieldError.RangeError, 8.0f);
+			}
+
+			// Test CbField(Integer+, MaxBinary32) as Float
+			{
+				byte[] Payload = new byte[9];
+				BitUtils.WriteVarUInt(((ulong)(1) << 24) - 1, Payload); // 16,777,215
+				CbField Field = new CbField(Payload, CbFieldType.IntegerPositive);
+				TestField(CbFieldType.Float32, Field, 16_777_215.0f);
+				TestField(CbFieldType.Float64, Field, 16_777_215.0);
+			}
+
+			// Test CbField(Integer+, MaxBinary32+1) as Float
+			{
+				byte[] Payload = new byte[9];
+				BitUtils.WriteVarUInt((ulong)(1) << 24, Payload); // 16,777,216
+				CbField Field = new CbField(Payload, CbFieldType.IntegerPositive);
+				TestFieldError(CbFieldType.Float32, Field, CbFieldError.RangeError, 8.0f);
+				TestField(CbFieldType.Float64, Field, 16_777_216.0);
+			}
+
+			// Test CbField(Integer+, MaxBinary64) as Float
+			{
+				byte[] Payload = new byte[9];
+				BitUtils.WriteVarUInt(((ulong)(1) << 53) - 1, Payload); // 9,007,199,254,740,991
+				CbField Field = new CbField(Payload, CbFieldType.IntegerPositive);
+				TestFieldError(CbFieldType.Float32, Field, CbFieldError.RangeError, 8.0f);
+				TestField(CbFieldType.Float64, Field, 9_007_199_254_740_991.0);
+			}
+
+			// Test CbField(Integer+, MaxBinary64+1) as Float
+			{
+				byte[] Payload = new byte[9];
+				BitUtils.WriteVarUInt((ulong)(1) << 53, Payload); // 9,007,199,254,740,992
+				CbField Field = new CbField(Payload, CbFieldType.IntegerPositive);
+				TestFieldError(CbFieldType.Float32, Field, CbFieldError.RangeError, 8.0f);
+				TestFieldError(CbFieldType.Float64, Field, CbFieldError.RangeError, 8.0);
+			}
+
+			// Test CbField(Integer+, MaxUInt64) as Float
+			{
+				byte[] Payload = new byte[9];
+				BitUtils.WriteVarUInt(~(ulong)0, Payload); // Max uint64
+				CbField Field = new CbField(Payload, CbFieldType.IntegerPositive);
+				TestFieldError(CbFieldType.Float32, Field, CbFieldError.RangeError, 8.0f);
+				TestFieldError(CbFieldType.Float64, Field, CbFieldError.RangeError, 8.0);
+			}
+
+			// Test CbField(Integer-, MaxBinary32) as Float
+			{
+				byte[] Payload = new byte[9];
+				BitUtils.WriteVarUInt(((ulong)(1) << 24) - 2, Payload); // -16,777,215
+				CbField Field = new CbField(Payload, CbFieldType.IntegerNegative);
+				TestField(CbFieldType.Float32, Field, -16_777_215.0f);
+				TestField(CbFieldType.Float64, Field, -16_777_215.0);
+			}
+
+			// Test CbField(Integer-, MaxBinary32+1) as Float
+			{
+				byte[] Payload = new byte[9];
+				BitUtils.WriteVarUInt(((ulong)(1) << 24) - 1, Payload); // -16,777,216
+				CbField Field = new CbField(Payload, CbFieldType.IntegerNegative);
+				TestFieldError(CbFieldType.Float32, Field, CbFieldError.RangeError, 8.0f);
+				TestField(CbFieldType.Float64, Field, -16_777_216.0);
+			}
+
+			// Test CbField(Integer-, MaxBinary64) as Float
+			{
+				byte[] Payload = new byte[9];
+				BitUtils.WriteVarUInt(((ulong)(1) << 53) - 2, Payload); // -9,007,199,254,740,991
+				CbField Field = new CbField(Payload, CbFieldType.IntegerNegative);
+				TestFieldError(CbFieldType.Float32, Field, CbFieldError.RangeError, 8.0f);
+				TestField(CbFieldType.Float64, Field, -9_007_199_254_740_991.0);
+			}
+
+			// Test CbField(Integer-, MaxBinary64+1) as Float
+			{
+				byte[] Payload = new byte[9];
+				BitUtils.WriteVarUInt(((ulong)(1) << 53) - 1, Payload); // -9,007,199,254,740,992
+				CbField Field = new CbField(Payload, CbFieldType.IntegerNegative);
+				TestFieldError(CbFieldType.Float32, Field, CbFieldError.RangeError, 8.0f);
+				TestFieldError(CbFieldType.Float64, Field, CbFieldError.RangeError, 8.0);
+			}
+
+			// Test CbField(None) as Float
+			{
+				CbField Field = new CbField();
+				TestFieldError(CbFieldType.Float32, Field, CbFieldError.TypeError, 8.0f);
+				TestFieldError(CbFieldType.Float64, Field, CbFieldError.TypeError, 8.0);
+			}
+		}
+
+		[TestMethod]
+		public void CbFieldBoolTest()
+		{
+			// Test CbField(Bool, False)
+			TestField(CbFieldType.BoolFalse, Array.Empty<byte>(), false, true);
+
+			// Test CbField(Bool, True)
+			TestField(CbFieldType.BoolTrue, Array.Empty<byte>(), true, false);
+
+			// Test CbField(None) as Bool
+			{
+				CbField DefaultField = new CbField();
+				TestFieldError(CbFieldType.BoolFalse, DefaultField, CbFieldError.TypeError, false);
+				TestFieldError(CbFieldType.BoolTrue, DefaultField, CbFieldError.TypeError, true);
+			}
+		}
+
+		[TestMethod]
+		public void CbFieldObjectAttachmentTest()
+		{
+			byte[] ZeroBytes = new byte[20];
+			byte[] SequentialBytes = new byte[]{1, 2, 3, 4, 5, 6, 7, 8, 9, 10, 11, 12, 13, 14, 15, 16, 17, 18, 19, 20};
+
+			// Test CbField(ObjectAttachment, Zero)
+			TestField(CbFieldType.ObjectAttachment, ZeroBytes);
+
+			// Test CbField(ObjectAttachment, NonZero)
+			TestField(CbFieldType.ObjectAttachment, SequentialBytes, new IoHash(SequentialBytes));
+
+			// Test CbField(ObjectAttachment, NonZero) AsAttachment
+			{
+				CbField Field = new CbField(SequentialBytes, CbFieldType.ObjectAttachment);
+				TestField(CbFieldType.ObjectAttachment, Field, new IoHash(SequentialBytes), new IoHash(), CbFieldError.None);
+			}
+
+			// Test CbField(None) as ObjectAttachment
+			{
+				CbField DefaultField = new CbField();
+				TestFieldError(CbFieldType.ObjectAttachment, DefaultField, CbFieldError.TypeError, new IoHash(SequentialBytes));
+			}
+		}
+
+		[TestMethod]
+		public void CbFieldBinaryAttachmentTest()
+		{
+			byte[] ZeroBytes = new byte[20];
+			byte[] SequentialBytes = new byte[]{1, 2, 3, 4, 5, 6, 7, 8, 9, 10, 11, 12, 13, 14, 15, 16, 17, 18, 19, 20};
+
+			// Test CbField(BinaryAttachment, Zero)
+			TestField(CbFieldType.BinaryAttachment, ZeroBytes);
+
+			// Test CbField(BinaryAttachment, NonZero)
+			TestField(CbFieldType.BinaryAttachment, SequentialBytes, new IoHash(SequentialBytes));
+
+			// Test CbField(BinaryAttachment, NonZero) AsAttachment
+			{
+				CbField Field = new CbField(SequentialBytes, CbFieldType.BinaryAttachment);
+				TestField(CbFieldType.BinaryAttachment, Field, new IoHash(SequentialBytes), new IoHash(), CbFieldError.None);
+			}
+
+			// Test CbField(None) as BinaryAttachment
+			{
+				CbField DefaultField = new CbField();
+				TestFieldError(CbFieldType.BinaryAttachment, DefaultField, CbFieldError.TypeError, new IoHash(SequentialBytes));
+			}
+		}
+
+		[TestMethod]
+		public void CbFieldHashTest()
+		{
+			byte[] ZeroBytes = new byte[20];
+			byte[] SequentialBytes = new byte[]{1, 2, 3, 4, 5, 6, 7, 8, 9, 10, 11, 12, 13, 14, 15, 16, 17, 18, 19, 20};
+
+			// Test CbField(Hash, Zero)
+			TestField(CbFieldType.Hash, ZeroBytes);
+
+			// Test CbField(Hash, NonZero)
+			TestField(CbFieldType.Hash, SequentialBytes, new IoHash(SequentialBytes));
+
+			// Test CbField(None) as Hash
+			{
+				CbField DefaultField = new CbField();
+				TestFieldError(CbFieldType.Hash, DefaultField, CbFieldError.TypeError, new IoHash(SequentialBytes));
+			}
+
+			// Test CbField(ObjectAttachment) as Hash
+			{
+				CbField Field = new CbField(SequentialBytes, CbFieldType.ObjectAttachment);
+				TestField(CbFieldType.Hash, Field, new IoHash(SequentialBytes));
+			}
+
+			// Test CbField(BinaryAttachment) as Hash
+			{
+				CbField Field = new CbField(SequentialBytes, CbFieldType.BinaryAttachment);
+				TestField(CbFieldType.Hash, Field, new IoHash(SequentialBytes));
+			}
+		}
+
+		[TestMethod]
+		public void CbFieldUuidTest()
+		{
+			byte[] ZeroBytes = new byte[]{0, 0, 0, 0, 0, 0, 0, 0, 0, 0, 0, 0, 0, 0, 0, 0};
+			byte[] SequentialBytes = new byte[]{0, 1, 2, 3, 4, 5, 6, 7, 8, 9, 10, 11, 12, 13, 14, 15};
+			Guid SequentialGuid = Guid.Parse("00010203-0405-0607-0809-0a0b0c0d0e0f");
+
+			// Test CbField(Uuid, Zero)
+			TestField(CbFieldType.Uuid, ZeroBytes, new Guid(), SequentialGuid);
+
+			// Test CbField(Uuid, NonZero)
+			TestField(CbFieldType.Uuid, SequentialBytes, SequentialGuid, new Guid());
+
+			// Test CbField(None) as Uuid
+			{
+				CbField DefaultField = new CbField();
+				TestFieldError(CbFieldType.Uuid, DefaultField, CbFieldError.TypeError, Guid.NewGuid());
+			}
+		}
+
+		[TestMethod]
+		public void CbFieldDateTimeTest()
+		{
+			// Test CbField(DateTime, Zero)
+			TestField(CbFieldType.DateTime, new byte[]{0, 0, 0, 0, 0, 0, 0, 0});
+
+			// Test CbField(DateTime, 0x1020_3040_5060_7080)
+			TestField(CbFieldType.DateTime, new byte[]{0x10, 0x20, 0x30, 0x40, 0x50, 0x60, 0x70, 0x80}, new DateTime(0x1020_3040_5060_7080L));
+
+			// Test CbField(DateTime, Zero) as FDateTime
+			{
+				byte[] Payload = new byte[]{0, 0, 0, 0, 0, 0, 0, 0};
+				CbField Field = new CbField(Payload, CbFieldType.DateTime);
+				Assert.AreEqual(Field.AsDateTime(), new DateTime(0));
+			}
+
+			// Test CbField(None) as DateTime
+			{
+				CbField DefaultField = new CbField();
+				TestFieldError(CbFieldType.DateTime, DefaultField, CbFieldError.TypeError);
+				DateTime DefaultValue = new DateTime(0x1020_3040_5060_7080L);
+				Assert.AreEqual(DefaultField.AsDateTime(DefaultValue), DefaultValue);
+			}
+		}
+
+		[TestMethod]
+		public void CbFieldTimeSpanTest()
+		{
+			// Test CbField(TimeSpan, Zero)
+			TestField(CbFieldType.TimeSpan, new byte[] {0, 0, 0, 0, 0, 0, 0, 0});
+
+			// Test CbField(TimeSpan, 0x1020_3040_5060_7080)
+			TestField(CbFieldType.TimeSpan, new byte[] {0x10, 0x20, 0x30, 0x40, 0x50, 0x60, 0x70, 0x80}, new TimeSpan(0x1020_3040_5060_7080L));
+
+			// Test CbField(TimeSpan, Zero) as FTimeSpan
+			{
+				byte[] Payload = new byte[]{0, 0, 0, 0, 0, 0, 0, 0};
+				CbField Field = new CbField(Payload, CbFieldType.TimeSpan);
+				Assert.AreEqual(Field.AsTimeSpan(), new TimeSpan(0));
+			}
+
+			// Test CbField(None) as TimeSpan
+			{
+				CbField DefaultField = new CbField();
+				TestFieldError(CbFieldType.TimeSpan, DefaultField, CbFieldError.TypeError);
+				TimeSpan DefaultValue = new TimeSpan(0x1020_3040_5060_7080L);
+				Assert.AreEqual(DefaultField.AsTimeSpan(DefaultValue), DefaultValue);
+			}
+		}
+
+#if false
+		[TestMethod]
+		public void CbFieldObjectIdTest()
+		{
+			// Test CbField(ObjectId, Zero)
+			TestField(CbFieldType.ObjectId, {0, 0, 0, 0, 0, 0, 0, 0, 0, 0, 0, 0});
+
+			// Test CbField(ObjectId, 0x102030405060708090A0B0C0)
+			TestField(CbFieldType.ObjectId, {0x10, 0x20, 0x30, 0x40, 0x50, 0x60, 0x70, 0x80, 0x90, 0xA0, 0xB0, 0xC0},
+				FCbObjectId(MakeMemoryView<byte>({0x10, 0x20, 0x30, 0x40, 0x50, 0x60, 0x70, 0x80, 0x90, 0xA0, 0xB0, 0xC0})));
+
+			// Test CbField(ObjectId, Zero) as FCbObjectId
+			{
+				byte[] Payload = new byte[]{0, 0, 0, 0, 0, 0, 0, 0, 0, 0, 0, 0};
+				CbField Field = new CbField(Payload, CbFieldType.ObjectId);
+				Assert.AreEqual(Field.AsObjectId(), FCbObjectId());
+			}
+
+			// Test CbField(None) as ObjectId
+			{
+				CbField DefaultField;
+				TestFieldError(CbFieldType.ObjectId, DefaultField, CbFieldError.TypeError);
+				FCbObjectId DefaultValue(MakeMemoryView<byte>({0x10, 0x20, 0x30, 0x40, 0x50, 0x60, 0x70, 0x80, 0x90, 0xA0, 0xB0, 0xC0}));
+				Assert.AreEqual(DefaultField.AsObjectId(DefaultValue), DefaultValue);
+			}
+		}
+
+		[TestMethod]
+		public void CbFieldCustomByIdTest()
+		{
+			struct FCustomByIdAccessor
+			{
+				explicit FCustomByIdAccessor(uint64 Id)
+					: AsType([Id](CbField& Field, FMemoryView Default) { return Field.AsCustom(Id, Default); })
+				{
+				}
+
+				bool (CbField.*IsType)() = &CbField.IsCustomById;
+				TUniqueFunction<FMemoryView (CbField& Field, FMemoryView Default)> AsType;
+			};
+
+			// Test CbField(CustomById, MinId, Empty)
+			{
+				byte[] Payload = new byte[]{1, 0};
+				TestField(CbFieldType.CustomById, Payload, FCbCustomById{0});
+				TestField(CbFieldType.CustomById, Payload, FMemoryView(), MakeMemoryView<byte>({1, 2, 3}), CbFieldError.None, FCustomByIdAccessor(0));
+				TestFieldError(CbFieldType.CustomById, Payload, CbFieldError.RangeError, MakeMemoryView<byte>({1, 2, 3}), FCustomByIdAccessor(MAX_uint64));
+			}
+
+			// Test CbField(CustomById, MinId, Value)
+			{
+				byte[] Payload = new byte[]{5, 0, 1, 2, 3, 4};
+				TestFieldNoClone(CbFieldType.CustomById, Payload, FCbCustomById{0, Payload.Right(4)});
+				TestFieldNoClone(CbFieldType.CustomById, Payload, Payload.Right(4), FMemoryView(), CbFieldError.None, FCustomByIdAccessor(0));
+				TestFieldError(CbFieldType.CustomById, Payload, CbFieldError.RangeError, MakeMemoryView<byte>({1, 2, 3}), FCustomByIdAccessor(MAX_uint64));
+			}
+
+			// Test CbField(CustomById, MaxId, Empty)
+			{
+				byte[] Payload = new byte[]{9, 0xff, 0xff, 0xff, 0xff, 0xff, 0xff, 0xff, 0xff, 0xff};
+				TestField(CbFieldType.CustomById, Payload, FCbCustomById{MAX_uint64});
+				TestField(CbFieldType.CustomById, Payload, FMemoryView(), MakeMemoryView<byte>({1, 2, 3}), CbFieldError.None, FCustomByIdAccessor(MAX_uint64));
+				TestFieldError(CbFieldType.CustomById, Payload, CbFieldError.RangeError, MakeMemoryView<byte>({1, 2, 3}), FCustomByIdAccessor(0));
+			}
+
+			// Test CbField(CustomById, MaxId, Value)
+			{
+				byte[] Payload = new byte[]{13, 0xff, 0xff, 0xff, 0xff, 0xff, 0xff, 0xff, 0xff, 0xff, 1, 2, 3, 4};
+				TestFieldNoClone(CbFieldType.CustomById, Payload, FCbCustomById{MAX_uint64, Payload.Right(4)});
+				TestFieldNoClone(CbFieldType.CustomById, Payload, Payload.Right(4), FMemoryView(), CbFieldError.None, FCustomByIdAccessor(MAX_uint64));
+				TestFieldError(CbFieldType.CustomById, Payload, CbFieldError.RangeError, MakeMemoryView<byte>({1, 2, 3}), FCustomByIdAccessor(0));
+			}
+
+			// Test CbField(None) as CustomById
+			{
+				CbField DefaultField;
+				TestFieldError(CbFieldType.CustomById, DefaultField, CbFieldError.TypeError, FCbCustomById{4, MakeMemoryView<byte>({1, 2, 3})});
+				TestFieldError(CbFieldType.CustomById, DefaultField, CbFieldError.TypeError, MakeMemoryView<byte>({1, 2, 3}), FCustomByIdAccessor(0));
+				byte[] DefaultValue = new byte[]{1, 2, 3};
+				Assert.AreEqual(DefaultField.AsCustom(0, DefaultValue), DefaultValue);
+			}
+
+			return true;
+		}
+
+		[TestMethod]
+		void CbFieldCustomByNameTest()
+		{
+			struct FCustomByNameAccessor
+			{
+				explicit FCustomByNameAccessor(FAnsiStringView Name)
+					: AsType([Name = FString(Name)](CbField& Field, FMemoryView Default) { return Field.AsCustom(TCHAR_TO_ANSI(*Name), Default); })
+				{
+				}
+
+				bool (CbField.*IsType)() = &CbField.IsCustomByName;
+				TUniqueFunction<FMemoryView (CbField& Field, FMemoryView Default)> AsType;
+			};
+
+			// Test CbField(CustomByName, ABC, Empty)
+			{
+				byte[] Payload = new byte[]{4, 3, 'A', 'B', 'C'};
+				TestField(CbFieldType.CustomByName, Payload, FCbCustomByName{"ABC"});
+				TestField(CbFieldType.CustomByName, Payload, FMemoryView(), MakeMemoryView<byte>({1, 2, 3}), CbFieldError.None, FCustomByNameAccessor("ABC"));
+				TestFieldError(CbFieldType.CustomByName, Payload, CbFieldError.RangeError, MakeMemoryView<byte>({1, 2, 3}), FCustomByNameAccessor("abc"));
+			}
+
+			// Test CbField(CustomByName, ABC, Value)
+			{
+				byte[] Payload = new byte[]{8, 3, 'A', 'B', 'C', 1, 2, 3, 4};
+				TestFieldNoClone(CbFieldType.CustomByName, Payload, FCbCustomByName{"ABC", Payload.Right(4)});
+				TestFieldNoClone(CbFieldType.CustomByName, Payload, Payload.Right(4), FMemoryView(), CbFieldError.None, FCustomByNameAccessor("ABC"));
+				TestFieldError(CbFieldType.CustomByName, Payload, CbFieldError.RangeError, MakeMemoryView<byte>({1, 2, 3}), FCustomByNameAccessor("abc"));
+			}
+
+			// Test CbField(None) as CustomByName
+			{
+				CbField DefaultField;
+				TestFieldError(CbFieldType.CustomByName, DefaultField, CbFieldError.TypeError, FCbCustomByName{"ABC", MakeMemoryView<byte>({1, 2, 3})});
+				TestFieldError(CbFieldType.CustomByName, DefaultField, CbFieldError.TypeError, MakeMemoryView<byte>({1, 2, 3}), FCustomByNameAccessor("ABC"));
+				byte[] DefaultValue = new byte[]{1, 2, 3};
+				Assert.AreEqual(DefaultField.AsCustom("ABC", DefaultValue), DefaultValue);
+			}
+		}
+
+		[TestMethod]
+		public void CbFieldIterateAttachmentsTest()
+		{
+			Func<uint, IoHash> MakeTestHash = (uint Index) =>
+			{
+				byte[] Data = new byte[4];
+				BinaryPrimitives.WriteUInt32LittleEndian(Data, Index);
+				return IoHash.Compute(Data);
+			};
+
+			CbFieldIterator Fields;
+			{
+				CbWriter Writer = new CbWriter();
+
+				Writer.SetName("IgnoredTypeInRoot").AddHash(MakeTestHash(100));
+				Writer.AddObjectAttachment(MakeTestHash(0));
+				Writer.AddBinaryAttachment(MakeTestHash(1));
+				Writer.SetName("ObjAttachmentInRoot").AddObjectAttachment(MakeTestHash(2));
+				Writer.SetName("BinAttachmentInRoot").AddBinaryAttachment(MakeTestHash(3));
+
+				// Uniform array of type to ignore.
+				Writer.BeginArray();
+				{
+					Writer << 1;
+					Writer << 2;
+				}
+				Writer.EndArray();
+				// Uniform array of binary attachments.
+				Writer.BeginArray();
+				{
+					Writer.AddBinaryAttachment(MakeTestHash(4));
+					Writer.AddBinaryAttachment(MakeTestHash(5));
+				}
+				Writer.EndArray();
+				// Uniform array of uniform arrays.
+				Writer.BeginArray();
+				{
+					Writer.BeginArray();
+					Writer.AddBinaryAttachment(MakeTestHash(6));
+					Writer.AddBinaryAttachment(MakeTestHash(7));
+					Writer.EndArray();
+					Writer.BeginArray();
+					Writer.AddBinaryAttachment(MakeTestHash(8));
+					Writer.AddBinaryAttachment(MakeTestHash(9));
+					Writer.EndArray();
+				}
+				Writer.EndArray();
+				// Uniform array of non-uniform arrays.
+				Writer.BeginArray();
+				{
+					Writer.BeginArray();
+					Writer << 0;
+					Writer << false;
+					Writer.EndArray();
+					Writer.BeginArray();
+					Writer.AddObjectAttachment(MakeTestHash(10));
+					Writer << false;
+					Writer.EndArray();
+				}
+				Writer.EndArray();
+				// Uniform array of uniform objects.
+				Writer.BeginArray();
+				{
+					Writer.BeginObject();
+					Writer.SetName("ObjAttachmentInUniObjInUniObj1").AddObjectAttachment(MakeTestHash(11));
+					Writer.SetName("ObjAttachmentInUniObjInUniObj2").AddObjectAttachment(MakeTestHash(12));
+					Writer.EndObject();
+					Writer.BeginObject();
+					Writer.SetName("ObjAttachmentInUniObjInUniObj3").AddObjectAttachment(MakeTestHash(13));
+					Writer.SetName("ObjAttachmentInUniObjInUniObj4").AddObjectAttachment(MakeTestHash(14));
+					Writer.EndObject();
+				}
+				Writer.EndArray();
+				// Uniform array of non-uniform objects.
+				Writer.BeginArray();
+				{
+					Writer.BeginObject();
+					Writer << "Int" << 0;
+					Writer << "Bool" << false;
+					Writer.EndObject();
+					Writer.BeginObject();
+					Writer.SetName("ObjAttachmentInNonUniObjInUniObj").AddObjectAttachment(MakeTestHash(15));
+					Writer << "Bool" << false;
+					Writer.EndObject();
+				}
+				Writer.EndArray();
+
+				// Uniform object of type to ignore.
+				Writer.BeginObject();
+				{
+					Writer << "Int1" << 1;
+					Writer << "Int2" << 2;
+				}
+				Writer.EndObject();
+				// Uniform object of binary attachments.
+				Writer.BeginObject();
+				{
+					Writer.SetName("BinAttachmentInUniObj1").AddBinaryAttachment(MakeTestHash(16));
+					Writer.SetName("BinAttachmentInUniObj2").AddBinaryAttachment(MakeTestHash(17));
+				}
+				Writer.EndObject();
+				// Uniform object of uniform arrays.
+				Writer.BeginObject();
+				{
+					Writer.SetName("Array1");
+					Writer.BeginArray();
+					Writer.AddBinaryAttachment(MakeTestHash(18));
+					Writer.AddBinaryAttachment(MakeTestHash(19));
+					Writer.EndArray();
+					Writer.SetName("Array2");
+					Writer.BeginArray();
+					Writer.AddBinaryAttachment(MakeTestHash(20));
+					Writer.AddBinaryAttachment(MakeTestHash(21));
+					Writer.EndArray();
+				}
+				Writer.EndObject();
+				// Uniform object of non-uniform arrays.
+				Writer.BeginObject();
+				{
+					Writer.SetName("Array1");
+					Writer.BeginArray();
+					Writer << 0;
+					Writer << false;
+					Writer.EndArray();
+					Writer.SetName("Array2");
+					Writer.BeginArray();
+					Writer.AddObjectAttachment(MakeTestHash(22));
+					Writer << false;
+					Writer.EndArray();
+				}
+				Writer.EndObject();
+				// Uniform object of uniform objects.
+				Writer.BeginObject();
+				{
+					Writer.SetName("Object1");
+					Writer.BeginObject();
+					Writer.SetName("ObjAttachmentInUniObjInUniObj1").AddObjectAttachment(MakeTestHash(23));
+					Writer.SetName("ObjAttachmentInUniObjInUniObj2").AddObjectAttachment(MakeTestHash(24));
+					Writer.EndObject();
+					Writer.SetName("Object2");
+					Writer.BeginObject();
+					Writer.SetName("ObjAttachmentInUniObjInUniObj3").AddObjectAttachment(MakeTestHash(25));
+					Writer.SetName("ObjAttachmentInUniObjInUniObj4").AddObjectAttachment(MakeTestHash(26));
+					Writer.EndObject();
+				}
+				Writer.EndObject();
+				// Uniform object of non-uniform objects.
+				Writer.BeginObject();
+				{
+					Writer.SetName("Object1");
+					Writer.BeginObject();
+					Writer << "Int" << 0;
+					Writer << "Bool" << false;
+					Writer.EndObject();
+					Writer.SetName("Object2");
+					Writer.BeginObject();
+					Writer.SetName("ObjAttachmentInNonUniObjInUniObj").AddObjectAttachment(MakeTestHash(27));
+					Writer << "Bool" << false;
+					Writer.EndObject();
+				}
+				Writer.EndObject();
+
+				Fields = Writer.Save();
+			}
+
+			Assert.AreEqual(ValidateCompactBinaryRange(Fields.GetOuterBuffer(), ECbValidateMode.All), ECbValidateError.None);
+
+			uint AttachmentIndex = 0;
+			Fields.IterateRangeAttachments([this, &AttachmentIndex, &MakeTestHash](CbField Field)
+				{
+					Assert.IsTrue(FString.Printf(AttachmentIndex), Field.IsAttachment());
+					Assert.AreEqual(FString.Printf(AttachmentIndex), Field.AsAttachment(), MakeTestHash(AttachmentIndex));
+					++AttachmentIndex;
+				});
+			Assert.AreEqual(AttachmentIndex, 28);
+		}
+
+		[TestMethod]
+		void CbFieldBufferTest()
+		{
+			static_assert(std.is_constructible<CbField>.value, "Missing constructor for CbField");
+			static_assert(std.is_constructible<CbField, CbField&>.value, "Missing constructor for CbField");
+			static_assert(std.is_constructible<CbField, CbField&&>.value, "Missing constructor for CbField");
+
+			static_assert(std.is_constructible<CbField, FSharedBuffer&>.value, "Missing constructor for CbField");
+			static_assert(std.is_constructible<CbField, FSharedBuffer&&>.value, "Missing constructor for CbField");
+
+			static_assert(std.is_constructible<CbField, CbField&, FSharedBuffer&>.value, "Missing constructor for CbField");
+			static_assert(std.is_constructible<CbField, CbField&, CbFieldIterator&>.value, "Missing constructor for CbField");
+			static_assert(std.is_constructible<CbField, CbField&, CbField&>.value, "Missing constructor for CbField");
+			static_assert(std.is_constructible<CbField, CbField&, CbArray&>.value, "Missing constructor for CbField");
+			static_assert(std.is_constructible<CbField, CbField&, FCbObject&>.value, "Missing constructor for CbField");
+
+			static_assert(std.is_constructible<CbField, CbField&, FSharedBuffer&&>.value, "Missing constructor for CbField");
+			static_assert(std.is_constructible<CbField, CbField&, CbFieldIterator&&>.value, "Missing constructor for CbField");
+			static_assert(std.is_constructible<CbField, CbField&, CbField&&>.value, "Missing constructor for CbField");
+			static_assert(std.is_constructible<CbField, CbField&, CbArray&&>.value, "Missing constructor for CbField");
+			static_assert(std.is_constructible<CbField, CbField&, FCbObject&&>.value, "Missing constructor for CbField");
+
+			// Test CbField()
+			{
+				CbField DefaultField;
+				Assert.IsFalse(DefaultField.HasValue());
+				Assert.IsFalse(DefaultField.IsOwned());
+				DefaultField.MakeOwned();
+				Assert.IsTrue(DefaultField.IsOwned());
+			}
+
+			// Test Field w/ Type from Shared Buffer
+			{
+				byte[] Payload = new byte[]{ (byte)(CbFieldType.Binary), 3, 4, 5, 6 }; // Size: 3, Data: 4/5/6
+				FSharedBuffer ViewBuffer = FSharedBuffer.MakeView(Payload);
+				FSharedBuffer OwnedBuffer = FSharedBuffer.Clone(ViewBuffer);
+
+				CbField View(ViewBuffer);
+				CbField ViewMove{FSharedBuffer(ViewBuffer)};
+				CbField ViewOuterField(ImplicitConv<CbField>(View), ViewBuffer);
+				CbField ViewOuterBuffer(ImplicitConv<CbField>(View), View);
+				CbField Owned(OwnedBuffer);
+				CbField OwnedMove{FSharedBuffer(OwnedBuffer)};
+				CbField OwnedOuterField(ImplicitConv<CbField>(Owned), OwnedBuffer);
+				CbField OwnedOuterBuffer(ImplicitConv<CbField>(Owned), Owned);
+
+				// These lines are expected to assert when uncommented.
+				//CbField InvalidOuterBuffer(ImplicitConv<CbField>(Owned), ViewBuffer);
+				//CbField InvalidOuterBufferMove(ImplicitConv<CbField>(Owned), FSharedBuffer(ViewBuffer));
+
+				Assert.AreEqual(View.AsBinaryView(), ViewBuffer.GetView().Right(3));
+				Assert.AreEqual(ViewMove.AsBinaryView(), View.AsBinaryView());
+				Assert.AreEqual(ViewOuterField.AsBinaryView(), View.AsBinaryView());
+				Assert.AreEqual(ViewOuterBuffer.AsBinaryView(), View.AsBinaryView());
+				Assert.AreEqual(Owned.AsBinaryView(), OwnedBuffer.GetView().Right(3));
+				Assert.AreEqual(OwnedMove.AsBinaryView(), Owned.AsBinaryView());
+				Assert.AreEqual(OwnedOuterField.AsBinaryView(), Owned.AsBinaryView());
+				Assert.AreEqual(OwnedOuterBuffer.AsBinaryView(), Owned.AsBinaryView());
+
+				Assert.IsFalse(View.IsOwned());
+				Assert.IsFalse(ViewMove.IsOwned());
+				Assert.IsFalse(ViewOuterField.IsOwned());
+				Assert.IsFalse(ViewOuterBuffer.IsOwned());
+				Assert.IsTrue(Owned.IsOwned());
+				Assert.IsTrue(OwnedMove.IsOwned());
+				Assert.IsTrue(OwnedOuterField.IsOwned());
+				Assert.IsTrue(OwnedOuterBuffer.IsOwned());
+
+				View.MakeOwned();
+				Owned.MakeOwned();
+				Assert.AreNotEqual(View.AsBinaryView(), ViewBuffer.GetView().Right(3));
+				Assert.IsTrue(View.IsOwned());
+				Assert.AreEqual(Owned.AsBinaryView(), OwnedBuffer.GetView().Right(3));
+				Assert.IsTrue(Owned.IsOwned());
+			}
+
+			// Test Field w/ Type
+			{
+				byte[] Payload = new byte[]{ (byte)(CbFieldType.Binary), 3, 4, 5, 6 }; // Size: 3, Data: 4/5/6
+				CbField Field = new CbField(Payload);
+
+				CbField VoidView = CbField.MakeView(Payload);
+				CbField VoidClone = CbField.Clone(Payload);
+				CbField FieldView = CbField.MakeView(Field);
+				CbField FieldClone = CbField.Clone(Field);
+				CbField FieldViewClone = CbField.Clone(FieldView);
+
+				Assert.AreEqual(VoidView.AsBinaryView(), Payload.Right(3));
+				Assert.AreNotEqual(VoidClone.AsBinaryView(), Payload.Right(3));
+				Assert.IsTrue(VoidClone.AsBinaryView().Span.SequenceEqual(VoidView.AsBinaryView()));
+				Assert.AreEqual(FieldView.AsBinaryView(), Payload.Right(3));
+				Assert.AreNotEqual(FieldClone.AsBinaryView(), Payload.Right(3));
+				Assert.IsTrue(FieldClone.AsBinaryView().Span.SequenceEqual(VoidView.AsBinaryView()));
+				Assert.AreNotEqual(FieldViewClone.AsBinaryView(), FieldView.AsBinaryView());
+				Assert.IsTrue(FieldViewClone.AsBinaryView().Span.SequenceEqual(VoidView.AsBinaryView()));
+
+				Assert.IsFalse(VoidView.IsOwned());
+				Assert.IsTrue(VoidClone.IsOwned());
+				Assert.IsFalse(FieldView.IsOwned());
+				Assert.IsTrue(FieldClone.IsOwned());
+				Assert.IsTrue(FieldViewClone.IsOwned());
+			}
+
+			// Test Field w/o Type
+			{
+				byte[] Payload = new byte[]{ 3, 4, 5, 6 }; // Size: 3, Data: 4/5/6
+				CbField Field = new CbField(Payload, CbFieldType.Binary);
+
+				CbField FieldView = CbField.MakeView(Field);
+				CbField FieldClone = CbField.Clone(Field);
+				CbField FieldViewClone = CbField.Clone(FieldView);
+
+				Assert.AreEqual(FieldView.AsBinaryView(), Payload.Right(3));
+				Assert.IsTrue(FieldClone.AsBinaryView().Span.SequenceEqual(FieldView.AsBinaryView()));
+				Assert.IsTrue(FieldViewClone.AsBinaryView().Span.SequenceEqual(FieldView.AsBinaryView()));
+
+				Assert.IsFalse(FieldView.IsOwned());
+				Assert.IsTrue(FieldClone.IsOwned());
+				Assert.IsTrue(FieldViewClone.IsOwned());
+
+				FieldView.MakeOwned();
+				Assert.IsTrue(FieldView.AsBinaryView().Span.SequenceEqual(Payload.Right(3)));
+				Assert.IsTrue(FieldView.IsOwned());
+			}
+
+			return true;
+		}
+
+		[TestMethod]
+		void CbArrayBufferTest()
+		{
+			static_assert(std.is_constructible<CbArray>.value, "Missing constructor for CbArray");
+			static_assert(std.is_constructible<CbArray, CbArray&>.value, "Missing constructor for CbArray");
+			static_assert(std.is_constructible<CbArray, CbArray&&>.value, "Missing constructor for CbArray");
+
+			static_assert(std.is_constructible<CbArray, CbArray&, FSharedBuffer&>.value, "Missing constructor for CbArray");
+			static_assert(std.is_constructible<CbArray, CbArray&, CbFieldIterator&>.value, "Missing constructor for CbArray");
+			static_assert(std.is_constructible<CbArray, CbArray&, CbField&>.value, "Missing constructor for CbArray");
+			static_assert(std.is_constructible<CbArray, CbArray&, CbArray&>.value, "Missing constructor for CbArray");
+			static_assert(std.is_constructible<CbArray, CbArray&, FCbObject&>.value, "Missing constructor for CbArray");
+
+			static_assert(std.is_constructible<CbArray, CbArray&, FSharedBuffer&&>.value, "Missing constructor for CbArray");
+			static_assert(std.is_constructible<CbArray, CbArray&, CbFieldIterator&&>.value, "Missing constructor for CbArray");
+			static_assert(std.is_constructible<CbArray, CbArray&, CbField&&>.value, "Missing constructor for CbArray");
+			static_assert(std.is_constructible<CbArray, CbArray&, CbArray&&>.value, "Missing constructor for CbArray");
+			static_assert(std.is_constructible<CbArray, CbArray&, FCbObject&&>.value, "Missing constructor for CbArray");
+
+			// Test CbArray()
+			{
+				CbArray DefaultArray;
+				Assert.IsFalse(DefaultArray.IsOwned());
+				DefaultArray.MakeOwned();
+				Assert.IsTrue(DefaultArray.IsOwned());
+			}
+		}
+
+		[TestMethod]
+		public void CbObjectBufferTest()
+		{
+			static_assert(std.is_constructible<FCbObject>.value, "Missing constructor for FCbObject");
+			static_assert(std.is_constructible<FCbObject, FCbObject&&>.value, "Missing constructor for FCbObject");
+			static_assert(std.is_constructible<FCbObject, FCbObject&&>.value, "Missing constructor for FCbObject");
+
+			static_assert(std.is_constructible<FCbObject, FCbObjectView&, FSharedBuffer&>.value, "Missing constructor for FCbObject");
+			static_assert(std.is_constructible<FCbObject, FCbObjectView&, CbFieldIterator&>.value, "Missing constructor for FCbObject");
+			static_assert(std.is_constructible<FCbObject, FCbObjectView&, CbField&>.value, "Missing constructor for FCbObject");
+			static_assert(std.is_constructible<FCbObject, FCbObjectView&, CbArray&>.value, "Missing constructor for FCbObject");
+			static_assert(std.is_constructible<FCbObject, FCbObjectView&, FCbObject&>.value, "Missing constructor for FCbObject");
+
+			static_assert(std.is_constructible<FCbObject, FCbObjectView&, FSharedBuffer&&>.value, "Missing constructor for FCbObject");
+			static_assert(std.is_constructible<FCbObject, FCbObjectView&, CbFieldIterator&&>.value, "Missing constructor for FCbObject");
+			static_assert(std.is_constructible<FCbObject, FCbObjectView&, CbField&&>.value, "Missing constructor for FCbObject");
+			static_assert(std.is_constructible<FCbObject, FCbObjectView&, CbArray&&>.value, "Missing constructor for FCbObject");
+			static_assert(std.is_constructible<FCbObject, FCbObjectView&, FCbObject&&>.value, "Missing constructor for FCbObject");
+
+			// Test FCbObject()
+			{
+				FCbObject DefaultObject;
+				Assert.IsFalse(DefaultObject.IsOwned());
+				DefaultObject.MakeOwned();
+				Assert.IsTrue(DefaultObject.IsOwned());
+			}
+
+			return true;
+		}
+
+		[TestMethod]
+		public void CbFieldBufferIterator()
+		{
+			static_assert(std.is_constructible<CbFieldIterator, CbFieldIterator&>.value, "Missing constructor for CbFieldIterator");
+			static_assert(std.is_constructible<CbFieldIterator, CbFieldIterator&&>.value, "Missing constructor for CbFieldIterator");
+
+			static_assert(std.is_constructible<CbFieldIterator, CbFieldIterator&>.value, "Missing constructor for CbFieldIterator");
+			static_assert(std.is_constructible<CbFieldIterator, CbFieldIterator&&>.value, "Missing constructor for CbFieldIterator");
+
+			auto GetCount = [](auto It) -> uint
+			{
+				uint Count = 0;
+				for (; It; ++It)
+				{
+					++Count;
+				}
+				return Count;
+			};
+
+			// Test CbField[View]Iterator()
+			{
+				Assert.AreEqual(GetCount(CbFieldIterator()), 0);
+				Assert.AreEqual(GetCount(CbFieldIterator()), 0);
+			}
+
+			// Test CbField[View]Iterator(Range)
+			{
+				byte T = (byte)(CbFieldType.IntegerPositive);
+				byte[] Payload = new byte[]{ T, 0, T, 1, T, 2, T, 3 };
+
+				FSharedBuffer View = FSharedBuffer.MakeView(Payload);
+				FSharedBuffer Clone = FSharedBuffer.Clone(View);
+
+				FMemoryView EmptyView;
+				FSharedBuffer NullBuffer;
+
+				CbFieldIterator FieldViewIt = CbFieldIterator.MakeRange(View);
+				CbFieldIterator FieldIt = CbFieldIterator.MakeRange(View);
+
+				Assert.AreEqual(FieldViewIt.GetRangeHash(), FBlake3.HashBuffer(View));
+				Assert.AreEqual(FieldIt.GetRangeHash(), FBlake3.HashBuffer(View));
+
+				FMemoryView RangeView;
+				Assert.IsTrue(FieldViewIt.TryGetRangeView(RangeView) && RangeView == Payload);
+				Assert.IsTrue(FieldIt.TryGetRangeView(RangeView) && RangeView == Payload);
+
+				Assert.AreEqual(GetCount(CbFieldIterator.CloneRange(CbFieldIterator())), 0);
+				Assert.AreEqual(GetCount(CbFieldIterator.CloneRange(CbFieldIterator())), 0);
+				CbFieldIterator FieldViewItClone = CbFieldIterator.CloneRange(FieldViewIt);
+				CbFieldIterator FieldItClone = CbFieldIterator.CloneRange(FieldIt);
+				Assert.AreEqual(GetCount(FieldViewItClone), 4);
+				Assert.AreEqual(GetCount(FieldItClone), 4);
+				Assert.AreNotEqual(FieldViewItClone, FieldIt);
+				Assert.AreNotEqual(FieldItClone, FieldIt);
+
+				Assert.AreEqual(GetCount(CbFieldIterator.MakeRange(EmptyView)), 0);
+				Assert.AreEqual(GetCount(CbFieldIterator.MakeRange(NullBuffer)), 0);
+				Assert.AreEqual(GetCount(CbFieldIterator.MakeRange(FSharedBuffer(NullBuffer))), 0);
+
+				Assert.AreEqual(GetCount(CbFieldIterator.MakeRange(Payload)), 4);
+				Assert.AreEqual(GetCount(CbFieldIterator.MakeRange(Clone)), 4);
+				Assert.AreEqual(GetCount(CbFieldIterator.MakeRange(FSharedBuffer(Clone))), 4);
+
+				Assert.AreEqual(GetCount(CbFieldIterator.MakeRangeView(CbFieldIterator.MakeRange(View), NullBuffer)), 4);
+				Assert.AreEqual(GetCount(CbFieldIterator.MakeRangeView(CbFieldIterator.MakeRange(View), FSharedBuffer(NullBuffer))), 4);
+				Assert.AreEqual(GetCount(CbFieldIterator.MakeRangeView(CbFieldIterator.MakeRange(View), View)), 4);
+				Assert.AreEqual(GetCount(CbFieldIterator.MakeRangeView(CbFieldIterator.MakeRange(View), FSharedBuffer(View))), 4);
+				Assert.AreEqual(GetCount(CbFieldIterator.MakeRangeView(CbFieldIterator.MakeRange(Clone), Clone)), 4);
+				Assert.AreEqual(GetCount(CbFieldIterator.MakeRangeView(CbFieldIterator.MakeRange(Clone), FSharedBuffer(Clone))), 4);
+
+				Assert.AreEqual(GetCount(CbFieldIterator(FieldIt)), 4);
+				Assert.AreEqual(GetCount(CbFieldIterator(CbFieldIterator(FieldIt))), 4);
+
+				// Uniform
+				byte[] UniformPayload = new byte[]{ 0, 1, 2, 3 };
+				CbFieldIterator UniformFieldViewIt = CbFieldIterator.MakeRange(UniformPayload, CbFieldType.IntegerPositive);
+				Assert.AreEqual(UniformFieldViewIt.GetRangeHash(), FieldViewIt.GetRangeHash());
+				Assert.IsFalse(UniformFieldViewIt.TryGetRangeView(RangeView));
+				FSharedBuffer UniformView = FSharedBuffer.MakeView(UniformPayload);
+				CbFieldIterator UniformFieldIt = CbFieldIterator.MakeRange(UniformView, CbFieldType.IntegerPositive);
+				Assert.AreEqual(UniformFieldIt.GetRangeHash(), FieldViewIt.GetRangeHash());
+				Assert.IsFalse(UniformFieldIt.TryGetRangeView(RangeView));
+
+				// Equals
+				Assert.IsTrue(FieldViewIt.Equals(FieldViewIt));
+				Assert.IsTrue(FieldViewIt.Equals(FieldIt));
+				Assert.IsTrue(FieldIt.Equals(FieldIt));
+				Assert.IsTrue(FieldIt.Equals(FieldViewIt));
+				Assert.IsFalse(FieldViewIt.Equals(FieldViewItClone));
+				Assert.IsFalse(FieldIt.Equals(FieldItClone));
+				Assert.IsTrue(UniformFieldViewIt.Equals(UniformFieldViewIt));
+				Assert.IsTrue(UniformFieldViewIt.Equals(UniformFieldIt));
+				Assert.IsTrue(UniformFieldIt.Equals(UniformFieldIt));
+				Assert.IsTrue(UniformFieldIt.Equals(UniformFieldViewIt));
+				Assert.IsFalse(TEXT("CbFieldIterator.Equals(SamePayload, DifferentEnd)"),
+					CbFieldIterator.MakeRange(UniformPayload, CbFieldType.IntegerPositive)
+						.Equals(CbFieldIterator.MakeRange(UniformPayload.LeftChop(1), CbFieldType.IntegerPositive)));
+				Assert.IsFalse(TEXT("CbFieldIterator.Equals(DifferentPayload, SameEnd)"),
+					CbFieldIterator.MakeRange(UniformPayload, CbFieldType.IntegerPositive)
+						.Equals(CbFieldIterator.MakeRange(UniformPayload.RightChop(1), CbFieldType.IntegerPositive)));
+
+				// CopyRangeTo
+				byte[] CopyBytes = new byte[Payload.Length];
+				FieldViewIt.CopyRangeTo(CopyBytes);
+				Assert.IsTrue(CopyBytes.Span.SequenceEqual(Payload));
+				FieldIt.CopyRangeTo(CopyBytes);
+				Assert.IsTrue(CopyBytes.Span.SequenceEqual(Payload));
+				UniformFieldViewIt.CopyRangeTo(CopyBytes);
+				Assert.IsTrue(CopyBytes.Span.SequenceEqual(Payload));
+
+				// MakeRangeOwned
+				CbFieldIterator OwnedFromView = UniformFieldIt;
+				OwnedFromView.MakeRangeOwned();
+				Assert.IsTrue(OwnedFromView.TryGetRangeView(RangeView) && RangeView.Span.SequenceEqual(Payload));
+				CbFieldIterator OwnedFromOwned = OwnedFromView;
+				OwnedFromOwned.MakeRangeOwned();
+				Assert.AreEqual(OwnedFromOwned, OwnedFromView);
+
+				// These lines are expected to assert when uncommented.
+				//FSharedBuffer ShortView = FSharedBuffer.MakeView(Payload.LeftChop(2));
+				//Assert.AreEqual(GetCount(CbFieldIterator.MakeRangeView(CbFieldIterator.MakeRange(*View), ShortView)), 4);
+				//Assert.AreEqual(GetCount(CbFieldIterator.MakeRangeView(CbFieldIterator.MakeRange(*View), FSharedBuffer(ShortView))), 4);
+			}
+
+			// Test CbField[View]Iterator(Scalar)
+			{
+				byte T = (byte)(CbFieldType.IntegerPositive);
+				byte[] Payload = new byte[]{ T, 0 };
+
+				FSharedBuffer View = FSharedBuffer.MakeView(Payload);
+				FSharedBuffer Clone = FSharedBuffer.Clone(View);
+
+				CbField FieldView(Payload);
+				CbField Field = new CbField(View);
+
+				Assert.AreEqual(GetCount(CbFieldIterator.MakeSingle(FieldView)), 1);
+				Assert.AreEqual(GetCount(CbFieldIterator.MakeSingle(CbField(FieldView))), 1);
+				Assert.AreEqual(GetCount(CbFieldIterator.MakeSingle(Field)), 1);
+				Assert.AreEqual(GetCount(CbFieldIterator.MakeSingle(CbField(Field))), 1);
+			}
+
+			return true;
+		}
+#endif
+		delegate void ParseObjectType(CbObject Object, ref uint A, ref uint B, ref uint C, ref uint D);
+
+		[TestMethod]
+		public void CbFieldParseTest()
+		{
+			// Test the optimal object parsing loop because it is expected to be required for high performance.
+			// Under ideal conditions, when the fields are in the expected order and there are no extra fields,
+			// the loop will execute once and only one comparison will be performed for each field name. Either
+			// way, each field will only be visited once even if the loop needs to execute several times.
+			ParseObjectType ParseObject = (CbObject Object, ref uint A, ref uint B, ref uint C, ref uint D) =>
+			{
+				for (CbFieldIterator It = Object.CreateViewIterator(); It;)
+				{
+					CbFieldIterator Last = It;
+					if (It.Current.GetName().Equals("A"))
+					{
+						A = It.Current.AsUInt32();
+						++It;
+					}
+					if (It.Current.GetName().Equals("B"))
+					{
+						B = It.Current.AsUInt32();
+						++It;
+					}
+					if (It.Current.GetName().Equals("C"))
+					{
+						C = It.Current.AsUInt32();
+						++It;
+					}
+					if (It.Current.GetName().Equals("D"))
+					{
+						D = It.Current.AsUInt32();
+						++It;
+					}
+					if (Last == It)
+					{
+						++It;
+					}
+				}
+			};
+
+			Func<byte[], uint, uint, uint, uint, bool> TestParseObject = (byte[] Data, uint A, uint B, uint C, uint D) =>
+			{
+				uint ParsedA = 0, ParsedB = 0, ParsedC = 0, ParsedD = 0;
+				ParseObject(new CbObject(Data, CbFieldType.Object), ref ParsedA, ref ParsedB, ref ParsedC, ref ParsedD);
+				return A == ParsedA && B == ParsedB && C == ParsedC && D == ParsedD;
+			};
+
+			byte T = (byte)(CbFieldType.IntegerPositive | CbFieldType.HasFieldName);
+			Assert.IsTrue(TestParseObject(new byte[]{0}, 0, 0, 0, 0));
+			Assert.IsTrue(TestParseObject(new byte[] { 16, T, 1, (byte)'A', 1, T, 1, (byte)'B', 2, T, 1, (byte)'C', 3, T, 1, (byte)'D', 4}, 1, 2, 3, 4));
+			Assert.IsTrue(TestParseObject(new byte[]{16, T, 1, (byte)'B', 2, T, 1, (byte)'C', 3, T, 1, (byte)'D', 4, T, 1, (byte)'A', 1}, 1, 2, 3, 4));
+			Assert.IsTrue(TestParseObject(new byte[]{12, T, 1, (byte)'B', 2, T, 1, (byte)'C', 3, T, 1, (byte)'D', 4}, 0, 2, 3, 4));
+			Assert.IsTrue(TestParseObject(new byte[]{8, T, 1, (byte)'B', 2, T, 1, (byte)'C', 3}, 0, 2, 3, 0));
+			Assert.IsTrue(TestParseObject(new byte[]{20, T, 1, (byte)'A', 1, T, 1, (byte)'B', 2, T, 1, (byte)'C', 3, T, 1, (byte)'D', 4, T, 1, (byte)'E', 5}, 1, 2, 3, 4));
+			Assert.IsTrue(TestParseObject(new byte[]{20, T, 1, (byte)'E', 5, T, 1, (byte)'A', 1, T, 1, (byte)'B', 2, T, 1, (byte)'C', 3, T, 1, (byte)'D', 4}, 1, 2, 3, 4));
+			Assert.IsTrue(TestParseObject(new byte[] { 16, T, 1, (byte)'D', 4, T, 1, (byte)'C', 3, T, 1, (byte)'B', 2, T, 1, (byte)'A', 1}, 1, 2, 3, 4));
+		}
+	}
+}
diff --git a/Engine/Source/Programs/Shared/EpicGames.Serialization.Tests/EpicGames.Serialization.Tests.csproj b/Engine/Source/Programs/Shared/EpicGames.Serialization.Tests/EpicGames.Serialization.Tests.csproj
new file mode 100644
--- /dev/null
+++ b/Engine/Source/Programs/Shared/EpicGames.Serialization.Tests/EpicGames.Serialization.Tests.csproj
@@ -0,0 +1,32 @@
+<Project Sdk="Microsoft.NET.Sdk">
+
+  <PropertyGroup>
+    <TargetFramework>netcoreapp3.1</TargetFramework>
+
+    <IsPackable>false</IsPackable>
+
+    <Nullable>enable</Nullable>
+  </PropertyGroup>
+
+  <ItemGroup>
+    <PackageReference Include="Microsoft.NET.Test.Sdk" Version="16.7.1" />
+    <PackageReference Include="MSTest.TestAdapter" Version="2.1.1" />
+    <PackageReference Include="MSTest.TestFramework" Version="2.1.1" />
+    <PackageReference Include="coverlet.collector" Version="1.3.0" />
+  </ItemGroup>
+
+  <ItemGroup>
+    <ProjectReference Include="..\EpicGames.IoHash\EpicGames.IoHash.csproj" />
+    <ProjectReference Include="..\EpicGames.Serialization\EpicGames.Serialization.csproj" />
+  </ItemGroup>
+
+  <ItemGroup>
+    <None Update="CompactBinaryObjects\build">
+      <CopyToOutputDirectory>PreserveNewest</CopyToOutputDirectory>
+    </None>
+    <None Update="CompactBinaryObjects\ReferenceOutput">
+      <CopyToOutputDirectory>PreserveNewest</CopyToOutputDirectory>
+    </None>
+  </ItemGroup>
+
+</Project>
diff --git a/Engine/Source/Programs/Shared/EpicGames.Serialization.Tests/JupiterTests.cs b/Engine/Source/Programs/Shared/EpicGames.Serialization.Tests/JupiterTests.cs
new file mode 100644
--- /dev/null
+++ b/Engine/Source/Programs/Shared/EpicGames.Serialization.Tests/JupiterTests.cs
@@ -0,0 +1,120 @@
+// Copyright Epic Games, Inc. All Rights Reserved.
+
+using System;
+using System.Collections.Generic;
+using System.IO;
+using System.Linq;
+using EpicGames.Core;
+using Microsoft.VisualStudio.TestTools.UnitTesting;
+
+namespace EpicGames.Serialization.Tests
+{
+	[TestClass]
+	public class JupiterTests
+	{
+		[TestMethod]
+		public void BuildObject()
+		{
+			byte[] bytes = File.ReadAllBytes("CompactBinaryObjects/build");
+
+			ReadOnlyMemory<byte> memory = new ReadOnlyMemory<byte>(bytes);
+			CbField o = new CbField(memory);
+
+			Assert.AreEqual("BuildAction", o.Name);
+			List<CbField> buildActionFields = o.ToList();
+			Assert.AreEqual(3, buildActionFields.Count);
+			Assert.AreEqual("Function", buildActionFields[0].Name);
+			Assert.AreEqual("Constants", buildActionFields[1].Name);
+			Assert.AreEqual("Inputs", buildActionFields[2].Name);
+
+			List<CbField> constantsFields = buildActionFields[1].ToList();
+			Assert.AreEqual(3, constantsFields.Count);
+			Assert.AreEqual("TextureBuildSettings", constantsFields[0].Name);
+			Assert.AreEqual("TextureOutputSettings", constantsFields[1].Name);
+			Assert.AreEqual("TextureSource", constantsFields[2].Name);
+
+			List<CbField> inputsFields = buildActionFields[2].ToList();
+			Assert.AreEqual(1, inputsFields.Count);
+			Assert.AreEqual("7587B323422942733DDD048A91709FDE", inputsFields[0].Name);
+			Assert.IsTrue(inputsFields[0].IsBinaryAttachment());
+			Assert.IsTrue(inputsFields[0].IsAttachment());
+			Assert.IsFalse(inputsFields[0].IsObjectAttachment());
+			Assert.AreEqual(IoHash.Parse("f855382171a0b1e5a1c653aa6c5121a05cbf4ba0"), inputsFields[0].AsHash());
+		}
+
+
+		[TestMethod]
+		public void ReferenceOutput()
+		{
+			byte[] bytes = File.ReadAllBytes("CompactBinaryObjects/ReferenceOutput");
+
+			ReadOnlyMemory<byte> memory = new ReadOnlyMemory<byte>(bytes);
+			CbField o = new CbField(memory);
+
+			Assert.AreEqual("BuildOutput", o.Name);
+			List<CbField> buildActionFields = o.ToList();
+			Assert.AreEqual(1, buildActionFields.Count);
+			CbField payloads = buildActionFields[0];
+			List<CbField> payloadFields = payloads.ToList();
+			Assert.AreEqual(3, payloadFields.Count);
+
+			Assert.AreEqual(IoHash.Parse("5d8a6dc277c968f0d027c98f879c955c1905c293"), payloadFields[0]["RawHash"]!.AsHash());
+			Assert.AreEqual(IoHash.Parse("313f0d0d334100d83aeb1ee2c42794fd087cb0ae"), payloadFields[1]["RawHash"]!.AsHash());
+			Assert.AreEqual(IoHash.Parse("c7a03f83c08cdca882110ecf2b5654ee3b09b11e"), payloadFields[2]["RawHash"]!.AsHash());
+		}
+
+
+		[TestMethod]
+		public void WriteArray()
+		{
+			var hash1 = IoHash.Parse("5d8a6dc277c968f0d027c98f879c955c1905c293");
+			var hash2 = IoHash.Parse("313f0d0d334100d83aeb1ee2c42794fd087cb0ae");
+
+			CbWriter writer = new CbWriter();
+			writer.BeginObject();
+			writer.AddUniformArray(new IoHash[] { hash1, hash2 }, CbFieldType.Hash, "needs");
+			writer.EndObject();
+
+			byte[] objectData = writer.Save();
+			ReadOnlyMemory<byte> memory = new ReadOnlyMemory<byte>(objectData);
+			CbField o = new CbField(memory);
+
+			// the top object has no name
+			Assert.AreEqual("", o.Name);
+			List<CbField> fields = o.ToList();
+			Assert.AreEqual(1, fields.Count);
+			CbField? needs = o["needs"];
+			List<CbField> blobList = needs!.AsArray().ToList();
+			IoHash[] blobs = blobList.Select(field => field.AsHash()).ToArray();
+			CollectionAssert.AreEqual(new IoHash[] { hash1, hash2 }, blobs);
+		}
+
+
+		[TestMethod]
+		public void WriteObject()
+		{
+			var hash1 = IoHash.Parse("5d8a6dc277c968f0d027c98f879c955c1905c293");
+
+			CbWriter writer = new CbWriter();
+			writer.BeginObject();
+			writer.AddString("test", "string");
+			writer.AddBinaryAttachment(hash1, "hash");
+			writer.EndObject();
+
+			byte[] objectData = writer.Save();
+			ReadOnlyMemory<byte> memory = new ReadOnlyMemory<byte>(objectData);
+			CbField o = new CbField(memory);
+
+			// the object has no name and 2 fields
+			Assert.AreEqual("", o.Name);
+			List<CbField> fields = o.ToList();
+			Assert.AreEqual(2, fields.Count);
+
+			CbField? stringField = o["string"];
+			Assert.AreEqual("test", stringField!.AsString());
+
+			CbField? hashField = o["hash"];
+			Assert.AreEqual(hash1, hashField!.AsAttachment());
+		}
+	}
+}
diff --git a/Engine/Source/Programs/Shared/EpicGames.Serialization/CompactBinary.cs b/Engine/Source/Programs/Shared/EpicGames.Serialization/CompactBinary.cs
new file mode 100644
--- /dev/null
+++ b/Engine/Source/Programs/Shared/EpicGames.Serialization/CompactBinary.cs
@@ -0,0 +1,2333 @@
+// Copyright Epic Games, Inc. All Rights Reserved.
+
+using EpicGames.Core;
+using System;
+using System.Buffers;
+using System.Buffers.Binary;
+using System.Collections;
+using System.Collections.Generic;
+using System.Linq;
+using System.Numerics;
+using System.Text;
+using System.Text.Json;
+
+namespace EpicGames.Serialization
+{
+	/// <summary>
+	/// Field types and flags for FCbField[View].
+	///
+	/// DO NOT CHANGE THE VALUE OF ANY MEMBERS OF THIS ENUM!
+	/// BACKWARD COMPATIBILITY REQUIRES THAT THESE VALUES BE FIXED!
+	/// SERIALIZATION USES HARD-CODED CONSTANTS BASED ON THESE VALUES!
+	/// </summary>
+	[Flags]
+	public enum CbFieldType : byte
+	{
+		/// <summary>
+		/// A field type that does not occur in a valid object.
+		/// </summary>
+		None = 0x00,
+
+		/// <summary>
+		/// Null. Payload is empty.
+		/// </summary>
+		Null = 0x01,
+
+		/// <summary>
+		/// Object is an array of fields with unique non-empty names.
+		///
+		/// Payload is a VarUInt byte count for the encoded fields followed by the fields.
+		/// </summary>
+		Object = 0x02,
+
+		/// <summary>
+		/// UniformObject is an array of fields with the same field types and unique non-empty names.
+		/// 
+		/// Payload is a VarUInt byte count for the encoded fields followed by the fields.
+		/// </summary>
+		UniformObject = 0x03,
+
+		/// <summary>
+		/// Array is an array of fields with no name that may be of different types.
+		/// 
+		/// Payload is a VarUInt byte count, followed by a VarUInt item count, followed by the fields.
+		/// </summary>
+		Array = 0x04,
+
+		/// <summary>
+		/// UniformArray is an array of fields with no name and with the same field type.
+		///
+		/// Payload is a VarUInt byte count, followed by a VarUInt item count, followed by field type,
+		/// followed by the fields without their field type.
+		/// </summary>
+		UniformArray = 0x05,
+
+		/// <summary> 
+		/// Binary. Payload is a VarUInt byte count followed by the data. 
+		/// /// </summary>
+		Binary = 0x06,
+
+		/// <summary>
+		/// String in UTF-8. Payload is a VarUInt byte count then an unterminated UTF-8 string. 
+		/// </summary>
+		String = 0x07,
+
+		/// <summary>
+		/// Non-negative integer with the range of a 64-bit unsigned integer.
+		/// 
+		/// Payload is the value encoded as a VarUInt.
+		/// </summary>
+		IntegerPositive = 0x08,
+
+		/// <summary>
+		/// Negative integer with the range of a 64-bit signed integer.
+		///
+		/// Payload is the ones' complement of the value encoded as a VarUInt.
+		/// </summary>
+		IntegerNegative = 0x09,
+
+		/// <summary>
+		/// Single precision float. Payload is one big endian IEEE 754 binary32 float.
+		/// /// </summary>
+		Float32 = 0x0a,
+
+		/// <summary>
+		/// Double precision float. Payload is one big endian IEEE 754 binary64 float. 
+		/// </summary>
+		Float64 = 0x0b,
+
+		/// <summary>
+		/// Boolean false value. Payload is empty. 
+		/// </summary>
+		BoolFalse = 0x0c,
+		
+		/// <summary>
+		/// Boolean true value. Payload is empty. 
+		/// </summary>
+		BoolTrue = 0x0d,
+
+		/// <summary>
+		/// CompactBinaryAttachment is a reference to a compact binary attachment stored externally.
+		///
+		/// Payload is a 160-bit hash digest of the referenced compact binary data.
+		/// </summary>
+		ObjectAttachment = 0x0e,
+
+		/// <summary>
+		/// BinaryAttachment is a reference to a binary attachment stored externally.
+		///
+		/// Payload is a 160-bit hash digest of the referenced binary data.
+		/// </summary>
+		BinaryAttachment = 0x0f,
+
+		/// <summary>
+		/// Hash. Payload is a 160-bit hash digest. 
+		/// </summary>
+		Hash = 0x10,
+
+		/// <summary>
+		/// UUID/GUID. Payload is a 128-bit UUID as defined by RFC 4122. 
+		/// </summary>
+		Uuid = 0x11,
+
+		/// <summary>
+		/// Date and time between 0001-01-01 00:00:00.0000000 and 9999-12-31 23:59:59.9999999.
+		///
+		/// Payload is a big endian int64 count of 100ns ticks since 0001-01-01 00:00:00.0000000.
+		/// </summary>
+		DateTime = 0x12,
+
+		/// <summary>
+		/// Difference between two date/time values.
+		/// 
+		/// Payload is a big endian int64 count of 100ns ticks in the span, and may be negative.
+		/// </summary>
+		TimeSpan = 0x13,
+
+		/// <summary>
+		/// ObjectId is an opaque object identifier. See FCbObjectId.
+		///
+		/// Payload is a 12-byte object identifier.
+		/// </summary>
+		ObjectId = 0x14,
+
+		/// <summary>
+		/// CustomById identifies the sub-type of its payload by an integer identifier.
+		///
+		/// Payload is a VarUInt byte count of the sub-type identifier and the sub-type payload, followed
+		/// by a VarUInt of the sub-type identifier then the payload of the sub-type.
+		/// </summary>
+		CustomById = 0x1e,
+
+		/// <summary>
+		/// CustomByType identifies the sub-type of its payload by a string identifier.
+		///
+		/// Payload is a VarUInt byte count of the sub-type identifier and the sub-type payload, followed
+		/// by a VarUInt byte count of the unterminated sub-type identifier, then the sub-type identifier
+		/// without termination, then the payload of the sub-type.
+		/// </summary>
+		CustomByName = 0x1f,
+
+		/// <summary>
+		/// Reserved for future use as a flag. Do not add types in this range. 
+		/// </summary>
+		Reserved = 0x20,
+
+		/// <summary>
+		/// A transient flag which indicates that the object or array containing this field has stored
+		/// the field type before the payload and name. Non-uniform objects and fields will set this.
+		/// 
+		/// Note: Since the flag must never be serialized, this bit may be repurposed in the future.
+		/// </summary>
+		HasFieldType = 0x40,
+
+		/// <summary>
+		/// A persisted flag which indicates that the field has a name stored before the payload. 
+		/// </summary>
+		HasFieldName = 0x80,
+	}
+
+	/// <summary>
+	/// Methods that operate on <see cref="CbFieldType"/>.
+	/// </summary>
+	public static class CbFieldUtils
+	{
+		private const CbFieldType SerializedTypeMask = (CbFieldType)0b_1001_1111;
+		private const CbFieldType TypeMask = (CbFieldType)0b_0001_1111;
+
+		private const CbFieldType ObjectMask = (CbFieldType)0b_0001_1110;
+		private const CbFieldType ObjectBase = (CbFieldType)0b_0000_0010;
+
+		private const CbFieldType ArrayMask = (CbFieldType)0b_0001_1110;
+		private const CbFieldType ArrayBase = (CbFieldType)0b_0000_0100;
+
+		private const CbFieldType IntegerMask = (CbFieldType)0b_0011_1110;
+		private const CbFieldType IntegerBase = (CbFieldType)0b_0000_1000;
+
+		private const CbFieldType FloatMask = (CbFieldType)0b_0001_1100;
+		private const CbFieldType FloatBase = (CbFieldType)0b_0000_1000;
+
+		private const CbFieldType BoolMask = (CbFieldType)0b_0001_1110;
+		private const CbFieldType BoolBase = (CbFieldType)0b_0000_1100;
+
+		private const CbFieldType AttachmentMask = (CbFieldType)0b_0001_1110;
+		private const CbFieldType AttachmentBase = (CbFieldType)0b_0000_1110;
+
+		/// <summary>
+		/// Removes flags from the given type
+		/// </summary>
+		/// <param name="Type">Type to check</param>
+		/// <returns>Type without flag fields</returns>
+		public static CbFieldType GetType(CbFieldType Type)
+		{
+			return Type & TypeMask;
+		}
+
+		/// <summary>
+		/// Gets the serialized type
+		/// </summary>
+		/// <param name="Type">Type to check</param>
+		/// <returns>Type without flag fields</returns>
+		public static CbFieldType GetSerializedType(CbFieldType Type)
+		{
+			return Type & SerializedTypeMask;
+		}
+
+		/// <summary>
+		/// Tests if the given field has a type
+		/// </summary>
+		/// <param name="Type">Type to check</param>
+		/// <returns>True if the field has a type</returns>
+		public static bool HasFieldType(CbFieldType Type)
+		{
+			return (Type & CbFieldType.HasFieldType) != 0;
+		}
+
+		/// <summary>
+		/// Tests if the given field has a name
+		/// </summary>
+		/// <param name="Type">Type to check</param>
+		/// <returns>True if the field has a name</returns>
+		public static bool HasFieldName(CbFieldType Type)
+		{
+			return (Type & CbFieldType.HasFieldName) != 0;
+		}
+
+		/// <summary>
+		/// Tests if the given field type is none
+		/// </summary>
+		/// <param name="Type">Type to check</param>
+		/// <returns>True if the field is none</returns>
+		public static bool IsNone(CbFieldType Type)
+		{
+			return GetType(Type) == CbFieldType.None;
+		}
+
+		/// <summary>
+		/// Tests if the given field type is a null value
+		/// </summary>
+		/// <param name="Type">Type to check</param>
+		/// <returns>True if the field is a null</returns>
+		public static bool IsNull(CbFieldType Type)
+		{
+			return GetType(Type) == CbFieldType.Null;
+		}
+
+		/// <summary>
+		/// Tests if the given field type is an object
+		/// </summary>
+		/// <param name="Type">Type to check</param>
+		/// <returns>True if the field is an object type</returns>
+		public static bool IsObject(CbFieldType Type)
+		{
+			return (Type & ObjectMask) == ObjectBase;
+		}
+
+		/// <summary>
+		/// Tests if the given field type is an array
+		/// </summary>
+		/// <param name="Type">Type to check</param>
+		/// <returns>True if the field is an array type</returns>
+		public static bool IsArray(CbFieldType Type)
+		{
+			return (Type & ArrayMask) == ArrayBase;
+		}
+
+		/// <summary>
+		/// Tests if the given field type is binary
+		/// </summary>
+		/// <param name="Type">Type to check</param>
+		/// <returns>True if the field is binary</returns>
+		public static bool IsBinary(CbFieldType type)
+		{
+			return GetType(type) == CbFieldType.Binary;
+		}
+
+		/// <summary>
+		/// Tests if the given field type is a string
+		/// </summary>
+		/// <param name="Type">Type to check</param>
+		/// <returns>True if the field is an array type</returns>
+		public static bool IsString(CbFieldType Type)
+		{
+			return GetType(Type) == CbFieldType.String;
+		}
+
+		/// <summary>
+		/// Tests if the given field type is an integer
+		/// </summary>
+		/// <param name="Type">Type to check</param>
+		/// <returns>True if the field is an integer type</returns>
+		public static bool IsInteger(CbFieldType Type)
+		{
+			return (Type & IntegerMask) == IntegerBase;
+		}
+
+		/// <summary>
+		/// Tests if the given field type is a float (or integer, due to implicit conversion)
+		/// </summary>
+		/// <param name="Type">Type to check</param>
+		/// <returns>True if the field is a float type</returns>
+		public static bool IsFloat(CbFieldType Type)
+		{
+			return (Type & FloatMask) == FloatBase;
+		}
+
+		/// <summary>
+		/// Tests if the given field type is a boolean
+		/// </summary>
+		/// <param name="Type">Type to check</param>
+		/// <returns>True if the field is an bool type</returns>
+		public static bool IsBool(CbFieldType Type)
+		{
+			return (Type & BoolMask) == BoolBase;
+		}
+
+		/// <summary>
+		/// Tests if the given field type is a compact binary attachment
+		/// </summary>
+		/// <param name="Type">Type to check</param>
+		/// <returns>True if the field is a compact binary attachment</returns>
+		public static bool IsObjectAttachment(CbFieldType type)
+		{
+			return GetType(type) == CbFieldType.ObjectAttachment;
+		}
+
+		/// <summary>
+		/// Tests if the given field type is a binary attachment
+		/// </summary>
+		/// <param name="Type">Type to check</param>
+		/// <returns>True if the field is a binary attachment</returns>
+		public static bool IsBinaryAttachment(CbFieldType type)
+		{
+			return GetType(type) == CbFieldType.BinaryAttachment;
+		}
+
+		/// <summary>
+		/// Tests if the given field type is an attachment
+		/// </summary>
+		/// <param name="Type">Type to check</param>
+		/// <returns>True if the field is an attachment type</returns>
+		public static bool IsAttachment(CbFieldType Type)
+		{
+			return (Type & AttachmentMask) == AttachmentBase;
+		}
+
+		/// <summary>
+		/// Tests if the given field type is a hash
+		/// </summary>
+		/// <param name="Type">Type to check</param>
+		/// <returns>True if the field is a hash</returns>
+		public static bool IsHash(CbFieldType Type)
+		{
+			return GetType(Type) == CbFieldType.Hash || IsAttachment(Type);
+		}
+
+		/// <summary>
+		/// Tests if the given field type is a UUID
+		/// </summary>
+		/// <param name="Type">Type to check</param>
+		/// <returns>True if the field is a UUID</returns>
+		public static bool IsUuid(CbFieldType type)
+		{
+			return GetType(type) == CbFieldType.Uuid;
+		}
+
+		/// <summary>
+		/// Tests if the given field type is a date/time
+		/// </summary>
+		/// <param name="Type">Type to check</param>
+		/// <returns>True if the field is a date/time</returns>
+		public static bool IsDateTime(CbFieldType type)
+		{
+			return GetType(type) == CbFieldType.DateTime;
+		}
+
+		/// <summary>
+		/// Tests if the given field type is a timespan
+		/// </summary>
+		/// <param name="Type">Type to check</param>
+		/// <returns>True if the field is a timespan</returns>
+		public static bool IsTimeSpan(CbFieldType type)
+		{
+			return GetType(type) == CbFieldType.TimeSpan;
+		}
+
+		/// <summary>
+		/// Tests if the given field type has fields
+		/// </summary>
+		/// <param name="Type">Type to check</param>
+		/// <returns>True if the field has fields</returns>
+		public static bool HasFields(CbFieldType Type)
+		{
+			CbFieldType NoFlags = GetType(Type);
+			return NoFlags >= CbFieldType.Object && NoFlags <= CbFieldType.UniformArray;
+		}
+
+		/// <summary>
+		/// Tests if the given field type has uniform fields (array/object)
+		/// </summary>
+		/// <param name="Type">Type to check</param>
+		/// <returns>True if the field has uniform fields</returns>
+		public static bool HasUniformFields(CbFieldType Type)
+		{
+			CbFieldType LocalType = GetType(Type);
+			return LocalType == CbFieldType.UniformObject || LocalType == CbFieldType.UniformArray;
+		}
+
+		/// <summary>
+		/// Tests if the type is or may contain fields of any attachment type.
+		/// </summary>
+		public static bool MayContainAttachments(CbFieldType Type)
+		{
+			return IsObject(Type) | IsArray(Type) | IsAttachment(Type);
+		}
+	}
+
+	/// <summary>
+	/// Errors that can occur when accessing a field. */
+	/// </summary>
+	public enum CbFieldError : byte
+	{
+		/// <summary>
+		/// The field is not in an error state.
+		/// </summary>
+		None,
+
+		/// <summary>
+		/// The value type does not match the requested type.
+		/// </summary>
+		TypeError,
+
+		/// <summary>
+		/// The value is out of range for the requested type.
+		/// </summary>
+		RangeError,
+	}
+
+	/// <summary>
+	/// An atom of data in the compact binary format.
+	///
+	/// Accessing the value of a field is always a safe operation, even if accessed as the wrong type.
+	/// An invalid access will return a default value for the requested type, and set an error code on
+	/// the field that can be checked with GetLastError and HasLastError. A valid access will clear an
+	/// error from a previous invalid access.
+	///
+	/// A field is encoded in one or more bytes, depending on its type and the type of object or array
+	/// that contains it. A field of an object or array which is non-uniform encodes its field type in
+	/// the first byte, and includes the HasFieldName flag for a field in an object. The field name is
+	/// encoded in a variable-length unsigned integer of its size in bytes, for named fields, followed
+	/// by that many bytes of the UTF-8 encoding of the name with no null terminator.The remainder of
+	/// the field is the payload and is described in the field type enum. Every field must be uniquely
+	/// addressable when encoded, which means a zero-byte field is not permitted, and only arises in a
+	/// uniform array of fields with no payload, where the answer is to encode as a non-uniform array.
+	/// </summary>
+	public class CbField : IEquatable<CbField>, IEnumerable<CbField>
+	{
+		/// <summary>
+		/// Default empty field
+		/// </summary>
+		public static CbField Empty { get; } = new CbField();
+
+		/// <summary>
+		/// The field type, with the transient HasFieldType flag if the field contains its type
+		/// </summary>
+		public CbFieldType TypeWithFlags { get; }
+
+		/// <summary>
+		/// Data for this field
+		/// </summary>
+		public ReadOnlyMemory<byte> Memory { get; }
+
+		/// <summary>
+		/// Offset of the name with the memory
+		/// </summary>
+		public int NameLen;
+
+		/// <summary>
+		/// Offset of the payload within the memory
+		/// </summary>
+		public int PayloadOffset;
+
+		/// <summary>
+		/// Error for parsing the current field type
+		/// </summary>
+		public CbFieldError Error { get; private set; }
+
+		/// <summary>
+		/// Default constructor
+		/// </summary>
+		public CbField()
+			: this(ReadOnlyMemory<byte>.Empty, CbFieldType.None)
+		{
+		}
+
+		/// <summary>
+		/// Copy constructor
+		/// </summary>
+		/// <param name="Other"></param>
+		public CbField(CbField Other)
+		{
+			this.TypeWithFlags = Other.TypeWithFlags;
+			this.Memory = Other.Memory;
+			this.NameLen = Other.NameLen;
+			this.PayloadOffset = Other.PayloadOffset;
+			this.Error = Other.Error;
+		}
+
+		/// <summary>
+		/// Construct a field from a pointer to its data and an optional externally-provided type.
+		/// </summary>
+		/// <param>Data Pointer to the start of the field data.</param>
+		/// <param>Type HasFieldType means that Data contains the type. Otherwise, use the given type.</param>
+		public CbField(ReadOnlyMemory<byte> Data, CbFieldType Type = CbFieldType.HasFieldType)
+		{
+			int Offset = 0;
+			if (CbFieldUtils.HasFieldType(Type))
+			{
+				Type = (CbFieldType)Data.Span[Offset] | CbFieldType.HasFieldType;
+				Offset++;
+			}
+
+			if (CbFieldUtils.HasFieldName(Type))
+			{
+				NameLen = (int)BitUtils.ReadVarUInt(Data.Slice(Offset).Span, out int NameLenByteCount);
+				Offset += NameLenByteCount + NameLen;
+			}
+
+			this.Memory = Data;
+			this.TypeWithFlags = Type;
+			this.PayloadOffset = Offset;
+
+			Memory = Memory.Slice(0, (int)Math.Min((ulong)Memory.Length, (ulong)PayloadOffset + GetPayloadSize()));
+		}
+
+		/// <summary>
+		/// Returns the name of the field if it has a name, otherwise an empty view.
+		/// </summary>
+		public ReadOnlyUtf8String Name => new ReadOnlyUtf8String(Memory.Slice(PayloadOffset - NameLen, NameLen));
+
+		/// <inheritdoc cref="Name"/>
+		public ReadOnlyUtf8String GetName() => Name;
+
+		/// <summary>
+		/// Access the field as an object. Defaults to an empty object on error. 
+		/// </summary>
+		/// <returns></returns>
+		public CbObject AsObject()
+		{
+			if (CbFieldUtils.IsObject(TypeWithFlags))
+			{
+				Error = CbFieldError.None;
+				return CbObject.FromFieldNoCheck(this);
+			}
+			else
+			{
+				Error = CbFieldError.TypeError;
+				return CbObject.Empty;
+			}
+		}
+
+		/// <summary>
+		/// Access the field as an object. Defaults to an empty object on error. 
+		/// </summary>
+		/// <returns></returns>
+		public CbObject AsObjectView() => AsObject();
+
+		/// <summary>
+		/// Access the field as an array. Defaults to an empty array on error. 
+		/// </summary>
+		/// <returns></returns>
+		public CbArray AsArray()
+		{
+			if (CbFieldUtils.IsArray(TypeWithFlags))
+			{
+				Error = CbFieldError.None;
+				return CbArray.FromFieldNoCheck(this);
+			}
+			else
+			{
+				Error = CbFieldError.TypeError;
+				return CbArray.Empty;
+			}
+		}
+
+		/// <summary>
+		/// Access the field as an array. Defaults to an empty array on error. 
+		/// </summary>
+		/// <returns></returns>
+		public CbArray AsArrayView() => AsArray();
+
+		/// <summary>
+		/// Access the field as binary data.
+		/// </summary>
+		/// <returns></returns>
+		public ReadOnlyMemory<byte> AsBinary(ReadOnlyMemory<byte> Default = default)
+		{
+			if (CbFieldUtils.IsBinary(TypeWithFlags))
+			{
+				Error = CbFieldError.None;
+				return BitUtils.GetVarUIntSizedPayload(GetPayloadView());
+			}
+			else
+			{
+				Error = CbFieldError.TypeError;
+				return Default;
+			}
+		}
+
+		/// <summary>
+		/// Access the field as binary data.
+		/// </summary>
+		/// <returns></returns>
+		public ReadOnlyMemory<byte> AsBinaryView(ReadOnlyMemory<byte> Default = default) => AsBinary(Default);
+
+
+		/// <summary>
+		/// Access the field as a UTF-8 string. Returns the provided default on error.
+		/// </summary>
+		/// <param name="Default">Default value to return</param>
+		/// <returns></returns>
+		public ReadOnlyUtf8String AsString(ReadOnlyUtf8String Default = default)
+		{
+			if (CbFieldUtils.IsString(TypeWithFlags))
+			{
+				ulong ValueSize = BitUtils.ReadVarUInt(Payload.Span, out int ValueSizeByteCount);
+				if (ValueSize >= (1UL << 31))
+				{
+					Error = CbFieldError.RangeError;
+					return Default;
+				}
+				else
+				{
+					Error = CbFieldError.None;
+					return new ReadOnlyUtf8String(Payload.Slice(ValueSizeByteCount, (int)ValueSize));
+				}
+			}
+			else
+			{
+				Error = CbFieldError.TypeError;
+				return Default;
+			}
+		}
+
+		/// <summary>
+		/// Access the field as an int8. Returns the provided default on error.
+		/// </summary>
+		public sbyte AsInt8(sbyte Default = 0)
+		{
+			return (sbyte)AsInteger((ulong)Default, 7, true);
+		}
+
+		/// <summary>
+		/// Access the field as an int16. Returns the provided default on error.
+		/// </summary>
+		public short AsInt16(short Default = 0)
+		{
+			return (short)AsInteger((ulong)Default, 15, true);
+		}
+
+		/// <summary>
+		/// Access the field as an int32. Returns the provided default on error.
+		/// </summary>
+		public int AsInt32(int Default = 0)
+		{
+			return (int)AsInteger((ulong)Default, 31, true);
+		}
+
+		/// <summary>
+		/// Access the field as an int64. Returns the provided default on error.
+		/// </summary>
+		public long AsInt64(long Default = 0)
+		{
+			return (long)AsInteger((ulong)Default, 63, true);
+		}
+
+		/// <summary>
+		/// Access the field as an int8. Returns the provided default on error.
+		/// </summary>
+		public byte AsUInt8(byte Default = 0)
+		{
+			return (byte)AsInteger(Default, 8, false);
+		}
+
+		/// <summary>
+		/// Access the field as an int16. Returns the provided default on error.
+		/// </summary>
+		public ushort AsUInt16(ushort Default = 0)
+		{
+			return (ushort)AsInteger(Default, 16, false);
+		}
+
+		/// <summary>
+		/// Access the field as an int32. Returns the provided default on error.
+		/// </summary>
+		public uint AsUInt32(uint Default = 0)
+		{
+			return (uint)AsInteger(Default, 32, false);
+		}
+
+		/// <summary>
+		/// Access the field as an int64. Returns the provided default on error.
+		/// </summary>
+		public ulong AsUInt64(ulong Default = 0)
+		{
+			return (ulong)AsInteger(Default, 64, false);
+		}
+
+		/// <summary>
+		/// Access the field as an integer, checking that it's in the correct range
+		/// </summary>
+		/// <param name="Default"></param>
+		/// <param name="MagnitudeBits"></param>
+		/// <param name="IsSigned"></param>
+		/// <returns></returns>
+		private ulong AsInteger(ulong Default, int MagnitudeBits, bool IsSigned)
+		{
+			if (CbFieldUtils.IsInteger(TypeWithFlags))
+			{
+				// A shift of a 64-bit value by 64 is undefined so shift by one less because magnitude is never zero.
+				ulong OutOfRangeMask = ~(ulong)1 << (MagnitudeBits - 1);
+				ulong IsNegative = (ulong)(byte)(TypeWithFlags) & 1;
+
+				int MagnitudeByteCount;
+				ulong Magnitude = BitUtils.ReadVarUInt(Payload.Span, out MagnitudeByteCount);
+				ulong Value = Magnitude ^ (ulong)-(long)(IsNegative);
+
+				if ((Magnitude & OutOfRangeMask) == 0 && (IsNegative == 0 || IsSigned))
+				{
+					Error = CbFieldError.None;
+					return Value;
+				}
+				else
+				{
+					Error = CbFieldError.RangeError;
+					return Default;
+				}
+			}
+			else
+			{
+				Error = CbFieldError.TypeError;
+				return Default;
+			}
+		}
+
+		/// <summary>
+		/// Access the field as a float. Returns the provided default on error.
+		/// </summary>
+		/// <param name="Default">Default value</param>
+		/// <returns>Value of the field</returns>
+		public float AsFloat(float Default = 0.0f)
+		{
+			switch (GetType())
+			{
+				case CbFieldType.IntegerPositive:
+				case CbFieldType.IntegerNegative:
+					{
+						ulong IsNegative = (ulong)TypeWithFlags & 1;
+						ulong OutOfRangeMask = ~((1UL << /*FLT_MANT_DIG*/ 24) - 1);
+
+						int MagnitudeByteCount;
+						ulong Magnitude = BitUtils.ReadVarUInt(Payload.Span, out MagnitudeByteCount) + IsNegative;
+						bool IsInRange = (Magnitude & OutOfRangeMask) == 0;
+						Error = IsInRange ? CbFieldError.None : CbFieldError.RangeError;
+						return IsInRange ? (float)((IsNegative != 0) ? (float)-(long)Magnitude : (float)Magnitude) : Default;
+					}
+				case CbFieldType.Float32:
+					Error = CbFieldError.None;
+					return BitConverter.Int32BitsToSingle(BinaryPrimitives.ReadInt32BigEndian(Payload.Span));
+				case CbFieldType.Float64:
+					Error = CbFieldError.RangeError;
+					return Default;
+				default:
+					Error = CbFieldError.TypeError;
+					return Default;
+			}
+		}
+
+		/// <summary>
+		/// Access the field as a double. Returns the provided default on error.
+		/// </summary>
+		/// <param name="Default">Default value</param>
+		/// <returns>Value of the field</returns>
+		public double AsDouble(double Default = 0.0f)
+		{
+			switch (GetType())
+			{
+				case CbFieldType.IntegerPositive:
+				case CbFieldType.IntegerNegative:
+					{
+						ulong IsNegative = (ulong)TypeWithFlags & 1;
+						ulong OutOfRangeMask = ~((1UL << /*DBL_MANT_DIG*/ 53) - 1);
+
+						int MagnitudeByteCount;
+						ulong Magnitude = BitUtils.ReadVarUInt(Payload.Span, out MagnitudeByteCount) + IsNegative;
+						bool IsInRange = (Magnitude & OutOfRangeMask) == 0;
+						Error = IsInRange ? CbFieldError.None : CbFieldError.RangeError;
+						return IsInRange ? (double)((IsNegative != 0) ? (double)-(long)Magnitude : (double)Magnitude) : Default;
+					}
+				case CbFieldType.Float32:
+					{
+						Error = CbFieldError.None;
+						return BitConverter.Int32BitsToSingle(BinaryPrimitives.ReadInt32BigEndian(Payload.Span));
+					}
+				case CbFieldType.Float64:
+					{
+						Error = CbFieldError.None;
+						return BitConverter.Int64BitsToDouble(BinaryPrimitives.ReadInt64BigEndian(Payload.Span));
+					}
+				default:
+					Error = CbFieldError.TypeError;
+					return Default;
+			}
+		}
+
+		/// <summary>
+		/// Access the field as a bool. Returns the provided default on error.
+		/// </summary>
+		/// <param name="Default">Default value</param>
+		/// <returns>Value of the field</returns>
+		public bool AsBool(bool Default = false)
+		{
+			switch (GetType())
+			{
+				case CbFieldType.BoolTrue:
+					Error = CbFieldError.None;
+					return true;
+				case CbFieldType.BoolFalse:
+					Error = CbFieldError.None;
+					return false;
+				default:
+					Error = CbFieldError.TypeError;
+					return Default;
+			}
+		}
+
+		/// <summary>
+		/// Access the field as a hash referencing an object attachment. Returns the provided default on error.
+		/// </summary>
+		/// <returns>Value of the field</returns>
+		public IoHash AsObjectAttachment() => AsObjectAttachment(IoHash.Zero);
+
+		/// <summary>
+		/// Access the field as a hash referencing an object attachment. Returns the provided default on error.
+		/// </summary>
+		/// <param name="Default">Default value</param>
+		/// <returns>Value of the field</returns>
+		public IoHash AsObjectAttachment(IoHash Default)
+		{
+			if (CbFieldUtils.IsObjectAttachment(TypeWithFlags))
+			{
+				Error = CbFieldError.None;
+				return new IoHash(Payload);
+			}
+			else
+			{
+				Error = CbFieldError.TypeError;
+				return Default;
+			}
+		}
+
+		/// <summary>
+		/// Access the field as a hash referencing a binary attachment. Returns the provided default on error.
+		/// </summary>
+		/// <param name="Default">Default value</param>
+		/// <returns>Value of the field</returns>
+		public IoHash AsBinaryAttachment() => AsBinaryAttachment(IoHash.Zero);
+
+		/// <summary>
+		/// Access the field as a hash referencing a binary attachment. Returns the provided default on error.
+		/// </summary>
+		/// <param name="Default">Default value</param>
+		/// <returns>Value of the field</returns>
+		public IoHash AsBinaryAttachment(IoHash Default)
+		{
+			if (CbFieldUtils.IsBinaryAttachment(TypeWithFlags))
+			{
+				Error = CbFieldError.None;
+				return new IoHash(Payload);
+			}
+			else
+			{
+				Error = CbFieldError.TypeError;
+				return Default;
+			}
+		}
+
+		/// <summary>
+		/// Access the field as a hash referencing an attachment. Returns the provided default on error.
+		/// </summary>
+		/// <param name="Default">Default value</param>
+		/// <returns>Value of the field</returns>
+		public IoHash AsAttachment() => AsAttachment(IoHash.Zero);
+
+		/// <summary>
+		/// Access the field as a hash referencing an attachment. Returns the provided default on error.
+		/// </summary>
+		/// <param name="Default">Default value</param>
+		/// <returns>Value of the field</returns>
+		public IoHash AsAttachment(IoHash Default)
+		{
+			if (CbFieldUtils.IsAttachment(TypeWithFlags))
+			{
+				Error = CbFieldError.None;
+				return new IoHash(Payload);
+			}
+			else
+			{
+				Error = CbFieldError.TypeError;
+				return Default;
+			}
+		}
+
+		/// <summary>
+		/// Access the field as a hash referencing an attachment. Returns the provided default on error.
+		/// </summary>
+		/// <returns>Value of the field</returns>
+		public IoHash AsHash() => AsHash(IoHash.Zero);
+
+		/// <summary>
+		/// Access the field as a hash referencing an attachment. Returns the provided default on error.
+		/// </summary>
+		/// <param name="Default">Default value</param>
+		/// <returns>Value of the field</returns>
+		public IoHash AsHash(IoHash Default)
+		{
+			if (CbFieldUtils.IsHash(TypeWithFlags))
+			{
+				Error = CbFieldError.None;
+				return new IoHash(Payload);
+			}
+			else
+			{
+				Error = CbFieldError.TypeError;
+				return Default;
+			}
+		}
+
+		/// <summary>
+		/// Access the field as a UUID. Returns a nil UUID on error.
+		/// </summary>
+		/// <param name="Default">Default value</param>
+		/// <returns>Value of the field</returns>
+		public Guid AsUuid(Guid Default = default(Guid))
+		{
+			if (CbFieldUtils.IsUuid(TypeWithFlags))
+			{
+				Error = CbFieldError.None;
+
+				ReadOnlySpan<byte> Span = Payload.Span;
+				uint A = BinaryPrimitives.ReadUInt32BigEndian(Span);
+				ushort B = BinaryPrimitives.ReadUInt16BigEndian(Span.Slice(4));
+				ushort C = BinaryPrimitives.ReadUInt16BigEndian(Span.Slice(6));
+
+				return new Guid(A, B, C, Span[8], Span[9], Span[10], Span[11], Span[12], Span[13], Span[14], Span[15]);
+			}
+			else
+			{
+				Error = CbFieldError.TypeError;
+				return Default;
+			}
+		}
+
+		/// <summary>
+		/// Reads a date time as number of ticks from the stream
+		/// </summary>
+		/// <param name="Default"></param>
+		/// <returns></returns>
+		public long AsDateTimeTicks(long Default = 0)
+		{
+			if (CbFieldUtils.IsDateTime(TypeWithFlags))
+			{
+				Error = CbFieldError.None;
+				return BinaryPrimitives.ReadInt64BigEndian(Payload.Span);
+			}
+			else
+			{
+				Error = CbFieldError.TypeError;
+				return Default;
+			}
+		}
+
+		/// <summary>
+		/// Access the field as a DateTime.
+		/// </summary>
+		/// <param name="Default"></param>
+		/// <returns></returns>
+		public DateTime AsDateTime(DateTime Default = default)
+		{
+			return new DateTime(AsDateTimeTicks(Default.Ticks));
+		}
+
+		/// <summary>
+		/// Reads a timespan as number of ticks from the stream
+		/// </summary>
+		/// <param name="Default"></param>
+		/// <returns></returns>
+		public long AsTimeSpanTicks(long Default = 0)
+		{
+			if (CbFieldUtils.IsTimeSpan(TypeWithFlags))
+			{
+				Error = CbFieldError.None;
+				return BinaryPrimitives.ReadInt64BigEndian(Payload.Span);
+			}
+			else
+			{
+				Error = CbFieldError.TypeError;
+				return Default;
+			}
+		}
+
+		/// <summary>
+		/// Reads a timespan as number of ticks from the stream
+		/// </summary>
+		/// <param name="Default"></param>
+		/// <returns></returns>
+		public TimeSpan AsTimeSpan(TimeSpan Default = default) => new TimeSpan(AsTimeSpanTicks(Default.Ticks));
+
+		/// <inheritdoc cref="CbFieldUtils.HasFieldName(CbFieldType)"/>
+		public bool HasName() => CbFieldUtils.HasFieldName(TypeWithFlags);
+
+		/// <inheritdoc cref="CbFieldUtils.IsNull(CbFieldType)"/>
+		public bool IsNull() => CbFieldUtils.IsNull(TypeWithFlags);
+
+		/// <inheritdoc cref="CbFieldUtils.IsObject(CbFieldType)"/>
+		public bool IsObject() => CbFieldUtils.IsObject(TypeWithFlags);
+
+		/// <inheritdoc cref="CbFieldUtils.IsArray(CbFieldType)"/>
+		public bool IsArray() => CbFieldUtils.IsArray(TypeWithFlags);
+
+		/// <inheritdoc cref="CbFieldUtils.IsBinary(CbFieldType)"/>
+		public bool IsBinary() => CbFieldUtils.IsBinary(TypeWithFlags);
+
+		/// <inheritdoc cref="CbFieldUtils.IsString(CbFieldType)"/>
+		public bool IsString() => CbFieldUtils.IsString(TypeWithFlags);
+
+		/// <inheritdoc cref="CbFieldUtils.IsInteger(CbFieldType)"/>
+		public bool IsInteger() => CbFieldUtils.IsInteger(TypeWithFlags);
+
+		/// <inheritdoc cref="CbFieldUtils.IsFloat(CbFieldType)"/>
+		public bool IsFloat() => CbFieldUtils.IsFloat(TypeWithFlags);
+
+		/// <inheritdoc cref="CbFieldUtils.IsBool(CbFieldType)"/>
+		public bool IsBool() => CbFieldUtils.IsBool(TypeWithFlags);
+
+		/// <inheritdoc cref="CbFieldUtils.IsObjectAttachment(CbFieldType)"/>
+		public bool IsObjectAttachment() => CbFieldUtils.IsObjectAttachment(TypeWithFlags);
+
+		/// <inheritdoc cref="CbFieldUtils.IsBinaryAttachment(CbFieldType)"/>
+		public bool IsBinaryAttachment() => CbFieldUtils.IsBinaryAttachment(TypeWithFlags);
+
+		/// <inheritdoc cref="CbFieldUtils.IsAttachment(CbFieldType)"/>
+		public bool IsAttachment() => CbFieldUtils.IsAttachment(TypeWithFlags);
+
+		/// <inheritdoc cref="CbFieldUtils.IsHash(CbFieldType)"/>
+		public bool IsHash() => CbFieldUtils.IsHash(TypeWithFlags);
+
+		/// <inheritdoc cref="CbFieldUtils.IsUuid(CbFieldType)"/>
+		public bool IsUuid() => CbFieldUtils.IsUuid(TypeWithFlags);
+
+		/// <inheritdoc cref="CbFieldUtils.IsDateTime(CbFieldType)"/>
+		public bool IsDateTime() => CbFieldUtils.IsDateTime(TypeWithFlags);
+
+		/// <inheritdoc cref="CbFieldUtils.IsTimeSpan(CbFieldType)"/>
+		public bool IsTimeSpan() => CbFieldUtils.IsTimeSpan(TypeWithFlags);
+
+		/// <summary>
+		/// Whether the field has a value
+		/// </summary>
+		/// <param name="Field"></param>
+		public static explicit operator bool(CbField Field) => Field.HasValue();
+
+		/// <summary>
+		/// Whether the field has a value.
+		///
+		/// All fields in a valid object or array have a value. A field with no value is returned when
+		/// finding a field by name fails or when accessing an iterator past the end.
+		/// </summary>
+		public bool HasValue() => !CbFieldUtils.IsNone(TypeWithFlags);
+
+		/// <summary>
+		/// Whether the last field access encountered an error.
+		/// </summary>
+		public bool HasError() => Error != CbFieldError.None;
+
+		/// <inheritdoc cref="Error"/>
+		public CbFieldError GetError() => Error;
+
+		/// <summary>
+		/// Returns the size of the field in bytes, including the type and name
+		/// </summary>
+		/// <returns></returns>
+		public int GetSize() => sizeof(CbFieldType) + GetViewNoType().Length;
+
+		/// <summary>
+		/// Calculate the hash of the field, including the type and name.
+		/// </summary>
+		/// <returns></returns>
+		public Blake3Hash GetHash()
+		{
+			Blake3.Hasher Hasher = Blake3.Hasher.New();
+			AppendHash(Hasher);
+
+			byte[] Hash = new byte[32];
+			Hasher.Finalize(Hash);
+
+			return new Blake3Hash(Hash);
+		}
+
+		/// <summary>
+		/// Append the hash of the field, including the type and name
+		/// </summary>
+		/// <param name="Hasher"></param>
+		void AppendHash(Blake3.Hasher Hasher)
+		{
+			byte[] Data = new byte[] { (byte)CbFieldUtils.GetSerializedType(TypeWithFlags) };
+			Hasher.Update(Data);
+			Hasher.Update(GetViewNoType().Span);
+		}
+
+		/// <summary>
+		/// Whether this field is identical to the other field.
+		/// 
+		/// Performs a deep comparison of any contained arrays or objects and their fields. Comparison
+		/// assumes that both fields are valid and are written in the canonical format. Fields must be
+		/// written in the same order in arrays and objects, and name comparison is case sensitive. If
+		/// these assumptions do not hold, this may return false for equivalent inputs. Validation can
+		/// be performed with ValidateCompactBinary, except for field order and field name case.
+		/// </summary>
+		/// <param name="Other"></param>
+		/// <returns></returns>
+		public bool Equals(CbField? Other)
+		{
+			return Other != null && CbFieldUtils.GetSerializedType(TypeWithFlags) == CbFieldUtils.GetSerializedType(Other.TypeWithFlags) && GetViewNoType().Span.SequenceEqual(Other.GetViewNoType().Span);
+		}
+
+		/// <summary>
+		/// Copy the field into a buffer of exactly GetSize() bytes, including the type and name.
+		/// </summary>
+		/// <param name="Buffer"></param>
+		public void CopyTo(Span<byte> Buffer)
+		{
+			Buffer[0] = (byte)CbFieldUtils.GetSerializedType(TypeWithFlags);
+			GetViewNoType().Span.CopyTo(Buffer.Slice(1));
+		}
+
+		/// <summary>
+		/// Invoke the visitor for every attachment in the field.
+		/// </summary>
+		/// <param name="Visitor"></param>
+		public void IterateAttachments(Action<CbField> Visitor)
+		{
+			switch (GetType())
+			{
+				case CbFieldType.Object:
+				case CbFieldType.UniformObject:
+					CbObject.FromFieldNoCheck(this).IterateAttachments(Visitor);
+					break;
+				case CbFieldType.Array:
+				case CbFieldType.UniformArray:
+					CbArray.FromFieldNoCheck(this).IterateAttachments(Visitor);
+					break;
+				case CbFieldType.ObjectAttachment:
+				case CbFieldType.BinaryAttachment:
+					Visitor(this);
+					break;
+			}
+		}
+
+		/// <summary>
+		/// Try to get a view of the field as it would be serialized, such as by CopyTo.
+		///
+		/// A view is available if the field contains its type. Access the equivalent for other fields
+		/// through FCbField::GetBuffer, FCbField::Clone, or CopyTo.
+		/// </summary>
+		/// <param name="OutView"></param>
+		/// <returns></returns>
+		public bool TryGetView(out ReadOnlyMemory<byte> OutView)
+		{
+			if (CbFieldUtils.HasFieldType(TypeWithFlags))
+			{
+				OutView = Memory;
+				return true;
+			}
+			else
+			{
+				OutView = ReadOnlyMemory<byte>.Empty;
+				return false;
+			}
+		}
+
+		/// <summary>
+		/// Find a field of an object by case-sensitive name comparison, otherwise a field with no value.
+		/// </summary>
+		/// <param name="Name"></param>
+		/// <returns></returns>
+		public CbField this[ReadOnlyUtf8String Name]
+		{
+			get { return this.FirstOrDefault(Field => Field.Name == Name) ?? CbField.Empty; }
+		}
+
+		/// <summary>
+		/// Create an iterator for the fields of an array or object, otherwise an empty iterator.
+		/// </summary>
+		/// <returns></returns>
+		public CbFieldIterator CreateIterator() => CreateViewIterator();
+
+		/// <summary>
+		/// Create an iterator for the fields of an array or object, otherwise an empty iterator.
+		/// </summary>
+		/// <returns></returns>
+		public CbFieldIterator CreateViewIterator()
+		{
+			CbFieldType LocalTypeWithFlags = TypeWithFlags;
+			if (CbFieldUtils.HasFields(LocalTypeWithFlags))
+			{
+				ReadOnlyMemory<byte> PayloadBytes = Payload;
+				int PayloadSizeByteCount;
+				int PayloadSize = (int)BitUtils.ReadVarUInt(PayloadBytes.Span, out PayloadSizeByteCount);
+				PayloadBytes = PayloadBytes.Slice(PayloadSizeByteCount);
+				int NumByteCount = CbFieldUtils.IsArray(LocalTypeWithFlags) ? (int)BitUtils.MeasureVarUInt(PayloadBytes.Span) : 0;
+				if (PayloadSize > NumByteCount)
+				{
+					PayloadBytes = PayloadBytes.Slice(NumByteCount);
+
+					CbFieldType UniformType = CbFieldType.HasFieldType;
+					if (CbFieldUtils.HasUniformFields(TypeWithFlags))
+					{
+						UniformType = (CbFieldType)PayloadBytes.Span[0];
+						PayloadBytes = PayloadBytes.Slice(1);
+					}
+
+					return new CbFieldIterator(PayloadBytes, UniformType);
+				}
+			}
+			return new CbFieldIterator(ReadOnlyMemory<byte>.Empty, CbFieldType.HasFieldType);
+		}
+
+		/// <inheritdoc/>
+		public IEnumerator<CbField> GetEnumerator()
+		{
+			for (CbFieldIterator Iter = CreateViewIterator(); Iter; Iter.MoveNext())
+			{
+				yield return Iter.Current;
+			}
+		}
+
+		/// <inheritdoc/>
+		IEnumerator IEnumerable.GetEnumerator() => GetEnumerator();
+
+		/// <summary>
+		/// Returns a view of the field, including the type and name when present.
+		/// </summary>
+		/// <returns></returns>
+		internal ReadOnlyMemory<byte> GetView()
+		{
+			return Memory;
+		}
+
+		/// <summary>
+		/// Returns a view of the name and value payload, which excludes the type.
+		/// </summary>
+		/// <returns></returns>
+		private ReadOnlyMemory<byte> GetViewNoType()
+		{
+			int NameSize = CbFieldUtils.HasFieldName(TypeWithFlags) ? NameLen + (int)BitUtils.MeasureVarUInt((uint)NameLen) : 0;
+			return Memory.Slice(PayloadOffset - NameSize);
+		}
+
+		/// <summary>
+		/// Accessor for the payload
+		/// </summary>
+		internal ReadOnlyMemory<byte> Payload => Memory.Slice(PayloadOffset);
+
+		/// <summary>
+		/// Returns a view of the value payload, which excludes the type and name.
+		/// </summary>
+		/// <returns></returns>
+		internal ReadOnlyMemory<byte> GetPayloadView() => Memory.Slice(PayloadOffset);
+
+		/// <summary>
+		/// Returns the type of the field excluding flags.
+		/// </summary>
+		internal new CbFieldType GetType() => CbFieldUtils.GetType(TypeWithFlags);
+
+		/// <summary>
+		/// Returns the type of the field excluding flags.
+		/// </summary>
+		internal CbFieldType GetTypeWithFlags() => TypeWithFlags;
+
+		/// <summary>
+		/// Returns the size of the value payload in bytes, which is the field excluding the type and name.
+		/// </summary>
+		/// <param name="Type"></param>
+		/// <param name="Payload"></param>
+		/// <returns></returns>
+		public ulong GetPayloadSize()
+		{
+			switch (GetType())
+			{
+				case CbFieldType.None:
+				case CbFieldType.Null:
+					return 0;
+				case CbFieldType.Object:
+				case CbFieldType.UniformObject:
+				case CbFieldType.Array:
+				case CbFieldType.UniformArray:
+				case CbFieldType.Binary:
+				case CbFieldType.String:
+					{
+						ulong PayloadSize = BitUtils.ReadVarUInt(Payload.Span, out int BytesRead);
+						return PayloadSize + (ulong)BytesRead;
+					}
+				case CbFieldType.IntegerPositive:
+				case CbFieldType.IntegerNegative:
+					{
+						return (ulong)BitUtils.MeasureVarUInt(Payload.Span);
+					}
+				case CbFieldType.Float32:
+					return 4;
+				case CbFieldType.Float64:
+					return 8;
+				case CbFieldType.BoolFalse:
+				case CbFieldType.BoolTrue:
+					return 0;
+				case CbFieldType.ObjectAttachment:
+				case CbFieldType.BinaryAttachment:
+				case CbFieldType.Hash:
+					return 20;
+				case CbFieldType.Uuid:
+					return 16;
+				case CbFieldType.DateTime:
+				case CbFieldType.TimeSpan:
+					return 8;
+				case CbFieldType.ObjectId:
+					return 12;
+				default:
+					return 0;
+			}
+		}
+
+		/*
+				/// <summary>
+				/// Enumerate fields within this field. (Originally CreateFieldViewIterator).
+				/// </summary>
+				private IEnumerable<CbField> Fields
+				{
+				   get
+				   {
+					   for (CbFieldIterator Iterator = CreateViewIterator(); Iterator; Iterator.MoveNext())
+					   {
+						   yield return Iterator.Current;
+					   }
+				   }
+			   }
+
+			   public CbFieldIterator CreateIterator() => CreateViewIterator();
+
+			   public static CbField MakeView(CbField Field) => Field;
+			   public CbField Find(ReadOnlyUtf8String Name) => FindView(Name);
+			   public CbField FindIgnoreCase(ReadOnlyUtf8String Name) => FindViewIgnoreCase(Name);
+			   public CbField FindView(ReadOnlyUtf8String Name) => this[Name];
+			   public CbField FindViewIgnoreCase(ReadOnlyUtf8String Name) => Fields.FirstOrDefault(Field => ReadOnlyUtf8StringComparer.OrdinalIgnoreCase.Equals(Field.Name, Name)) ?? new CbField();
+
+			   /// <summary>
+			   /// Fetch a field by index. This operation is O(n); prefer to iterate through fields once.
+			   /// </summary>
+			   /// <param name="Index"></param>
+			   /// <returns></returns>
+			   public CbField? this[int Index]
+			   {
+				   get { return Fields.Skip(Index).FirstOrDefault(); }
+			   }
+		*/
+		#region Mimic inheritance from TCbBufferFactory
+
+		public static CbField Clone(ReadOnlyMemory<byte> Data) => Clone(new CbField(Data));
+		public static CbField Clone(CbField Other) => Other;
+		public static CbField MakeView(ReadOnlyMemory<byte> Data) => new CbField(Data);
+		public static CbField MakeView(CbField Other) => Other;
+
+		#endregion
+	}
+
+	public class CbFieldEnumerator : IEnumerator<CbField>
+	{
+		/// <summary>
+		/// The underlying buffer
+		/// </summary>
+		ReadOnlyMemory<byte> Data;
+
+		/// <summary>
+		/// Type for all fields
+		/// </summary>
+		CbFieldType UniformType { get; }
+
+		/// <inheritdoc/>
+		public CbField Current { get; private set; } = null!;
+
+		/// <inheritdoc/>
+		object? IEnumerator.Current => Current;
+
+		/// <summary>
+		/// Constructor
+		/// </summary>
+		/// <param name="Data"></param>
+		/// <param name="UniformType"></param>
+		public CbFieldEnumerator(ReadOnlyMemory<byte> Data, CbFieldType UniformType)
+		{
+			this.Data = Data;
+			this.UniformType = UniformType;
+		}
+
+		/// <inheritdoc/>
+		public void Dispose()
+		{
+		}
+
+		/// <inheritdoc/>
+		public void Reset()
+		{
+			throw new InvalidOperationException();
+		}
+
+		/// <inheritdoc/>
+		public bool MoveNext()
+		{
+			if (Data.Length > 0)
+			{
+				Current = new CbField(Data, UniformType);
+				return true;
+			}
+			else
+			{
+				Current = null!;
+				return false;
+			}
+		}
+
+		/// <summary>
+		/// Clone this enumerator
+		/// </summary>
+		/// <returns></returns>
+		public CbFieldEnumerator Clone()
+		{
+			return new CbFieldEnumerator(Data, UniformType);
+		}
+	}
+
+	/// <summary>
+	/// Iterator for fields
+	/// </summary>
+	public class CbFieldIterator
+	{
+		/// <summary>
+		/// The underlying buffer
+		/// </summary>
+		ReadOnlyMemory<byte> NextData;
+
+		/// <summary>
+		/// Type for all fields
+		/// </summary>
+		CbFieldType UniformType;
+
+		/// <summary>
+		/// The current iterator
+		/// </summary>
+		public CbField Current { get; private set; } = null!;
+
+		/// <summary>
+		/// Default constructor
+		/// </summary>
+		public CbFieldIterator()
+			: this(ReadOnlyMemory<byte>.Empty, CbFieldType.HasFieldType)
+		{
+		}
+
+		/// <summary>
+		/// Constructor for single field iterator
+		/// </summary>
+		/// <param name="Field"></param>
+		private CbFieldIterator(CbField Field)
+		{
+			NextData = ReadOnlyMemory<byte>.Empty;
+			Current = Field;
+		}
+
+		/// <summary>
+		/// Constructor
+		/// </summary>
+		/// <param name="Data"></param>
+		/// <param name="UniformType"></param>
+		public CbFieldIterator(ReadOnlyMemory<byte> Data, CbFieldType UniformType)
+		{
+			this.NextData = Data;
+			this.UniformType = UniformType;
+
+			MoveNext();
+		}
+
+		/// <summary>
+		/// Copy constructor
+		/// </summary>
+		/// <param name="Other"></param>
+		public CbFieldIterator(CbFieldIterator Other)
+		{
+			this.NextData = Other.NextData;
+			this.UniformType = Other.UniformType;
+			this.Current = Other.Current;
+		}
+
+		/// <summary>
+		/// Construct a field range that contains exactly one field.
+		/// </summary>
+		/// <param name="Field"></param>
+		/// <returns></returns>
+		public static CbFieldIterator MakeSingle(CbField Field)
+		{
+			return new CbFieldIterator(Field);
+		}
+
+		/// <summary>
+		/// Construct a field range from a buffer containing zero or more valid fields.
+		/// </summary>
+		/// <param name="View">A buffer containing zero or more valid fields.</param>
+		/// <param name="Type">HasFieldType means that View contains the type.Otherwise, use the given type.</param>
+		/// <returns></returns>
+		public static CbFieldIterator MakeRange(ReadOnlyMemory<byte> View, CbFieldType Type = CbFieldType.HasFieldType)
+		{
+			return new CbFieldIterator(View, Type);
+		}
+
+		/// <summary>
+		/// Copy the field range into a buffer of exactly GetRangeSize() bytes.
+		/// </summary>
+		/// <param name="Buffer"></param>
+		public void CopyRangeTo(Span<byte> Buffer)
+		{
+			ReadOnlyMemory<byte> Source;
+			if (TryGetRangeView(out Source))
+			{
+				Source.Span.CopyTo(Buffer);
+			}
+			else
+			{
+				for (CbFieldIterator It = new CbFieldIterator(this); It; It.MoveNext())
+				{
+					int Size = It.Current.GetSize();
+					It.Current.CopyTo(Buffer);
+					Buffer = Buffer.Slice(Size);
+				}
+			}
+		}
+
+		/// <summary>
+		/// Invoke the visitor for every attachment in the field range.
+		/// </summary>
+		/// <param name="Visitor"></param>
+		public void IterateRangeAttachments(Action<CbField> Visitor)
+		{
+			// Always iterate over non-uniform ranges because we do not know if they contain an attachment.
+			if (CbFieldUtils.HasFieldType(Current.GetTypeWithFlags()))
+			{
+				for (CbFieldIterator It = new CbFieldIterator(this); It; ++It)
+				{
+					if (CbFieldUtils.MayContainAttachments(It.Current.GetTypeWithFlags()))
+					{
+						It.Current.IterateAttachments(Visitor);
+					}
+				}
+			}
+			// Only iterate over uniform ranges if the uniform type may contain an attachment.
+			else
+			{
+				if (CbFieldUtils.MayContainAttachments(Current.GetTypeWithFlags()))
+				{
+					for (CbFieldIterator It = new CbFieldIterator(this); It; ++It)
+					{
+						It.Current.IterateAttachments(Visitor);
+					}
+				}
+			}
+		}
+
+		/// <summary>
+		/// Try to get a view of every field in the range as they would be serialized.
+		///
+		/// A view is available if each field contains its type. Access the equivalent for other field
+		/// ranges through FCbFieldIterator::CloneRange or CopyRangeTo.
+		/// </summary>
+		/// <returns></returns>
+		bool TryGetRangeView(out ReadOnlyMemory<byte> OutView)
+		{
+			throw new NotImplementedException();
+/*			FMemoryView View;
+			if (FieldType::TryGetView(View))
+			{
+				OutView = MakeMemoryView(View.GetData(), FieldsEnd);
+				return true;
+			}
+			return false;*/
+		}
+
+		/// <summary>
+		/// Move to the next element
+		/// </summary>
+		/// <returns></returns>
+		public bool MoveNext()
+		{
+			if (NextData.Length > 0)
+			{
+				Current = new CbField(NextData, UniformType);
+				NextData = NextData.Slice(Current.Memory.Length);
+				return true;
+			}
+			else
+			{
+				Current = CbField.Empty;
+				return false;
+			}
+		}
+
+		/// <summary>
+		/// Test whether the iterator is valid
+		/// </summary>
+		/// <param name="Iterator"></param>
+		public static implicit operator bool(CbFieldIterator Iterator)
+		{
+			return Iterator.Current.GetType() != CbFieldType.None;
+		}
+
+		/// <summary>
+		/// Move to the next item
+		/// </summary>
+		/// <param name="Iterator"></param>
+		/// <returns></returns>
+		public static CbFieldIterator operator ++(CbFieldIterator Iterator)
+		{
+			return new CbFieldIterator(Iterator.NextData, Iterator.UniformType);
+		}
+
+		public override bool Equals(object? obj)
+		{
+			throw new NotImplementedException();
+		}
+
+		public override int GetHashCode()
+		{
+			throw new NotImplementedException();
+		}
+
+		public static bool operator ==(CbFieldIterator A, CbFieldIterator B)
+		{
+			return A.Current.Equals(B.Current);
+		}
+		public static bool operator !=(CbFieldIterator A, CbFieldIterator B)
+		{
+			return !A.Current.Equals(B.Current);
+		}
+	}
+	/*
+	/// <summary>
+	/// Iterator for fields
+	/// </summary>
+	public class CbFieldIterator : IEnumerable<CbField>
+	{
+		/// <summary>
+		/// The underlying buffer
+		/// </summary>
+		ReadOnlyMemory<byte> Data;
+
+		/// <summary>
+		/// Type for all fields
+		/// </summary>
+		CbFieldType UniformType;
+
+		/// <summary>
+		/// Constructor
+		/// </summary>
+		/// <param name="Data"></param>
+		/// <param name="UniformType"></param>
+		public CbFieldIterator(ReadOnlyMemory<byte> Data, CbFieldType UniformType)
+		{
+			this.Data = Data;
+			this.UniformType = UniformType;
+		}
+
+		/// <inheritdoc/>
+		public IEnumerator<CbField> GetEnumerator()
+		{
+			ReadOnlyMemory<byte> LocalData = Data;
+			while (LocalData.Length > 0)
+			{
+				CbField Field = new CbField(LocalData, UniformType);
+				yield return Field;
+				LocalData = LocalData.Slice(Field.Memory.Length);
+			}
+		}
+
+		/// <inheritdoc/>
+		IEnumerator IEnumerable.GetEnumerator() => GetEnumerator();
+	}
+	*/
+	/// <summary>
+	/// Array of CbField that have no names.
+	///
+	/// Accessing a field of the array requires iteration. Access by index is not provided because the
+	/// cost of accessing an item by index scales linearly with the index.
+	/// </summary>
+	public class CbArray : IEnumerable<CbField>
+	{
+		/// <summary>
+		/// The field containing this array
+		/// </summary>
+		readonly CbField InnerField;
+
+		/// <summary>
+		/// Empty array constant
+		/// </summary>
+		public static CbArray Empty { get; } = new CbArray(new byte[] { (byte)CbFieldType.Array, 1, 0 });
+
+		/// <summary>
+		/// Construct an array with no fields
+		/// </summary>
+		public CbArray()
+		{
+			InnerField = Empty.InnerField;
+		}
+
+		/// <summary>
+		/// Constructor
+		/// </summary>
+		/// <param name="Field"></param>
+		private CbArray(CbField Field)
+		{
+			InnerField = Field;
+		}
+
+		/// <summary>
+		/// Constructor
+		/// </summary>
+		/// <param name="Data"></param>
+		/// <param name="Type"></param>
+		public CbArray(ReadOnlyMemory<byte> Data, CbFieldType Type = CbFieldType.HasFieldType)
+		{
+			InnerField = new CbField(Data, Type);
+		}
+
+		/// <summary>
+		/// Returns the number of items in the array.
+		/// </summary>
+		/// <returns></returns>
+		public int Num()
+		{
+			ReadOnlyMemory<byte> PayloadBytes = InnerField.Payload;
+			PayloadBytes = PayloadBytes.Slice(BitUtils.MeasureVarUInt(PayloadBytes.Span));
+			return (int)BitUtils.ReadVarUInt(PayloadBytes.Span, out int NumByteCount);
+		}
+
+		/// <summary>
+		/// Access the array as an array field.
+		/// </summary>
+		/// <returns></returns>
+		public CbField AsField() => InnerField;
+
+		/// <summary>
+		/// Access the array as an array field.
+		/// </summary>
+		/// <returns></returns>
+		public CbField AsFieldView() => InnerField;
+
+		/// <summary>
+		/// Construct an array from an array field. No type check is performed!
+		/// </summary>
+		/// <param name="Field"></param>
+		/// <returns></returns>
+		public static CbArray FromFieldNoCheck(CbField Field) => new CbArray(Field);
+
+		/// <summary>
+		/// Whether the array has any fields.
+		/// </summary>
+		/// <param name="Array"></param>
+		public static explicit operator bool(CbArray Array) => Array.Num() > 0;
+
+		/// <summary>
+		/// Returns the size of the array in bytes if serialized by itself with no name.
+		/// </summary>
+		/// <returns></returns>
+		public int GetSize()
+		{
+			return (int)Math.Min((ulong)sizeof(CbFieldType) + InnerField.GetPayloadSize(), int.MaxValue);
+		}
+
+		/// <summary>
+		/// Calculate the hash of the array if serialized by itself with no name.
+		/// </summary>
+		/// <returns></returns>
+		public Blake3Hash GetHash()
+		{
+			Blake3.Hasher Hasher = Blake3.Hasher.New();
+			AppendHash(Hasher);
+
+			byte[] Result = new byte[Blake3Hash.NumBytes];
+			Hasher.Finalize(Result);
+
+			return new Blake3Hash(Result);
+		}
+
+		/// <summary>
+		/// Append the hash of the array if serialized by itself with no name.
+		/// </summary>
+		public void AppendHash(Blake3.Hasher Hasher)
+		{
+			byte[] SerializedType = new byte[] { (byte)InnerField.GetType() };
+			Hasher.Update(SerializedType);
+			Hasher.Update(InnerField.Payload.Span);
+		}
+
+		/// <inheritdoc/>
+		public override bool Equals(object? Obj) => Equals(Obj as CbArray);
+
+		/// <inheritdoc/>
+		public override int GetHashCode() => BinaryPrimitives.ReadInt32BigEndian(GetHash().Span);
+
+		/// <summary>
+		/// Whether this array is identical to the other array.
+		///
+		/// Performs a deep comparison of any contained arrays or objects and their fields. Comparison
+		/// assumes that both fields are valid and are written in the canonical format.Fields must be
+		/// written in the same order in arrays and objects, and name comparison is case sensitive.If
+		/// these assumptions do not hold, this may return false for equivalent inputs. Validation can
+		/// be done with the All mode to check these assumptions about the format of the inputs.
+		/// </summary>
+		/// <param name="Other"></param>
+		/// <returns></returns>
+		public bool Equals(CbArray? Other)
+		{
+			return Other != null && GetType() == Other.GetType() && GetPayloadView().Span.SequenceEqual(Other.GetPayloadView().Span);
+		}
+
+		/// <summary>
+		/// Copy the array into a buffer of exactly GetSize() bytes, with no name.
+		/// </summary>
+		/// <param name="Buffer"></param>
+		public void CopyTo(Span<byte> Buffer)
+		{
+			Buffer[0] = (byte)GetType();
+			GetPayloadView().Span.CopyTo(Buffer.Slice(1));
+		}
+
+		/** Invoke the visitor for every attachment in the array. */
+		public void IterateAttachments(Action<CbField> Visitor) => CreateViewIterator().IterateRangeAttachments(Visitor);
+
+		/// <summary>
+		/// Try to get a view of the array as it would be serialized, such as by CopyTo.
+		/// 
+		/// A view is available if the array contains its type and has no name. Access the equivalent
+		/// for other arrays through FCbArray::GetBuffer, FCbArray::Clone, or CopyTo.
+		/// </summary>
+		public bool TryGetView(out ReadOnlyMemory<byte> OutView)
+		{
+			if(InnerField.HasName())
+			{
+				OutView = ReadOnlyMemory<byte>.Empty;
+				return false;
+			}
+			return InnerField.TryGetView(out OutView);
+		}
+
+		/// <inheritdoc cref="CbField.CreateIterator"/>
+		public CbFieldIterator CreateIterator() => InnerField.CreateIterator();
+
+		/// <inheritdoc cref="CbField.CreateViewIterator"/>
+		public CbFieldIterator CreateViewIterator() => InnerField.CreateViewIterator();
+
+		/// <inheritdoc/>
+		public IEnumerator<CbField> GetEnumerator() => InnerField.GetEnumerator();
+
+		/// <inheritdoc/>
+		IEnumerator IEnumerable.GetEnumerator() => GetEnumerator();
+
+		#region Mimic inheritance from CbField
+
+		/// <inheritdoc cref="CbField.GetType"/>
+		internal new CbFieldType GetType() => InnerField.GetType();
+
+		/// <inheritdoc cref="CbField.GetPayloadView"/>
+		internal ReadOnlyMemory<byte> GetPayloadView() => InnerField.GetPayloadView();
+
+		#endregion
+
+		#region Mimic inheritance from TCbBufferFactory
+
+		public static CbArray Clone(ReadOnlyMemory<byte> Data) => Clone(new CbArray(Data));
+		public static CbArray Clone(CbArray Other) => Other;
+		public static CbArray MakeView(ReadOnlyMemory<byte> Data) => new CbArray(Data);
+		public static CbArray MakeView(CbArray Other) => Other;
+
+		#endregion
+	}
+
+	/// <summary>
+	/// Array of CbField that have unique names.
+	///
+	/// Accessing the fields of an object is always a safe operation, even if the requested field does
+	/// not exist. Fields may be accessed by name or through iteration. When a field is requested that
+	/// is not found in the object, the field that it returns has no value (evaluates to false) though
+	/// attempting to access the empty field is also safe, as described by FCbFieldView.
+	/// </summary>
+	public class CbObject : IEnumerable<CbField>
+	{
+		/// <summary>
+		/// Empty array constant
+		/// </summary>
+		public static CbObject Empty = CbObject.FromFieldNoCheck(new CbField(new byte[] { (byte)CbFieldType.Object, 0 }));
+
+		/// <summary>
+		/// The inner field object
+		/// </summary>
+		private CbField InnerField;
+
+		/// <summary>
+		/// Constructor
+		/// </summary>
+		/// <param name="Buffer"></param>
+		private CbObject(CbField Field)
+		{
+			InnerField = new CbField(Field.Memory, Field.TypeWithFlags);
+		}
+
+		/// <summary>
+		/// Constructor
+		/// </summary>
+		/// <param name="Buffer"></param>
+		public CbObject(ReadOnlyMemory<byte> Buffer, CbFieldType FieldType = CbFieldType.HasFieldType)
+		{
+			InnerField = new CbField(Buffer, FieldType);
+		}
+
+		/// <inheritdoc cref="FindView(ReadOnlyUtf8String)"/>
+		public CbField Find(ReadOnlyUtf8String Name) => FindView(Name);
+
+		/// <summary>
+		/// Find a field by case-sensitive name comparison.
+		///
+		/// The cost of this operation scales linearly with the number of fields in the object. Prefer to
+		/// iterate over the fields only once when consuming an object.
+		/// </summary>
+		/// <param name="Name">The name of the field.</param>
+		/// <returns>The matching field if found, otherwise a field with no value.</returns>
+		public CbField FindView(ReadOnlyUtf8String Name) => InnerField[Name];
+
+		/// <inheritdoc cref="FindView(ReadOnlyUtf8String)"/>
+		public CbField FindIgnoreCase(ReadOnlyUtf8String Name) => FindViewIgnoreCase(Name);
+
+		/// <summary>
+		/// Find a field by case-insensitive name comparison.
+		/// </summary>
+		/// <param name="Name">The name of the field.</param>
+		/// <returns>The matching field if found, otherwise a field with no value.</returns>
+		public CbField FindViewIgnoreCase(ReadOnlyUtf8String Name) => InnerField.FirstOrDefault(Field => ReadOnlyUtf8StringComparer.OrdinalIgnoreCase.Equals(Field.Name, Name)) ?? new CbField();
+
+		/// <summary>
+		/// Find a field by case-sensitive name comparison.
+		/// </summary>
+		/// <param name="Name">The name of the field.</param>
+		/// <returns>The matching field if found, otherwise a field with no value.</returns>
+		public CbField this[ReadOnlyUtf8String Name] => InnerField[Name];
+
+		/// <inheritdoc cref="AsFieldView"/>
+		public CbField AsField() => InnerField;
+
+		/// <summary>
+		/// Access the object as an object field.
+		/// </summary>
+		/// <returns></returns>
+		public CbField AsFieldView() => InnerField;
+
+		/// <summary>
+		/// Construct an object from an object field. No type check is performed!
+		/// </summary>
+		/// <param name="Field"></param>
+		/// <returns></returns>
+		public static CbObject FromFieldNoCheck(CbField Field) => new CbObject(Field);
+
+		/// <summary>
+		/// Whether the object has any fields.
+		/// </summary>
+		/// <param name="Object"></param>
+		public static explicit operator bool(CbObject Object)
+		{
+			throw new NotImplementedException();
+		}
+
+		/// <summary>
+		/// Returns the size of the object in bytes if serialized by itself with no name.
+		/// </summary>
+		/// <returns></returns>
+		public int GetSize()
+		{
+			return sizeof(CbFieldType) + InnerField.Payload.Length;
+		}
+
+		/// <summary>
+		/// Calculate the hash of the object if serialized by itself with no name.
+		/// </summary>
+		/// <returns></returns>
+		public Blake3Hash GetHash()
+		{
+			Blake3.Hasher Hasher = Blake3.Hasher.New();
+			AppendHash(Hasher);
+
+			byte[] Data = new byte[Blake3Hash.NumBytes];
+			Hasher.Finalize(Data);
+
+			return new Blake3Hash(Data);
+		}
+
+		/// <summary>
+		/// Append the hash of the object if serialized by itself with no name.
+		/// </summary>
+		/// <param name="Hasher"></param>
+		public void AppendHash(Blake3.Hasher Hasher)
+		{
+			byte[] Temp = new byte[] { (byte)InnerField.GetType() };
+			Hasher.Update(Temp);
+			Hasher.Update(InnerField.Payload.Span);
+		}
+
+		/// <inheritdoc/>
+		public override bool Equals(object? Obj) => Equals(Obj as CbObject);
+
+		/// <inheritdoc/>
+		public override int GetHashCode() => BinaryPrimitives.ReadInt32BigEndian(GetHash().Span);
+
+		/// <summary>
+		/// Whether this object is identical to the other object.
+		/// 
+		/// Performs a deep comparison of any contained arrays or objects and their fields. Comparison
+		/// assumes that both fields are valid and are written in the canonical format. Fields must be
+		/// written in the same order in arrays and objects, and name comparison is case sensitive. If
+		/// these assumptions do not hold, this may return false for equivalent inputs. Validation can
+		/// be done with the All mode to check these assumptions about the format of the inputs.
+		/// </summary>
+		/// <param name="Other"></param>
+		/// <returns></returns>
+		public bool Equals(CbObject? Other)
+		{
+			return Other != null && InnerField.GetType() == Other.InnerField.GetType() && InnerField.Payload.Span.SequenceEqual(Other.InnerField.Payload.Span);
+		}
+
+		/// <summary>
+		/// Copy the object into a buffer of exactly GetSize() bytes, with no name.
+		/// </summary>
+		/// <param name="Buffer"></param>
+		public void CopyTo(Span<byte> Buffer)
+		{
+			Buffer[0] = (byte)InnerField.GetType();
+			InnerField.Payload.Span.CopyTo(Buffer.Slice(1));
+		}
+
+		/// <summary>
+		/// Invoke the visitor for every attachment in the object.
+		/// </summary>
+		/// <param name="Visitor"></param>
+		public void IterateAttachments(Action<CbField> Visitor) => CreateViewIterator().IterateRangeAttachments(Visitor);
+
+		/// <summary>
+		/// Try to get a view of the object as it would be serialized, such as by CopyTo.
+		/// 
+		/// A view is available if the object contains its type and has no name. Access the equivalent
+		/// for other objects through FCbObject::GetBuffer, FCbObject::Clone, or CopyTo.
+		/// </summary>
+		/// <param name="OutView"></param>
+		/// <returns></returns>
+		public bool TryGetView(out ReadOnlyMemory<byte> OutView)
+		{
+			if (InnerField.HasName())
+			{
+				OutView = ReadOnlyMemory<byte>.Empty;
+				return false;
+			}
+			return InnerField.TryGetView(out OutView);
+		}
+
+		/// <inheritdoc cref="CbField.CreateIterator"/>
+		public CbFieldIterator CreateIterator() => InnerField.CreateIterator();
+
+		/// <inheritdoc cref="CbField.CreateViewIterator"/>
+		public CbFieldIterator CreateViewIterator() => InnerField.CreateViewIterator();
+
+		/// <inheritdoc/>
+		public IEnumerator<CbField> GetEnumerator() => InnerField.GetEnumerator();
+
+		/// <inheritdoc/>
+		IEnumerator IEnumerable.GetEnumerator() => InnerField.GetEnumerator();
+
+		/// <summary>
+		/// Clone this object
+		/// </summary>
+		/// <param name="Object"></param>
+		/// <returns></returns>
+		public static CbObject Clone(CbObject Object) => Object;
+
+		#region Conversion to Json
+		/// <summary>
+		/// Convert this object to JSON
+		/// </summary>
+		/// <returns></returns>
+		public string ToJson()
+		{
+			ArrayBufferWriter<byte> Buffer = new ArrayBufferWriter<byte>();
+			using (Utf8JsonWriter JsonWriter = new Utf8JsonWriter(Buffer))
+			{
+				ToJson(JsonWriter);
+			}
+			return Encoding.UTF8.GetString(Buffer.WrittenMemory.Span);
+		}
+
+		/// <summary>
+		/// Write this object to JSON
+		/// </summary>
+		/// <param name="Writer"></param>
+		public void ToJson(Utf8JsonWriter Writer)
+		{
+			Writer.WriteStartObject();
+			foreach (CbField Field in InnerField)
+			{
+				WriteField(Field, Writer);
+			}
+			Writer.WriteEndObject();
+		}
+
+		/// <summary>
+		/// Write a single field to a writer
+		/// </summary>
+		/// <param name="Field"></param>
+		/// <param name="Writer"></param>
+		private static void WriteField(CbField Field, Utf8JsonWriter Writer)
+		{
+			if (Field.IsObject())
+			{
+				Writer.WriteStartObject();
+				CbObject Object = Field.AsObject();
+				foreach (CbField ObjectField in Object.InnerField)
+				{
+					WriteField(ObjectField, Writer);
+				}
+				Writer.WriteEndObject();
+			}
+			else if (Field.IsArray())
+			{
+				Writer.WriteStartArray();
+				Writer.WriteEndArray();
+			}
+			else if (Field.IsInteger())
+			{
+				if (Field.GetType() == CbFieldType.IntegerNegative)
+				{
+					Writer.WriteNumber(Field.Name.Span, -Field.AsInt64());
+				}
+				else
+				{
+					Writer.WriteNumber(Field.Name.Span, Field.AsUInt64());
+				}
+			}
+			else if (Field.IsBool())
+			{
+				Writer.WriteBoolean(Field.Name.Span, Field.AsBool());
+			}
+			else if (Field.IsNull())
+			{
+				Writer.WriteNullValue();
+			}
+			else if (Field.IsDateTime())
+			{
+				Writer.WriteString(Field.Name.Span, Field.AsDateTime());
+			}
+			else if (Field.IsHash())
+			{
+				Writer.WriteString(Field.Name.Span, StringUtils.FormatUtf8HexString(Field.AsHash().Span).Span);
+			}
+			else if (Field.IsString())
+			{
+				Writer.WriteString(Field.Name.Span, Field.AsString().Span);
+			}
+			else
+			{
+				throw new NotImplementedException($"Unhandled type {Field.GetType()} when attempting to convert to json");
+			}
+		}
+		#endregion
+	}
+
+	/// <summary>
+	/// Methods for reading VarUInt values
+	/// </summary>
+	public static class BitUtils
+	{
+		/// <summary>
+		/// Extracts the payload from a buffer prefixed with a variable-length value
+		/// </summary>
+		/// <param name="Memory"></param>
+		/// <returns></returns>
+		public static ReadOnlyMemory<byte> GetVarUIntSizedPayload(ReadOnlyMemory<byte> Memory)
+		{
+			ulong Length = ReadVarUInt(Memory.Span, out int BytesRead);
+			return Memory.Slice(BytesRead, (int)Length);
+		}
+
+		/// <summary>
+		/// Read a variable-length unsigned integer.
+		/// </summary>
+		/// <param name="Buffer">A variable-length encoding of an unsigned integer</param>
+		/// <param name="BytesRead">The number of bytes consumed from the input</param>
+		/// <returns></returns>
+		public static ulong ReadVarUInt(ReadOnlySpan<byte> Buffer, out int BytesRead)
+		{
+			BytesRead = (int)MeasureVarUInt(Buffer);
+
+			ulong Value = (ulong)(Buffer[0] & (0xff >> BytesRead));
+			for (int i = 1; i < BytesRead; i++)
+			{
+				Value <<= 8;
+				Value |= Buffer[i];
+			}
+			return Value;
+		}
+
+		/// <summary>
+		/// Measure the length in bytes (1-9) of an encoded variable-length integer.
+		/// </summary>
+		/// <param name="Buffer">A variable-length encoding of an(signed or unsigned) integer.</param>
+		/// <returns>The number of bytes used to encode the integer, in the range 1-9.</returns>
+		public static int MeasureVarUInt(ReadOnlySpan<byte> Buffer)
+		{
+			byte b = Buffer[0];
+			b = (byte)~b;
+			return BitOperations.LeadingZeroCount(b) - 23;
+		}
+
+		/// <summary>
+		/// Measure the number of bytes (1-5) required to encode the 32-bit input.
+		/// </summary>
+		/// <param name="Value"></param>
+		/// <returns></returns>
+		public static int MeasureVarUInt(int Value)
+		{
+			return MeasureVarUInt((uint)Value);
+		}
+
+		/// <summary>
+		/// Measure the number of bytes (1-5) required to encode the 32-bit input.
+		/// </summary>
+		/// <param name="Value"></param>
+		/// <returns></returns>
+		public static int MeasureVarUInt(uint Value)
+		{
+			return BitOperations.Log2(Value) / 7 + 1;
+		}
+
+		/// <summary>
+		/// Measure the number of bytes (1-9) required to encode the 64-bit input.
+		/// </summary>
+		/// <param name="value"></param>
+		/// <returns></returns>
+		public static int MeasureVarUInt(ulong value)
+		{
+			return Math.Min(BitOperations.Log2(value) / 7 + 1, 9);
+		}
+
+		/// <summary>
+		/// Write a variable-length unsigned integer.
+		/// </summary>
+		/// <param name="Value">An unsigned integer to encode</param>
+		/// <param name="Buffer">A buffer of at least 9 bytes to write the output to.</param>
+		/// <param name="BufferOffset"></param>
+		/// <returns>The number of bytes used in the output</returns>
+		public static int WriteVarUInt(long Value, byte[] Buffer, int BufferOffset = 0)
+		{
+			return WriteVarUInt((ulong)Value, Buffer, BufferOffset);
+		}
+
+		/// <summary>
+		/// Write a variable-length unsigned integer.
+		/// </summary>
+		/// <param name="Value">An unsigned integer to encode</param>
+		/// <param name="Buffer">A buffer of at least 9 bytes to write the output to.</param>
+		/// <param name="BufferOffset"></param>
+		/// <returns>The number of bytes used in the output</returns>
+		public static int WriteVarUInt(ulong Value, byte[] Buffer, int BufferOffset = 0)
+		{
+			int ByteCount = MeasureVarUInt(Value);
+
+			for (uint Idx = 1; Idx < ByteCount; Idx++)
+			{
+				Buffer[BufferOffset + ByteCount - Idx] = (byte)Value;
+				Value >>= 8;
+			}
+			Buffer[BufferOffset] = (byte)((0xff << (9 - (int)ByteCount)) | (byte)Value);
+			return ByteCount;
+		}
+	}
+}
diff --git a/Engine/Source/Programs/Shared/EpicGames.Serialization/CompactBinaryWriter.cs b/Engine/Source/Programs/Shared/EpicGames.Serialization/CompactBinaryWriter.cs
new file mode 100644
--- /dev/null
+++ b/Engine/Source/Programs/Shared/EpicGames.Serialization/CompactBinaryWriter.cs
@@ -0,0 +1,313 @@
+// Copyright Epic Games, Inc. All Rights Reserved.
+
+using EpicGames.Core;
+using System;
+using System.Collections.Generic;
+using System.Linq;
+using System.Text;
+
+namespace EpicGames.Serialization
+{
+	public class CbWriter
+	{
+		readonly List<Field> _fields = new List<Field>();
+		private Field _topField;
+
+		struct Field
+		{
+			public CbFieldType FieldType;
+			public byte[]? Name;
+			public byte[]? Payload;
+		}
+
+		public byte[] Save()
+		{
+			int bufferLength = 1; // field type;
+
+			if (_topField.Name != null)
+			{
+				bufferLength += _topField.Name!.Length;
+			}
+
+			int payloadBytesLength = 0;
+			if (_topField.Payload != null)
+			{
+				payloadBytesLength = BitUtils.MeasureVarUInt((uint)_topField.Payload!.Length);
+				bufferLength += payloadBytesLength;
+				bufferLength += _topField.Payload!.Length;
+			}
+			byte[] buffer = new byte[bufferLength];
+			buffer[0] = (byte)_topField.FieldType;
+			int offset = 1;
+
+			if (_topField.Name != null)
+			{
+				Array.Copy(_topField.Name!, 0, buffer, offset, _topField.Name!.Length);
+				offset += _topField.Name.Length;
+			}
+
+			if (_topField.Payload != null)
+			{
+				BitUtils.WriteVarUInt((uint)_topField.Payload!.Length, buffer, offset);
+				offset += (int)payloadBytesLength;
+				Array.Copy(_topField.Payload!, 0, buffer, offset, _topField.Payload!.Length);
+				offset += _topField.Payload.Length;
+			}
+
+			return buffer;
+		}
+
+		public void BeginObject()
+		{
+			_topField = BeginField();
+		}
+
+		public void EndObject()
+		{
+			_topField.Payload = SerializePayload(_fields);
+
+			EndField(ref this._topField, CbFieldType.Object);
+		}
+
+		private static byte[] SerializePayload(List<Field> fields)
+		{
+			uint totalSizeFields = 0;
+			foreach (Field field in fields)
+			{
+				totalSizeFields += 1; // the field type
+				if (field.Payload != null)
+					totalSizeFields += (uint)field.Payload.Length;
+				if (field.Name != null)
+					totalSizeFields += (uint)field.Name.Length;
+			}
+
+			byte[] buffer = new byte[totalSizeFields];
+			long payloadOffset = 0;
+			foreach (Field field in fields)
+			{
+				buffer[payloadOffset] = (byte)field.FieldType;
+				payloadOffset += 1;
+
+				if (CbFieldUtils.HasFieldName(field.FieldType))
+				{
+					Array.Copy(field.Name!, 0, buffer, payloadOffset, field.Name!.Length);
+					payloadOffset += field.Name.Length;
+				}
+
+				if ((field.Payload?.Length ?? 0) != 0)
+				{
+					Array.Copy(field.Payload!, 0, buffer, payloadOffset, field.Payload!.Length);
+					payloadOffset += field.Payload.Length;
+				}
+			}
+
+			return buffer;
+		}
+
+		public void AddString(string value, string? fieldName = null)
+		{
+			Field field = BeginField();
+			SetName(ref field, fieldName);
+
+			int length = Encoding.UTF8.GetByteCount(value);
+			int bytesCount = BitUtils.MeasureVarUInt(length);
+			byte[] buffer = new byte[length + bytesCount];
+			BitUtils.WriteVarUInt(length, buffer);
+			int bytesWritten = Encoding.UTF8.GetBytes(value, 0, value.Length, buffer, (int)bytesCount);
+			if (bytesWritten != length)
+				throw new Exception("Failed to write string into buffer");
+
+			field.Payload = buffer;
+			EndField(ref field, CbFieldType.String);
+		}
+
+
+		private Field BeginField()
+		{
+			// TODO: We could add checks similar to what we do in C++ to make sure the writer is used correctly
+			return new Field();
+		}
+
+		private void EndField(ref Field field, CbFieldType fieldType)
+		{
+			field.FieldType = fieldType;
+			if (field.Name != null && field.Name.Length != 0)
+			{
+				field.FieldType |= CbFieldType.HasFieldName;
+			}
+			// TODO: check if types are uniform and check the field type to be uniform thus removing redudant type info
+
+			_fields.Add(field);
+		}
+
+		public void AddBinaryAttachment(IoHash hash, string? fieldName = null)
+		{
+			Field field = BeginField();
+			SetName(ref field, fieldName);
+
+			if (hash.Memory.Length != 20)
+				throw new Exception("Hash data is assumed to be 20 bytes");
+			field.Payload = hash.Memory.ToArray();
+
+			EndField(ref field, CbFieldType.BinaryAttachment);
+		}
+
+		public void AddCompactBinaryAttachment(IoHash hash, string? fieldName = null)
+		{
+			Field field = BeginField();
+			SetName(ref field, fieldName);
+
+			if (hash.Memory.Length != 20)
+				throw new Exception("Hash data is assumed to be 20 bytes");
+			field.Payload = hash.Memory.ToArray();
+
+			EndField(ref field, CbFieldType.ObjectAttachment);
+		}
+
+		public void AddNull(string? fieldName = null)
+		{
+			Field field = BeginField();
+			SetName(ref field, fieldName);
+			EndField(ref field, CbFieldType.Null);
+		}
+
+
+		public void AddDateTime(DateTime dateTime, string? fieldName = null)
+		{
+			Field field = BeginField();
+			SetName(ref field, fieldName);
+			long ticks = dateTime.Ticks;
+			byte[] payload = BitConverter.GetBytes(ticks);
+			if (BitConverter.IsLittleEndian)
+			{
+				Array.Reverse(payload, 0, payload.Length);
+			}
+			field.Payload = payload;
+
+			EndField(ref field, CbFieldType.DateTime);
+		}
+
+
+		public void AddBinary(byte[] data, string? fieldName = null)
+		{
+			Field field = BeginField();
+			SetName(ref field, fieldName);
+
+			long bytesLength = BitUtils.MeasureVarUInt((uint)data.Length);
+			byte[] payload = new byte[data.Length + bytesLength];
+			BitUtils.WriteVarUInt((uint)data.Length, payload);
+			Array.Copy(data, 0, payload, bytesLength, data.Length);
+			field.Payload = payload;
+			EndField(ref field, CbFieldType.Binary);
+		}
+
+		public void AddUniformArray<T>(T[] members, CbFieldType fieldType, string? fieldName = null)
+		{
+			Field field = BeginField();
+			SetName(ref field, fieldName);
+
+			int bytesMembers = BitUtils.MeasureVarUInt((uint)members.Length);
+
+			List<byte[]> payloads = new List<byte[]>();
+
+			foreach (T member in members)
+			{
+				byte[] buffer = GetByteValue<T>(member, fieldType);
+				payloads.Add(buffer);
+			}
+
+			int payloadsBytesLength = payloads.Sum(payload => payload.Length);
+			int payloadsLength = 1 + bytesMembers + payloadsBytesLength;
+			int bytesPayloads = BitUtils.MeasureVarUInt(payloadsLength);
+			
+			// Uniform array payload is a VarUInt byte count, followed by a VarUInt item count, followed by field type, followed by the fields without their field type.
+			byte[] payloadBuffer = new byte[payloadsLength + bytesPayloads];
+			BitUtils.WriteVarUInt(payloadsLength, payloadBuffer);
+			int bufferOffset = (int)bytesPayloads;
+			BitUtils.WriteVarUInt((uint)members.Length, payloadBuffer, bufferOffset);
+			bufferOffset += (int)bytesMembers;
+			payloadBuffer[bufferOffset] = (byte)fieldType;
+			bufferOffset += 1;
+
+			foreach (byte[] payload in payloads)
+			{
+				int payloadLength = payload.Length;
+				Array.Copy(payload, 0, payloadBuffer, bufferOffset, payloadLength);
+				bufferOffset += payloadLength;
+			}
+
+			field.Payload = payloadBuffer;
+			EndField(ref field, CbFieldType.UniformArray);
+		}
+
+		private byte[] GetByteValue<T>(T member, CbFieldType fieldType)
+		{
+			switch (fieldType)
+			{
+				case CbFieldType.None:
+				case CbFieldType.Null:
+					return Array.Empty<byte>();
+				case CbFieldType.Object:
+				case CbFieldType.UniformObject:
+				case CbFieldType.Array:
+				case CbFieldType.UniformArray:
+					throw new NotImplementedException();
+				case CbFieldType.Binary:
+					throw new NotImplementedException();
+				case CbFieldType.String:
+					if (member is string s)
+					{
+						return Encoding.UTF8.GetBytes(s);
+					}
+					break;
+				case CbFieldType.IntegerPositive:
+				case CbFieldType.IntegerNegative:
+					throw new NotImplementedException();
+				case CbFieldType.Float32:
+				case CbFieldType.Float64:
+					throw new NotImplementedException();
+				case CbFieldType.BoolFalse:
+				case CbFieldType.BoolTrue:
+					return Array.Empty<byte>();
+				case CbFieldType.ObjectAttachment:
+				case CbFieldType.BinaryAttachment:
+				case CbFieldType.Hash:
+					if (member is IoHash blob)
+					{
+						return blob.Memory.ToArray();
+					}
+
+					break;
+				case CbFieldType.Uuid:
+				case CbFieldType.DateTime:
+				case CbFieldType.TimeSpan:
+				case CbFieldType.ObjectId:
+				case CbFieldType.CustomById:
+				case CbFieldType.CustomByName:
+					throw new NotImplementedException();
+				default:
+					throw new ArgumentOutOfRangeException(nameof(fieldType), fieldType, null);
+			}
+			throw new ArgumentException($"{member} of type {typeof(T).Name} not convert-able to field-type {fieldType}", nameof(member));
+		}
+
+		private void SetName(ref Field field, string? fieldName)
+		{
+			if (string.IsNullOrEmpty(fieldName))
+				return;
+
+			// TODO: We could add checks similar to what we do in C++ to make sure the writer is used correctly
+
+			int nameLength = fieldName.Length;
+			int nameBytesCount = BitUtils.MeasureVarUInt(nameLength);
+
+			byte[] buffer = new byte[nameLength + nameBytesCount];
+			BitUtils.WriteVarUInt(nameLength, buffer);
+			int bytesWritten = Encoding.ASCII.GetBytes(fieldName, 0, fieldName.Length, buffer, (int)nameBytesCount);
+			if (bytesWritten != nameLength)
+				throw new Exception("Failed to write name into buffer");
+
+			field.Name = buffer;
+		}
+	}
+}
\ No newline at end of file
diff --git a/Engine/Source/Programs/Shared/EpicGames.Serialization/EpicGames.Serialization.csproj b/Engine/Source/Programs/Shared/EpicGames.Serialization/EpicGames.Serialization.csproj
new file mode 100644
--- /dev/null
+++ b/Engine/Source/Programs/Shared/EpicGames.Serialization/EpicGames.Serialization.csproj
@@ -0,0 +1,32 @@
+<Project Sdk="Microsoft.NET.Sdk">
+  <PropertyGroup>
+    <TargetFramework>netcoreapp3.1</TargetFramework>
+    <Configuration Condition=" '$(Configuration)' == '' ">Development</Configuration>
+    <OutputType>Library</OutputType>
+    <Deterministic>true</Deterministic>
+    <RootNamespace>EpicGames.Serialization</RootNamespace>
+    <GenerateAssemblyInfo>false</GenerateAssemblyInfo>
+    <GenerateTargetFrameworkAttribute>false</GenerateTargetFrameworkAttribute>
+    <Configurations>Debug;Release;Development</Configurations>
+    <Nullable>enable</Nullable>
+  </PropertyGroup>
+  <PropertyGroup Condition=" '$(Configuration)|$(Platform)' == 'Debug|AnyCPU' ">
+    <DocumentationFile></DocumentationFile>
+  </PropertyGroup>
+  <PropertyGroup Condition="'$(Configuration)|$(Platform)' == 'Development|AnyCPU'">
+    <DocumentationFile></DocumentationFile>
+    <Optimize>true</Optimize>
+  </PropertyGroup>
+  <PropertyGroup Condition="'$(Configuration)|$(Platform)' == 'Release|AnyCPU'">
+    <DocumentationFile></DocumentationFile>
+    <Optimize>true</Optimize>
+  </PropertyGroup>
+  <ItemGroup>
+    <Compile Include="..\..\DotNETCommon\MetaData.cs">
+      <Link>Properties\MetaData.cs</Link>
+    </Compile>
+  </ItemGroup>
+  <ItemGroup>
+    <ProjectReference Include="..\EpicGames.IoHash\EpicGames.IoHash.csproj" />
+  </ItemGroup>
+</Project>
\ No newline at end of file
