Missing file for CB serialization.

diff --git a/Engine/Source/Programs/Shared/EpicGames.Core/StringUtils.cs b/Engine/Source/Programs/Shared/EpicGames.Core/StringUtils.cs
--- a/Engine/Source/Programs/Shared/EpicGames.Core/StringUtils.cs
+++ b/Engine/Source/Programs/Shared/EpicGames.Core/StringUtils.cs
@@ -16,6 +16,11 @@ namespace EpicGames.Core
 		/// </summary>
 		static sbyte[] HexDigits;
 
+		/// <summary>
+		/// Hex digits to utf8 byte
+		/// </summary>
+		static byte[] HexDigitToUtf8Byte = Encoding.UTF8.GetBytes("0123456789abcdef");
+
 		/// <summary>
 		/// Array mapping human readable size of bytes, 1024^x. long max is within the range of Exabytes.
 		/// </summary>
@@ -400,6 +405,22 @@ namespace EpicGames.Core
 			return new string(Characters);
 		}
 
+		/// <summary>
+		/// Formats an array of bytes as a hexadecimal string
+		/// </summary>
+		/// <param name="Bytes">An array of bytes</param>
+		/// <returns>String representation of the array</returns>
+		public static ReadOnlyUtf8String FormatUtf8HexString(ReadOnlySpan<byte> Bytes)
+		{
+			byte[] Characters = new byte[Bytes.Length * 2];
+			for (int Idx = 0; Idx < Bytes.Length; Idx++)
+			{
+				Characters[Idx * 2 + 0] = HexDigitToUtf8Byte[Bytes[Idx] >> 4];
+				Characters[Idx * 2 + 1] = HexDigitToUtf8Byte[Bytes[Idx] & 15];
+			}
+			return new ReadOnlyUtf8String(Characters);
+		}
+
 		/// <summary>
 		/// Quotes a string as a command line argument
 		/// </summary>
