Horde: Add TypeConverter for IoHash objects, to allow passing them via MVC routes.

diff --git a/Engine/Source/Programs/Shared/EpicGames.IoHash/IoHash.cs b/Engine/Source/Programs/Shared/EpicGames.IoHash/IoHash.cs
--- a/Engine/Source/Programs/Shared/EpicGames.IoHash/IoHash.cs
+++ b/Engine/Source/Programs/Shared/EpicGames.IoHash/IoHash.cs
@@ -7,12 +7,15 @@ using System.Security.Cryptography;
 using System.Text;
 using System.Buffers.Binary;
 using System.Diagnostics.CodeAnalysis;
+using System.ComponentModel;
+using System.Globalization;
 
 namespace EpicGames.Core
 {
 	/// <summary>
 	/// Struct representing a strongly typed IoHash value (a 20-byte Blake3 hash).
 	/// </summary>
+	[TypeConverter(typeof(IoHashTypeConverter))]
 	public struct IoHash : IEquatable<IoHash>, IComparable<IoHash>
 	{
 		/// <summary>
@@ -54,6 +57,17 @@ namespace EpicGames.Core
 			this.Memory = Memory;
 		}
 
+		/// <summary>
+		/// Construct 
+		/// </summary>
+		/// <param name="Hasher">The hasher to construct from</param>
+		public IoHash(Blake3.Hasher Hasher)
+		{
+			byte[] Output = new byte[32];
+			Hasher.Finalize(Output);
+			Memory = Output.AsMemory(0, NumBytes);
+		}
+
 		/// <summary>
 		/// Creates a content hash for a block of data, using a given algorithm.
 		/// </summary>
@@ -161,4 +175,22 @@ namespace EpicGames.Core
 			Writer.WriteFixedLengthBytes(Hash.Span);
 		}
 	}
+
+	/// <summary>
+	/// Type converter from strings to IoHash objects
+	/// </summary>
+	sealed class IoHashTypeConverter : TypeConverter
+	{
+		/// <inheritdoc/>
+		public override bool CanConvertFrom(ITypeDescriptorContext Context, Type SourceType)
+		{
+			return SourceType == typeof(string);
+		}
+
+		/// <inheritdoc/>
+		public override object ConvertFrom(ITypeDescriptorContext Context, CultureInfo Culture, object Value)
+		{
+			return IoHash.Parse((string)Value);
+		}
+	}
 }
