Missing file for Blake3 hashes.

diff --git a/Engine/Source/Programs/Shared/EpicGames.IoHash/Blake3Hash.cs b/Engine/Source/Programs/Shared/EpicGames.IoHash/Blake3Hash.cs
new file mode 100644
--- /dev/null
+++ b/Engine/Source/Programs/Shared/EpicGames.IoHash/Blake3Hash.cs
@@ -0,0 +1,164 @@
+// Copyright Epic Games, Inc. All Rights Reserved.
+
+using System;
+using System.Collections.Generic;
+using System.Buffers;
+using System.Security.Cryptography;
+using System.Text;
+using System.Buffers.Binary;
+using System.Diagnostics.CodeAnalysis;
+
+namespace EpicGames.Core
+{
+	/// <summary>
+	/// Struct representing a strongly typed Blake3 hash value (a 32-byte Blake3 hash).
+	/// </summary>
+	public struct Blake3Hash : IEquatable<Blake3Hash>, IComparable<Blake3Hash>
+	{
+		/// <summary>
+		/// Length of an Blake3Hash
+		/// </summary>
+		public const int NumBytes = 32;
+
+		/// <summary>
+		/// Length of the hash in bits
+		/// </summary>
+		public const int NumBits = NumBytes * 8;
+
+		/// <summary>
+		/// Memory storing the digest data
+		/// </summary>
+		public ReadOnlyMemory<byte> Memory;
+
+		/// <summary>
+		/// Span for the underlying memory
+		/// </summary>
+		public ReadOnlySpan<byte> Span => Memory.Span;
+
+		/// <summary>
+		/// Hash consisting of zeroes
+		/// </summary>
+		public static Blake3Hash Zero { get; } = new Blake3Hash(new byte[NumBytes]);
+
+		/// <summary>
+		/// Constructor
+		/// </summary>
+		/// <param name="Memory">Memory to construct from</param>
+		public Blake3Hash(ReadOnlyMemory<byte> Memory)
+		{
+			if (Memory.Length != NumBytes)
+			{
+				throw new ArgumentException($"Blake3Hash must be {NumBytes} bytes long");
+			}
+
+			this.Memory = Memory;
+		}
+
+		/// <summary>
+		/// Creates a content hash for a block of data, using a given algorithm.
+		/// </summary>
+		/// <param name="Data">Data to compute the hash for</param>
+		/// <returns>New content hash instance containing the hash of the data</returns>
+		public static Blake3Hash Compute(ReadOnlySpan<byte> Data)
+		{
+			byte[] Output = new byte[32];
+			Blake3.Hasher.Hash(Data, Output);
+			return new Blake3Hash(Output);
+		}
+
+		/// <summary>
+		/// Parses a digest from the given hex string
+		/// </summary>
+		/// <param name="Text"></param>
+		/// <returns></returns>
+		public static Blake3Hash Parse(string Text)
+		{
+			return new Blake3Hash(StringUtils.ParseHexString(Text));
+		}
+
+		/// <inheritdoc cref="IComparable{T}.CompareTo(T)"/>
+		public int CompareTo(Blake3Hash Other)
+		{
+			ReadOnlySpan<byte> A = Span;
+			ReadOnlySpan<byte> B = Other.Span;
+
+			for (int Idx = 0; Idx < A.Length && Idx < B.Length; Idx++)
+			{
+				int Compare = A[Idx] - B[Idx];
+				if (Compare != 0)
+				{
+					return Compare;
+				}
+			}
+			return A.Length - B.Length;
+		}
+
+		/// <inheritdoc/>
+		public bool Equals(Blake3Hash Other) => Span.SequenceEqual(Other.Span);
+
+		/// <inheritdoc/>
+		public override bool Equals(object? Obj) => (Obj is Blake3Hash Hash) && Hash.Span.SequenceEqual(Span);
+
+		/// <inheritdoc/>
+		public override int GetHashCode() => BinaryPrimitives.ReadInt32LittleEndian(Span);
+
+		/// <inheritdoc/>
+		public override string ToString() => StringUtils.FormatHexString(Memory.Span);
+
+		/// <summary>
+		/// Test two hash values for equality
+		/// </summary>
+		public static bool operator ==(Blake3Hash A, Blake3Hash B) => A.Span.SequenceEqual(B.Span);
+
+		/// <summary>
+		/// Test two hash values for equality
+		/// </summary>
+		public static bool operator !=(Blake3Hash A, Blake3Hash B) => !(A == B);
+
+		/// <summary>
+		/// Tests whether A > B
+		/// </summary>
+		public static bool operator >(Blake3Hash A, Blake3Hash B) => A.CompareTo(B) > 0;
+
+		/// <summary>
+		/// Tests whether A is less than B
+		/// </summary>
+		public static bool operator <(Blake3Hash A, Blake3Hash B) => A.CompareTo(B) < 0;
+
+		/// <summary>
+		/// Tests whether A is greater than or equal to B
+		/// </summary>
+		public static bool operator >=(Blake3Hash A, Blake3Hash B) => A.CompareTo(B) >= 0;
+
+		/// <summary>
+		/// Tests whether A is less than or equal to B
+		/// </summary>
+		public static bool operator <=(Blake3Hash A, Blake3Hash B) => A.CompareTo(B) <= 0;
+	}
+
+	/// <summary>
+	/// Extension methods for dealing with Blake3Hash values
+	/// </summary>
+	public static class Blake3HashExtensions
+	{
+		/// <summary>
+		/// Read an <see cref="Blake3Hash"/> from a memory reader
+		/// </summary>
+		/// <param name="Reader"></param>
+		/// <returns></returns>
+		public static Blake3Hash ReadBlake3Hash(this MemoryReader Reader)
+		{
+			return new Blake3Hash(Reader.ReadFixedLengthBytes(Blake3Hash.NumBytes));
+		}
+
+		/// <summary>
+		/// Write an <see cref="Blake3Hash"/> to a memory writer
+		/// </summary>
+		/// <param name="Writer"></param>
+		/// <param name="Hash"></param>
+		public static void WriteBlake3Hash(this MemoryWriter Writer, Blake3Hash Hash)
+		{
+			Writer.WriteFixedLengthBytes(Hash.Span);
+		}
+	}
+}
