Horde: Various compact binary library improvements.

CbObject/CbBinary/CbArray:
* Added proxy classes for viewing objects in the debugger. CbArray now displays as an array, CbObject displays as an array of fields, etc...
* Removed overloaded "View" methods, which are not required in C# due to GC.

CbWriter:
* Significantly fewer memory allocations/copying. Output buffer is now allocated in large chunks, with length fields only inserted when copying to the final output buffer.
* Methods for writing named/unnamed fields are now more idomatic with other C# serialization classes (ie. WriteXXX(Name, Value) and WriteXXXValue(Value)).

Other:
* Added serialization tests to Horde projects and fixed NuGet dependencies.

diff --git a/Engine/Source/Programs/Horde/HordeServerTests/StorageTests.cs b/Engine/Source/Programs/Horde/HordeServerTests/StorageTests.cs
--- a/Engine/Source/Programs/Horde/HordeServerTests/StorageTests.cs
+++ b/Engine/Source/Programs/Horde/HordeServerTests/StorageTests.cs
@@ -61,7 +61,7 @@ namespace HordeServerTests
 
 			CbObject ObjectD = CbObject.Build(Writer =>
 			{
-				Writer.AddString("hello", "field");
+				Writer.WriteString("hello", "field");
 			});
 			IoHash HashD = ObjectD.GetHash();
 
@@ -95,16 +95,16 @@ namespace HordeServerTests
 
 			CbObject ObjectD = CbObject.Build(Writer =>
 			{
-				Writer.AddBinaryAttachment(HashA, "A");
-				Writer.AddBinaryAttachment(HashB, "B");
-				Writer.AddBinaryAttachment(HashC, "C");
+				Writer.WriteBinaryAttachment("A", HashA);
+				Writer.WriteBinaryAttachment("B", HashB);
+				Writer.WriteBinaryAttachment("C", HashC);
 			});
 			IoHash HashD = ObjectD.GetHash();
 
 			CbObject ObjectE = CbObject.Build(Writer =>
 			{
-				Writer.AddBinaryAttachment(HashA, "A");
-				Writer.AddCompactBinaryAttachment(HashD, "D");
+				Writer.WriteBinaryAttachment("A", HashA);
+				Writer.WriteObjectAttachment("D", HashD);
 			});
 			IoHash HashE = ObjectE.GetHash();
 
diff --git a/Engine/Source/Programs/Shared/EpicGames.Serialization.Tests/EngineTests.cs b/Engine/Source/Programs/Shared/EpicGames.Serialization.Tests/EngineTests.cs
--- a/Engine/Source/Programs/Shared/EpicGames.Serialization.Tests/EngineTests.cs
+++ b/Engine/Source/Programs/Shared/EpicGames.Serialization.Tests/EngineTests.cs
@@ -41,11 +41,11 @@ namespace EpicGames.Serialization.Tests
 
 		Dictionary<CbFieldType, CbFieldAccessors> TypeAccessors = new Dictionary<CbFieldType, CbFieldAccessors>
 		{
-			[CbFieldType.Object] = new CbFieldAccessors(CbObject.Empty, x => x.IsObject(), x => x.AsObjectView(), (x, y) => x.AsObjectView()),
-			[CbFieldType.UniformObject] = new CbFieldAccessors(CbObject.Empty, x => x.IsObject(), x => x.AsObjectView(), (x, y) => x.AsObjectView()),
-			[CbFieldType.Array] = new CbFieldAccessors(CbArray.Empty, x => x.IsArray(), x => x.AsArrayView(), (x, y) => x.AsArrayView()),
-			[CbFieldType.UniformArray] = new CbFieldAccessors(CbArray.Empty, x => x.IsArray(), x => x.AsArrayView(), (x, y) => x.AsArrayView()),
-			[CbFieldType.Binary] = CbFieldAccessors.FromStruct<ReadOnlyMemory<byte>>(x => x.IsBinary(), (x, y) => x.AsBinaryView(y), (x, y) => x.Span.SequenceEqual(y.Span)),
+			[CbFieldType.Object] = new CbFieldAccessors(CbObject.Empty, x => x.IsObject(), x => x.AsObject(), (x, y) => x.AsObject()),
+			[CbFieldType.UniformObject] = new CbFieldAccessors(CbObject.Empty, x => x.IsObject(), x => x.AsObject(), (x, y) => x.AsObject()),
+			[CbFieldType.Array] = new CbFieldAccessors(CbArray.Empty, x => x.IsArray(), x => x.AsArray(), (x, y) => x.AsArray()),
+			[CbFieldType.UniformArray] = new CbFieldAccessors(CbArray.Empty, x => x.IsArray(), x => x.AsArray(), (x, y) => x.AsArray()),
+			[CbFieldType.Binary] = CbFieldAccessors.FromStruct<ReadOnlyMemory<byte>>(x => x.IsBinary(), (x, y) => x.AsBinary(y), (x, y) => x.Span.SequenceEqual(y.Span)),
 			[CbFieldType.String] = new CbFieldAccessors(ReadOnlyUtf8String.Empty, x => x.IsString(), x => x.AsString(), (x, Default) => x.AsString((ReadOnlyUtf8String)Default)),
 			[CbFieldType.IntegerPositive] = CbFieldAccessors.FromStruct<ulong>(x => x.IsInteger(), (x, y) => x.AsUInt64(y)),
 			[CbFieldType.IntegerNegative] = CbFieldAccessors.FromStruct<long>(x => x.IsInteger(), (x, y) => x.AsInt64(y)),
@@ -249,12 +249,12 @@ namespace EpicGames.Serialization.Tests
 				TestIntObject(Object, 0, 1);
 
 				// Find fields that do not exist.
-				Assert.IsFalse(Object.FindView("Field").HasValue());
-				Assert.IsFalse(Object.FindViewIgnoreCase("Field").HasValue());
+				Assert.IsFalse(Object.Find("Field").HasValue());
+				Assert.IsFalse(Object.FindIgnoreCase("Field").HasValue());
 				Assert.IsFalse(Object["Field"].HasValue());
 
 				// Advance an iterator past the last field.
-				CbFieldIterator It = Object.CreateViewIterator();
+				CbFieldIterator It = Object.CreateIterator();
 				Assert.IsFalse((bool)It);
 				Assert.IsTrue(!It);
 				for (int Count = 16; Count > 0; --Count)
@@ -272,14 +272,14 @@ namespace EpicGames.Serialization.Tests
 				byte[] Payload = { 12, IntType, 1, (byte)'A', 1, IntType, 1, (byte)'B', 2, IntType, 1, (byte)'C', 3 };
 				CbField Field = new CbField(Payload, CbFieldType.Object);
 				TestField(CbFieldType.Object, Field, new CbObject(Payload, CbFieldType.Object));
-				CbObject Object = CbObject.Clone(Field.AsObjectView());
+				CbObject Object = CbObject.Clone(Field.AsObject());
 				TestIntObject(Object, 3, Payload.Length);
-				TestIntObject(Field.AsObjectView(), 3, Payload.Length);
-				Assert.IsTrue(Object.Equals(Field.AsObjectView()));
-				Assert.AreEqual(Object.FindView("B").AsInt32(), 2);
-				Assert.AreEqual(Object.FindView("b").AsInt32(4), 4);
-				Assert.AreEqual(Object.FindViewIgnoreCase("B").AsInt32(), 2);
-				Assert.AreEqual(Object.FindViewIgnoreCase("b").AsInt32(), 2);
+				TestIntObject(Field.AsObject(), 3, Payload.Length);
+				Assert.IsTrue(Object.Equals(Field.AsObject()));
+				Assert.AreEqual(Object.Find("B").AsInt32(), 2);
+				Assert.AreEqual(Object.Find("b").AsInt32(4), 4);
+				Assert.AreEqual(Object.FindIgnoreCase("B").AsInt32(), 2);
+				Assert.AreEqual(Object.FindIgnoreCase("b").AsInt32(), 2);
 				Assert.AreEqual(Object["B"].AsInt32(), 2);
 				Assert.AreEqual(Object["b"].AsInt32(4), 4);
 			}
@@ -290,17 +290,17 @@ namespace EpicGames.Serialization.Tests
 				byte[] Payload = { 10, IntType, 1, (byte)'A', 1, 1, (byte)'B', 2, 1, (byte)'C', 3 };
 				CbField Field = new CbField(Payload, CbFieldType.UniformObject);
 				TestField(CbFieldType.UniformObject, Field, new CbObject(Payload, CbFieldType.UniformObject));
-				CbObject Object = CbObject.Clone(Field.AsObjectView());
+				CbObject Object = CbObject.Clone(Field.AsObject());
 				TestIntObject(Object, 3, Payload.Length);
 				TestIntObject(Field.AsObject(), 3, Payload.Length);
 				Assert.IsTrue(Object.Equals(Field.AsObject()));
-				Assert.AreEqual(Object.FindView("B").AsInt32(), 2);
 				Assert.AreEqual(Object.Find("B").AsInt32(), 2);
-				Assert.AreEqual(Object.FindView("b").AsInt32(4), 4);
+				Assert.AreEqual(Object.Find("B").AsInt32(), 2);
+				Assert.AreEqual(Object.Find("b").AsInt32(4), 4);
 				Assert.AreEqual(Object.Find("b").AsInt32(4), 4);
-				Assert.AreEqual(Object.FindViewIgnoreCase("B").AsInt32(), 2);
 				Assert.AreEqual(Object.FindIgnoreCase("B").AsInt32(), 2);
-				Assert.AreEqual(Object.FindViewIgnoreCase("b").AsInt32(), 2);
+				Assert.AreEqual(Object.FindIgnoreCase("B").AsInt32(), 2);
+				Assert.AreEqual(Object.FindIgnoreCase("b").AsInt32(), 2);
 				Assert.AreEqual(Object.FindIgnoreCase("b").AsInt32(), 2);
 				Assert.AreEqual(Object["B"].AsInt32(), 2);
 				Assert.AreEqual(Object["b"].AsInt32(4), 4);
@@ -312,15 +312,15 @@ namespace EpicGames.Serialization.Tests
 
 				// CopyTo
 				byte[] CopyBytes = new byte[Payload.Length + 1];
-				Field.AsObjectView().CopyTo(CopyBytes);
+				Field.AsObject().CopyTo(CopyBytes);
 				Assert.IsTrue(Payload.AsSpan().SequenceEqual(CopyBytes.AsSpan(1)));
-				NamedField.AsObjectView().CopyTo(CopyBytes);
+				NamedField.AsObject().CopyTo(CopyBytes);
 				Assert.IsTrue(Payload.AsSpan().SequenceEqual(CopyBytes.AsSpan(1)));
 
 				// TryGetView
 				ReadOnlyMemory<byte> View;
-				Assert.IsFalse(Field.AsObjectView().TryGetView(out View));
-				Assert.IsFalse(NamedField.AsObjectView().TryGetView(out View));
+				Assert.IsFalse(Field.AsObject().TryGetView(out View));
+				Assert.IsFalse(NamedField.AsObject().TryGetView(out View));
 			}
 
 			// Test CbField(None) as Object
@@ -358,7 +358,7 @@ namespace EpicGames.Serialization.Tests
 			{
 				int Count = 0;
 				CbObject Object = CbObject.Empty;
-				for (CbFieldIterator It = CbFieldIterator.MakeSingle(Object.AsFieldView()); It; ++It)
+				for (CbFieldIterator It = CbFieldIterator.MakeSingle(Object.AsField()); It; ++It)
 				{
 					CbField Field = It.Current;
 					Assert.IsTrue(Field.IsObject());
@@ -376,7 +376,7 @@ namespace EpicGames.Serialization.Tests
 				Assert.AreEqual(Array.Num(), ExpectedNum);
 
 				int ActualNum = 0;
-				for (CbFieldIterator It = Array.CreateViewIterator(); It; ++It)
+				for (CbFieldIterator It = Array.CreateIterator(); It; ++It)
 				{
 					++ActualNum;
 					Assert.AreEqual(It.Current.AsInt32(), ActualNum);
@@ -392,7 +392,7 @@ namespace EpicGames.Serialization.Tests
 				Assert.AreEqual(ActualNum, ExpectedNum);
 
 				ActualNum = 0;
-				foreach (CbField Field in Array.AsFieldView())
+				foreach (CbField Field in Array.AsField())
 				{
 					++ActualNum;
 					Assert.AreEqual(Field.AsInt32(), ActualNum);
@@ -409,7 +409,7 @@ namespace EpicGames.Serialization.Tests
 				TestIntArray(Array, 0, 2);
 
 				// Advance an iterator past the last field.
-				CbFieldIterator It = Array.CreateViewIterator();
+				CbFieldIterator It = Array.CreateIterator();
 				Assert.IsFalse((bool)It);
 				Assert.IsTrue(!It);
 				for (int Count = 16; Count > 0; --Count)
@@ -427,10 +427,10 @@ namespace EpicGames.Serialization.Tests
 				byte[] Payload = new byte[]{ 7, 3, IntType, 1, IntType, 2, IntType, 3 };
 				CbField Field = new CbField(Payload, CbFieldType.Array);
 				TestField(CbFieldType.Array, Field, new CbArray(Payload, CbFieldType.Array));
-				CbArray Array = CbArray.Clone(Field.AsArrayView());
+				CbArray Array = CbArray.Clone(Field.AsArray());
 				TestIntArray(Array, 3, Payload.Length);
-				TestIntArray(Field.AsArrayView(), 3, Payload.Length);
-				Assert.IsTrue(Array.Equals(Field.AsArrayView()));
+				TestIntArray(Field.AsArray(), 3, Payload.Length);
+				Assert.IsTrue(Array.Equals(Field.AsArray()));
 			}
 
 			// Test CbField(UniformArray)
@@ -439,32 +439,32 @@ namespace EpicGames.Serialization.Tests
 				byte[] Payload = new byte[]{ 5, 3, IntType, 1, 2, 3 };
 				CbField Field = new CbField(Payload, CbFieldType.UniformArray);
 				TestField(CbFieldType.UniformArray, Field, new CbArray(Payload, CbFieldType.UniformArray));
-				CbArray Array = CbArray.Clone(Field.AsArrayView());
+				CbArray Array = CbArray.Clone(Field.AsArray());
 				TestIntArray(Array, 3, Payload.Length);
-				TestIntArray(Field.AsArrayView(), 3, Payload.Length);
-				Assert.IsTrue(Array.Equals(Field.AsArrayView()));
+				TestIntArray(Field.AsArray(), 3, Payload.Length);
+				Assert.IsTrue(Array.Equals(Field.AsArray()));
 
 //				Assert.IsTrue(Array.GetOuterBuffer() == Array.AsField().AsArray().GetOuterBuffer());
 
 				// Equals
 				byte[] NamedPayload = new byte[]{ 1, (byte)'A', 5, 3, IntType, 1, 2, 3 };
 				CbField NamedField = new CbField(NamedPayload, CbFieldType.UniformArray | CbFieldType.HasFieldName);
-				Assert.IsTrue(Field.AsArrayView().Equals(NamedField.AsArrayView()));
-				Assert.IsTrue(Field.Equals(Field.AsArrayView().AsFieldView()));
-				Assert.IsTrue(NamedField.Equals(NamedField.AsArrayView().AsFieldView()));
+				Assert.IsTrue(Field.AsArray().Equals(NamedField.AsArray()));
+				Assert.IsTrue(Field.Equals(Field.AsArray().AsField()));
+				Assert.IsTrue(NamedField.Equals(NamedField.AsArray().AsField()));
 
 				// CopyTo
 				byte[] CopyBytes = new byte[Payload.Length + 1];
-				Field.AsArrayView().CopyTo(CopyBytes);
+				Field.AsArray().CopyTo(CopyBytes);
 				Assert.IsTrue(Payload.AsSpan().SequenceEqual(CopyBytes.AsSpan(1)));
-				NamedField.AsArrayView().CopyTo(CopyBytes);
+				NamedField.AsArray().CopyTo(CopyBytes);
 				Assert.IsTrue(Payload.AsSpan().SequenceEqual(CopyBytes.AsSpan(1)));
 
 				// TryGetView
 				ReadOnlyMemory<byte> View;
 //				Assert.IsTrue(Array.TryGetView(out View) && View == Array.GetOuterBuffer().GetView());
-				Assert.IsFalse(Field.AsArrayView().TryGetView(out View));
-				Assert.IsFalse(NamedField.AsArrayView().TryGetView(out View));
+				Assert.IsFalse(Field.AsArray().TryGetView(out View));
+				Assert.IsFalse(NamedField.AsArray().TryGetView(out View));
 
 //				// GetBuffer
 //				Assert.IsTrue(Array.GetBuffer().Flatten().GetView() == Array.GetOuterBuffer().GetView());
@@ -519,7 +519,7 @@ namespace EpicGames.Serialization.Tests
 			{
 				uint Count = 0;
 				CbArray Array = new CbArray();
-				for (CbFieldIterator Iter = CbFieldIterator.MakeSingle(Array.AsFieldView()); Iter; ++Iter)
+				for (CbFieldIterator Iter = CbFieldIterator.MakeSingle(Array.AsField()); Iter; ++Iter)
 				{
 					CbField Field = Iter.Current;
 					Assert.IsTrue(Field.IsArray());
@@ -588,7 +588,7 @@ namespace EpicGames.Serialization.Tests
 			// Test CbField(String, OutOfRangeSize)
 			{
 				byte[] Payload = new byte[9];
-				BitUtils.WriteVarUInt((ulong)(1) << 31, Payload);
+				VarInt.Write(Payload, (ulong)(1) << 31);
 				TestFieldError(CbFieldType.String, Payload, CbFieldError.RangeError, new ReadOnlyUtf8String("ABC"));
 			}
 
@@ -631,7 +631,7 @@ namespace EpicGames.Serialization.Tests
 		{
 			byte[] Payload = new byte[9];
 			ulong Negative = (ulong)((byte)FieldType & 1);
-			BitUtils.WriteVarUInt(Magnitude - Negative, Payload);
+			VarInt.Write(Payload, Magnitude - Negative);
 			ulong DefaultValue = 8;
 			ulong ExpectedValue = (Negative != 0)? (ulong)(-(long)(Magnitude)) : Magnitude;
 			CbField Field = new CbField(Payload, FieldType);
@@ -724,7 +724,7 @@ namespace EpicGames.Serialization.Tests
 			// Test CbField(Integer+, MaxBinary32) as Float
 			{
 				byte[] Payload = new byte[9];
-				BitUtils.WriteVarUInt(((ulong)(1) << 24) - 1, Payload); // 16,777,215
+				VarInt.Write(Payload, ((ulong)(1) << 24) - 1); // 16,777,215
 				CbField Field = new CbField(Payload, CbFieldType.IntegerPositive);
 				TestField(CbFieldType.Float32, Field, 16_777_215.0f);
 				TestField(CbFieldType.Float64, Field, 16_777_215.0);
@@ -733,7 +733,7 @@ namespace EpicGames.Serialization.Tests
 			// Test CbField(Integer+, MaxBinary32+1) as Float
 			{
 				byte[] Payload = new byte[9];
-				BitUtils.WriteVarUInt((ulong)(1) << 24, Payload); // 16,777,216
+				VarInt.Write(Payload, (ulong)(1) << 24); // 16,777,216
 				CbField Field = new CbField(Payload, CbFieldType.IntegerPositive);
 				TestFieldError(CbFieldType.Float32, Field, CbFieldError.RangeError, 8.0f);
 				TestField(CbFieldType.Float64, Field, 16_777_216.0);
@@ -742,7 +742,7 @@ namespace EpicGames.Serialization.Tests
 			// Test CbField(Integer+, MaxBinary64) as Float
 			{
 				byte[] Payload = new byte[9];
-				BitUtils.WriteVarUInt(((ulong)(1) << 53) - 1, Payload); // 9,007,199,254,740,991
+				VarInt.Write(Payload, ((ulong)(1) << 53) - 1); // 9,007,199,254,740,991
 				CbField Field = new CbField(Payload, CbFieldType.IntegerPositive);
 				TestFieldError(CbFieldType.Float32, Field, CbFieldError.RangeError, 8.0f);
 				TestField(CbFieldType.Float64, Field, 9_007_199_254_740_991.0);
@@ -751,7 +751,7 @@ namespace EpicGames.Serialization.Tests
 			// Test CbField(Integer+, MaxBinary64+1) as Float
 			{
 				byte[] Payload = new byte[9];
-				BitUtils.WriteVarUInt((ulong)(1) << 53, Payload); // 9,007,199,254,740,992
+				VarInt.Write(Payload, (ulong)(1) << 53); // 9,007,199,254,740,992
 				CbField Field = new CbField(Payload, CbFieldType.IntegerPositive);
 				TestFieldError(CbFieldType.Float32, Field, CbFieldError.RangeError, 8.0f);
 				TestFieldError(CbFieldType.Float64, Field, CbFieldError.RangeError, 8.0);
@@ -760,7 +760,7 @@ namespace EpicGames.Serialization.Tests
 			// Test CbField(Integer+, MaxUInt64) as Float
 			{
 				byte[] Payload = new byte[9];
-				BitUtils.WriteVarUInt(~(ulong)0, Payload); // Max uint64
+				VarInt.Write(Payload, ~(ulong)0); // Max uint64
 				CbField Field = new CbField(Payload, CbFieldType.IntegerPositive);
 				TestFieldError(CbFieldType.Float32, Field, CbFieldError.RangeError, 8.0f);
 				TestFieldError(CbFieldType.Float64, Field, CbFieldError.RangeError, 8.0);
@@ -769,7 +769,7 @@ namespace EpicGames.Serialization.Tests
 			// Test CbField(Integer-, MaxBinary32) as Float
 			{
 				byte[] Payload = new byte[9];
-				BitUtils.WriteVarUInt(((ulong)(1) << 24) - 2, Payload); // -16,777,215
+				VarInt.Write(Payload, ((ulong)(1) << 24) - 2); // -16,777,215
 				CbField Field = new CbField(Payload, CbFieldType.IntegerNegative);
 				TestField(CbFieldType.Float32, Field, -16_777_215.0f);
 				TestField(CbFieldType.Float64, Field, -16_777_215.0);
@@ -778,7 +778,7 @@ namespace EpicGames.Serialization.Tests
 			// Test CbField(Integer-, MaxBinary32+1) as Float
 			{
 				byte[] Payload = new byte[9];
-				BitUtils.WriteVarUInt(((ulong)(1) << 24) - 1, Payload); // -16,777,216
+				VarInt.Write(Payload, ((ulong)(1) << 24) - 1); // -16,777,216
 				CbField Field = new CbField(Payload, CbFieldType.IntegerNegative);
 				TestFieldError(CbFieldType.Float32, Field, CbFieldError.RangeError, 8.0f);
 				TestField(CbFieldType.Float64, Field, -16_777_216.0);
@@ -787,7 +787,7 @@ namespace EpicGames.Serialization.Tests
 			// Test CbField(Integer-, MaxBinary64) as Float
 			{
 				byte[] Payload = new byte[9];
-				BitUtils.WriteVarUInt(((ulong)(1) << 53) - 2, Payload); // -9,007,199,254,740,991
+				VarInt.Write(Payload, ((ulong)(1) << 53) - 2); // -9,007,199,254,740,991
 				CbField Field = new CbField(Payload, CbFieldType.IntegerNegative);
 				TestFieldError(CbFieldType.Float32, Field, CbFieldError.RangeError, 8.0f);
 				TestField(CbFieldType.Float64, Field, -9_007_199_254_740_991.0);
@@ -796,7 +796,7 @@ namespace EpicGames.Serialization.Tests
 			// Test CbField(Integer-, MaxBinary64+1) as Float
 			{
 				byte[] Payload = new byte[9];
-				BitUtils.WriteVarUInt(((ulong)(1) << 53) - 1, Payload); // -9,007,199,254,740,992
+				VarInt.Write(Payload, ((ulong)(1) << 53) - 1); // -9,007,199,254,740,992
 				CbField Field = new CbField(Payload, CbFieldType.IntegerNegative);
 				TestFieldError(CbFieldType.Float32, Field, CbFieldError.RangeError, 8.0f);
 				TestFieldError(CbFieldType.Float64, Field, CbFieldError.RangeError, 8.0);
@@ -1621,7 +1621,7 @@ namespace EpicGames.Serialization.Tests
 			// way, each field will only be visited once even if the loop needs to execute several times.
 			ParseObjectType ParseObject = (CbObject Object, ref uint A, ref uint B, ref uint C, ref uint D) =>
 			{
-				for (CbFieldIterator It = Object.CreateViewIterator(); It;)
+				for (CbFieldIterator It = Object.CreateIterator(); It;)
 				{
 					CbFieldIterator Last = It;
 					if (It.Current.GetName().Equals("A"))
diff --git a/Engine/Source/Programs/Shared/EpicGames.Serialization.Tests/EpicGames.Serialization.Tests.csproj b/Engine/Source/Programs/Shared/EpicGames.Serialization.Tests/EpicGames.Serialization.Tests.csproj
--- a/Engine/Source/Programs/Shared/EpicGames.Serialization.Tests/EpicGames.Serialization.Tests.csproj
+++ b/Engine/Source/Programs/Shared/EpicGames.Serialization.Tests/EpicGames.Serialization.Tests.csproj
@@ -1,28 +1,19 @@
-<Project Sdk="Microsoft.NET.Sdk">
+﻿<Project Sdk="Microsoft.NET.Sdk">
 
   <PropertyGroup>
     <TargetFramework>netcoreapp3.1</TargetFramework>
     <Configuration Condition=" '$(Configuration)' == '' ">Development</Configuration>
     <OutputType>Library</OutputType>
     <Deterministic>true</Deterministic>
-    <RootNamespace>EpicGames.Serialization</RootNamespace>
-    <GenerateAssemblyInfo>false</GenerateAssemblyInfo>
-    <GenerateTargetFrameworkAttribute>false</GenerateTargetFrameworkAttribute>
+    <RootNamespace>EpicGames.Serialization.Tests</RootNamespace>
+    <GenerateAssemblyCompanyAttribute>false</GenerateAssemblyCompanyAttribute>
+    <GenerateAssemblyInformationalVersionAttribute>false</GenerateAssemblyInformationalVersionAttribute>
+    <GenerateAssemblyProductAttribute>false</GenerateAssemblyProductAttribute>
+    <GenerateAssemblyVersionAttribute>false</GenerateAssemblyVersionAttribute>
     <Configurations>Debug;Release;Development</Configurations>
     <IsPackable>false</IsPackable>
     <Nullable>enable</Nullable>
   </PropertyGroup>
-  <PropertyGroup Condition=" '$(Configuration)|$(Platform)' == 'Debug|AnyCPU' ">
-    <DocumentationFile></DocumentationFile>
-  </PropertyGroup>
-  <PropertyGroup Condition="'$(Configuration)|$(Platform)' == 'Development|AnyCPU'">
-    <DocumentationFile></DocumentationFile>
-    <Optimize>true</Optimize>
-  </PropertyGroup>
-  <PropertyGroup Condition="'$(Configuration)|$(Platform)' == 'Release|AnyCPU'">
-    <DocumentationFile></DocumentationFile>
-    <Optimize>true</Optimize>
-  </PropertyGroup>
   <ItemGroup>
     <Compile Include="..\..\DotNETCommon\MetaData.cs">
       <Link>Properties\MetaData.cs</Link>
@@ -30,10 +21,10 @@
   </ItemGroup>
 
   <ItemGroup>
-    <PackageReference Include="Microsoft.NET.Test.Sdk" Version="16.7.1" />
-    <PackageReference Include="MSTest.TestAdapter" Version="2.1.1" />
-    <PackageReference Include="MSTest.TestFramework" Version="2.1.1" />
-    <PackageReference Include="coverlet.collector" Version="1.3.0" />
+    <PackageReference Include="Microsoft.NET.Test.Sdk" Version="16.5.0" />
+    <PackageReference Include="MSTest.TestAdapter" Version="2.1.0" />
+    <PackageReference Include="MSTest.TestFramework" Version="2.1.0" />
+    <PackageReference Include="coverlet.collector" Version="1.2.0" />
   </ItemGroup>
 
   <ItemGroup>
diff --git a/Engine/Source/Programs/Shared/EpicGames.Serialization.Tests/JupiterTests.cs b/Engine/Source/Programs/Shared/EpicGames.Serialization.Tests/JupiterTests.cs
--- a/Engine/Source/Programs/Shared/EpicGames.Serialization.Tests/JupiterTests.cs
+++ b/Engine/Source/Programs/Shared/EpicGames.Serialization.Tests/JupiterTests.cs
@@ -72,10 +72,13 @@ namespace EpicGames.Serialization.Tests
 
 			CbWriter writer = new CbWriter();
 			writer.BeginObject();
-			writer.AddUniformArray(new IoHash[] { hash1, hash2 }, CbFieldType.Hash, "needs");
+			writer.BeginUniformArray("needs", CbFieldType.Hash);
+			writer.WriteHashValue(hash1);
+			writer.WriteHashValue(hash2);
+			writer.EndUniformArray();
 			writer.EndObject();
 
-			byte[] objectData = writer.Save();
+			byte[] objectData = writer.ToByteArray();
 			ReadOnlyMemory<byte> memory = new ReadOnlyMemory<byte>(objectData);
 			CbField o = new CbField(memory);
 
@@ -97,11 +100,11 @@ namespace EpicGames.Serialization.Tests
 
 			CbWriter writer = new CbWriter();
 			writer.BeginObject();
-			writer.AddString("test", "string");
-			writer.AddBinaryAttachment(hash1, "hash");
+			writer.WriteString("string", "test");
+			writer.WriteBinaryAttachment("hash", hash1);
 			writer.EndObject();
 
-			byte[] objectData = writer.Save();
+			byte[] objectData = writer.ToByteArray();
 			ReadOnlyMemory<byte> memory = new ReadOnlyMemory<byte>(objectData);
 			CbField o = new CbField(memory);
 
diff --git a/Engine/Source/Programs/Shared/EpicGames.Serialization/CompactBinary.cs b/Engine/Source/Programs/Shared/EpicGames.Serialization/CompactBinary.cs
--- a/Engine/Source/Programs/Shared/EpicGames.Serialization/CompactBinary.cs
+++ b/Engine/Source/Programs/Shared/EpicGames.Serialization/CompactBinary.cs
@@ -6,6 +6,7 @@ using System.Buffers;
 using System.Buffers.Binary;
 using System.Collections;
 using System.Collections.Generic;
+using System.Diagnostics;
 using System.Linq;
 using System.Numerics;
 using System.Text;
@@ -465,6 +466,28 @@ namespace EpicGames.Serialization
 		RangeError,
 	}
 
+	/// <summary>
+	/// Simplified view of <see cref="CbField"/> in the debugger, for fields with a name
+	/// </summary>
+	class CbFieldWithNameDebugView
+	{
+		public string? Name;
+		public object? Value;
+	}
+
+	/// <summary>
+	/// Simplified view of <see cref="CbField"/> for the debugger
+	/// </summary>
+	class CbFieldDebugView
+	{
+		public CbFieldDebugView(CbField Field) => Value = Field.HasName()
+				? new CbFieldWithNameDebugView { Name = Field.Name.ToString(), Value = Field.Value } 
+				: Field.Value;
+
+		[DebuggerBrowsable(DebuggerBrowsableState.RootHidden)]
+		public object? Value { get; }
+	}
+
 	/// <summary>
 	/// An atom of data in the compact binary format.
 	///
@@ -482,8 +505,28 @@ namespace EpicGames.Serialization
 	/// addressable when encoded, which means a zero-byte field is not permitted, and only arises in a
 	/// uniform array of fields with no payload, where the answer is to encode as a non-uniform array.
 	/// </summary>
+	[DebuggerDisplay("{DebugValue,nq}")]
+	[DebuggerTypeProxy(typeof(CbFieldDebugView))]
 	public class CbField : IEquatable<CbField>, IEnumerable<CbField>
 	{
+		/// <summary>
+		/// Type returned for none values
+		/// </summary>
+		[DebuggerDisplay("<none>")]
+		class NoneValueType
+		{
+		}
+
+		/// <summary>
+		/// Special value returned for "none" fields.
+		/// </summary>
+		static NoneValueType None { get; } = new NoneValueType();
+
+		/// <summary>
+		/// Formatter for the debug string
+		/// </summary>
+		object? DebugValue => HasName() ? $"{Name} = {Value}" : Value;
+
 		/// <summary>
 		/// Default empty field
 		/// </summary>
@@ -551,7 +594,7 @@ namespace EpicGames.Serialization
 
 			if (CbFieldUtils.HasFieldName(Type))
 			{
-				NameLen = (int)BitUtils.ReadVarUInt(Data.Slice(Offset).Span, out int NameLenByteCount);
+				NameLen = (int)VarInt.Read(Data.Slice(Offset).Span, out int NameLenByteCount);
 				Offset += NameLenByteCount + NameLen;
 			}
 
@@ -567,6 +610,58 @@ namespace EpicGames.Serialization
 		/// </summary>
 		public ReadOnlyUtf8String Name => new ReadOnlyUtf8String(Memory.Slice(PayloadOffset - NameLen, NameLen));
 
+		/// <summary>
+		/// Gets the value of this field
+		/// </summary>
+		public object? Value
+		{
+			get
+			{
+				CbFieldType FieldType = CbFieldUtils.GetType(TypeWithFlags);
+				switch (FieldType)
+				{
+					case CbFieldType.None:
+						return None;
+					case CbFieldType.Null:
+						return null;
+					case CbFieldType.Object:
+					case CbFieldType.UniformObject:
+						return AsObject();
+					case CbFieldType.Array:
+					case CbFieldType.UniformArray:
+						return AsArray();
+					case CbFieldType.Binary:
+						return AsBinary();
+					case CbFieldType.String:
+						return AsString();
+					case CbFieldType.IntegerPositive:
+						return AsUInt64();
+					case CbFieldType.IntegerNegative:
+						return AsInt64();
+					case CbFieldType.Float32:
+						return AsFloat();
+					case CbFieldType.Float64:
+						return AsDouble();
+					case CbFieldType.BoolFalse:
+						return false;
+					case CbFieldType.BoolTrue:
+						return true;
+					case CbFieldType.ObjectAttachment:
+					case CbFieldType.BinaryAttachment:
+					case CbFieldType.Hash:
+						return AsAttachment();
+					case CbFieldType.Uuid:
+						return AsUuid();
+					case CbFieldType.DateTime:
+						return AsDateTime();
+					case CbFieldType.TimeSpan:
+						return AsTimeSpan();
+					default:
+						throw new NotImplementedException($"Unknown field type ({FieldType})");
+				}
+			}
+		}
+
 		/// <inheritdoc cref="Name"/>
 		public ReadOnlyUtf8String GetName() => Name;
 
@@ -588,12 +683,6 @@ namespace EpicGames.Serialization
 			}
 		}
 
-		/// <summary>
-		/// Access the field as an object. Defaults to an empty object on error. 
-		/// </summary>
-		/// <returns></returns>
-		public CbObject AsObjectView() => AsObject();
-
 		/// <summary>
 		/// Access the field as an array. Defaults to an empty array on error. 
 		/// </summary>
@@ -612,12 +701,6 @@ namespace EpicGames.Serialization
 			}
 		}
 
-		/// <summary>
-		/// Access the field as an array. Defaults to an empty array on error. 
-		/// </summary>
-		/// <returns></returns>
-		public CbArray AsArrayView() => AsArray();
-
 		/// <summary>
 		/// Access the field as binary data.
 		/// </summary>
@@ -627,7 +710,9 @@ namespace EpicGames.Serialization
 			if (CbFieldUtils.IsBinary(TypeWithFlags))
 			{
 				Error = CbFieldError.None;
-				return BitUtils.GetVarUIntSizedPayload(GetPayloadView());
+
+				ulong Length = VarInt.Read(Memory.Span, out int BytesRead);
+				return Memory.Slice(BytesRead, (int)Length);
 			}
 			else
 			{
@@ -636,13 +721,6 @@ namespace EpicGames.Serialization
 			}
 		}
 
-		/// <summary>
-		/// Access the field as binary data.
-		/// </summary>
-		/// <returns></returns>
-		public ReadOnlyMemory<byte> AsBinaryView(ReadOnlyMemory<byte> Default = default) => AsBinary(Default);
-
-
 		/// <summary>
 		/// Access the field as a UTF-8 string. Returns the provided default on error.
 		/// </summary>
@@ -652,7 +730,7 @@ namespace EpicGames.Serialization
 		{
 			if (CbFieldUtils.IsString(TypeWithFlags))
 			{
-				ulong ValueSize = BitUtils.ReadVarUInt(Payload.Span, out int ValueSizeByteCount);
+				ulong ValueSize = VarInt.Read(Payload.Span, out int ValueSizeByteCount);
 				if (ValueSize >= (1UL << 31))
 				{
 					Error = CbFieldError.RangeError;
@@ -751,7 +829,7 @@ namespace EpicGames.Serialization
 				ulong IsNegative = (ulong)(byte)(TypeWithFlags) & 1;
 
 				int MagnitudeByteCount;
-				ulong Magnitude = BitUtils.ReadVarUInt(Payload.Span, out MagnitudeByteCount);
+				ulong Magnitude = VarInt.Read(Payload.Span, out MagnitudeByteCount);
 				ulong Value = Magnitude ^ (ulong)-(long)(IsNegative);
 
 				if ((Magnitude & OutOfRangeMask) == 0 && (IsNegative == 0 || IsSigned))
@@ -788,7 +866,7 @@ namespace EpicGames.Serialization
 						ulong OutOfRangeMask = ~((1UL << /*FLT_MANT_DIG*/ 24) - 1);
 
 						int MagnitudeByteCount;
-						ulong Magnitude = BitUtils.ReadVarUInt(Payload.Span, out MagnitudeByteCount) + IsNegative;
+						ulong Magnitude = VarInt.Read(Payload.Span, out MagnitudeByteCount) + IsNegative;
 						bool IsInRange = (Magnitude & OutOfRangeMask) == 0;
 						Error = IsInRange ? CbFieldError.None : CbFieldError.RangeError;
 						return IsInRange ? (float)((IsNegative != 0) ? (float)-(long)Magnitude : (float)Magnitude) : Default;
@@ -821,7 +899,7 @@ namespace EpicGames.Serialization
 						ulong OutOfRangeMask = ~((1UL << /*DBL_MANT_DIG*/ 53) - 1);
 
 						int MagnitudeByteCount;
-						ulong Magnitude = BitUtils.ReadVarUInt(Payload.Span, out MagnitudeByteCount) + IsNegative;
+						ulong Magnitude = VarInt.Read(Payload.Span, out MagnitudeByteCount) + IsNegative;
 						bool IsInRange = (Magnitude & OutOfRangeMask) == 0;
 						Error = IsInRange ? CbFieldError.None : CbFieldError.RangeError;
 						return IsInRange ? (double)((IsNegative != 0) ? (double)-(long)Magnitude : (double)Magnitude) : Default;
@@ -1232,22 +1310,16 @@ namespace EpicGames.Serialization
 		/// Create an iterator for the fields of an array or object, otherwise an empty iterator.
 		/// </summary>
 		/// <returns></returns>
-		public CbFieldIterator CreateIterator() => CreateViewIterator();
-
-		/// <summary>
-		/// Create an iterator for the fields of an array or object, otherwise an empty iterator.
-		/// </summary>
-		/// <returns></returns>
-		public CbFieldIterator CreateViewIterator()
+		public CbFieldIterator CreateIterator()
 		{
 			CbFieldType LocalTypeWithFlags = TypeWithFlags;
 			if (CbFieldUtils.HasFields(LocalTypeWithFlags))
 			{
 				ReadOnlyMemory<byte> PayloadBytes = Payload;
 				int PayloadSizeByteCount;
-				int PayloadSize = (int)BitUtils.ReadVarUInt(PayloadBytes.Span, out PayloadSizeByteCount);
+				int PayloadSize = (int)VarInt.Read(PayloadBytes.Span, out PayloadSizeByteCount);
 				PayloadBytes = PayloadBytes.Slice(PayloadSizeByteCount);
-				int NumByteCount = CbFieldUtils.IsArray(LocalTypeWithFlags) ? (int)BitUtils.MeasureVarUInt(PayloadBytes.Span) : 0;
+				int NumByteCount = CbFieldUtils.IsArray(LocalTypeWithFlags) ? (int)VarInt.Measure(PayloadBytes.Span) : 0;
 				if (PayloadSize > NumByteCount)
 				{
 					PayloadBytes = PayloadBytes.Slice(NumByteCount);
@@ -1268,7 +1340,7 @@ namespace EpicGames.Serialization
 		/// <inheritdoc/>
 		public IEnumerator<CbField> GetEnumerator()
 		{
-			for (CbFieldIterator Iter = CreateViewIterator(); Iter; Iter.MoveNext())
+			for (CbFieldIterator Iter = CreateIterator(); Iter; Iter.MoveNext())
 			{
 				yield return Iter.Current;
 			}
@@ -1277,22 +1349,13 @@ namespace EpicGames.Serialization
 		/// <inheritdoc/>
 		IEnumerator IEnumerable.GetEnumerator() => GetEnumerator();
 
-		/// <summary>
-		/// Returns a view of the field, including the type and name when present.
-		/// </summary>
-		/// <returns></returns>
-		internal ReadOnlyMemory<byte> GetView()
-		{
-			return Memory;
-		}
-
 		/// <summary>
 		/// Returns a view of the name and value payload, which excludes the type.
 		/// </summary>
 		/// <returns></returns>
 		private ReadOnlyMemory<byte> GetViewNoType()
 		{
-			int NameSize = CbFieldUtils.HasFieldName(TypeWithFlags) ? NameLen + (int)BitUtils.MeasureVarUInt((uint)NameLen) : 0;
+			int NameSize = CbFieldUtils.HasFieldName(TypeWithFlags) ? NameLen + (int)VarInt.Measure((uint)NameLen) : 0;
 			return Memory.Slice(PayloadOffset - NameSize);
 		}
 
@@ -1337,13 +1400,13 @@ namespace EpicGames.Serialization
 				case CbFieldType.Binary:
 				case CbFieldType.String:
 					{
-						ulong PayloadSize = BitUtils.ReadVarUInt(Payload.Span, out int BytesRead);
+						ulong PayloadSize = VarInt.Read(Payload.Span, out int BytesRead);
 						return PayloadSize + (ulong)BytesRead;
 					}
 				case CbFieldType.IntegerPositive:
 				case CbFieldType.IntegerNegative:
 					{
-						return (ulong)BitUtils.MeasureVarUInt(Payload.Span);
+						return (ulong)VarInt.Measure(Payload.Span);
 					}
 				case CbFieldType.Float32:
 					return 4;
@@ -1368,39 +1431,6 @@ namespace EpicGames.Serialization
 			}
 		}
 
-		/*
-				/// <summary>
-				/// Enumerate fields within this field. (Originally CreateFieldViewIterator).
-				/// </summary>
-				private IEnumerable<CbField> Fields
-				{
-				   get
-				   {
-					   for (CbFieldIterator Iterator = CreateViewIterator(); Iterator; Iterator.MoveNext())
-					   {
-						   yield return Iterator.Current;
-					   }
-				   }
-			   }
-
-			   public CbFieldIterator CreateIterator() => CreateViewIterator();
-
-			   public static CbField MakeView(CbField Field) => Field;
-			   public CbField Find(ReadOnlyUtf8String Name) => FindView(Name);
-			   public CbField FindIgnoreCase(ReadOnlyUtf8String Name) => FindViewIgnoreCase(Name);
-			   public CbField FindView(ReadOnlyUtf8String Name) => this[Name];
-			   public CbField FindViewIgnoreCase(ReadOnlyUtf8String Name) => Fields.FirstOrDefault(Field => ReadOnlyUtf8StringComparer.OrdinalIgnoreCase.Equals(Field.Name, Name)) ?? new CbField();
-
-			   /// <summary>
-			   /// Fetch a field by index. This operation is O(n); prefer to iterate through fields once.
-			   /// </summary>
-			   /// <param name="Index"></param>
-			   /// <returns></returns>
-			   public CbField? this[int Index]
-			   {
-				   get { return Fields.Skip(Index).FirstOrDefault(); }
-			   }
-		*/
 		#region Mimic inheritance from TCbBufferFactory
 
 		public static CbField Clone(ReadOnlyMemory<byte> Data) => Clone(new CbField(Data));
@@ -1687,55 +1717,28 @@ namespace EpicGames.Serialization
 			return !A.Current.Equals(B.Current);
 		}
 	}
-	/*
+
 	/// <summary>
-	/// Iterator for fields
+	/// Simplified view of <see cref="CbArray"/> for display in the debugger
 	/// </summary>
-	public class CbFieldIterator : IEnumerable<CbField>
+	class CbArrayDebugView
 	{
-		/// <summary>
-		/// The underlying buffer
-		/// </summary>
-		ReadOnlyMemory<byte> Data;
-
-		/// <summary>
-		/// Type for all fields
-		/// </summary>
-		CbFieldType UniformType;
+		CbArray Array;
 
-		/// <summary>
-		/// Constructor
-		/// </summary>
-		/// <param name="Data"></param>
-		/// <param name="UniformType"></param>
-		public CbFieldIterator(ReadOnlyMemory<byte> Data, CbFieldType UniformType)
-		{
-			this.Data = Data;
-			this.UniformType = UniformType;
-		}
+		public CbArrayDebugView(CbArray Array) => this.Array = Array;
 
-		/// <inheritdoc/>
-		public IEnumerator<CbField> GetEnumerator()
-		{
-			ReadOnlyMemory<byte> LocalData = Data;
-			while (LocalData.Length > 0)
-			{
-				CbField Field = new CbField(LocalData, UniformType);
-				yield return Field;
-				LocalData = LocalData.Slice(Field.Memory.Length);
-			}
-		}
-
-		/// <inheritdoc/>
-		IEnumerator IEnumerable.GetEnumerator() => GetEnumerator();
+		[DebuggerBrowsable(DebuggerBrowsableState.RootHidden)]
+		public object?[] Value => Array.Select(x => x.Value).ToArray();
 	}
-	*/
+
 	/// <summary>
 	/// Array of CbField that have no names.
 	///
 	/// Accessing a field of the array requires iteration. Access by index is not provided because the
 	/// cost of accessing an item by index scales linearly with the index.
 	/// </summary>
+	[DebuggerDisplay("Count = {Count}")]
+	[DebuggerTypeProxy(typeof(CbArrayDebugView))]
 	public class CbArray : IEnumerable<CbField>
 	{
 		/// <summary>
@@ -1776,14 +1779,16 @@ namespace EpicGames.Serialization
 		}
 
 		/// <summary>
-		/// Returns the number of items in the array.
+		/// Number of items in this array
 		/// </summary>
-		/// <returns></returns>
-		public int Num()
+		public int Count
 		{
-			ReadOnlyMemory<byte> PayloadBytes = InnerField.Payload;
-			PayloadBytes = PayloadBytes.Slice(BitUtils.MeasureVarUInt(PayloadBytes.Span));
-			return (int)BitUtils.ReadVarUInt(PayloadBytes.Span, out int NumByteCount);
+			get
+			{
+				ReadOnlyMemory<byte> PayloadBytes = InnerField.Payload;
+				PayloadBytes = PayloadBytes.Slice(VarInt.Measure(PayloadBytes.Span));
+				return (int)VarInt.Read(PayloadBytes.Span, out int NumByteCount);
+			}
 		}
 
 		/// <summary>
@@ -1792,12 +1797,6 @@ namespace EpicGames.Serialization
 		/// <returns></returns>
 		public CbField AsField() => InnerField;
 
-		/// <summary>
-		/// Access the array as an array field.
-		/// </summary>
-		/// <returns></returns>
-		public CbField AsFieldView() => InnerField;
-
 		/// <summary>
 		/// Construct an array from an array field. No type check is performed!
 		/// </summary>
@@ -1805,12 +1804,6 @@ namespace EpicGames.Serialization
 		/// <returns></returns>
 		public static CbArray FromFieldNoCheck(CbField Field) => new CbArray(Field);
 
-		/// <summary>
-		/// Whether the array has any fields.
-		/// </summary>
-		/// <param name="Array"></param>
-		public static explicit operator bool(CbArray Array) => Array.Num() > 0;
-
 		/// <summary>
 		/// Returns the size of the array in bytes if serialized by itself with no name.
 		/// </summary>
@@ -1878,7 +1871,7 @@ namespace EpicGames.Serialization
 		}
 
 		/** Invoke the visitor for every attachment in the array. */
-		public void IterateAttachments(Action<CbField> Visitor) => CreateViewIterator().IterateRangeAttachments(Visitor);
+		public void IterateAttachments(Action<CbField> Visitor) => CreateIterator().IterateRangeAttachments(Visitor);
 
 		/// <summary>
 		/// Try to get a view of the array as it would be serialized, such as by CopyTo.
@@ -1899,9 +1892,6 @@ namespace EpicGames.Serialization
 		/// <inheritdoc cref="CbField.CreateIterator"/>
 		public CbFieldIterator CreateIterator() => InnerField.CreateIterator();
 
-		/// <inheritdoc cref="CbField.CreateViewIterator"/>
-		public CbFieldIterator CreateViewIterator() => InnerField.CreateViewIterator();
-
 		/// <inheritdoc/>
 		public IEnumerator<CbField> GetEnumerator() => InnerField.GetEnumerator();
 
@@ -1928,6 +1918,26 @@ namespace EpicGames.Serialization
 		#endregion
 	}
 
+	/// <summary>
+	/// Simplified view of <see cref="CbObject"/> for display in the debugger
+	/// </summary>
+	class CbObjectDebugView
+	{
+		[DebuggerDisplay("{Name}: {Value}")]
+		public class Property
+		{
+			public string? Name { get; set; }
+			public object? Value { get; set; }
+		}
+
+		CbObject Object;
+
+		public CbObjectDebugView(CbObject Object) => this.Object = Object;
+
+		[DebuggerBrowsable(DebuggerBrowsableState.RootHidden)]
+		public Property[] Properties => Object.Select(x => new Property { Name = x.Name.ToString(), Value = x.Value }).ToArray();
+	}
+
 	/// <summary>
 	/// Array of CbField that have unique names.
 	///
@@ -1936,6 +1946,7 @@ namespace EpicGames.Serialization
 	/// is not found in the object, the field that it returns has no value (evaluates to false) though
 	/// attempting to access the empty field is also safe, as described by FCbFieldView.
 	/// </summary>
+	[DebuggerTypeProxy(typeof(CbObjectDebugView))]
 	public class CbObject : IEnumerable<CbField>
 	{
 		/// <summary>
@@ -1977,12 +1988,9 @@ namespace EpicGames.Serialization
 			Writer.BeginObject();
 			Build(Writer);
 			Writer.EndObject();
-			return new CbObject(Writer.Save());
+			return new CbObject(Writer.ToByteArray());
 		}
 
-		/// <inheritdoc cref="FindView(ReadOnlyUtf8String)"/>
-		public CbField Find(ReadOnlyUtf8String Name) => FindView(Name);
-
 		/// <summary>
 		/// Find a field by case-sensitive name comparison.
 		///
@@ -1991,17 +1999,14 @@ namespace EpicGames.Serialization
 		/// </summary>
 		/// <param name="Name">The name of the field.</param>
 		/// <returns>The matching field if found, otherwise a field with no value.</returns>
-		public CbField FindView(ReadOnlyUtf8String Name) => InnerField[Name];
-
-		/// <inheritdoc cref="FindView(ReadOnlyUtf8String)"/>
-		public CbField FindIgnoreCase(ReadOnlyUtf8String Name) => FindViewIgnoreCase(Name);
+		public CbField Find(ReadOnlyUtf8String Name) => InnerField[Name];
 
 		/// <summary>
 		/// Find a field by case-insensitive name comparison.
 		/// </summary>
 		/// <param name="Name">The name of the field.</param>
 		/// <returns>The matching field if found, otherwise a field with no value.</returns>
-		public CbField FindViewIgnoreCase(ReadOnlyUtf8String Name) => InnerField.FirstOrDefault(Field => ReadOnlyUtf8StringComparer.OrdinalIgnoreCase.Equals(Field.Name, Name)) ?? new CbField();
+		public CbField FindIgnoreCase(ReadOnlyUtf8String Name) => InnerField.FirstOrDefault(Field => ReadOnlyUtf8StringComparer.OrdinalIgnoreCase.Equals(Field.Name, Name)) ?? new CbField();
 
 		/// <summary>
 		/// Find a field by case-sensitive name comparison.
@@ -2013,12 +2018,6 @@ namespace EpicGames.Serialization
 		/// <inheritdoc cref="AsFieldView"/>
 		public CbField AsField() => InnerField;
 
-		/// <summary>
-		/// Access the object as an object field.
-		/// </summary>
-		/// <returns></returns>
-		public CbField AsFieldView() => InnerField;
-
 		/// <summary>
 		/// Construct an object from an object field. No type check is performed!
 		/// </summary>
@@ -2026,15 +2025,6 @@ namespace EpicGames.Serialization
 		/// <returns></returns>
 		public static CbObject FromFieldNoCheck(CbField Field) => new CbObject(Field);
 
-		/// <summary>
-		/// Whether the object has any fields.
-		/// </summary>
-		/// <param name="Object"></param>
-		public static explicit operator bool(CbObject Object)
-		{
-			throw new NotImplementedException();
-		}
-
 		/// <summary>
 		/// Returns the size of the object in bytes if serialized by itself with no name.
 		/// </summary>
@@ -2106,7 +2096,7 @@ namespace EpicGames.Serialization
 		/// Invoke the visitor for every attachment in the object.
 		/// </summary>
 		/// <param name="Visitor"></param>
-		public void IterateAttachments(Action<CbField> Visitor) => CreateViewIterator().IterateRangeAttachments(Visitor);
+		public void IterateAttachments(Action<CbField> Visitor) => CreateIterator().IterateRangeAttachments(Visitor);
 
 		/// <summary>
 		/// Creates a view of the object, excluding the name
@@ -2145,9 +2135,6 @@ namespace EpicGames.Serialization
 		/// <inheritdoc cref="CbField.CreateIterator"/>
 		public CbFieldIterator CreateIterator() => InnerField.CreateIterator();
 
-		/// <inheritdoc cref="CbField.CreateViewIterator"/>
-		public CbFieldIterator CreateViewIterator() => InnerField.CreateViewIterator();
-
 		/// <inheritdoc/>
 		public IEnumerator<CbField> GetEnumerator() => InnerField.GetEnumerator();
 
@@ -2250,114 +2237,4 @@ namespace EpicGames.Serialization
 		}
 		#endregion
 	}
-
-	/// <summary>
-	/// Methods for reading VarUInt values
-	/// </summary>
-	public static class BitUtils
-	{
-		/// <summary>
-		/// Extracts the payload from a buffer prefixed with a variable-length value
-		/// </summary>
-		/// <param name="Memory"></param>
-		/// <returns></returns>
-		public static ReadOnlyMemory<byte> GetVarUIntSizedPayload(ReadOnlyMemory<byte> Memory)
-		{
-			ulong Length = ReadVarUInt(Memory.Span, out int BytesRead);
-			return Memory.Slice(BytesRead, (int)Length);
-		}
-
-		/// <summary>
-		/// Read a variable-length unsigned integer.
-		/// </summary>
-		/// <param name="Buffer">A variable-length encoding of an unsigned integer</param>
-		/// <param name="BytesRead">The number of bytes consumed from the input</param>
-		/// <returns></returns>
-		public static ulong ReadVarUInt(ReadOnlySpan<byte> Buffer, out int BytesRead)
-		{
-			BytesRead = (int)MeasureVarUInt(Buffer);
-
-			ulong Value = (ulong)(Buffer[0] & (0xff >> BytesRead));
-			for (int i = 1; i < BytesRead; i++)
-			{
-				Value <<= 8;
-				Value |= Buffer[i];
-			}
-			return Value;
-		}
-
-		/// <summary>
-		/// Measure the length in bytes (1-9) of an encoded variable-length integer.
-		/// </summary>
-		/// <param name="Buffer">A variable-length encoding of an(signed or unsigned) integer.</param>
-		/// <returns>The number of bytes used to encode the integer, in the range 1-9.</returns>
-		public static int MeasureVarUInt(ReadOnlySpan<byte> Buffer)
-		{
-			byte b = Buffer[0];
-			b = (byte)~b;
-			return BitOperations.LeadingZeroCount(b) - 23;
-		}
-
-		/// <summary>
-		/// Measure the number of bytes (1-5) required to encode the 32-bit input.
-		/// </summary>
-		/// <param name="Value"></param>
-		/// <returns></returns>
-		public static int MeasureVarUInt(int Value)
-		{
-			return MeasureVarUInt((uint)Value);
-		}
-
-		/// <summary>
-		/// Measure the number of bytes (1-5) required to encode the 32-bit input.
-		/// </summary>
-		/// <param name="Value"></param>
-		/// <returns></returns>
-		public static int MeasureVarUInt(uint Value)
-		{
-			return BitOperations.Log2(Value) / 7 + 1;
-		}
-
-		/// <summary>
-		/// Measure the number of bytes (1-9) required to encode the 64-bit input.
-		/// </summary>
-		/// <param name="value"></param>
-		/// <returns></returns>
-		public static int MeasureVarUInt(ulong value)
-		{
-			return Math.Min(BitOperations.Log2(value) / 7 + 1, 9);
-		}
-
-		/// <summary>
-		/// Write a variable-length unsigned integer.
-		/// </summary>
-		/// <param name="Value">An unsigned integer to encode</param>
-		/// <param name="Buffer">A buffer of at least 9 bytes to write the output to.</param>
-		/// <param name="BufferOffset"></param>
-		/// <returns>The number of bytes used in the output</returns>
-		public static int WriteVarUInt(long Value, byte[] Buffer, int BufferOffset = 0)
-		{
-			return WriteVarUInt((ulong)Value, Buffer, BufferOffset);
-		}
-
-		/// <summary>
-		/// Write a variable-length unsigned integer.
-		/// </summary>
-		/// <param name="Value">An unsigned integer to encode</param>
-		/// <param name="Buffer">A buffer of at least 9 bytes to write the output to.</param>
-		/// <param name="BufferOffset"></param>
-		/// <returns>The number of bytes used in the output</returns>
-		public static int WriteVarUInt(ulong Value, byte[] Buffer, int BufferOffset = 0)
-		{
-			int ByteCount = MeasureVarUInt(Value);
-
-			for (uint Idx = 1; Idx < ByteCount; Idx++)
-			{
-				Buffer[BufferOffset + ByteCount - Idx] = (byte)Value;
-				Value >>= 8;
-			}
-			Buffer[BufferOffset] = (byte)((0xff << (9 - (int)ByteCount)) | (byte)Value);
-			return ByteCount;
-		}
-	}
 }
diff --git a/Engine/Source/Programs/Shared/EpicGames.Serialization/CompactBinaryWriter.cs b/Engine/Source/Programs/Shared/EpicGames.Serialization/CompactBinaryWriter.cs
--- a/Engine/Source/Programs/Shared/EpicGames.Serialization/CompactBinaryWriter.cs
+++ b/Engine/Source/Programs/Shared/EpicGames.Serialization/CompactBinaryWriter.cs
@@ -2,317 +2,660 @@
 
 using EpicGames.Core;
 using System;
+using System.Buffers.Binary;
 using System.Collections.Generic;
 using System.Linq;
 using System.Text;
+using System.Text.Json;
 
 namespace EpicGames.Serialization
 {
-	public class CbWriter
+	/// <summary>
+	/// Exception for <see cref="CbWriter"/>
+	/// </summary>
+	public class CbWriterException : Exception
 	{
-		readonly List<Field> _fields = new List<Field>();
-		private Field _topField;
+		/// <summary>
+		/// Constructor
+		/// </summary>
+		/// <param name="Message"></param>
+		public CbWriterException(string Message)
+			: base(Message)
+		{
+		}
 
-		struct Field
+		/// <summary>
+		/// Constructor
+		/// </summary>
+		/// <param name="Message"></param>
+		/// <param name="Ex"></param>
+		public CbWriterException(string Message, Exception? Ex)
+			: base(Message, Ex)
+		{
+		}
+	}
+
+	/// <summary>
+	/// Forward-only writer for compact binary objects
+	/// </summary>
+	public class CbWriter
+	{
+		/// <summary>
+		/// Stores information about an object or array scope within the written buffer which requires a header to be inserted containing
+		/// the size or number of elements when copied to an output buffer
+		/// </summary>
+		class Scope
 		{
 			public CbFieldType FieldType;
-			public byte[]? Name;
-			public byte[]? Payload;
+			public CbFieldType UniformFieldType;
+			public int Offset; // Offset to insert the length/count
+			public int Length; // Excludes the size of this field's headers, and child fields' headers.
+			public int Count;
+			public List<Scope>? Children;
+			public int SizeOfChildHeaders; // Sum of additional headers for child items, recursively.
+
+			public Scope(CbFieldType FieldType, CbFieldType UniformFieldType, int Offset)
+			{
+				this.FieldType = FieldType;
+				this.UniformFieldType = UniformFieldType;
+				this.Offset = Offset;
+			}
 		}
 
-		public CbObject ToObject()
+		/// <summary>
+		/// Chunk of written data. Chunks are allocated as needed and chained together with scope annotations to produce the output data.
+		/// </summary>
+		class Chunk
 		{
-			return new CbObject(Save());
+			public int Offset;
+			public int Length;
+			public byte[] Data;
+			public List<Scope> Scopes = new List<Scope>();
+
+			public Chunk(int Offset, int MaxLength)
+			{
+				this.Offset = Offset;
+				this.Data = new byte[MaxLength];
+			}
 		}
 
-		public byte[] Save()
+		const int DefaultChunkSize = 1024;
+
+		List<Chunk> Chunks = new List<Chunk>();
+		Stack<Scope> OpenScopes = new Stack<Scope>();
+		Chunk CurrentChunk => Chunks[Chunks.Count - 1];
+		Scope CurrentScope => OpenScopes.Peek();
+		int CurrentOffset;
+
+		/// <summary>
+		/// Constructor
+		/// </summary>
+		public CbWriter()
+			: this(DefaultChunkSize)
 		{
-			int bufferLength = 1; // field type;
+		}
 
-			if (_topField.Name != null)
+		/// <summary>
+		/// Constructor
+		/// </summary>
+		/// <param name="Reserve">Amount of data to reserve for output</param>
+		public CbWriter(int Reserve)
+		{
+			Chunks.Add(new Chunk(0, Reserve));
+			OpenScopes.Push(new Scope(CbFieldType.Array, CbFieldType.None, 0));
+		}
+
+		/// <summary>
+		/// Ensure that a block of contiguous memory of the given length is available in the output buffer
+		/// </summary>
+		/// <param name="Length"></param>
+		/// <returns>The allocated memory</returns>
+		Memory<byte> Allocate(int Length)
+		{
+			Chunk LastChunk = CurrentChunk;
+			if (LastChunk.Length + Length > LastChunk.Data.Length)
 			{
-				bufferLength += _topField.Name!.Length;
+				int ChunkSize = Math.Max(Length, DefaultChunkSize);
+				LastChunk = new Chunk(CurrentOffset, ChunkSize);
+				Chunks.Add(LastChunk);
 			}
 
-			int payloadBytesLength = 0;
-			if (_topField.Payload != null)
+			Memory<byte> Buffer = LastChunk.Data.AsMemory(LastChunk.Length, Length);
+			LastChunk.Length += Length;
+			CurrentOffset += Length;
+			return Buffer;
+		}
+
+		/// <summary>
+		/// Insert a new scope
+		/// </summary>
+		/// <param name="FieldType"></param>
+		/// <param name="UniformFieldType"></param>
+		void PushScope(CbFieldType FieldType, CbFieldType UniformFieldType)
+		{
+			Scope NewScope = new Scope(FieldType, UniformFieldType, CurrentOffset);
+			CurrentScope.Children ??= new List<Scope>();
+			CurrentScope.Children.Add(NewScope);
+			OpenScopes.Push(NewScope);
+
+			CurrentChunk.Scopes.Add(NewScope);
+		}
+
+		/// <summary>
+		/// Pop a scope from the current open list
+		/// </summary>
+		void PopScope()
+		{
+			Scope Scope = CurrentScope;
+			Scope.Length = CurrentOffset - Scope.Offset;
+			Scope.SizeOfChildHeaders = 0;
+
+			if (Scope.Children != null)
 			{
-				payloadBytesLength = BitUtils.MeasureVarUInt((uint)_topField.Payload!.Length);
-				bufferLength += payloadBytesLength;
-				bufferLength += _topField.Payload!.Length;
+				foreach (Scope ChildScope in Scope.Children)
+				{
+					switch (ChildScope.FieldType)
+					{
+						case CbFieldType.Object:
+						case CbFieldType.UniformObject:
+							Scope.SizeOfChildHeaders += ChildScope.SizeOfChildHeaders + VarInt.Measure(ChildScope.Length + ChildScope.SizeOfChildHeaders);
+							break;
+						case CbFieldType.Array:
+						case CbFieldType.UniformArray:
+							int ArrayCountLength = VarInt.Measure(Scope.Count);
+							Scope.SizeOfChildHeaders += ChildScope.SizeOfChildHeaders + VarInt.Measure(ChildScope.Length + ChildScope.SizeOfChildHeaders + ArrayCountLength) + ArrayCountLength;
+							break;
+						default:
+							throw new InvalidOperationException();
+					}
+				}
 			}
-			byte[] buffer = new byte[bufferLength];
-			buffer[0] = (byte)_topField.FieldType;
-			int offset = 1;
 
-			if (_topField.Name != null)
+			OpenScopes.Pop();
+		}
+
+		/// <summary>
+		/// Writes the header for an unnamed field
+		/// </summary>
+		/// <param name="Type"></param>
+		void WriteField(CbFieldType Type)
+		{
+			Scope Scope = CurrentScope;
+			if (!CbFieldUtils.IsArray(Scope.FieldType))
 			{
-				Array.Copy(_topField.Name!, 0, buffer, offset, _topField.Name!.Length);
-				offset += _topField.Name.Length;
+				throw new CbWriterException($"Anonymous fields are not allowed within fields of type {Scope.FieldType}");
 			}
+			
+			if (Scope.UniformFieldType == CbFieldType.None)
+			{
+				Allocate(1).Span[0] = (byte)Type;
+			}
+			else if (Scope.UniformFieldType != Type)
+			{
+				throw new CbWriterException($"Mismatched type for uniform array - expected {Scope.UniformFieldType}, not {Type}");
+			}
+			Scope.Count++;
+		}
 
-			if (_topField.Payload != null)
+		/// <summary>
+		/// Writes the header for a named field
+		/// </summary>
+		/// <param name="Type"></param>
+		/// <param name="Name"></param>
+		void WriteField(CbFieldType Type, ReadOnlyUtf8String Name)
+		{
+			Scope Scope = CurrentScope;
+			if (!CbFieldUtils.IsObject(Scope.FieldType))
 			{
-				BitUtils.WriteVarUInt((uint)_topField.Payload!.Length, buffer, offset);
-				offset += (int)payloadBytesLength;
-				Array.Copy(_topField.Payload!, 0, buffer, offset, _topField.Payload!.Length);
-				offset += _topField.Payload.Length;
+				throw new CbWriterException($"Named fields are not allowed within fields of type {Scope.FieldType}");
 			}
 
-			return buffer;
+			int NameVarIntLength = VarInt.Measure(Name.Length);
+			if (Scope.UniformFieldType == CbFieldType.None)
+			{
+				Span<byte> Buffer = Allocate(1 + NameVarIntLength + Name.Length).Span;
+				Buffer[0] = (byte)(Type | CbFieldType.HasFieldName);
+				WriteBinaryPayload(Buffer[1..], Name.Span);
+			}
+			else
+			{
+				if (Scope.UniformFieldType != Type)
+				{
+					throw new CbWriterException($"Mismatched type for uniform object - expected {Scope.UniformFieldType}, not {Type}");
+				}
+				WriteBinaryPayload(Name.Span);
+			}
+			Scope.Count++;
 		}
 
+		/// <summary>
+		/// Begin writing an object field
+		/// </summary>
 		public void BeginObject()
 		{
-			_topField = BeginField();
+			WriteField(CbFieldType.Object);
+			PushScope(CbFieldType.Object, CbFieldType.None);
 		}
 
-		public void EndObject()
+		/// <summary>
+		/// Begin writing an object field
+		/// </summary>
+		/// <param name="Name">Name of the field</param>
+		public void BeginObject(ReadOnlyUtf8String Name)
 		{
-			_topField.Payload = SerializePayload(_fields);
+			WriteField(CbFieldType.Object, Name);
+			PushScope(CbFieldType.Object, CbFieldType.None);
+		}
 
-			EndField(ref this._topField, CbFieldType.Object);
+		/// <summary>
+		/// End the current object
+		/// </summary>
+		public void EndObject()
+		{
+			PopScope();
 		}
 
-		private static byte[] SerializePayload(List<Field> fields)
+		/// <summary>
+		/// Begin writing an array field
+		/// </summary>
+		public void BeginArray()
 		{
-			uint totalSizeFields = 0;
-			foreach (Field field in fields)
-			{
-				totalSizeFields += 1; // the field type
-				if (field.Payload != null)
-					totalSizeFields += (uint)field.Payload.Length;
-				if (field.Name != null)
-					totalSizeFields += (uint)field.Name.Length;
-			}
+			WriteField(CbFieldType.Array);
+			PushScope(CbFieldType.Array, CbFieldType.None);
+		}
 
-			byte[] buffer = new byte[totalSizeFields];
-			long payloadOffset = 0;
-			foreach (Field field in fields)
-			{
-				buffer[payloadOffset] = (byte)field.FieldType;
-				payloadOffset += 1;
+		/// <summary>
+		/// Begin writing a named array field
+		/// </summary>
+		/// <param name="Name"></param>
+		public void BeginArray(ReadOnlyUtf8String Name)
+		{
+			WriteField(CbFieldType.Array, Name);
+			PushScope(CbFieldType.Array, CbFieldType.None);
+		}
 
-				if (CbFieldUtils.HasFieldName(field.FieldType))
-				{
-					Array.Copy(field.Name!, 0, buffer, payloadOffset, field.Name!.Length);
-					payloadOffset += field.Name.Length;
-				}
+		/// <summary>
+		/// End the current array
+		/// </summary>
+		public void EndArray()
+		{
+			PopScope();
+		}
 
-				if ((field.Payload?.Length ?? 0) != 0)
-				{
-					Array.Copy(field.Payload!, 0, buffer, payloadOffset, field.Payload!.Length);
-					payloadOffset += field.Payload.Length;
-				}
-			}
+		/// <summary>
+		/// Begin writing a uniform array field
+		/// </summary>
+		/// <param name="FieldType">The field type for elements in the array</param>
+		public void BeginUniformArray(CbFieldType FieldType)
+		{
+			WriteField(CbFieldType.UniformArray);
+			PushScope(CbFieldType.UniformArray, FieldType);
+			Allocate(1).Span[0] = (byte)FieldType;
+		}
 
-			return buffer;
+		/// <summary>
+		/// Begin writing a named uniform array field
+		/// </summary>
+		/// <param name="Name">Name of the field</param>
+		/// <param name="FieldType">The field type for elements in the array</param>
+		public void BeginUniformArray(ReadOnlyUtf8String Name, CbFieldType FieldType)
+		{
+			WriteField(CbFieldType.UniformArray, Name);
+			PushScope(CbFieldType.UniformArray, FieldType);
+			Allocate(1).Span[0] = (byte)FieldType;
 		}
 
-		public void AddString(string value, string? fieldName = null)
+		/// <summary>
+		/// End the current array
+		/// </summary>
+		public void EndUniformArray()
 		{
-			Field field = BeginField();
-			SetName(ref field, fieldName);
+			PopScope();
+		}
 
-			int length = Encoding.UTF8.GetByteCount(value);
-			int bytesCount = BitUtils.MeasureVarUInt(length);
-			byte[] buffer = new byte[length + bytesCount];
-			BitUtils.WriteVarUInt(length, buffer);
-			int bytesWritten = Encoding.UTF8.GetBytes(value, 0, value.Length, buffer, (int)bytesCount);
-			if (bytesWritten != length)
-				throw new Exception("Failed to write string into buffer");
+		/// <summary>
+		/// Write a null field
+		/// </summary>
+		public void WriteNullValue()
+		{
+			WriteField(CbFieldType.Null);
+		}
 
-			field.Payload = buffer;
-			EndField(ref field, CbFieldType.String);
+		/// <summary>
+		/// Write a named null field
+		/// </summary>
+		/// <param name="Name">Name of the field</param>
+		public void WriteNull(ReadOnlyUtf8String Name)
+		{
+			WriteField(CbFieldType.Null, Name);
 		}
 
+		/// <summary>
+		/// Writes the payload for an integer
+		/// </summary>
+		/// <param name="Value">Value to write</param>
+		void WriteIntegerPayload(ulong Value)
+		{
+			int Length = VarInt.Measure(Value);
+			Span<byte> Buffer = Allocate(Length).Span;
+			VarInt.Write(Buffer, Value);
+		}
 
-		private Field BeginField()
+		/// <summary>
+		/// Writes an unnamed integer field
+		/// </summary>
+		/// <param name="Value">Value to be written</param>
+		public void WriteIntegerValue(long Value)
 		{
-			// TODO: We could add checks similar to what we do in C++ to make sure the writer is used correctly
-			return new Field();
+			if (Value >= 0)
+			{
+				WriteField(CbFieldType.IntegerPositive);
+				WriteIntegerPayload((ulong)Value);
+			}
+			else
+			{
+				WriteField(CbFieldType.IntegerNegative);
+				WriteIntegerPayload((ulong)-Value);
+			}
 		}
 
-		private void EndField(ref Field field, CbFieldType fieldType)
+		/// <summary>
+		/// Writes an named integer field
+		/// </summary>
+		/// <param name="Name">Name of the field</param>
+		/// <param name="Value">Value to be written</param>
+		public void WriteInteger(ReadOnlyUtf8String Name, long Value)
 		{
-			field.FieldType = fieldType;
-			if (field.Name != null && field.Name.Length != 0)
+			if (Value >= 0)
 			{
-				field.FieldType |= CbFieldType.HasFieldName;
+				WriteField(CbFieldType.IntegerPositive, Name);
+				WriteIntegerPayload((ulong)Value);
 			}
-			// TODO: check if types are uniform and check the field type to be uniform thus removing redudant type info
+			else
+			{
+				WriteField(CbFieldType.IntegerNegative, Name);
+				WriteIntegerPayload((ulong)-Value);
+			}
+		}
 
-			_fields.Add(field);
+		/// <summary>
+		/// Writes an unnamed integer field
+		/// </summary>
+		/// <param name="Value">Value to be written</param>
+		public void WriteIntegerValue(ulong Value)
+		{
+			WriteField(CbFieldType.IntegerPositive);
+			WriteIntegerPayload(Value);
 		}
 
-		public void AddBinaryAttachment(IoHash hash, string? fieldName = null)
+		/// <summary>
+		/// Writes a named integer field
+		/// </summary>
+		/// <param name="Name">Name of the field</param>
+		/// <param name="Value">Value to be written</param>
+		public void WriteInteger(ReadOnlyUtf8String Name, ulong Value)
 		{
-			Field field = BeginField();
-			SetName(ref field, fieldName);
+			WriteField(CbFieldType.IntegerPositive, Name);
+			WriteIntegerPayload(Value);
+		}
 
-			if (hash.Memory.Length != 20)
-				throw new Exception("Hash data is assumed to be 20 bytes");
-			field.Payload = hash.Memory.ToArray();
+		/// <summary>
+		/// Writes the payload for a <see cref="DateTime"/> value
+		/// </summary>
+		/// <param name="DateTime">The value to write</param>
+		void WriteDateTimePayload(DateTime DateTime)
+		{
+			Span<byte> Buffer = Allocate(sizeof(long)).Span;
+			BinaryPrimitives.WriteInt64BigEndian(Buffer, DateTime.Ticks);
+		}
 
-			EndField(ref field, CbFieldType.BinaryAttachment);
+		/// <summary>
+		/// Writes an unnamed <see cref="DateTime"/> field
+		/// </summary>
+		/// <param name="Value">Value to be written</param>
+		public void WriteDateTimeValue(DateTime Value)
+		{
+			WriteField(CbFieldType.DateTime);
+			WriteDateTimePayload(Value);
 		}
 
-		public void AddCompactBinaryAttachment(IoHash hash, string? fieldName = null)
+		/// <summary>
+		/// Writes a named <see cref="DateTime"/> field
+		/// </summary>
+		/// <param name="Name">Name of the field</param>
+		/// <param name="Value">Value to be written</param>
+		public void WriteDateTime(ReadOnlyUtf8String Name, DateTime Value)
 		{
-			Field field = BeginField();
-			SetName(ref field, fieldName);
+			WriteField(CbFieldType.DateTime, Name);
+			WriteDateTimePayload(Value);
+		}
 
-			if (hash.Memory.Length != 20)
-				throw new Exception("Hash data is assumed to be 20 bytes");
-			field.Payload = hash.Memory.ToArray();
+		/// <summary>
+		/// Writes the payload for a hash
+		/// </summary>
+		/// <param name="Hash"></param>
+		void WriteHashPayload(IoHash Hash)
+		{
+			Span<byte> Buffer = Allocate(IoHash.NumBytes).Span;
+			Hash.Span.CopyTo(Buffer);
+		}
 
-			EndField(ref field, CbFieldType.ObjectAttachment);
+		/// <summary>
+		/// Writes an unnamed <see cref="IoHash"/> field
+		/// </summary>
+		/// <param name="Hash"></param>
+		public void WriteHashValue(IoHash Hash)
+		{
+			WriteField(CbFieldType.Hash);
+			WriteHashPayload(Hash);
 		}
 
-		public void AddNull(string? fieldName = null)
+		/// <summary>
+		/// Writes a named <see cref="IoHash"/> field
+		/// </summary>
+		/// <param name="Name">Name of the field</param>
+		/// <param name="Value">Value to be written</param>
+		public void WriteHash(ReadOnlyUtf8String Name, IoHash Value)
 		{
-			Field field = BeginField();
-			SetName(ref field, fieldName);
-			EndField(ref field, CbFieldType.Null);
+			WriteField(CbFieldType.Hash, Name);
+			WriteHashPayload(Value);
 		}
 
+		/// <summary>
+		/// Writes an unnamed reference to a binary attachment
+		/// </summary>
+		/// <param name="Hash">Hash of the attachment</param>
+		public void WriteBinaryAttachmentValue(IoHash Hash)
+		{
+			WriteField(CbFieldType.BinaryAttachment);
+			WriteHashPayload(Hash);
+		}
 
-		public void AddDateTime(DateTime dateTime, string? fieldName = null)
+		/// <summary>
+		/// Writes a named reference to a binary attachment
+		/// </summary>
+		/// <param name="Name">Name of the field</param>
+		/// <param name="Hash">Hash of the attachment</param>
+		public void WriteBinaryAttachment(ReadOnlyUtf8String Name, IoHash Hash)
 		{
-			Field field = BeginField();
-			SetName(ref field, fieldName);
-			long ticks = dateTime.Ticks;
-			byte[] payload = BitConverter.GetBytes(ticks);
-			if (BitConverter.IsLittleEndian)
-			{
-				Array.Reverse(payload, 0, payload.Length);
-			}
-			field.Payload = payload;
+			WriteField(CbFieldType.BinaryAttachment, Name);
+			WriteHashPayload(Hash);
+		}
 
-			EndField(ref field, CbFieldType.DateTime);
+		/// <summary>
+		/// Writes an unnamed reference to an object attachment
+		/// </summary>
+		/// <param name="Hash">Hash of the attachment</param>
+		public void WriteObjectAttachmentValue(IoHash Hash)
+		{
+			WriteField(CbFieldType.ObjectAttachment);
+			WriteHashPayload(Hash);
 		}
 
+		/// <summary>
+		/// Writes a named reference to an object attachment
+		/// </summary>
+		/// <param name="Name">Name of the field</param>
+		/// <param name="Hash">Hash of the attachment</param>
+		public void WriteObjectAttachment(ReadOnlyUtf8String Name, IoHash Hash)
+		{
+			WriteField(CbFieldType.ObjectAttachment, Name);
+			WriteHashPayload(Hash);
+		}
 
-		public void AddBinary(byte[] data, string? fieldName = null)
+		/// <summary>
+		/// Writes the payload for a binary value
+		/// </summary>
+		/// <param name="Output">Output buffer</param>
+		/// <param name="Value">Value to be written</param>
+		static void WriteBinaryPayload(Span<byte> Output, ReadOnlySpan<byte> Value)
 		{
-			Field field = BeginField();
-			SetName(ref field, fieldName);
+			int VarIntLength = VarInt.Write(Output, Value.Length);
+			Output = Output[VarIntLength..];
 
-			long bytesLength = BitUtils.MeasureVarUInt((uint)data.Length);
-			byte[] payload = new byte[data.Length + bytesLength];
-			BitUtils.WriteVarUInt((uint)data.Length, payload);
-			Array.Copy(data, 0, payload, bytesLength, data.Length);
-			field.Payload = payload;
-			EndField(ref field, CbFieldType.Binary);
+			Value.CopyTo(Output);
+			CheckSize(Output, Value.Length);
 		}
 
-		public void AddUniformArray<T>(T[] members, CbFieldType fieldType, string? fieldName = null)
+		/// <summary>
+		/// Writes the payload for a binary value
+		/// </summary>
+		/// <param name="Value">Value to be written</param>
+		void WriteBinaryPayload(ReadOnlySpan<byte> Value)
 		{
-			Field field = BeginField();
-			SetName(ref field, fieldName);
+			int ValueVarIntLength = VarInt.Measure(Value.Length);
+			Span<byte> Buffer = Allocate(ValueVarIntLength + Value.Length).Span;
+			WriteBinaryPayload(Buffer, Value);
+		}
 
-			int bytesMembers = BitUtils.MeasureVarUInt((uint)members.Length);
+		/// <summary>
+		/// Writes an unnamed string value
+		/// </summary>
+		/// <param name="Value">Value to be written</param>
+		public void WriteStringValue(ReadOnlyUtf8String Value)
+		{
+			WriteField(CbFieldType.String);
+			WriteBinaryPayload(Value.Span);
+		}
 
-			List<byte[]> payloads = new List<byte[]>();
+		/// <summary>
+		/// Writes a named string value
+		/// </summary>
+		/// <param name="Name">Name of the field</param>
+		/// <param name="Value">Value to be written</param>
+		public void WriteString(ReadOnlyUtf8String Name, ReadOnlyUtf8String Value)
+		{
+			WriteField(CbFieldType.String, Name);
+			WriteBinaryPayload(Value.Span);
+		}
 
-			foreach (T member in members)
+		/// <summary>
+		/// Writes an unnamed binary value
+		/// </summary>
+		/// <param name="Value">Value to be written</param>
+		public void WriteBinaryValue(ReadOnlySpan<byte> Value)
+		{
+			WriteField(CbFieldType.Binary);
+			WriteBinaryPayload(Value);
+		}
+
+		/// <summary>
+		/// Writes a named binary value
+		/// </summary>
+		/// <param name="Name">Name of the field</param>
+		/// <param name="Value">Value to be written</param>
+		public void WriteBinary(ReadOnlyUtf8String Name, ReadOnlySpan<byte> Value)
+		{
+			WriteField(CbFieldType.Binary, Name);
+			WriteBinaryPayload(Value);
+		}
+
+		/// <summary>
+		/// Check that the given span is the required size
+		/// </summary>
+		/// <param name="Span"></param>
+		/// <param name="ExpectedSize"></param>
+		static void CheckSize(Span<byte> Span, int ExpectedSize)
+		{
+			if (Span.Length != ExpectedSize)
 			{
-				byte[] buffer = GetByteValue<T>(member, fieldType);
-				payloads.Add(buffer);
+				throw new Exception("Size of buffer is not correct");
 			}
+		}
 
-			int payloadsBytesLength = payloads.Sum(payload => payload.Length);
-			int payloadsLength = 1 + bytesMembers + payloadsBytesLength;
-			int bytesPayloads = BitUtils.MeasureVarUInt(payloadsLength);
-			
-			// Uniform array payload is a VarUInt byte count, followed by a VarUInt item count, followed by field type, followed by the fields without their field type.
-			byte[] payloadBuffer = new byte[payloadsLength + bytesPayloads];
-			BitUtils.WriteVarUInt(payloadsLength, payloadBuffer);
-			int bufferOffset = (int)bytesPayloads;
-			BitUtils.WriteVarUInt((uint)members.Length, payloadBuffer, bufferOffset);
-			bufferOffset += (int)bytesMembers;
-			payloadBuffer[bufferOffset] = (byte)fieldType;
-			bufferOffset += 1;
-
-			foreach (byte[] payload in payloads)
+		/// <summary>
+		/// Gets the size of the serialized data
+		/// </summary>
+		/// <returns></returns>
+		public int GetSize()
+		{
+			if (OpenScopes.Count != 1)
 			{
-				int payloadLength = payload.Length;
-				Array.Copy(payload, 0, payloadBuffer, bufferOffset, payloadLength);
-				bufferOffset += payloadLength;
+				throw new CbWriterException("Unfinished scope in writer");
+			}
+			return CurrentOffset + CurrentScope.SizeOfChildHeaders;
+		}
+
+		/// <summary>
+		/// Copy the data from this writer to a buffer
+		/// </summary>
+		/// <param name="Buffer"></param>
+		public void CopyTo(Span<byte> Buffer)
+		{
+			int BufferOffset = 0;
+
+			int SourceOffset = 0;
+			foreach (Chunk Chunk in Chunks)
+			{
+				foreach (Scope Scope in Chunk.Scopes)
+				{
+					ReadOnlySpan<byte> SourceData = Chunk.Data.AsSpan(SourceOffset - Chunk.Offset, Scope.Offset - SourceOffset);
+					SourceData.CopyTo(Buffer.Slice(BufferOffset));
+
+					BufferOffset += SourceData.Length;
+					SourceOffset += SourceData.Length;
+
+					BufferOffset += WriteScopeHeader(Buffer.Slice(BufferOffset), Scope);
+				}
+
+				ReadOnlySpan<byte> LastSourceData = Chunk.Data.AsSpan(SourceOffset - Chunk.Offset, (Chunk.Offset + Chunk.Length) - SourceOffset);
+				LastSourceData.CopyTo(Buffer.Slice(BufferOffset));
+				BufferOffset += LastSourceData.Length;
 			}
+		}
 
-			field.Payload = payloadBuffer;
-			EndField(ref field, CbFieldType.UniformArray);
+		/// <summary>
+		/// Convert the data into a flat array
+		/// </summary>
+		/// <returns></returns>
+		public byte[] ToByteArray()
+		{
+			byte[] Buffer = new byte[GetSize()];
+			CopyTo(Buffer);
+			return Buffer;
 		}
 
-		private byte[] GetByteValue<T>(T member, CbFieldType fieldType)
+		/// <summary>
+		/// Writes the header for a particular scope
+		/// </summary>
+		/// <param name="Span"></param>
+		/// <param name="Scope"></param>
+		/// <returns></returns>
+		static int WriteScopeHeader(Span<byte> Span, Scope Scope)
 		{
-			switch (fieldType)
+			switch (Scope.FieldType)
 			{
-				case CbFieldType.None:
-				case CbFieldType.Null:
-					return Array.Empty<byte>();
 				case CbFieldType.Object:
 				case CbFieldType.UniformObject:
+					return VarInt.Write(Span, Scope.Length + Scope.SizeOfChildHeaders);
 				case CbFieldType.Array:
 				case CbFieldType.UniformArray:
-					throw new NotImplementedException();
-				case CbFieldType.Binary:
-					throw new NotImplementedException();
-				case CbFieldType.String:
-					if (member is string s)
-					{
-						return Encoding.UTF8.GetBytes(s);
-					}
-					break;
-				case CbFieldType.IntegerPositive:
-				case CbFieldType.IntegerNegative:
-					throw new NotImplementedException();
-				case CbFieldType.Float32:
-				case CbFieldType.Float64:
-					throw new NotImplementedException();
-				case CbFieldType.BoolFalse:
-				case CbFieldType.BoolTrue:
-					return Array.Empty<byte>();
-				case CbFieldType.ObjectAttachment:
-				case CbFieldType.BinaryAttachment:
-				case CbFieldType.Hash:
-					if (member is IoHash blob)
-					{
-						return blob.Memory.ToArray();
-					}
-
-					break;
-				case CbFieldType.Uuid:
-				case CbFieldType.DateTime:
-				case CbFieldType.TimeSpan:
-				case CbFieldType.ObjectId:
-				case CbFieldType.CustomById:
-				case CbFieldType.CustomByName:
-					throw new NotImplementedException();
+					int NumItemsLength = VarInt.Measure(Scope.Count);
+					int Offset = VarInt.Write(Span, Scope.Length + Scope.SizeOfChildHeaders + NumItemsLength);
+					return Offset + VarInt.Write(Span.Slice(Offset), Scope.Count);
 				default:
-					throw new ArgumentOutOfRangeException(nameof(fieldType), fieldType, null);
+					throw new InvalidOperationException();
 			}
-			throw new ArgumentException($"{member} of type {typeof(T).Name} not convert-able to field-type {fieldType}", nameof(member));
-		}
-
-		private void SetName(ref Field field, string? fieldName)
-		{
-			if (string.IsNullOrEmpty(fieldName))
-				return;
-
-			// TODO: We could add checks similar to what we do in C++ to make sure the writer is used correctly
-
-			int nameLength = fieldName.Length;
-			int nameBytesCount = BitUtils.MeasureVarUInt(nameLength);
-
-			byte[] buffer = new byte[nameLength + nameBytesCount];
-			BitUtils.WriteVarUInt(nameLength, buffer);
-			int bytesWritten = Encoding.ASCII.GetBytes(fieldName, 0, fieldName.Length, buffer, (int)nameBytesCount);
-			if (bytesWritten != nameLength)
-				throw new Exception("Failed to write name into buffer");
-
-			field.Name = buffer;
 		}
 	}
 }
\ No newline at end of file
diff --git a/Engine/Source/Programs/Shared/EpicGames.Serialization/VarInt.cs b/Engine/Source/Programs/Shared/EpicGames.Serialization/VarInt.cs
new file mode 100644
--- /dev/null
+++ b/Engine/Source/Programs/Shared/EpicGames.Serialization/VarInt.cs
@@ -0,0 +1,116 @@
+// Copyright Epic Games, Inc. All Rights Reserved.
+
+using System;
+using System.Collections.Generic;
+using System.Numerics;
+using System.Text;
+
+namespace EpicGames.Serialization
+{
+	/// <summary>
+	/// Methods for reading VarUInt values
+	/// </summary>
+	public static class VarInt
+	{
+		/// <summary>
+		/// Read a variable-length unsigned integer.
+		/// </summary>
+		/// <param name="Buffer">A variable-length encoding of an unsigned integer</param>
+		/// <returns></returns>
+		public static ulong Read(ReadOnlySpan<byte> Buffer)
+		{
+			return Read(Buffer, out _);
+		}
+
+		/// <summary>
+		/// Read a variable-length unsigned integer.
+		/// </summary>
+		/// <param name="Buffer">A variable-length encoding of an unsigned integer</param>
+		/// <param name="BytesRead">The number of bytes consumed from the input</param>
+		/// <returns></returns>
+		public static ulong Read(ReadOnlySpan<byte> Buffer, out int BytesRead)
+		{
+			BytesRead = (int)Measure(Buffer);
+
+			ulong Value = (ulong)(Buffer[0] & (0xff >> BytesRead));
+			for (int i = 1; i < BytesRead; i++)
+			{
+				Value <<= 8;
+				Value |= Buffer[i];
+			}
+			return Value;
+		}
+
+		/// <summary>
+		/// Measure the length in bytes (1-9) of an encoded variable-length integer.
+		/// </summary>
+		/// <param name="Buffer">A variable-length encoding of an(signed or unsigned) integer.</param>
+		/// <returns>The number of bytes used to encode the integer, in the range 1-9.</returns>
+		public static int Measure(ReadOnlySpan<byte> Buffer)
+		{
+			byte b = Buffer[0];
+			b = (byte)~b;
+			return BitOperations.LeadingZeroCount(b) - 23;
+		}
+
+		/// <summary>
+		/// Measure the number of bytes (1-5) required to encode the 32-bit input.
+		/// </summary>
+		/// <param name="Value"></param>
+		/// <returns></returns>
+		public static int Measure(int Value)
+		{
+			return Measure((uint)Value);
+		}
+
+		/// <summary>
+		/// Measure the number of bytes (1-5) required to encode the 32-bit input.
+		/// </summary>
+		/// <param name="Value"></param>
+		/// <returns></returns>
+		public static int Measure(uint Value)
+		{
+			return BitOperations.Log2(Value) / 7 + 1;
+		}
+
+		/// <summary>
+		/// Measure the number of bytes (1-9) required to encode the 64-bit input.
+		/// </summary>
+		/// <param name="value"></param>
+		/// <returns></returns>
+		public static int Measure(ulong value)
+		{
+			return Math.Min(BitOperations.Log2(value) / 7 + 1, 9);
+		}
+
+		/// <summary>
+		/// Write a variable-length unsigned integer.
+		/// </summary>
+		/// <param name="Value">An unsigned integer to encode</param>
+		/// <param name="Buffer">A buffer of at least 9 bytes to write the output to.</param>
+		/// <returns>The number of bytes used in the output</returns>
+		public static int Write(Span<byte> Buffer, long Value)
+		{
+			return Write(Buffer, (ulong)Value);
+		}
+
+		/// <summary>
+		/// Write a variable-length unsigned integer.
+		/// </summary>
+		/// <param name="Value">An unsigned integer to encode</param>
+		/// <param name="Buffer">A buffer of at least 9 bytes to write the output to.</param>
+		/// <returns>The number of bytes used in the output</returns>
+		public static int Write(Span<byte> Buffer, ulong Value)
+		{
+			int ByteCount = Measure(Value);
+
+			for (int Idx = 1; Idx < ByteCount; Idx++)
+			{
+				Buffer[ByteCount - Idx] = (byte)Value;
+				Value >>= 8;
+			}
+			Buffer[0] = (byte)((0xff << (9 - (int)ByteCount)) | (byte)Value);
+			return ByteCount;
+		}
+	}
+}
