Fix buffer overrun in compact binary writer.

diff --git a/Engine/Source/Programs/Shared/EpicGames.Serialization/CompactBinaryWriter.cs b/Engine/Source/Programs/Shared/EpicGames.Serialization/CompactBinaryWriter.cs
--- a/Engine/Source/Programs/Shared/EpicGames.Serialization/CompactBinaryWriter.cs
+++ b/Engine/Source/Programs/Shared/EpicGames.Serialization/CompactBinaryWriter.cs
@@ -148,29 +148,7 @@ namespace EpicGames.Serialization
 		{
 			Scope Scope = CurrentScope;
 			Scope.Length = CurrentOffset - Scope.Offset;
-			Scope.SizeOfChildHeaders = 0;
-
-			if (Scope.Children != null)
-			{
-				foreach (Scope ChildScope in Scope.Children)
-				{
-					switch (ChildScope.FieldType)
-					{
-						case CbFieldType.Object:
-						case CbFieldType.UniformObject:
-							Scope.SizeOfChildHeaders += ChildScope.SizeOfChildHeaders + VarInt.Measure(ChildScope.Length + ChildScope.SizeOfChildHeaders);
-							break;
-						case CbFieldType.Array:
-						case CbFieldType.UniformArray:
-							int ArrayCountLength = VarInt.Measure(Scope.Count);
-							Scope.SizeOfChildHeaders += ChildScope.SizeOfChildHeaders + VarInt.Measure(ChildScope.Length + ChildScope.SizeOfChildHeaders + ArrayCountLength) + ArrayCountLength;
-							break;
-						default:
-							throw new InvalidOperationException();
-					}
-				}
-			}
-
+			Scope.SizeOfChildHeaders = ComputeSizeOfChildHeaders(Scope);
 			OpenScopes.Pop();
 		}
 
@@ -593,7 +571,8 @@ namespace EpicGames.Serialization
 			{
 				throw new CbWriterException("Unfinished scope in writer");
 			}
-			return CurrentOffset + CurrentScope.SizeOfChildHeaders;
+
+			return CurrentOffset + ComputeSizeOfChildHeaders(CurrentScope);
 		}
 
 		/// <summary>
@@ -624,6 +603,15 @@ namespace EpicGames.Serialization
 			}
 		}
 
+		/// <summary>
+		/// Convert the data into a compact binary object
+		/// </summary>
+		/// <returns></returns>
+		public CbObject ToObject()
+		{
+			return new CbObject(ToByteArray());
+		}
+
 		/// <summary>
 		/// Convert the data into a flat array
 		/// </summary>
@@ -635,6 +623,36 @@ namespace EpicGames.Serialization
 			return Buffer;
 		}
 
+		/// <summary>
+		/// Comptues the size of any child headers
+		/// </summary>
+		/// <param name="Scope"></param>
+		static int ComputeSizeOfChildHeaders(Scope Scope)
+		{
+			int SizeOfChildHeaders = 0;
+			if (Scope.Children != null)
+			{
+				foreach (Scope ChildScope in Scope.Children)
+				{
+					switch (ChildScope.FieldType)
+					{
+						case CbFieldType.Object:
+						case CbFieldType.UniformObject:
+							SizeOfChildHeaders += ChildScope.SizeOfChildHeaders + VarInt.Measure(ChildScope.Length + ChildScope.SizeOfChildHeaders);
+							break;
+						case CbFieldType.Array:
+						case CbFieldType.UniformArray:
+							int ArrayCountLength = VarInt.Measure(Scope.Count);
+							SizeOfChildHeaders += ChildScope.SizeOfChildHeaders + VarInt.Measure(ChildScope.Length + ChildScope.SizeOfChildHeaders + ArrayCountLength) + ArrayCountLength;
+							break;
+						default:
+							throw new InvalidOperationException();
+					}
+				}
+			}
+			return SizeOfChildHeaders;
+		}
+
 		/// <summary>
 		/// Writes the header for a particular scope
 		/// </summary>
