Rename ReadOnlyUtf8String to Utf8String. Strings are always immutable in C#.

diff --git a/Engine/Source/Programs/Shared/EpicGames.Core/FileUtils.cs b/Engine/Source/Programs/Shared/EpicGames.Core/FileUtils.cs
--- a/Engine/Source/Programs/Shared/EpicGames.Core/FileUtils.cs
+++ b/Engine/Source/Programs/Shared/EpicGames.Core/FileUtils.cs
@@ -51,7 +51,7 @@ namespace EpicGames.Core
 		/// <summary>
 		/// Utf8 string comparer that should be used for native path comparisons
 		/// </summary>
-		public static IEqualityComparer<ReadOnlyUtf8String> PlatformPathComparerUtf8 = RuntimePlatform.IsLinux ? ReadOnlyUtf8StringComparer.Ordinal : ReadOnlyUtf8StringComparer.OrdinalIgnoreCase;
+		public static IEqualityComparer<Utf8String> PlatformPathComparerUtf8 = RuntimePlatform.IsLinux ? Utf8StringComparer.Ordinal : Utf8StringComparer.OrdinalIgnoreCase;
 
 		/// <summary>
 		/// Read all text for a file
diff --git a/Engine/Source/Programs/Shared/EpicGames.Core/ReadOnlyUtf8String.cs b/Engine/Source/Programs/Shared/EpicGames.Core/ReadOnlyUtf8String.cs
deleted file mode 100644
diff --git a/Engine/Source/Programs/Shared/EpicGames.Core/StringUtils.cs b/Engine/Source/Programs/Shared/EpicGames.Core/StringUtils.cs
--- a/Engine/Source/Programs/Shared/EpicGames.Core/StringUtils.cs
+++ b/Engine/Source/Programs/Shared/EpicGames.Core/StringUtils.cs
@@ -410,7 +410,7 @@ namespace EpicGames.Core
 		/// </summary>
 		/// <param name="Bytes">An array of bytes</param>
 		/// <returns>String representation of the array</returns>
-		public static ReadOnlyUtf8String FormatUtf8HexString(ReadOnlySpan<byte> Bytes)
+		public static Utf8String FormatUtf8HexString(ReadOnlySpan<byte> Bytes)
 		{
 			byte[] Characters = new byte[Bytes.Length * 2];
 			for (int Idx = 0; Idx < Bytes.Length; Idx++)
@@ -418,7 +418,7 @@ namespace EpicGames.Core
 				Characters[Idx * 2 + 0] = HexDigitToUtf8Byte[Bytes[Idx] >> 4];
 				Characters[Idx * 2 + 1] = HexDigitToUtf8Byte[Bytes[Idx] & 15];
 			}
-			return new ReadOnlyUtf8String(Characters);
+			return new Utf8String(Characters);
 		}
 
 		/// <summary>
diff --git a/Engine/Source/Programs/Shared/EpicGames.Core/Utf8String.cs b/Engine/Source/Programs/Shared/EpicGames.Core/Utf8String.cs
new file mode 100644
--- /dev/null
+++ b/Engine/Source/Programs/Shared/EpicGames.Core/Utf8String.cs
@@ -0,0 +1,361 @@
+// Copyright Epic Games, Inc. All Rights Reserved.
+
+using System;
+using System.Collections.Generic;
+using System.Diagnostics.CodeAnalysis;
+using System.Runtime.CompilerServices;
+using System.Text;
+
+namespace EpicGames.Core
+{
+	/// <summary>
+	/// Represents a memory region which can be treated as a utf-8 string.
+	/// </summary>
+	public struct Utf8String : IEquatable<Utf8String>
+	{
+		/// <summary>
+		/// An empty string
+		/// </summary>
+		public static readonly Utf8String Empty = new Utf8String();
+
+		/// <summary>
+		/// The data represented by this string
+		/// </summary>
+		public ReadOnlyMemory<byte> Memory { get; }
+
+		/// <summary>
+		/// Returns read only span for this string
+		/// </summary>
+		public ReadOnlySpan<byte> Span
+		{
+			get { return Memory.Span; } 
+		}
+
+		/// <summary>
+		/// Determines if this string is empty
+		/// </summary>
+		public bool IsEmpty
+		{
+			get { return Memory.IsEmpty; }
+		}
+
+		/// <summary>
+		/// Returns the length of this string
+		/// </summary>
+		public int Length
+		{
+			get { return Memory.Length; }
+		}
+
+		/// <summary>
+		/// Allows indexing individual bytes of the data
+		/// </summary>
+		/// <param name="Index">Byte index</param>
+		/// <returns>Byte at the given index</returns>
+		public byte this[int Index]
+		{
+			get { return Span[Index]; }
+		}
+
+		/// <summary>
+		/// Constructor
+		/// </summary>
+		/// <param name="Text">Text to construct from</param>
+		public Utf8String(string Text)
+		{
+			this.Memory = Encoding.UTF8.GetBytes(Text);
+		}
+
+		/// <summary>
+		/// Constructor
+		/// </summary>
+		/// <param name="Memory">The data to construct from</param>
+		public Utf8String(ReadOnlyMemory<byte> Memory)
+		{
+			this.Memory = Memory;
+		}
+
+		/// <summary>
+		/// Constructor
+		/// </summary>
+		/// <param name="Buffer">The buffer to construct from</param>
+		/// <param name="Offset">Offset within the buffer</param>
+		/// <param name="Length">Length of the string within the buffer</param>
+		public Utf8String(byte[] Buffer, int Offset, int Length)
+		{
+			this.Memory = new ReadOnlyMemory<byte>(Buffer, Offset, Length);
+		}
+
+		/// <summary>
+		/// Tests two strings for equality
+		/// </summary>
+		/// <param name="A">The first string to compare</param>
+		/// <param name="B">The second string to compare</param>
+		/// <returns>True if the strings are equal</returns>
+		public static bool operator ==(Utf8String A, Utf8String B)
+		{
+			return A.Equals(B);
+		}
+
+		/// <summary>
+		/// Tests two strings for inequality
+		/// </summary>
+		/// <param name="A">The first string to compare</param>
+		/// <param name="B">The second string to compare</param>
+		/// <returns>True if the strings are not equal</returns>
+		public static bool operator !=(Utf8String A, Utf8String B)
+		{
+			return !A.Equals(B);
+		}
+
+		/// <summary>
+		/// Tests whether two strings are equal
+		/// </summary>
+		/// <param name="Other">The other string to compare to</param>
+		/// <returns>True if the strings are equal</returns>
+		public bool Equals(Utf8String Other)
+		{
+			return Span.SequenceEqual(Other.Span);
+		}
+
+		/// <summary>
+		/// Tests if this string starts with another string
+		/// </summary>
+		/// <param name="Other">The string to check against</param>
+		/// <returns>True if this string starts with the other string</returns>
+		public bool StartsWith(Utf8String Other)
+		{
+			return Span.StartsWith(Other.Span);
+		}
+
+		/// <summary>
+		/// Tests if this string ends with another string
+		/// </summary>
+		/// <param name="Other">The string to check against</param>
+		/// <returns>True if this string ends with the other string</returns>
+		public bool EndsWith(Utf8String Other)
+		{
+			return Span.EndsWith(Other.Span);
+		}
+
+		/// <summary>
+		/// Slices this string at the given start position
+		/// </summary>
+		/// <param name="Start">Start position</param>
+		/// <returns>String starting at the given position</returns>
+		public Utf8String Slice(int Start)
+		{
+			return new Utf8String(Memory.Slice(Start));
+		}
+
+		/// <summary>
+		/// Slices this string at the given start position and length
+		/// </summary>
+		/// <param name="Start">Start position</param>
+		/// <returns>String starting at the given position</returns>
+		public Utf8String Slice(int Start, int Count)
+		{
+			return new Utf8String(Memory.Slice(Start, Count));
+		}
+
+		/// <summary>
+		/// Tests if this string is equal to the other object
+		/// </summary>
+		/// <param name="obj">Object to compare to</param>
+		/// <returns>True if the objects are equivalent</returns>
+		public override bool Equals(object? obj)
+		{
+			Utf8String? Other = obj as Utf8String?;
+			return Other != null && Equals(Other.Value);
+		}
+
+		/// <summary>
+		/// Returns the hash code of this string
+		/// </summary>
+		/// <returns>Hash code for the string</returns>
+		public override int GetHashCode()
+		{
+			int Hash = 5381;
+			for(int Idx = 0; Idx < Memory.Length; Idx++)
+			{
+				Hash += (Hash << 5) + Span[Idx];
+			}
+			return Hash;
+		}
+
+		/// <summary>
+		/// Gets the string represented by this data
+		/// </summary>
+		/// <returns>The utf-8 string</returns>
+		public override string ToString()
+		{
+			return Encoding.UTF8.GetString(Span);
+		}
+
+		/// <summary>
+		/// Converts a string to a utf-8 string
+		/// </summary>
+		/// <param name="Text">Text to convert</param>
+		public static implicit operator Utf8String(string Text)
+		{
+			return new Utf8String(new ReadOnlyMemory<byte>(Encoding.UTF8.GetBytes(Text)));
+		}
+	}
+
+	/// <summary>
+	/// Comparison classes for utf8 strings
+	/// </summary>
+	public abstract class Utf8StringComparer : IEqualityComparer<Utf8String>, IComparer<Utf8String>
+	{
+		/// <summary>
+		/// Ordinal comparer for utf8 strings
+		/// </summary>
+		class OrdinalComparer : Utf8StringComparer
+		{
+			/// <inheritdoc/>
+			public override bool Equals(Utf8String StrA, Utf8String StrB)
+			{
+				return StrA.Equals(StrB);
+			}
+
+			/// <inheritdoc/>
+			public override int GetHashCode(Utf8String String)
+			{
+				return String.GetHashCode();
+			}
+
+			public override int Compare(Utf8String StrA, Utf8String StrB)
+			{
+				return StrA.Span.SequenceCompareTo(StrB.Span);
+			}
+		}
+
+		/// <summary>
+		/// Comparison between ReadOnlyUtf8String objects that ignores case for ASCII characters
+		/// </summary>
+		class OrdinalIgnoreCaseComparer : Utf8StringComparer 
+		{
+			/// <inheritdoc/>
+			public override bool Equals(Utf8String StrA, Utf8String StrB)
+			{
+				if (StrA.Length != StrB.Length)
+				{
+					return false;
+				}
+
+				for (int Idx = 0; Idx < StrA.Length; Idx++)
+				{
+					if (StrA[Idx] != StrB[Idx] && ToUpper(StrA[Idx]) != ToUpper(StrB[Idx]))
+					{
+						return false;
+					}
+				}
+
+				return true;
+			}
+
+			/// <inheritdoc/>
+			public override int GetHashCode(Utf8String String)
+			{
+				HashCode HashCode = new HashCode();
+				for (int Idx = 0; Idx < String.Length; Idx++)
+				{
+					HashCode.Add(ToUpper(String[Idx]));
+				}
+				return HashCode.ToHashCode();
+			}
+
+			/// <inheritdoc/>
+			public override int Compare(Utf8String StrA, Utf8String StrB)
+			{
+				ReadOnlySpan<byte> SpanA = StrA.Span;
+				ReadOnlySpan<byte> SpanB = StrB.Span;
+
+				int Length = Math.Min(StrA.Length, StrB.Length);
+				for (int Idx = 0; Idx < Length; Idx++)
+				{
+					if (SpanA[Idx] != SpanB[Idx])
+					{
+						int UpperA = ToUpper(SpanA[Idx]);
+						int UpperB = ToUpper(SpanB[Idx]);
+						if (UpperA != UpperB)
+						{
+							return UpperA - UpperB;
+						}
+					}
+				}
+
+				return StrA.Length - StrB.Length;
+			}
+
+			/// <summary>
+			/// Convert a character to uppercase
+			/// </summary>
+			/// <param name="Character">Character to convert</param>
+			/// <returns>The uppercase version of the character</returns>
+			static byte ToUpper(byte Character)
+			{
+				return (Character >= 'a' && Character <= 'z') ? (byte)(Character - 'a' + 'A') : Character;
+			}
+		}
+
+		/// <summary>
+		/// Static instance of the ordinal utf8 ordinal comparer
+		/// </summary>
+		public static Utf8StringComparer Ordinal = new OrdinalComparer();
+
+		/// <summary>
+		/// Static instance of the case-insensitive utf8 ordinal string comparer
+		/// </summary>
+		public static Utf8StringComparer OrdinalIgnoreCase = new OrdinalIgnoreCaseComparer();
+
+		/// <inheritdoc/>
+		public abstract bool Equals(Utf8String StrA, Utf8String StrB);
+
+		/// <inheritdoc/>
+		public abstract int GetHashCode(Utf8String String);
+
+		/// <inheritdoc/>
+		public abstract int Compare(Utf8String StrA, Utf8String StrB);
+	}
+
+	/// <summary>
+	/// Extension methods for ReadOnlyUtf8String objects
+	/// </summary>
+	public static class MemoryWriterExtensions
+	{
+		/// <summary>
+		/// Reads a null-terminated utf8 string from the buffer
+		/// </summary>
+		/// <returns>The string data</returns>
+		public static Utf8String ReadString(this MemoryReader Reader)
+		{
+			ReadOnlySpan<byte> Span = Reader.Span;
+			int Length = Span.IndexOf((byte)0);
+			Utf8String Value = new Utf8String(Reader.ReadFixedLengthBytes(Length));
+			Reader.ReadInt8();
+			return Value;
+		}
+
+		/// <summary>
+		/// Writes a UTF8 string into memory with a null terminator
+		/// </summary>
+		/// <param name="Writer">The memory writer to serialize to</param>
+		/// <param name="String">String to write</param>
+		public static void WriteString(this MemoryWriter Writer, Utf8String String)
+		{
+			Writer.WriteFixedLengthBytes(String.Span);
+			Writer.WriteInt8(0);
+		}
+
+		/// <summary>
+		/// Determines the size of a serialized utf-8 string
+		/// </summary>
+		/// <param name="String">The string to measure</param>
+		/// <returns>Size of the serialized string</returns>
+		public static int GetSerializedSize(this Utf8String String)
+		{
+			return String.Length + 1;
+		}
+	}
+}
diff --git a/Engine/Source/Programs/Shared/EpicGames.Perforce.Managed/FileContentId.cs b/Engine/Source/Programs/Shared/EpicGames.Perforce.Managed/FileContentId.cs
--- a/Engine/Source/Programs/Shared/EpicGames.Perforce.Managed/FileContentId.cs
+++ b/Engine/Source/Programs/Shared/EpicGames.Perforce.Managed/FileContentId.cs
@@ -19,12 +19,12 @@ namespace EpicGames.Perforce.Managed
 			get;
 		}
 
-		public ReadOnlyUtf8String Type
+		public Utf8String Type
 		{
 			get;
 		}
 
-		public FileContentId(Digest<Md5> Digest, ReadOnlyUtf8String Type)
+		public FileContentId(Digest<Md5> Digest, Utf8String Type)
 		{
 			this.Digest = Digest;
 			this.Type = Type;
@@ -46,7 +46,7 @@ namespace EpicGames.Perforce.Managed
 		public static FileContentId ReadFileContentId(this MemoryReader Reader)
 		{
 			Digest<Md5> Digest = Reader.ReadDigest<Md5>();
-			ReadOnlyUtf8String Type = Reader.ReadString();
+			Utf8String Type = Reader.ReadString();
 			return new FileContentId(Digest, Type);
 		}
 
diff --git a/Engine/Source/Programs/Shared/EpicGames.Perforce.Managed/ManagedWorkspace.cs b/Engine/Source/Programs/Shared/EpicGames.Perforce.Managed/ManagedWorkspace.cs
--- a/Engine/Source/Programs/Shared/EpicGames.Perforce.Managed/ManagedWorkspace.cs
+++ b/Engine/Source/Programs/Shared/EpicGames.Perforce.Managed/ManagedWorkspace.cs
@@ -945,7 +945,7 @@ namespace EpicGames.Perforce.Managed
 				foreach (WhereRecord WriteFile in WriteFiles)
 				{
 					string Path = Regex.Replace(WriteFile.ClientFile, "^//[^/]+/", "");
-					Workspace.AddFile(new ReadOnlyUtf8String(Path), 0, 0, false, new FileContentId(Digest<Md5>.Zero, default));
+					Workspace.AddFile(new Utf8String(Path), 0, 0, false, new FileContentId(Digest<Md5>.Zero, default));
 				}
 				await SaveAsync(TransactionState.Clean, CancellationToken.None);
 			}
@@ -1350,26 +1350,26 @@ namespace EpicGames.Perforce.Managed
 			}
 
 			public static readonly string[] FieldNames = Enum.GetNames(typeof(Field));
-			public static readonly ReadOnlyUtf8String[] Utf8FieldNames = Array.ConvertAll(FieldNames, x => new ReadOnlyUtf8String(x));
+			public static readonly Utf8String[] Utf8FieldNames = Array.ConvertAll(FieldNames, x => new Utf8String(x));
 
 			public PerforceValue[] Values = new PerforceValue[FieldNames.Length];
 
-			public ReadOnlyUtf8String DepotFile
+			public Utf8String DepotFile
 			{
 				get { return Values[(int)Field.depotFile].GetString(); }
 			}
 
-			public ReadOnlyUtf8String ClientFile
+			public Utf8String ClientFile
 			{
 				get { return Values[(int)Field.clientFile].GetString(); }
 			}
 
-			public ReadOnlyUtf8String HeadType
+			public Utf8String HeadType
 			{
 				get { return Values[(int)Field.headType].GetString(); }
 			}
 
-			public ReadOnlyUtf8String HaveRev
+			public Utf8String HaveRev
 			{
 				get { return Values[(int)Field.haveRev].GetString(); }
 			}
@@ -1379,7 +1379,7 @@ namespace EpicGames.Perforce.Managed
 				get { return Values[(int)Field.fileSize].AsLong(); }
 			}
 
-			public ReadOnlyUtf8String Digest
+			public Utf8String Digest
 			{
 				get { return Values[(int)Field.digest].GetString(); }
 			}
@@ -1402,11 +1402,11 @@ namespace EpicGames.Perforce.Managed
 				Stopwatch Timer = Stopwatch.StartNew();
 
 				// Get the expected prefix for any paths in client syntax
-				ReadOnlyUtf8String ClientPrefix = $"//{PerforceClient.ClientName}/";
+				Utf8String ClientPrefix = $"//{PerforceClient.ClientName}/";
 
 				// List of the last path fragments. Since file records that are returned are typically sorted by their position in the tree, we can save quite a lot of processing by
 				// reusing as many fragemnts as possible.
-				List<(ReadOnlyUtf8String, StreamSnapshotBuilder)> Fragments = new List<(ReadOnlyUtf8String, StreamSnapshotBuilder)>();
+				List<(Utf8String, StreamSnapshotBuilder)> Fragments = new List<(Utf8String, StreamSnapshotBuilder)>();
 
 				// Handler for each returned record
 				FStatIndexedRecord Record = new FStatIndexedRecord();
@@ -1452,7 +1452,7 @@ namespace EpicGames.Perforce.Managed
 						}
 
 						// Get the fragment text
-						ReadOnlyUtf8String Fragment = new ReadOnlyUtf8String(Record.ClientFile.Memory.Slice(FragmentMinIdx, FragmentMaxIdx - FragmentMinIdx));
+						Utf8String Fragment = new Utf8String(Record.ClientFile.Memory.Slice(FragmentMinIdx, FragmentMaxIdx - FragmentMinIdx));
 
 						// If this fragment matches the same fragment from the previous iteration, take the last stream directory straight away
 						if (FragmentIdx < Fragments.Count)
@@ -1470,7 +1470,7 @@ namespace EpicGames.Perforce.Managed
 						// Otherwise, find or add a directory for this fragment into the last directory
 						if (FragmentIdx >= Fragments.Count)
 						{
-							ReadOnlyUtf8String UnescapedFragment = PerforceUtils.UnescapePath(Fragment);
+							Utf8String UnescapedFragment = PerforceUtils.UnescapePath(Fragment);
 
 							StreamSnapshotBuilder? NextStreamDirectory;
 							if (!LastStreamDirectory.NameToSubDirectory.TryGetValue(UnescapedFragment, out NextStreamDirectory))
@@ -1495,10 +1495,10 @@ namespace EpicGames.Perforce.Managed
 					Record.DepotFile.Span.CopyTo(DepotFileAndRevisionData.AsSpan());
 					DepotFileAndRevisionData[Record.DepotFile.Length] = (byte)'#';
 					Record.HaveRev.Span.CopyTo(DepotFileAndRevisionData.AsSpan(Record.DepotFile.Length + 1));
-					ReadOnlyUtf8String DepotFileAndRevision = new ReadOnlyUtf8String(DepotFileAndRevisionData.AsMemory());
+					Utf8String DepotFileAndRevision = new Utf8String(DepotFileAndRevisionData.AsMemory());
 
 					// Add a new StreamFileInfo to the last directory object
-					ReadOnlyUtf8String FileName = PerforceUtils.UnescapePath(Record.ClientFile.Slice(FragmentMinIdx));
+					Utf8String FileName = PerforceUtils.UnescapePath(Record.ClientFile.Slice(FragmentMinIdx));
 					LastStreamDirectory.NameToFile.Add(FileName, new StreamFileInfo(FileName, Record.FileSize, ContentId, DepotFileAndRevision));
 				};
 
@@ -1849,7 +1849,7 @@ namespace EpicGames.Perforce.Managed
 			await SaveAsync(TransactionState.Clean, CancellationToken);
 		}
 
-		static readonly ReadOnlyUtf8String StatsFileName = "StatsV2.json";
+		static readonly Utf8String StatsFileName = "StatsV2.json";
 
 		/// <summary>
 		/// Syncs a batch of files
diff --git a/Engine/Source/Programs/Shared/EpicGames.Perforce.Managed/StreamDirectoryInfo.cs b/Engine/Source/Programs/Shared/EpicGames.Perforce.Managed/StreamDirectoryInfo.cs
--- a/Engine/Source/Programs/Shared/EpicGames.Perforce.Managed/StreamDirectoryInfo.cs
+++ b/Engine/Source/Programs/Shared/EpicGames.Perforce.Managed/StreamDirectoryInfo.cs
@@ -19,12 +19,12 @@ namespace EpicGames.Perforce.Managed
 		/// <summary>
 		/// Map of name to file within the directory
 		/// </summary>
-		public Dictionary<ReadOnlyUtf8String, StreamFileInfo> NameToFile = new Dictionary<ReadOnlyUtf8String, StreamFileInfo>();
+		public Dictionary<Utf8String, StreamFileInfo> NameToFile = new Dictionary<Utf8String, StreamFileInfo>();
 
 		/// <summary>
 		/// Map of name to subdirectory
 		/// </summary>
-		public Dictionary<ReadOnlyUtf8String, IoHash> NameToSubDirectory = new Dictionary<ReadOnlyUtf8String, IoHash>(FileUtils.PlatformPathComparerUtf8);
+		public Dictionary<Utf8String, IoHash> NameToSubDirectory = new Dictionary<Utf8String, IoHash>(FileUtils.PlatformPathComparerUtf8);
 	}
 
 	/// <summary>
@@ -50,7 +50,7 @@ namespace EpicGames.Perforce.Managed
 			int NumSubDirectories = Reader.ReadInt32();
 			for (int Idx = 0; Idx < NumSubDirectories; Idx++)
 			{
-				ReadOnlyUtf8String Name = Reader.ReadString();
+				Utf8String Name = Reader.ReadString();
 				DirectoryInfo.NameToSubDirectory[Name] = Reader.ReadIoHash();
 			}
 
@@ -71,7 +71,7 @@ namespace EpicGames.Perforce.Managed
 			}
 
 			Writer.WriteInt32(DirectoryInfo.NameToSubDirectory.Count);
-			foreach ((ReadOnlyUtf8String Name, IoHash Hash) in DirectoryInfo.NameToSubDirectory)
+			foreach ((Utf8String Name, IoHash Hash) in DirectoryInfo.NameToSubDirectory)
 			{
 				Writer.WriteString(Name);
 				Writer.WriteIoHash(Hash);
diff --git a/Engine/Source/Programs/Shared/EpicGames.Perforce.Managed/StreamFileInfo.cs b/Engine/Source/Programs/Shared/EpicGames.Perforce.Managed/StreamFileInfo.cs
--- a/Engine/Source/Programs/Shared/EpicGames.Perforce.Managed/StreamFileInfo.cs
+++ b/Engine/Source/Programs/Shared/EpicGames.Perforce.Managed/StreamFileInfo.cs
@@ -27,7 +27,7 @@ namespace EpicGames.Perforce.Managed
 		/// <summary>
 		/// Name of this file
 		/// </summary>
-		public ReadOnlyUtf8String Name { get; }
+		public Utf8String Name { get; }
 
 		/// <summary>
 		/// Length of the file, as reported by the server (actual size on disk may be different due to workspace options).
@@ -42,12 +42,12 @@ namespace EpicGames.Perforce.Managed
 		/// <summary>
 		/// Depot path for this file
 		/// </summary>
-		public ReadOnlyUtf8String DepotFileAndRevision { get; }
+		public Utf8String DepotFileAndRevision { get; }
 
 		/// <summary>
 		/// Constructor
 		/// </summary>
-		public StreamFileInfo(ReadOnlyUtf8String Name, long Length, FileContentId ContentId, ReadOnlyUtf8String DepotPathAndRevision)
+		public StreamFileInfo(Utf8String Name, long Length, FileContentId ContentId, Utf8String DepotPathAndRevision)
 		{
 			this.Name = Name;
 			this.Length = Length;
@@ -67,10 +67,10 @@ namespace EpicGames.Perforce.Managed
 		/// <param name="Reader">Binary reader to read data from</param>
 		public static StreamFileInfo ReadStreamFileInfo(this MemoryReader Reader)
 		{
-			ReadOnlyUtf8String Name = Reader.ReadString();
+			Utf8String Name = Reader.ReadString();
 			long Length = Reader.ReadInt64();
 			FileContentId ContentId = Reader.ReadFileContentId();
-			ReadOnlyUtf8String DepotPathAndRevision = Reader.ReadString();
+			Utf8String DepotPathAndRevision = Reader.ReadString();
 			return new StreamFileInfo(Name, Length, ContentId, DepotPathAndRevision);
 		}
 
diff --git a/Engine/Source/Programs/Shared/EpicGames.Perforce.Managed/StreamSnapshotBuilder.cs b/Engine/Source/Programs/Shared/EpicGames.Perforce.Managed/StreamSnapshotBuilder.cs
--- a/Engine/Source/Programs/Shared/EpicGames.Perforce.Managed/StreamSnapshotBuilder.cs
+++ b/Engine/Source/Programs/Shared/EpicGames.Perforce.Managed/StreamSnapshotBuilder.cs
@@ -15,12 +15,12 @@ namespace EpicGames.Perforce.Managed
 		/// <summary>
 		/// Map of name to file within the directory
 		/// </summary>
-		public Dictionary<ReadOnlyUtf8String, StreamFileInfo> NameToFile = new Dictionary<ReadOnlyUtf8String, StreamFileInfo>();
+		public Dictionary<Utf8String, StreamFileInfo> NameToFile = new Dictionary<Utf8String, StreamFileInfo>();
 
 		/// <summary>
 		/// Map of name to subdirectory
 		/// </summary>
-		public Dictionary<ReadOnlyUtf8String, StreamSnapshotBuilder> NameToSubDirectory = new Dictionary<ReadOnlyUtf8String, StreamSnapshotBuilder>(FileUtils.PlatformPathComparerUtf8);
+		public Dictionary<Utf8String, StreamSnapshotBuilder> NameToSubDirectory = new Dictionary<Utf8String, StreamSnapshotBuilder>(FileUtils.PlatformPathComparerUtf8);
 
 		/// <summary>
 		/// Default constructor
@@ -46,11 +46,11 @@ namespace EpicGames.Perforce.Managed
 		public StreamSnapshotBuilder(StreamSnapshot Snapshot, IoHash Hash)
 		{
 			StreamDirectoryInfo OldDirectory = Snapshot.Lookup(Hash);
-			foreach ((ReadOnlyUtf8String Name, StreamFileInfo File) in OldDirectory.NameToFile)
+			foreach ((Utf8String Name, StreamFileInfo File) in OldDirectory.NameToFile)
 			{
 				NameToFile[Name] = File;
 			}
-			foreach ((ReadOnlyUtf8String Name, IoHash SubDirHash) in OldDirectory.NameToSubDirectory)
+			foreach ((Utf8String Name, IoHash SubDirHash) in OldDirectory.NameToSubDirectory)
 			{
 				NameToSubDirectory[Name] = new StreamSnapshotBuilder(Snapshot, SubDirHash);
 			}
@@ -66,7 +66,7 @@ namespace EpicGames.Perforce.Managed
 			// Create the stream directory object
 			StreamDirectoryInfo Directory = new StreamDirectoryInfo();
 			Directory.NameToFile = NameToFile;
-			foreach ((ReadOnlyUtf8String Name, StreamSnapshotBuilder SubDirBuilder) in NameToSubDirectory)
+			foreach ((Utf8String Name, StreamSnapshotBuilder SubDirBuilder) in NameToSubDirectory)
 			{
 				Directory.NameToSubDirectory[Name] = SubDirBuilder.Encode(HashToDirectory);
 			}
diff --git a/Engine/Source/Programs/Shared/EpicGames.Perforce.Managed/Transactions/AddTransaction.cs b/Engine/Source/Programs/Shared/EpicGames.Perforce.Managed/Transactions/AddTransaction.cs
--- a/Engine/Source/Programs/Shared/EpicGames.Perforce.Managed/Transactions/AddTransaction.cs
+++ b/Engine/Source/Programs/Shared/EpicGames.Perforce.Managed/Transactions/AddTransaction.cs
@@ -88,7 +88,7 @@ namespace EpicGames.Perforce.Managed
 
 			// Update all the subdirectories
 			StreamDirectoryInfo StreamDir = StreamSnapshot.Lookup(StreamDirHash);
-			foreach((ReadOnlyUtf8String SubDirName, IoHash SubDirHash) in StreamDir.NameToSubDirectory)
+			foreach((Utf8String SubDirName, IoHash SubDirHash) in StreamDir.NameToSubDirectory)
 			{
 				WorkspaceDirectoryInfo? WorkspaceSubDir;
 				if(WorkspaceDir.NameToSubDirectory.TryGetValue(SubDirName, out WorkspaceSubDir))
@@ -124,7 +124,7 @@ namespace EpicGames.Perforce.Managed
 			Directory.CreateDirectory(NewWorkspaceDir.GetFullName());
 
 			// Add all the sub directories and files
-			foreach((ReadOnlyUtf8String SubDirName, IoHash SubDirHash) in StreamDir.NameToSubDirectory)
+			foreach((Utf8String SubDirName, IoHash SubDirHash) in StreamDir.NameToSubDirectory)
 			{
 				AddSubDirectory(SubDirName, NewWorkspaceDir, SubDirHash, Queue);
 			}
@@ -134,14 +134,14 @@ namespace EpicGames.Perforce.Managed
 			}
 		}
 
-		void MergeSubDirectory(ReadOnlyUtf8String Name, WorkspaceDirectoryInfo WorkspaceSubDir, IoHash StreamSubDirHash, WorkspaceDirectoryInfo NewWorkspaceDir, ThreadPoolWorkQueue Queue)
+		void MergeSubDirectory(Utf8String Name, WorkspaceDirectoryInfo WorkspaceSubDir, IoHash StreamSubDirHash, WorkspaceDirectoryInfo NewWorkspaceDir, ThreadPoolWorkQueue Queue)
 		{
 			WorkspaceDirectoryInfo NewWorkspaceSubDir = new WorkspaceDirectoryInfo(NewWorkspaceDir, Name, StreamSubDirHash);
 			NewWorkspaceDir.NameToSubDirectory.Add(Name, NewWorkspaceSubDir);
 			Queue.Enqueue(() => MergeDirectory(WorkspaceSubDir, NewWorkspaceSubDir, StreamSubDirHash, Queue));
 		}
 
-		void AddSubDirectory(ReadOnlyUtf8String Name, WorkspaceDirectoryInfo NewWorkspaceDir, IoHash StreamSubDirHash, ThreadPoolWorkQueue Queue)
+		void AddSubDirectory(Utf8String Name, WorkspaceDirectoryInfo NewWorkspaceDir, IoHash StreamSubDirHash, ThreadPoolWorkQueue Queue)
 		{
 			WorkspaceDirectoryInfo NewWorkspaceSubDir = new WorkspaceDirectoryInfo(NewWorkspaceDir, Name, StreamSubDirHash);
 			NewWorkspaceDir.NameToSubDirectory.Add(Name, NewWorkspaceSubDir);
diff --git a/Engine/Source/Programs/Shared/EpicGames.Perforce.Managed/Transactions/RemoveTransaction.cs b/Engine/Source/Programs/Shared/EpicGames.Perforce.Managed/Transactions/RemoveTransaction.cs
--- a/Engine/Source/Programs/Shared/EpicGames.Perforce.Managed/Transactions/RemoveTransaction.cs
+++ b/Engine/Source/Programs/Shared/EpicGames.Perforce.Managed/Transactions/RemoveTransaction.cs
@@ -65,7 +65,7 @@ namespace EpicGames.Perforce.Managed
 			}
 		}
 
-		void MergeSubDirectory(ReadOnlyUtf8String Name, WorkspaceDirectoryInfo NewWorkspaceDir, WorkspaceDirectoryInfo WorkspaceSubDir, IoHash StreamSubDirHash, ThreadPoolWorkQueue Queue)
+		void MergeSubDirectory(Utf8String Name, WorkspaceDirectoryInfo NewWorkspaceDir, WorkspaceDirectoryInfo WorkspaceSubDir, IoHash StreamSubDirHash, ThreadPoolWorkQueue Queue)
 		{
 			WorkspaceDirectoryInfo NewWorkspaceSubDir = new WorkspaceDirectoryInfo(NewWorkspaceDir, Name, StreamSubDirHash);
 			NewWorkspaceDir.NameToSubDirectory.Add(NewWorkspaceSubDir.Name, NewWorkspaceSubDir);
diff --git a/Engine/Source/Programs/Shared/EpicGames.Perforce.Managed/WorkspaceDirectoryInfo.cs b/Engine/Source/Programs/Shared/EpicGames.Perforce.Managed/WorkspaceDirectoryInfo.cs
--- a/Engine/Source/Programs/Shared/EpicGames.Perforce.Managed/WorkspaceDirectoryInfo.cs
+++ b/Engine/Source/Programs/Shared/EpicGames.Perforce.Managed/WorkspaceDirectoryInfo.cs
@@ -23,7 +23,7 @@ namespace EpicGames.Perforce.Managed
 		/// <summary>
 		/// Name of this directory
 		/// </summary>
-		public ReadOnlyUtf8String Name { get; }
+		public Utf8String Name { get; }
 
 		/// <summary>
 		/// Digest of the matching stream directory info. This should be set to zero if the workspace is modified.
@@ -33,12 +33,12 @@ namespace EpicGames.Perforce.Managed
 		/// <summary>
 		/// Map of name to file
 		/// </summary>
-		public Dictionary<ReadOnlyUtf8String, WorkspaceFileInfo> NameToFile { get; set; }
+		public Dictionary<Utf8String, WorkspaceFileInfo> NameToFile { get; set; }
 
 		/// <summary>
 		/// Map of name to subdirectory
 		/// </summary>
-		public Dictionary<ReadOnlyUtf8String, WorkspaceDirectoryInfo> NameToSubDirectory { get; set; }
+		public Dictionary<Utf8String, WorkspaceDirectoryInfo> NameToSubDirectory { get; set; }
 
 		/// <summary>
 		/// Constructor
@@ -55,13 +55,13 @@ namespace EpicGames.Perforce.Managed
 		/// <param name="ParentDirectory">The parent directory</param>
 		/// <param name="Name">Name of this directory</param>
 		/// <param name="Digest">The corresponding stream digest</param>
-		public WorkspaceDirectoryInfo(WorkspaceDirectoryInfo? ParentDirectory, ReadOnlyUtf8String Name, IoHash Digest)
+		public WorkspaceDirectoryInfo(WorkspaceDirectoryInfo? ParentDirectory, Utf8String Name, IoHash Digest)
 		{
 			this.ParentDirectory = ParentDirectory;
 			this.Name = Name;
 			this.StreamDirectoryDigest = Digest;
-			this.NameToFile = new Dictionary<ReadOnlyUtf8String, WorkspaceFileInfo>(ReadOnlyUtf8StringComparer.Ordinal);
-			this.NameToSubDirectory = new Dictionary<ReadOnlyUtf8String, WorkspaceDirectoryInfo>(FileUtils.PlatformPathComparerUtf8);
+			this.NameToFile = new Dictionary<Utf8String, WorkspaceFileInfo>(Utf8StringComparer.Ordinal);
+			this.NameToSubDirectory = new Dictionary<Utf8String, WorkspaceDirectoryInfo>(FileUtils.PlatformPathComparerUtf8);
 		}
 
 		/// <summary>
@@ -72,7 +72,7 @@ namespace EpicGames.Perforce.Managed
 		/// <param name="LastModifiedTicks">Last modified time of the file</param>
 		/// <param name="bReadOnly">Whether the file is read only</param>
 		/// <param name="ContentId">Unique identifier for the server content</param>
-		public void AddFile(ReadOnlyUtf8String Path, long Length, long LastModifiedTicks, bool bReadOnly, FileContentId ContentId)
+		public void AddFile(Utf8String Path, long Length, long LastModifiedTicks, bool bReadOnly, FileContentId ContentId)
 		{
 			StreamDirectoryDigest = IoHash.Zero;
 
@@ -83,7 +83,7 @@ namespace EpicGames.Perforce.Managed
 			}
 			else
 			{
-				ReadOnlyUtf8String Name = Path.Slice(0, Idx);
+				Utf8String Name = Path.Slice(0, Idx);
 
 				WorkspaceDirectoryInfo? SubDirectory;
 				if (!NameToSubDirectory.TryGetValue(Name, out SubDirectory))
@@ -115,7 +115,7 @@ namespace EpicGames.Perforce.Managed
 		{
 			Files.AddRange(NameToFile.Values);
 
-			foreach (KeyValuePair<ReadOnlyUtf8String, WorkspaceDirectoryInfo> Pair in NameToSubDirectory)
+			foreach (KeyValuePair<Utf8String, WorkspaceDirectoryInfo> Pair in NameToSubDirectory)
 			{
 				Pair.Value.GetFilesInternal(Files);
 			}
@@ -150,7 +150,7 @@ namespace EpicGames.Perforce.Managed
 		void Refresh(DirectoryInfo Info, bool bRemoveUntracked, ConcurrentBag<FileInfo> FilesToDelete, ConcurrentBag<DirectoryInfo> DirectoriesToDelete, ThreadPoolWorkQueue Queue)
 		{
 			// Recurse through subdirectories
-			Dictionary<ReadOnlyUtf8String, WorkspaceDirectoryInfo> NewNameToSubDirectory = new Dictionary<ReadOnlyUtf8String, WorkspaceDirectoryInfo>(NameToSubDirectory.Count, NameToSubDirectory.Comparer);
+			Dictionary<Utf8String, WorkspaceDirectoryInfo> NewNameToSubDirectory = new Dictionary<Utf8String, WorkspaceDirectoryInfo>(NameToSubDirectory.Count, NameToSubDirectory.Comparer);
 			foreach (DirectoryInfo SubDirectoryInfo in Info.EnumerateDirectories())
 			{
 				WorkspaceDirectoryInfo? SubDirectory;
@@ -167,7 +167,7 @@ namespace EpicGames.Perforce.Managed
 			NameToSubDirectory = NewNameToSubDirectory;
 
 			// Figure out which files have changed.
-			Dictionary<ReadOnlyUtf8String, WorkspaceFileInfo> NewNameToFile = new Dictionary<ReadOnlyUtf8String, WorkspaceFileInfo>(NameToFile.Count, NameToFile.Comparer);
+			Dictionary<Utf8String, WorkspaceFileInfo> NewNameToFile = new Dictionary<Utf8String, WorkspaceFileInfo>(NameToFile.Count, NameToFile.Comparer);
 			foreach (FileInfo File in Info.EnumerateFiles())
 			{
 				WorkspaceFileInfo? StagedFile;
@@ -228,7 +228,7 @@ namespace EpicGames.Perforce.Managed
 		void FindDifferences(DirectoryInfo Directory, string Path, ConcurrentBag<string> Paths, ThreadPoolWorkQueue Queue)
 		{
 			// Recurse through subdirectories
-			HashSet<ReadOnlyUtf8String> RemainingSubDirectoryNames = new HashSet<ReadOnlyUtf8String>(NameToSubDirectory.Keys);
+			HashSet<Utf8String> RemainingSubDirectoryNames = new HashSet<Utf8String>(NameToSubDirectory.Keys);
 			foreach (DirectoryInfo SubDirectory in Directory.EnumerateDirectories())
 			{
 				WorkspaceDirectoryInfo? StagedSubDirectory;
@@ -240,13 +240,13 @@ namespace EpicGames.Perforce.Managed
 				}
 				Paths.Add(String.Format("+{0}{1}/...", Path, SubDirectory.Name));
 			}
-			foreach (ReadOnlyUtf8String RemainingSubDirectoryName in RemainingSubDirectoryNames)
+			foreach (Utf8String RemainingSubDirectoryName in RemainingSubDirectoryNames)
 			{
 				Paths.Add(String.Format("-{0}{1}/...", Path, RemainingSubDirectoryName));
 			}
 
 			// Search through files
-			HashSet<ReadOnlyUtf8String> RemainingFileNames = new HashSet<ReadOnlyUtf8String>(NameToFile.Keys);
+			HashSet<Utf8String> RemainingFileNames = new HashSet<Utf8String>(NameToFile.Keys);
 			foreach (FileInfo File in Directory.EnumerateFiles())
 			{
 				WorkspaceFileInfo? StagedFile;
@@ -264,7 +264,7 @@ namespace EpicGames.Perforce.Managed
 					RemainingFileNames.Remove(File.Name);
 				}
 			}
-			foreach (ReadOnlyUtf8String RemainingFileName in RemainingFileNames)
+			foreach (Utf8String RemainingFileName in RemainingFileNames)
 			{
 				Paths.Add(String.Format("-{0}{1}", Path, RemainingFileName));
 			}
@@ -356,7 +356,7 @@ namespace EpicGames.Perforce.Managed
 			int NumSubDirectories = Reader.ReadInt32();
 			for (int Idx = 0; Idx < NumSubDirectories; Idx++)
 			{
-				ReadOnlyUtf8String Name = Reader.ReadString();
+				Utf8String Name = Reader.ReadString();
 
 				WorkspaceDirectoryInfo SubDirectory = new WorkspaceDirectoryInfo(DirectoryInfo, Name, IoHash.Zero);
 				Reader.ReadWorkspaceDirectoryInfo(SubDirectory, Version);
diff --git a/Engine/Source/Programs/Shared/EpicGames.Perforce.Managed/WorkspaceFileInfo.cs b/Engine/Source/Programs/Shared/EpicGames.Perforce.Managed/WorkspaceFileInfo.cs
--- a/Engine/Source/Programs/Shared/EpicGames.Perforce.Managed/WorkspaceFileInfo.cs
+++ b/Engine/Source/Programs/Shared/EpicGames.Perforce.Managed/WorkspaceFileInfo.cs
@@ -19,25 +19,25 @@ namespace EpicGames.Perforce.Managed
 	class WorkspaceFileInfo
 	{
 		public readonly WorkspaceDirectoryInfo Directory;
-		public readonly ReadOnlyUtf8String Name;
+		public readonly Utf8String Name;
 		public long Length;
 		public long LastModifiedTicks;
 		public bool bReadOnly;
 		public readonly FileContentId ContentId;
 
-		public WorkspaceFileInfo(WorkspaceDirectoryInfo Directory, ReadOnlyUtf8String Name, FileContentId ContentId)
+		public WorkspaceFileInfo(WorkspaceDirectoryInfo Directory, Utf8String Name, FileContentId ContentId)
 		{
 			this.Directory = Directory;
 			this.Name = Name;
 			this.ContentId = ContentId;
 		}
 
-		public WorkspaceFileInfo(WorkspaceDirectoryInfo Directory, ReadOnlyUtf8String Name, FileInfo Info, FileContentId ContentId)
+		public WorkspaceFileInfo(WorkspaceDirectoryInfo Directory, Utf8String Name, FileInfo Info, FileContentId ContentId)
 			: this(Directory, Name, Info.Length, Info.LastWriteTimeUtc.Ticks, Info.Attributes.HasFlag(FileAttributes.ReadOnly), ContentId)
 		{
 		}
 
-		public WorkspaceFileInfo(WorkspaceDirectoryInfo Directory, ReadOnlyUtf8String Name, long Length, long LastModifiedTicks, bool bReadOnly, FileContentId ContentId)
+		public WorkspaceFileInfo(WorkspaceDirectoryInfo Directory, Utf8String Name, long Length, long LastModifiedTicks, bool bReadOnly, FileContentId ContentId)
 		{
 			this.Directory = Directory;
 			this.Name = Name;
@@ -102,7 +102,7 @@ namespace EpicGames.Perforce.Managed
 	{
 		public static WorkspaceFileInfo ReadWorkspaceFileInfo(this MemoryReader Reader, WorkspaceDirectoryInfo Directory)
 		{
-			ReadOnlyUtf8String Name = Reader.ReadString();
+			Utf8String Name = Reader.ReadString();
 			long Length = Reader.ReadInt64();
 			long LastModifiedTicks = Reader.ReadInt64();
 			bool bReadOnly = Reader.ReadBoolean();
diff --git a/Engine/Source/Programs/Shared/EpicGames.Perforce/PerforceChildProcess.cs b/Engine/Source/Programs/Shared/EpicGames.Perforce/PerforceChildProcess.cs
--- a/Engine/Source/Programs/Shared/EpicGames.Perforce/PerforceChildProcess.cs
+++ b/Engine/Source/Programs/Shared/EpicGames.Perforce/PerforceChildProcess.cs
@@ -30,10 +30,10 @@ namespace EpicGames.Perforce
 		/// </summary>
 		static class ReadOnlyUtf8StringConstants
 		{
-			public static ReadOnlyUtf8String Code = "code";
-			public static ReadOnlyUtf8String Stat = "stat";
-			public static ReadOnlyUtf8String Info = "info";
-			public static ReadOnlyUtf8String Error = "error";
+			public static Utf8String Code = "code";
+			public static Utf8String Stat = "stat";
+			public static Utf8String Info = "info";
+			public static Utf8String Error = "error";
 		}
 
 		/// <summary>
@@ -426,7 +426,7 @@ namespace EpicGames.Perforce
 		public async Task ReadRecordsAsync(Action<PerforceRecord> HandleRecord, CancellationToken CancellationToken)
 		{
 			PerforceRecord Record = new PerforceRecord();
-			Record.Rows = new List<KeyValuePair<ReadOnlyUtf8String, PerforceValue>>();
+			Record.Rows = new List<KeyValuePair<Utf8String, PerforceValue>>();
 
 			for(; ;)
 			{
@@ -500,7 +500,7 @@ namespace EpicGames.Perforce
 		/// </summary>
 		/// <param name="Rows">List of rows to read into</param>
 		/// <returns>True if a record could be read; false if more data is required</returns>
-		bool ReadRecord(List<KeyValuePair<ReadOnlyUtf8String, PerforceValue>> Rows)
+		bool ReadRecord(List<KeyValuePair<Utf8String, PerforceValue>> Rows)
 		{
 			Rows.Clear();
 
@@ -537,7 +537,7 @@ namespace EpicGames.Perforce
 				}
 
 				// Read the tag
-				ReadOnlyUtf8String Key;
+				Utf8String Key;
 				if (!TryReadString(out Key))
 				{
 					return false;
@@ -557,7 +557,7 @@ namespace EpicGames.Perforce
 				PerforceValue Value;
 				if (ValueType == 's')
 				{
-					ReadOnlyUtf8String String;
+					Utf8String String;
 					if (!TryReadString(out String))
 					{
 						return false;
@@ -605,7 +605,7 @@ namespace EpicGames.Perforce
 			}
 			BufferPos += Prefix.Length;
 
-			ReadOnlyUtf8String Code;
+			Utf8String Code;
 			if (!TryReadString(out Code))
 			{
 				Response = null;
@@ -616,7 +616,7 @@ namespace EpicGames.Perforce
 			object? Record;
 			if (Code == ReadOnlyUtf8StringConstants.Stat && StatRecordInfo != null)
 			{
-				if(!TryReadTypedRecord(ReadOnlyUtf8String.Empty, StatRecordInfo, out Record))
+				if(!TryReadTypedRecord(Utf8String.Empty, StatRecordInfo, out Record))
 				{
 					Response = null;
 					return false;
@@ -624,7 +624,7 @@ namespace EpicGames.Perforce
 			}
 			else if (Code == ReadOnlyUtf8StringConstants.Info)
 			{
-				if(!TryReadTypedRecord(ReadOnlyUtf8String.Empty, PerforceReflection.InfoRecordInfo, out Record))
+				if(!TryReadTypedRecord(Utf8String.Empty, PerforceReflection.InfoRecordInfo, out Record))
 				{
 					Response = null;
 					return false;
@@ -632,7 +632,7 @@ namespace EpicGames.Perforce
 			}
 			else if (Code == ReadOnlyUtf8StringConstants.Error)
 			{
-				if (!TryReadTypedRecord(ReadOnlyUtf8String.Empty, PerforceReflection.ErrorRecordInfo, out Record))
+				if (!TryReadTypedRecord(Utf8String.Empty, PerforceReflection.ErrorRecordInfo, out Record))
 				{
 					Response = null;
 					return false;
@@ -661,7 +661,7 @@ namespace EpicGames.Perforce
 		/// <param name="RequiredSuffix">The required suffix for any subobject arrays.</param>
 		/// <param name="RecordInfo">Reflection information for the type being serialized into.</param>
 		/// <returns>The parsed object.</returns>
-		bool TryReadTypedRecord(ReadOnlyUtf8String RequiredSuffix, CachedRecordInfo RecordInfo, [NotNullWhen(true)] out object? Record)
+		bool TryReadTypedRecord(Utf8String RequiredSuffix, CachedRecordInfo RecordInfo, [NotNullWhen(true)] out object? Record)
 		{
 			// Create a bitmask for all the required tags
 			ulong RequiredTagsBitMask = 0;
@@ -699,7 +699,7 @@ namespace EpicGames.Perforce
 				BufferPos++;
 
 				// Read the tag
-				ReadOnlyUtf8String Tag;
+				Utf8String Tag;
 				if(!TryReadString(out Tag))
 				{
 					Record = null;
@@ -714,7 +714,7 @@ namespace EpicGames.Perforce
 				}
 
 				// Separate the key into tag and suffix
-				ReadOnlyUtf8String Suffix = Tag.Slice(SuffixIdx);
+				Utf8String Suffix = Tag.Slice(SuffixIdx);
 				Tag = Tag.Slice(0, SuffixIdx);
 
 				// Try to find the matching field
@@ -833,7 +833,7 @@ namespace EpicGames.Perforce
 			// Parse the appropriate value
 			if (ValueType == 's')
 			{
-				ReadOnlyUtf8String String;
+				Utf8String String;
 				if (!TryReadString(out String))
 				{
 					return false;
@@ -904,12 +904,12 @@ namespace EpicGames.Perforce
 		/// </summary>
 		/// <param name="String">Receives the value that was read</param>
 		/// <returns>True if a string was read from the buffer, false if there was not enough data</returns>
-		bool TryReadStringWithType(out ReadOnlyUtf8String String)
+		bool TryReadStringWithType(out Utf8String String)
 		{
 			byte ValueType;
 			if(!TryReadByte(out ValueType))
 			{
-				String = new ReadOnlyUtf8String();
+				String = new Utf8String();
 				return false;
 			}
 			if (ValueType != 's')
@@ -924,22 +924,22 @@ namespace EpicGames.Perforce
 		/// </summary>
 		/// <param name="String">Receives the value that was read</param>
 		/// <returns>True if a string was read from the buffer, false if there was not enough data</returns>
-		bool TryReadString(out ReadOnlyUtf8String String)
+		bool TryReadString(out Utf8String String)
 		{
 			int Length;
 			if(!TryReadInt(out Length))
 			{
-				String = new ReadOnlyUtf8String();
+				String = new Utf8String();
 				return false;
 			}
 
 			if(BufferPos + Length > BufferEnd)
 			{
-				String = new ReadOnlyUtf8String();
+				String = new Utf8String();
 				return false;
 			}
 
-			String = new ReadOnlyUtf8String(Buffer, BufferPos, Length);
+			String = new Utf8String(Buffer, BufferPos, Length);
 			BufferPos += Length;
 			return true;
 		}
@@ -951,7 +951,7 @@ namespace EpicGames.Perforce
 		/// <param name="Prefix">The required prefix</param>
 		/// <param name="Index">The required index</param>
 		/// <returns>True if the index is correct</returns>
-		static bool IsCorrectIndex(ReadOnlyUtf8String Text, ReadOnlyUtf8String Prefix, int Index)
+		static bool IsCorrectIndex(Utf8String Text, Utf8String Prefix, int Index)
 		{
 			if (Prefix.Length > 0)
 			{
diff --git a/Engine/Source/Programs/Shared/EpicGames.Perforce/PerforceRecord.cs b/Engine/Source/Programs/Shared/EpicGames.Perforce/PerforceRecord.cs
--- a/Engine/Source/Programs/Shared/EpicGames.Perforce/PerforceRecord.cs
+++ b/Engine/Source/Programs/Shared/EpicGames.Perforce/PerforceRecord.cs
@@ -80,7 +80,7 @@ namespace EpicGames.Perforce
 		/// <returns>The boolean value</returns>
 		public bool AsBool()
 		{
-			ReadOnlyUtf8String String = GetString();
+			Utf8String String = GetString();
 			return String.Length == 0 || String == StringConstants.True;
 		}
 
@@ -96,7 +96,7 @@ namespace EpicGames.Perforce
 			}
 			else if (Type == PerforceValueType.String)
 			{
-				ReadOnlyUtf8String String = GetString();
+				Utf8String String = GetString();
 
 				int Value;
 				int BytesConsumed;
@@ -133,7 +133,7 @@ namespace EpicGames.Perforce
 		/// <returns>The converted value</returns>
 		public long AsLong()
 		{
-			ReadOnlyUtf8String String = AsString();
+			Utf8String String = AsString();
 
 			long Value;
 			int BytesConsumed;
@@ -159,14 +159,14 @@ namespace EpicGames.Perforce
 		/// Converts this value to a string
 		/// </summary>
 		/// <returns>The converted value</returns>
-		public ReadOnlyUtf8String AsString()
+		public Utf8String AsString()
 		{
 			switch(Type)
 			{
 				case PerforceValueType.String:
 					return GetString();
 				case PerforceValueType.Integer:
-					return new ReadOnlyUtf8String(GetInteger().ToString());
+					return new Utf8String(GetInteger().ToString());
 				default:
 					throw new NotImplementedException();
 			}
@@ -189,13 +189,13 @@ namespace EpicGames.Perforce
 		/// Gets the contents of this string
 		/// </summary>
 		/// <returns></returns>
-		public ReadOnlyUtf8String GetString()
+		public Utf8String GetString()
 		{
 			if(Type != PerforceValueType.String)
 			{
 				throw new InvalidOperationException("Value is not a string");
 			}
-			return new ReadOnlyUtf8String(Data.Slice(5));
+			return new Utf8String(Data.Slice(5));
 		}
 
 		/// <summary>
@@ -216,14 +216,14 @@ namespace EpicGames.Perforce
 		/// <summary>
 		/// The rows in this record
 		/// </summary>
-		public List<KeyValuePair<ReadOnlyUtf8String, PerforceValue>> Rows;
+		public List<KeyValuePair<Utf8String, PerforceValue>> Rows;
 
 		/// <summary>
 		/// Copy this record into the given array of values. This method is O(n) if every record key being in the list of keys in the same order, but O(n^2) if not.
 		/// </summary>
 		/// <param name="Keys">List of keys to parse</param>
 		/// <param name="Values">Array to receive the list of values</param>
-		public void CopyInto(ReadOnlyUtf8String[] Keys, PerforceValue[] Values)
+		public void CopyInto(Utf8String[] Keys, PerforceValue[] Values)
 		{
 			// Clear out the current values array
 			for (int ValueIdx = 0; ValueIdx < Values.Length; ValueIdx++)
@@ -238,10 +238,10 @@ namespace EpicGames.Perforce
 				int InitialRowIdx = RowIdx;
 
 				// Find the key that matches this row.
-				ReadOnlyUtf8String RowKey = Rows[RowIdx].Key;
+				Utf8String RowKey = Rows[RowIdx].Key;
 				for (int KeyIdx = 0; KeyIdx < Keys.Length; KeyIdx++)
 				{
-					ReadOnlyUtf8String Key = Keys[KeyIdx];
+					Utf8String Key = Keys[KeyIdx];
 					if (RowKey == Key)
 					{
 						// Copy the value to the output array, then move to the next row, and try to match that against the next key.
diff --git a/Engine/Source/Programs/Shared/EpicGames.Perforce/PerforceReflection.cs b/Engine/Source/Programs/Shared/EpicGames.Perforce/PerforceReflection.cs
--- a/Engine/Source/Programs/Shared/EpicGames.Perforce/PerforceReflection.cs
+++ b/Engine/Source/Programs/Shared/EpicGames.Perforce/PerforceReflection.cs
@@ -20,10 +20,10 @@ namespace EpicGames.Perforce
 	/// </summary>
 	static class StringConstants
 	{
-		public static readonly ReadOnlyUtf8String True = new ReadOnlyUtf8String("true");
-		public static readonly ReadOnlyUtf8String New = new ReadOnlyUtf8String("new");
-		public static readonly ReadOnlyUtf8String None = new ReadOnlyUtf8String("none");
-		public static readonly ReadOnlyUtf8String Default = new ReadOnlyUtf8String("default");
+		public static readonly Utf8String True = new Utf8String("true");
+		public static readonly Utf8String New = new Utf8String("new");
+		public static readonly Utf8String None = new Utf8String("none");
+		public static readonly Utf8String Default = new Utf8String("default");
 	}
 
 	/// <summary>
@@ -34,7 +34,7 @@ namespace EpicGames.Perforce
 		/// <summary>
 		/// Name of the tag. Specified in the attribute or inferred from the field name.
 		/// </summary>
-		public ReadOnlyUtf8String Name;
+		public Utf8String Name;
 
 		/// <summary>
 		/// Whether this tag is optional or not.
@@ -54,7 +54,7 @@ namespace EpicGames.Perforce
 		/// <summary>
 		/// Parser for this field type
 		/// </summary>
-		public Action<object, ReadOnlyUtf8String> SetFromString;
+		public Action<object, Utf8String> SetFromString;
 
 		/// <summary>
 		/// Index into the bitmask of required types
@@ -68,7 +68,7 @@ namespace EpicGames.Perforce
 		/// <param name="Optional"></param>
 		/// <param name="Field"></param>
 		/// <param name="RequiredTagBitMask"></param>
-		public CachedTagInfo(ReadOnlyUtf8String Name, bool Optional, FieldInfo Field, ulong RequiredTagBitMask)
+		public CachedTagInfo(Utf8String Name, bool Optional, FieldInfo Field, ulong RequiredTagBitMask)
 		{
 			this.Name = Name;
 			this.Optional = Optional;
@@ -108,7 +108,7 @@ namespace EpicGames.Perforce
 		/// <summary>
 		/// Map of name to tag info
 		/// </summary>
-		public Dictionary<ReadOnlyUtf8String, CachedTagInfo> NameToInfo = new Dictionary<ReadOnlyUtf8String, CachedTagInfo>();
+		public Dictionary<Utf8String, CachedTagInfo> NameToInfo = new Dictionary<Utf8String, CachedTagInfo>();
 
 		/// <summary>
 		/// Bitmask of all the required tags. Formed by bitwise-or'ing the RequiredTagBitMask fields for each required CachedTagInfo.
@@ -170,7 +170,7 @@ namespace EpicGames.Perforce
 		/// <summary>
 		/// Map of name to value
 		/// </summary>
-		public Dictionary<ReadOnlyUtf8String, int> NameToValue = new Dictionary<ReadOnlyUtf8String, int>();
+		public Dictionary<Utf8String, int> NameToValue = new Dictionary<Utf8String, int>();
 
 		/// <summary>
 		/// List of name/value pairs
@@ -196,7 +196,7 @@ namespace EpicGames.Perforce
 					object? Value = Field.GetValue(null);
 					if (Value != null)
 					{
-						ReadOnlyUtf8String Name = new ReadOnlyUtf8String(Attribute.Name);
+						Utf8String Name = new Utf8String(Attribute.Name);
 						NameToValue[Name] = (int)Value;
 
 						NameValuePairs.Add(new KeyValuePair<string, int>(Attribute.Name, (int)Value));
@@ -220,7 +220,7 @@ namespace EpicGames.Perforce
 		/// </summary>
 		/// <param name="Text">The text to parse</param>
 		/// <returns>The enum value corresponding to the given text</returns>
-		public object ParseString(ReadOnlyUtf8String Text)
+		public object ParseString(Utf8String Text)
 		{
 			return Enum.ToObject(EnumType, ParseToInteger(Text));
 		}
@@ -230,7 +230,7 @@ namespace EpicGames.Perforce
 		/// </summary>
 		/// <param name="Text"></param>
 		/// <returns></returns>
-		public int ParseToInteger(ReadOnlyUtf8String Name)
+		public int ParseToInteger(Utf8String Name)
 		{
 			if (bHasFlagsAttribute)
 			{
@@ -251,7 +251,7 @@ namespace EpicGames.Perforce
 						}
 
 						// Take the subset
-						ReadOnlyUtf8String Item = Name.Slice(StartOffset, Offset - StartOffset);
+						Utf8String Item = Name.Slice(StartOffset, Offset - StartOffset);
 
 						// Get the value
 						int ItemValue;
@@ -419,7 +419,7 @@ namespace EpicGames.Perforce
 							Record.RequiredTagsBitMask |= RequiredTagBitMask;
 						}
 
-						CachedTagInfo TagInfo = new CachedTagInfo(new ReadOnlyUtf8String(TagName), TagAttribute.Optional, Field, RequiredTagBitMask);
+						CachedTagInfo TagInfo = new CachedTagInfo(new Utf8String(TagName), TagAttribute.Optional, Field, RequiredTagBitMask);
 
 						Type FieldType = Field.FieldType;
 
@@ -491,12 +491,12 @@ namespace EpicGames.Perforce
 			return Record;
 		}
 
-		static object ParseString(ReadOnlyUtf8String String)
+		static object ParseString(Utf8String String)
 		{
 			return String.ToString();
 		}
 
-		static object ParseStringAsDateTime(ReadOnlyUtf8String String)
+		static object ParseStringAsDateTime(Utf8String String)
 		{
 			string Text = String.ToString();
 
@@ -511,17 +511,17 @@ namespace EpicGames.Perforce
 			}
 		}
 
-		static object ParseStringAsBool(ReadOnlyUtf8String String)
+		static object ParseStringAsBool(Utf8String String)
 		{
 			return String.Length == 0 || String == StringConstants.True;
 		}
 
-		static object ParseStringAsNullableBool(ReadOnlyUtf8String String)
+		static object ParseStringAsNullableBool(Utf8String String)
 		{
 			return String == StringConstants.True;
 		}
 
-		static object ParseStringAsInt(ReadOnlyUtf8String String)
+		static object ParseStringAsInt(Utf8String String)
 		{
 			int Value;
 			int BytesConsumed;
@@ -547,7 +547,7 @@ namespace EpicGames.Perforce
 			}
 		}
 
-		static object ParseStringAsLong(ReadOnlyUtf8String String)
+		static object ParseStringAsLong(Utf8String String)
 		{
 			long Value;
 			int BytesConsumed;
@@ -558,7 +558,7 @@ namespace EpicGames.Perforce
 			return Value;
 		}
 
-		static object ParseStringAsNullableDateTimeOffset(ReadOnlyUtf8String String)
+		static object ParseStringAsNullableDateTimeOffset(Utf8String String)
 		{
 			string Text = String.ToString();
 			return DateTimeOffset.Parse(Regex.Replace(Text, "[a-zA-Z ]*$", "")); // Strip timezone name (eg. "EST")
diff --git a/Engine/Source/Programs/Shared/EpicGames.Perforce/PerforceUtils.cs b/Engine/Source/Programs/Shared/EpicGames.Perforce/PerforceUtils.cs
--- a/Engine/Source/Programs/Shared/EpicGames.Perforce/PerforceUtils.cs
+++ b/Engine/Source/Programs/Shared/EpicGames.Perforce/PerforceUtils.cs
@@ -33,7 +33,7 @@ namespace EpicGames.Perforce
 			return NewPath;
 		}
 
-		static public ReadOnlyUtf8String UnescapePath(ReadOnlyUtf8String Path)
+		static public Utf8String UnescapePath(Utf8String Path)
 		{
 			ReadOnlySpan<byte> PathSpan = Path.Span;
 			for (int InputIdx = 0; InputIdx < PathSpan.Length - 2; InputIdx++)
@@ -69,7 +69,7 @@ namespace EpicGames.Perforce
 					}
 
 					// Copy the last chunk of data to the output buffer
-					Path = new ReadOnlyUtf8String(Buffer.AsMemory(0, OutputIdx));
+					Path = new Utf8String(Buffer.AsMemory(0, OutputIdx));
 					break;
 				}
 			}
diff --git a/Engine/Source/Programs/Shared/EpicGames.Perforce/Records/FileType.cs b/Engine/Source/Programs/Shared/EpicGames.Perforce/Records/FileType.cs
--- a/Engine/Source/Programs/Shared/EpicGames.Perforce/Records/FileType.cs
+++ b/Engine/Source/Programs/Shared/EpicGames.Perforce/Records/FileType.cs
@@ -133,7 +133,7 @@ namespace EpicGames.Perforce.Managed
 		/// <summary>
 		/// Array of file type bases, with order matching <see cref="FileTypeBase"/>
 		/// </summary>
-		static ReadOnlyUtf8String[] BaseNames =
+		static Utf8String[] BaseNames =
 		{
 			"text",
 			"binary",
diff --git a/Engine/Source/Programs/Shared/EpicGames.Serialization.Tests/EngineTests.cs b/Engine/Source/Programs/Shared/EpicGames.Serialization.Tests/EngineTests.cs
--- a/Engine/Source/Programs/Shared/EpicGames.Serialization.Tests/EngineTests.cs
+++ b/Engine/Source/Programs/Shared/EpicGames.Serialization.Tests/EngineTests.cs
@@ -46,7 +46,7 @@ namespace EpicGames.Serialization.Tests
 			[CbFieldType.Array] = new CbFieldAccessors(CbArray.Empty, x => x.IsArray(), x => x.AsArray(), (x, y) => x.AsArray()),
 			[CbFieldType.UniformArray] = new CbFieldAccessors(CbArray.Empty, x => x.IsArray(), x => x.AsArray(), (x, y) => x.AsArray()),
 			[CbFieldType.Binary] = CbFieldAccessors.FromStruct<ReadOnlyMemory<byte>>(x => x.IsBinary(), (x, y) => x.AsBinary(y), (x, y) => x.Span.SequenceEqual(y.Span)),
-			[CbFieldType.String] = new CbFieldAccessors(ReadOnlyUtf8String.Empty, x => x.IsString(), x => x.AsString(), (x, Default) => x.AsString((ReadOnlyUtf8String)Default)),
+			[CbFieldType.String] = new CbFieldAccessors(Utf8String.Empty, x => x.IsString(), x => x.AsString(), (x, Default) => x.AsString((Utf8String)Default)),
 			[CbFieldType.IntegerPositive] = CbFieldAccessors.FromStruct<ulong>(x => x.IsInteger(), (x, y) => x.AsUInt64(y)),
 			[CbFieldType.IntegerNegative] = CbFieldAccessors.FromStruct<long>(x => x.IsInteger(), (x, y) => x.AsInt64(y)),
 			[CbFieldType.Float32] = CbFieldAccessors.FromStruct<float>(x => x.IsFloat(), (x, y) => x.AsFloat(y)),
@@ -373,7 +373,7 @@ namespace EpicGames.Serialization.Tests
 			Action<CbArray, int, int> TestIntArray = (CbArray Array, int ExpectedNum, int ExpectedPayloadSize) =>
 			{
 				Assert.AreEqual(Array.GetSize(), ExpectedPayloadSize + sizeof(CbFieldType));
-				Assert.AreEqual(Array.Num(), ExpectedNum);
+				Assert.AreEqual(Array.Count, ExpectedNum);
 
 				int ActualNum = 0;
 				for (CbFieldIterator It = Array.CreateIterator(); It; ++It)
@@ -582,20 +582,20 @@ namespace EpicGames.Serialization.Tests
 			// Test CbField(String, Value)
 			{
 				byte[] Payload = { 3, (byte)'A', (byte)'B', (byte)'C' }; // Size: 3, Data: ABC
-				TestField(CbFieldType.String, Payload, new ReadOnlyUtf8String(Payload.AsMemory(1, 3)));
+				TestField(CbFieldType.String, Payload, new Utf8String(Payload.AsMemory(1, 3)));
 			}
 
 			// Test CbField(String, OutOfRangeSize)
 			{
 				byte[] Payload = new byte[9];
 				VarInt.Write(Payload, (ulong)(1) << 31);
-				TestFieldError(CbFieldType.String, Payload, CbFieldError.RangeError, new ReadOnlyUtf8String("ABC"));
+				TestFieldError(CbFieldType.String, Payload, CbFieldError.RangeError, new Utf8String("ABC"));
 			}
 
 			// Test CbField(None) as String
 			{
 				CbField Field = new CbField();
-				TestFieldError(CbFieldType.String, Field, CbFieldError.TypeError, new ReadOnlyUtf8String("ABC"));
+				TestFieldError(CbFieldType.String, Field, CbFieldError.TypeError, new Utf8String("ABC"));
 			}
 		}
 
diff --git a/Engine/Source/Programs/Shared/EpicGames.Serialization/CompactBinary.cs b/Engine/Source/Programs/Shared/EpicGames.Serialization/CompactBinary.cs
--- a/Engine/Source/Programs/Shared/EpicGames.Serialization/CompactBinary.cs
+++ b/Engine/Source/Programs/Shared/EpicGames.Serialization/CompactBinary.cs
@@ -608,7 +608,7 @@ namespace EpicGames.Serialization
 		/// <summary>
 		/// Returns the name of the field if it has a name, otherwise an empty view.
 		/// </summary>
-		public ReadOnlyUtf8String Name => new ReadOnlyUtf8String(Memory.Slice(PayloadOffset - NameLen, NameLen));
+		public Utf8String Name => new Utf8String(Memory.Slice(PayloadOffset - NameLen, NameLen));
 
 		/// <summary>
 		/// Gets the value of this field
@@ -663,7 +663,7 @@ namespace EpicGames.Serialization
 		}
 
 		/// <inheritdoc cref="Name"/>
-		public ReadOnlyUtf8String GetName() => Name;
+		public Utf8String GetName() => Name;
 
 		/// <summary>
 		/// Access the field as an object. Defaults to an empty object on error. 
@@ -726,7 +726,7 @@ namespace EpicGames.Serialization
 		/// </summary>
 		/// <param name="Default">Default value to return</param>
 		/// <returns></returns>
-		public ReadOnlyUtf8String AsString(ReadOnlyUtf8String Default = default)
+		public Utf8String AsString(Utf8String Default = default)
 		{
 			if (CbFieldUtils.IsString(TypeWithFlags))
 			{
@@ -739,7 +739,7 @@ namespace EpicGames.Serialization
 				else
 				{
 					Error = CbFieldError.None;
-					return new ReadOnlyUtf8String(Payload.Slice(ValueSizeByteCount, (int)ValueSize));
+					return new Utf8String(Payload.Slice(ValueSizeByteCount, (int)ValueSize));
 				}
 			}
 			else
@@ -1301,7 +1301,7 @@ namespace EpicGames.Serialization
 		/// </summary>
 		/// <param name="Name"></param>
 		/// <returns></returns>
-		public CbField this[ReadOnlyUtf8String Name]
+		public CbField this[Utf8String Name]
 		{
 			get { return this.FirstOrDefault(Field => Field.Name == Name) ?? CbField.Empty; }
 		}
@@ -1999,21 +1999,21 @@ namespace EpicGames.Serialization
 		/// </summary>
 		/// <param name="Name">The name of the field.</param>
 		/// <returns>The matching field if found, otherwise a field with no value.</returns>
-		public CbField Find(ReadOnlyUtf8String Name) => InnerField[Name];
+		public CbField Find(Utf8String Name) => InnerField[Name];
 
 		/// <summary>
 		/// Find a field by case-insensitive name comparison.
 		/// </summary>
 		/// <param name="Name">The name of the field.</param>
 		/// <returns>The matching field if found, otherwise a field with no value.</returns>
-		public CbField FindIgnoreCase(ReadOnlyUtf8String Name) => InnerField.FirstOrDefault(Field => ReadOnlyUtf8StringComparer.OrdinalIgnoreCase.Equals(Field.Name, Name)) ?? new CbField();
+		public CbField FindIgnoreCase(Utf8String Name) => InnerField.FirstOrDefault(Field => Utf8StringComparer.OrdinalIgnoreCase.Equals(Field.Name, Name)) ?? new CbField();
 
 		/// <summary>
 		/// Find a field by case-sensitive name comparison.
 		/// </summary>
 		/// <param name="Name">The name of the field.</param>
 		/// <returns>The matching field if found, otherwise a field with no value.</returns>
-		public CbField this[ReadOnlyUtf8String Name] => InnerField[Name];
+		public CbField this[Utf8String Name] => InnerField[Name];
 
 		/// <inheritdoc cref="AsFieldView"/>
 		public CbField AsField() => InnerField;
diff --git a/Engine/Source/Programs/Shared/EpicGames.Serialization/CompactBinaryWriter.cs b/Engine/Source/Programs/Shared/EpicGames.Serialization/CompactBinaryWriter.cs
--- a/Engine/Source/Programs/Shared/EpicGames.Serialization/CompactBinaryWriter.cs
+++ b/Engine/Source/Programs/Shared/EpicGames.Serialization/CompactBinaryWriter.cs
@@ -180,7 +180,7 @@ namespace EpicGames.Serialization
 		/// </summary>
 		/// <param name="Type"></param>
 		/// <param name="Name"></param>
-		void WriteField(CbFieldType Type, ReadOnlyUtf8String Name)
+		void WriteField(CbFieldType Type, Utf8String Name)
 		{
 			Scope Scope = CurrentScope;
 			if (!CbFieldUtils.IsObject(Scope.FieldType))
@@ -219,7 +219,7 @@ namespace EpicGames.Serialization
 		/// Begin writing an object field
 		/// </summary>
 		/// <param name="Name">Name of the field</param>
-		public void BeginObject(ReadOnlyUtf8String Name)
+		public void BeginObject(Utf8String Name)
 		{
 			WriteField(CbFieldType.Object, Name);
 			PushScope(CbFieldType.Object, CbFieldType.None);
@@ -246,7 +246,7 @@ namespace EpicGames.Serialization
 		/// Begin writing a named array field
 		/// </summary>
 		/// <param name="Name"></param>
-		public void BeginArray(ReadOnlyUtf8String Name)
+		public void BeginArray(Utf8String Name)
 		{
 			WriteField(CbFieldType.Array, Name);
 			PushScope(CbFieldType.Array, CbFieldType.None);
@@ -276,7 +276,7 @@ namespace EpicGames.Serialization
 		/// </summary>
 		/// <param name="Name">Name of the field</param>
 		/// <param name="FieldType">The field type for elements in the array</param>
-		public void BeginUniformArray(ReadOnlyUtf8String Name, CbFieldType FieldType)
+		public void BeginUniformArray(Utf8String Name, CbFieldType FieldType)
 		{
 			WriteField(CbFieldType.UniformArray, Name);
 			PushScope(CbFieldType.UniformArray, FieldType);
@@ -303,7 +303,7 @@ namespace EpicGames.Serialization
 		/// Write a named null field
 		/// </summary>
 		/// <param name="Name">Name of the field</param>
-		public void WriteNull(ReadOnlyUtf8String Name)
+		public void WriteNull(Utf8String Name)
 		{
 			WriteField(CbFieldType.Null, Name);
 		}
@@ -342,7 +342,7 @@ namespace EpicGames.Serialization
 		/// </summary>
 		/// <param name="Name">Name of the field</param>
 		/// <param name="Value">Value to be written</param>
-		public void WriteInteger(ReadOnlyUtf8String Name, long Value)
+		public void WriteInteger(Utf8String Name, long Value)
 		{
 			if (Value >= 0)
 			{
@@ -371,7 +371,7 @@ namespace EpicGames.Serialization
 		/// </summary>
 		/// <param name="Name">Name of the field</param>
 		/// <param name="Value">Value to be written</param>
-		public void WriteInteger(ReadOnlyUtf8String Name, ulong Value)
+		public void WriteInteger(Utf8String Name, ulong Value)
 		{
 			WriteField(CbFieldType.IntegerPositive, Name);
 			WriteIntegerPayload(Value);
@@ -402,7 +402,7 @@ namespace EpicGames.Serialization
 		/// </summary>
 		/// <param name="Name">Name of the field</param>
 		/// <param name="Value">Value to be written</param>
-		public void WriteDateTime(ReadOnlyUtf8String Name, DateTime Value)
+		public void WriteDateTime(Utf8String Name, DateTime Value)
 		{
 			WriteField(CbFieldType.DateTime, Name);
 			WriteDateTimePayload(Value);
@@ -433,7 +433,7 @@ namespace EpicGames.Serialization
 		/// </summary>
 		/// <param name="Name">Name of the field</param>
 		/// <param name="Value">Value to be written</param>
-		public void WriteHash(ReadOnlyUtf8String Name, IoHash Value)
+		public void WriteHash(Utf8String Name, IoHash Value)
 		{
 			WriteField(CbFieldType.Hash, Name);
 			WriteHashPayload(Value);
@@ -454,7 +454,7 @@ namespace EpicGames.Serialization
 		/// </summary>
 		/// <param name="Name">Name of the field</param>
 		/// <param name="Hash">Hash of the attachment</param>
-		public void WriteBinaryAttachment(ReadOnlyUtf8String Name, IoHash Hash)
+		public void WriteBinaryAttachment(Utf8String Name, IoHash Hash)
 		{
 			WriteField(CbFieldType.BinaryAttachment, Name);
 			WriteHashPayload(Hash);
@@ -475,7 +475,7 @@ namespace EpicGames.Serialization
 		/// </summary>
 		/// <param name="Name">Name of the field</param>
 		/// <param name="Hash">Hash of the attachment</param>
-		public void WriteObjectAttachment(ReadOnlyUtf8String Name, IoHash Hash)
+		public void WriteObjectAttachment(Utf8String Name, IoHash Hash)
 		{
 			WriteField(CbFieldType.ObjectAttachment, Name);
 			WriteHashPayload(Hash);
@@ -510,7 +510,7 @@ namespace EpicGames.Serialization
 		/// Writes an unnamed string value
 		/// </summary>
 		/// <param name="Value">Value to be written</param>
-		public void WriteStringValue(ReadOnlyUtf8String Value)
+		public void WriteStringValue(Utf8String Value)
 		{
 			WriteField(CbFieldType.String);
 			WriteBinaryPayload(Value.Span);
@@ -521,7 +521,7 @@ namespace EpicGames.Serialization
 		/// </summary>
 		/// <param name="Name">Name of the field</param>
 		/// <param name="Value">Value to be written</param>
-		public void WriteString(ReadOnlyUtf8String Name, ReadOnlyUtf8String Value)
+		public void WriteString(Utf8String Name, Utf8String Value)
 		{
 			WriteField(CbFieldType.String, Name);
 			WriteBinaryPayload(Value.Span);
@@ -542,7 +542,7 @@ namespace EpicGames.Serialization
 		/// </summary>
 		/// <param name="Name">Name of the field</param>
 		/// <param name="Value">Value to be written</param>
-		public void WriteBinary(ReadOnlyUtf8String Name, ReadOnlySpan<byte> Value)
+		public void WriteBinary(Utf8String Name, ReadOnlySpan<byte> Value)
 		{
 			WriteField(CbFieldType.Binary, Name);
 			WriteBinaryPayload(Value);
