Add missing Md5Hash file.

diff --git a/Engine/Source/Programs/Shared/EpicGames.Core/Md5Hash.cs b/Engine/Source/Programs/Shared/EpicGames.Core/Md5Hash.cs
new file mode 100644
--- /dev/null
+++ b/Engine/Source/Programs/Shared/EpicGames.Core/Md5Hash.cs
@@ -0,0 +1,194 @@
+using System;
+using System.Buffers.Binary;
+using System.Collections.Generic;
+using System.ComponentModel;
+using System.Globalization;
+using System.Security.Cryptography;
+using System.Text;
+
+namespace EpicGames.Core
+{
+	/// <summary>
+	/// Struct representing a strongly typed Md5Hash value
+	/// </summary>
+	[TypeConverter(typeof(Md5HashTypeConverter))]
+	public struct Md5Hash : IEquatable<Md5Hash>, IComparable<Md5Hash>
+	{
+		/// <summary>
+		/// Length of an Md5Hash
+		/// </summary>
+		public const int NumBytes = 16;
+
+		/// <summary>
+		/// Length of the hash in bits
+		/// </summary>
+		public const int NumBits = NumBytes * 8;
+
+		/// <summary>
+		/// Memory storing the digest data
+		/// </summary>
+		public ReadOnlyMemory<byte> Memory;
+
+		/// <summary>
+		/// Span for the underlying memory
+		/// </summary>
+		public ReadOnlySpan<byte> Span => Memory.Span;
+
+		/// <summary>
+		/// Hash consisting of zeroes
+		/// </summary>
+		public static Md5Hash Zero { get; } = new Md5Hash(new byte[NumBytes]);
+
+		/// <summary>
+		/// Constructor
+		/// </summary>
+		/// <param name="Memory">Memory to construct from</param>
+		public Md5Hash(ReadOnlyMemory<byte> Memory)
+		{
+			if (Memory.Length != NumBytes)
+			{
+				throw new ArgumentException($"Md5Hash must be {NumBytes} bytes long");
+			}
+
+			this.Memory = Memory;
+		}
+
+		/// <summary>
+		/// Creates a content hash for a block of data, using a given algorithm.
+		/// </summary>
+		/// <param name="Data">Data to compute the hash for</param>
+		/// <returns>New content hash instance containing the hash of the data</returns>
+		public static Md5Hash Compute(ReadOnlySpan<byte> Data)
+		{
+			byte[] Output = new byte[NumBytes];
+			using (MD5 Hasher = MD5.Create())
+			{
+				Hasher.TryComputeHash(Data, Output, out _);
+			}
+			return new Md5Hash(Output);
+		}
+
+		/// <summary>
+		/// Parses a digest from the given hex string
+		/// </summary>
+		/// <param name="Text"></param>
+		/// <returns></returns>
+		public static Md5Hash Parse(string Text)
+		{
+			return new Md5Hash(StringUtils.ParseHexString(Text));
+		}
+
+		/// <summary>
+		/// Parses a digest from the given hex string
+		/// </summary>
+		/// <param name="Text"></param>
+		/// <returns></returns>
+		public static Md5Hash Parse(Utf8String Text)
+		{
+			return new Md5Hash(StringUtils.ParseHexString(Text.Span));
+		}
+
+		/// <inheritdoc cref="IComparable{T}.CompareTo(T)"/>
+		public int CompareTo(Md5Hash Other)
+		{
+			ReadOnlySpan<byte> A = Span;
+			ReadOnlySpan<byte> B = Other.Span;
+
+			for (int Idx = 0; Idx < A.Length && Idx < B.Length; Idx++)
+			{
+				int Compare = A[Idx] - B[Idx];
+				if (Compare != 0)
+				{
+					return Compare;
+				}
+			}
+			return A.Length - B.Length;
+		}
+
+		/// <inheritdoc/>
+		public bool Equals(Md5Hash Other) => Span.SequenceEqual(Other.Span);
+
+		/// <inheritdoc/>
+		public override bool Equals(object? Obj) => (Obj is Md5Hash Hash) && Hash.Span.SequenceEqual(Span);
+
+		/// <inheritdoc/>
+		public override int GetHashCode() => BinaryPrimitives.ReadInt32LittleEndian(Span);
+
+		/// <inheritdoc/>
+		public override string ToString() => StringUtils.FormatHexString(Memory.Span);
+
+		/// <summary>
+		/// Test two hash values for equality
+		/// </summary>
+		public static bool operator ==(Md5Hash A, Md5Hash B) => A.Span.SequenceEqual(B.Span);
+
+		/// <summary>
+		/// Test two hash values for equality
+		/// </summary>
+		public static bool operator !=(Md5Hash A, Md5Hash B) => !(A == B);
+
+		/// <summary>
+		/// Tests whether A > B
+		/// </summary>
+		public static bool operator >(Md5Hash A, Md5Hash B) => A.CompareTo(B) > 0;
+
+		/// <summary>
+		/// Tests whether A is less than B
+		/// </summary>
+		public static bool operator <(Md5Hash A, Md5Hash B) => A.CompareTo(B) < 0;
+
+		/// <summary>
+		/// Tests whether A is greater than or equal to B
+		/// </summary>
+		public static bool operator >=(Md5Hash A, Md5Hash B) => A.CompareTo(B) >= 0;
+
+		/// <summary>
+		/// Tests whether A is less than or equal to B
+		/// </summary>
+		public static bool operator <=(Md5Hash A, Md5Hash B) => A.CompareTo(B) <= 0;
+	}
+
+	/// <summary>
+	/// Extension methods for dealing with Md5Hash values
+	/// </summary>
+	public static class Md5HashExtensions
+	{
+		/// <summary>
+		/// Read an <see cref="Md5Hash"/> from a memory reader
+		/// </summary>
+		/// <param name="Reader"></param>
+		/// <returns></returns>
+		public static Md5Hash ReadMd5Hash(this MemoryReader Reader)
+		{
+			return new Md5Hash(Reader.ReadFixedLengthBytes(Md5Hash.NumBytes));
+		}
+
+		/// <summary>
+		/// Write an <see cref="Md5Hash"/> to a memory writer
+		/// </summary>
+		/// <param name="Writer"></param>
+		/// <param name="Hash"></param>
+		public static void WriteMd5Hash(this MemoryWriter Writer, Md5Hash Hash)
+		{
+			Writer.WriteFixedLengthBytes(Hash.Span);
+		}
+	}
+
+	/// <summary>
+	/// Type converter from strings to Md5Hash objects
+	/// </summary>
+	sealed class Md5HashTypeConverter : TypeConverter
+	{
+		/// <inheritdoc/>
+		public override bool CanConvertFrom(ITypeDescriptorContext Context, Type SourceType)
+		{
+			return SourceType == typeof(string);
+		}
+
+		/// <inheritdoc/>
+		public override object ConvertFrom(ITypeDescriptorContext Context, CultureInfo Culture, object Value)
+		{
+			return Md5Hash.Parse((string)Value);
+		}
+	}
+}
