Missing file.

diff --git a/Engine/Source/Programs/Shared/EpicGames.Perforce.Managed/StreamSnapshotFromMemory.cs b/Engine/Source/Programs/Shared/EpicGames.Perforce.Managed/StreamSnapshotFromMemory.cs
--- a/Engine/Source/Programs/Shared/EpicGames.Perforce.Managed/StreamSnapshotFromMemory.cs
+++ b/Engine/Source/Programs/Shared/EpicGames.Perforce.Managed/StreamSnapshotFromMemory.cs
@@ -1,6 +1,7 @@
 // Copyright Epic Games, Inc. All Rights Reserved.
 
 using EpicGames.Core;
+using EpicGames.Serialization;
 using System;
 using System.Collections.Generic;
 using System.IO;
@@ -19,27 +20,27 @@ namespace EpicGames.Perforce.Managed
 		/// <summary>
 		/// The current signature for saved directory objects
 		/// </summary>
-		static readonly byte[] CurrentSignature = { (byte)'W', (byte)'S', (byte)'D', 3 };
+		static readonly byte[] CurrentSignature = { (byte)'W', (byte)'S', (byte)'D', 4 };
 
 		/// <summary>
 		/// The root digest
 		/// </summary>
-		public override IoHash Root { get; }
+		public override StreamTreeRef Root { get; }
 
 		/// <summary>
 		/// Map of digest to directory
 		/// </summary>
-		public IReadOnlyDictionary<IoHash, StreamDirectoryInfo> HashToDirectory { get; }
+		public IReadOnlyDictionary<IoHash, CbObject> HashToTree { get; }
 
 		/// <summary>
 		/// Constructor
 		/// </summary>
 		/// <param name="Root"></param>
-		/// <param name="HashToDirectory"></param>
-		public StreamSnapshotFromMemory(IoHash Root, Dictionary<IoHash, StreamDirectoryInfo> HashToDirectory)
+		/// <param name="HashToTree"></param>
+		public StreamSnapshotFromMemory(StreamTreeRef Root, Dictionary<IoHash, CbObject> HashToTree)
 		{
 			this.Root = Root;
-			this.HashToDirectory = HashToDirectory;
+			this.HashToTree = HashToTree;
 		}
 
 		/// <summary>
@@ -48,15 +49,15 @@ namespace EpicGames.Perforce.Managed
 		/// <param name="Builder"></param>
 		public StreamSnapshotFromMemory(StreamSnapshotBuilder Builder)
 		{
-			Dictionary<IoHash, StreamDirectoryInfo> HashToDirectory = new Dictionary<IoHash, StreamDirectoryInfo>();
-			this.Root = Builder.Encode(HashToDirectory);
-			this.HashToDirectory = HashToDirectory;
+			Dictionary<IoHash, CbObject> HashToTree = new Dictionary<IoHash, CbObject>();
+			this.Root = Builder.Encode(HashToTree);
+			this.HashToTree = HashToTree;
 		}
 
 		/// <inheritdoc/>
-		public override StreamDirectoryInfo Lookup(IoHash Hash)
+		public override StreamTree Lookup(StreamTreeRef Ref)
 		{
-			return HashToDirectory[Hash];
+			return new StreamTree(HashToTree[Ref.Hash], Ref.Path);
 		}
 
 		/// <summary>
@@ -65,7 +66,7 @@ namespace EpicGames.Perforce.Managed
 		/// <param name="InputFile">File to read from</param>
 		/// <param name="CancellationToken">Cancellation token</param>
 		/// <returns>New StreamDirectoryInfo object</returns>
-		public static async Task<StreamSnapshotFromMemory?> TryLoadAsync(FileReference InputFile, CancellationToken CancellationToken)
+		public static async Task<StreamSnapshotFromMemory?> TryLoadAsync(FileReference InputFile, Utf8String BasePath, CancellationToken CancellationToken)
 		{
 			byte[] Data = await FileReference.ReadAllBytesAsync(InputFile);
 			if (!Data.AsSpan().StartsWith(CurrentSignature))
@@ -73,42 +74,51 @@ namespace EpicGames.Perforce.Managed
 				return null;
 			}
 
-			MemoryReader Reader = new MemoryReader(Data.AsMemory(CurrentSignature.Length));
+			CbObject RootObj = new CbObject(Data.AsMemory(CurrentSignature.Length));
 
-			IoHash Root = Reader.ReadIoHash();
+			CbObject RootObj2 = RootObj["root"].AsObject();
+			StreamTreeRef Root = new StreamTreeRef(RootObj2, BasePath);
 
-			int NumItems = Reader.ReadInt32();
+			CbArray Array = RootObj["items"].AsArray();
 
-			Dictionary<IoHash, StreamDirectoryInfo> HashToDirectory = new Dictionary<IoHash, StreamDirectoryInfo>(NumItems);
-			for (int Idx = 0; Idx < NumItems; Idx++)
+			Dictionary<IoHash, CbObject> HashToTree = new Dictionary<IoHash, CbObject>(Array.Count);
+			foreach (CbField Element in Array)
 			{
-				IoHash Hash = Reader.ReadIoHash();
-				StreamDirectoryInfo Info = Reader.ReadStreamDirectoryInfo();
-				HashToDirectory[Hash] = Info;
+				CbObject ObjectElement = Element.AsObject();
+				IoHash Hash = ObjectElement["hash"].AsHash();
+				CbObject Tree = ObjectElement["tree"].AsObject();
+				HashToTree[Hash] = Tree;
 			}
 
-			return new StreamSnapshotFromMemory(Root, HashToDirectory);
+			return new StreamSnapshotFromMemory(Root, HashToTree);
 		}
 
 		/// <summary>
 		/// Saves the contents of this object to disk
 		/// </summary>
 		/// <param name="OutputFile">The output file to write to</param>
-		public async Task Save(FileReference OutputFile)
+		public async Task Save(FileReference OutputFile, Utf8String BasePath)
 		{
-			int Length = Digest<Sha1>.Length + sizeof(int) + HashToDirectory.Sum(x => Digest<Sha1>.Length + x.Value.GetSerializedSize());
-			byte[] Data = new byte[Length];
-
-			MemoryWriter Writer = new MemoryWriter(Data.AsMemory());
-			Writer.WriteIoHash(Root);
-			Writer.WriteInt32(HashToDirectory.Count);
-			foreach ((IoHash Hash, StreamDirectoryInfo Info) in HashToDirectory)
+			CbWriter Writer = new CbWriter();
+			Writer.BeginObject();
+						
+			Writer.BeginObject("root");
+			Root.Write(Writer, BasePath);
+			Writer.EndObject();
+			
+			Writer.BeginArray("items");
+			foreach ((IoHash Hash, CbObject Tree) in HashToTree)
 			{
-				Writer.WriteIoHash(Hash);
-				Writer.WriteStreamDirectoryInfo(Info);
+				Writer.BeginObject();
+				Writer.WriteHash("hash", Hash);
+				Writer.WriteObject("tree", Tree);
+				Writer.EndObject();
 			}
-			Writer.CheckOffset(Data.Length);
+			Writer.EndArray();
+
+			Writer.EndObject();
 
+			byte[] Data = Writer.ToByteArray();
 			using (FileStream OutputStream = FileReference.Open(OutputFile, FileMode.Create, FileAccess.Write, FileShare.Read))
 			{
 				await OutputStream.WriteAsync(CurrentSignature, 0, CurrentSignature.Length);
