Missing copyright notice.

diff --git a/Engine/Source/Programs/Shared/EpicGames.Core/Md5Hash.cs b/Engine/Source/Programs/Shared/EpicGames.Core/Md5Hash.cs
--- a/Engine/Source/Programs/Shared/EpicGames.Core/Md5Hash.cs
+++ b/Engine/Source/Programs/Shared/EpicGames.Core/Md5Hash.cs
@@ -1,3 +1,5 @@
+// Copyright Epic Games, Inc. All Rights Reserved.
+
 using System;
 using System.Buffers.Binary;
 using System.Collections.Generic;
