Add cpu time spent per compilation unit in parallel and task executor
Add top 20 report when compilation ends for parallel and task executor
Add total cpu usage across all cores when compilation finishes for parallel and task executor
Use real processor time reported by OS when possible to account for internal usage of multiple threads (i.e. cl.exe codegen uses 8 threads)
Enabled by default for Epic developers to increase awareness on most problematic compilation times

#rnx

diff --git a/Engine/Source/Programs/Shared/EpicGames.Core/ManagedProcess.cs b/Engine/Source/Programs/Shared/EpicGames.Core/ManagedProcess.cs
--- a/Engine/Source/Programs/Shared/EpicGames.Core/ManagedProcess.cs
+++ b/Engine/Source/Programs/Shared/EpicGames.Core/ManagedProcess.cs
@@ -75,14 +75,31 @@ namespace EpicGames.Core
 			public UIntPtr PeakProcessMemoryUsed;
 			public UIntPtr PeakJobMemoryUsed;
 		}
+		
+		[StructLayout(LayoutKind.Sequential)]
+		struct JOBOBJECT_BASIC_ACCOUNTING_INFORMATION
+		{
+			public UInt64 TotalUserTime;
+			public UInt64 TotalKernelTime;
+			public UInt64 ThisPeriodTotalUserTime;
+			public UInt64 ThisPeriodTotalKernelTime;
+			public UInt32 TotalPageFaultCount;
+			public UInt32 TotalProcesses;
+			public UInt32 ActiveProcesses;
+			public UInt32 TotalTerminatedProcesses;
+		}
 
 		[DllImport("kernel32.dll", SetLastError=true)]
 		static extern SafeFileHandle CreateJobObject(IntPtr SecurityAttributes, IntPtr Name);
 
+		const int JobObjectBasicAccountingInformation = 1;
 		const int JobObjectExtendedLimitInformation = 9;
 
         [DllImport("kernel32.dll", SetLastError=true)]
         static extern int SetInformationJobObject(SafeFileHandle hJob, int JobObjectInfoClass, IntPtr lpJobObjectInfo, int cbJobObjectInfoLength);
+		
+		[DllImport("kernel32.dll", SetLastError = true)]
+		static extern bool QueryInformationJobObject(SafeFileHandle hJob, int JobObjectInformationClass, ref JOBOBJECT_BASIC_ACCOUNTING_INFORMATION lpJobObjectInformation, int cbJobObjectInformationLength);
 
 		/// <summary>
 		/// Handle to the native job object that this process is added to. This handle is closed by the Dispose() method (and will automatically be closed by the OS on process exit),
@@ -131,6 +148,23 @@ namespace EpicGames.Core
 			}
 		}
 
+		/// <summary>
+		/// Returns the total CPU time usage for this job.
+		/// </summary>
+		public TimeSpan TotalProcessorTime
+		{
+			get
+			{
+				JOBOBJECT_BASIC_ACCOUNTING_INFORMATION AccountingInformation = new JOBOBJECT_BASIC_ACCOUNTING_INFORMATION();
+				if (QueryInformationJobObject(JobHandle!, JobObjectBasicAccountingInformation, ref AccountingInformation, Marshal.SizeOf(typeof(JOBOBJECT_BASIC_ACCOUNTING_INFORMATION))) == false) 
+				{
+					throw new Win32Exception();
+				}
+
+				return new TimeSpan((long)AccountingInformation.TotalUserTime + (long)AccountingInformation.TotalKernelTime);
+			}
+		}
+
 		public void Add(IntPtr ProcessHandle)
 		{
 		}
@@ -260,6 +294,9 @@ namespace EpicGames.Core
 		[DllImport("kernel32.dll", SetLastError=true)]
 		static extern UInt32 WaitForSingleObject(SafeHandleZeroOrMinusOneIsInvalid hHandle, UInt32 dwMilliseconds);
 
+		[DllImport("kernel32.dll", CharSet = CharSet.Auto, SetLastError = true)]
+		public static extern bool GetProcessTimes(SafeHandleZeroOrMinusOneIsInvalid handle, out long creation, out long exit, out long kernel, out long user);
+
 		const int ERROR_ACCESS_DENIED = 5;
 
 		/// <summary>
@@ -338,6 +375,11 @@ namespace EpicGames.Core
 		/// </summary>
 		static object LockObject = new object();
 
+		/// <summary>
+		/// Used to perform CPU usage and resource accounting of all children process involved in a single compilation unit.
+		/// </summary>
+		ManagedProcessGroup? AccountingProcessGroup;
+
 		/// <summary>
 		/// Spawns a new managed process.
 		/// </summary>
@@ -541,7 +583,7 @@ namespace EpicGames.Core
 					// Add it to our job object
 					if (Group != null && AssignProcessToJobObject(Group.JobHandle, ProcessInfo.hProcess) == 0)
 					{
-						// Support for nested job objects was only addeed in Windows 8; prior to that, assigning processes to job objects would fail. Figure out if we're already in a job, and ignore the error if we are.
+						// Support for nested job objects was only added in Windows 8; prior to that, assigning processes to job objects would fail. Figure out if we're already in a job, and ignore the error if we are.
 						int OriginalError = Marshal.GetLastWin32Error();
 
 						bool bProcessInJob;
@@ -553,6 +595,16 @@ namespace EpicGames.Core
 						}
 					}
 
+					// Create a JobObject for each spawned process to do CPU usage accounting of spawned process and all its children
+					AccountingProcessGroup = new ManagedProcessGroup();
+					if (AccountingProcessGroup != null)
+					{
+						if (AssignProcessToJobObject(AccountingProcessGroup.JobHandle, ProcessInfo.hProcess) == 0) 
+						{
+							throw new Win32Exception();
+						}
+					}
+
 					// Allow the thread to start running
 					if (ResumeThread(ProcessInfo.hThread) == -1)
 					{
@@ -749,6 +801,10 @@ namespace EpicGames.Core
 				StdOutRead.Dispose();
 				StdOutRead = null;
 			}
+			if (AccountingProcessGroup != null) 
+			{
+				AccountingProcessGroup.Dispose();
+			}
 			if(FrameworkProcess != null)
 			{
 				if(!FrameworkProcess.HasExited)
@@ -921,6 +977,40 @@ namespace EpicGames.Core
 			return await StdOutText!.ReadLineAsync();
 		}
 
+		public TimeSpan TotalProcessorTime
+		{
+			get
+			{
+				// Prefer the job accounting over the GetProcessTimes because we need to account for children process spawned by cl-filter too.
+				long AccountingProcessorTime = 0;
+				if (AccountingProcessGroup != null && AccountingProcessGroup.TotalProcessorTime.Ticks != 0) 
+				{
+					return AccountingProcessGroup.TotalProcessorTime;
+				}
+
+				// Normal process accounting is better than nothing if the processgroup one is not available
+				if (FrameworkProcess == null) 
+				{
+					long creation = 0;
+					long exit = 0;
+					long kernel = 0;
+					long user = 0;
+					if (GetProcessTimes(ProcessHandle!, out creation, out exit, out kernel, out user)) 
+					{
+						return new TimeSpan(Math.Max(user + kernel, AccountingProcessorTime));
+					}
+					else 
+					{
+						throw new Win32Exception();
+					}
+				}
+				else 
+				{
+					return FrameworkProcess.TotalProcessorTime;
+				}
+			}
+		}
+
 		/// <summary>
 		/// Block until the process outputs a line of text, or terminates.
 		/// </summary>
diff --git a/Engine/Source/Programs/UnrealBuildTool/Executors/Experimental/TaskExecutor.cs b/Engine/Source/Programs/UnrealBuildTool/Executors/Experimental/TaskExecutor.cs
--- a/Engine/Source/Programs/UnrealBuildTool/Executors/Experimental/TaskExecutor.cs
+++ b/Engine/Source/Programs/UnrealBuildTool/Executors/Experimental/TaskExecutor.cs
@@ -44,6 +44,12 @@ namespace UnrealBuildTool
 		[XmlConfigFile]
 		bool bStopCompilationAfterErrors = false;
 
+		/// <summary>
+		/// Whether to show compilation times along with worst offenders or not.
+		/// </summary>
+		[XmlConfigFile]
+		bool bShowCompilationTimes = false;
+
 		/// <summary>
 		/// How many processes that will be executed in parallel
 		/// </summary>
@@ -97,7 +103,16 @@ namespace UnrealBuildTool
 		{
 			public List<string> LogLines { get; private set; }
 			public int ExitCode { get; private set; }
+			public TimeSpan ExecutionTime { get; private set; }
+			public TimeSpan ProcessorTime { get; private set; }
 
+			public ExecuteResults(List<string> LogLines, int ExitCode, TimeSpan ExecutionTime, TimeSpan ProcessorTime)
+			{
+				this.LogLines = LogLines;
+				this.ExitCode = ExitCode;
+				this.ProcessorTime = ProcessorTime;
+				this.ExecutionTime = ExecutionTime;
+			}
 			public ExecuteResults(List<string> LogLines, int ExitCode)
 			{
 				this.LogLines = LogLines;
@@ -144,6 +159,25 @@ namespace UnrealBuildTool
 			// Wait for all tasks to complete
 			Task.WaitAll(AllTasks.ToArray());
 
+			if (bShowCompilationTimes)
+			{
+				Log.TraceInformation("");
+				Log.TraceInformation("Total CPU Time: {0} s", ProcessGroup.TotalProcessorTime.TotalSeconds);
+				Log.TraceInformation("");
+
+				if (Tasks.Count > 0)
+				{
+					Log.TraceInformation("Compilation Time Top {0}", Math.Min(20, Tasks.Count));
+					Log.TraceInformation("");
+					foreach (var Pair in Tasks.OrderByDescending(x => x.Value.Result.ExecutionTime).Take(20))
+					{
+						string Description = $"{(Pair.Key.Inner.CommandDescription != null ? Pair.Key.Inner.CommandDescription : Pair.Key.Inner.CommandPath.GetFileNameWithoutExtension())} {Pair.Key.Inner.StatusDescription}".Trim();
+						Log.TraceInformation("{0} [ Wall Time {1:0.00} s / CPU Time {2:0.00} s ]", Description, Pair.Value.Result.ExecutionTime.TotalSeconds, Pair.Value.Result.ProcessorTime.TotalSeconds);
+					}
+					Log.TraceInformation("");
+				}
+			}
+
 			// Return if all tasks succeeded
 			return Tasks.Values.All(x => x.Result.ExitCode == 0);
 		}
@@ -172,6 +206,9 @@ namespace UnrealBuildTool
 				string Description = string.Empty;
 				List<string> LogLines;
 				int ExitCode;
+				DateTime StartTime = DateTime.Now;
+				TimeSpan ExecutionTime;
+				TimeSpan ProcessorTime;
 				using (ManagedProcess Process = new ManagedProcess(ProcessGroup, Action.CommandPath.FullName, Action.CommandArguments, Action.WorkingDirectory.FullName, null, null, ProcessPriorityClass.BelowNormal))
 				{
 					MemoryStream StdOutStream = new MemoryStream();
@@ -179,9 +216,11 @@ namespace UnrealBuildTool
 					CancellationToken.ThrowIfCancellationRequested();
 					LogLines = Console.OutputEncoding.GetString(StdOutStream.GetBuffer(), 0, Convert.ToInt32(StdOutStream.Length)).Split(LineEndingSplit, StringSplitOptions.RemoveEmptyEntries).ToList();
 					ExitCode = Process.ExitCode;
+					ProcessorTime = Process.TotalProcessorTime;
+					ExecutionTime = DateTime.Now - StartTime;
 				}
 
-				return new ExecuteResults(LogLines, ExitCode);
+				return new ExecuteResults(LogLines, ExitCode, ExecutionTime, ProcessorTime);
 			}
 			catch (OperationCanceledException)
 			{
diff --git a/Engine/Source/Programs/UnrealBuildTool/Executors/ParallelExecutor.cs b/Engine/Source/Programs/UnrealBuildTool/Executors/ParallelExecutor.cs
--- a/Engine/Source/Programs/UnrealBuildTool/Executors/ParallelExecutor.cs
+++ b/Engine/Source/Programs/UnrealBuildTool/Executors/ParallelExecutor.cs
@@ -29,6 +29,9 @@ namespace UnrealBuildTool
 			public List<string> LogLines = new List<string>();
 			public int ExitCode = -1;
 
+			public TimeSpan ExecutionTime;
+			public TimeSpan ProcessorTime;
+
 			public BuildAction(LinkedAction Inner)
 			{
 				this.Inner = Inner;
@@ -63,6 +66,12 @@ namespace UnrealBuildTool
 		[XmlConfigFile]
 		bool bStopCompilationAfterErrors = false;
 
+		/// <summary>
+		/// Whether to show compilation times along with worst offenders or not.
+		/// </summary>
+		[XmlConfigFile]
+		bool bShowCompilationTimes = false;
+
 		/// <summary>
 		/// How many processes that will be executed in parallel
 		/// </summary>
@@ -167,6 +176,8 @@ namespace UnrealBuildTool
 				}
 			}
 
+			List<BuildAction> AllCompletedActions = new List<BuildAction>();
+
 			// Execute the actions
 			using (LogIndentScope Indent = new LogIndentScope("  "))
 			{
@@ -175,6 +186,7 @@ namespace UnrealBuildTool
 				Dictionary<BuildAction, Thread> ExecutingActions = new Dictionary<BuildAction, Thread>();
 				List<BuildAction> CompletedActions = new List<BuildAction>();
 
+				TimeSpan TotalProcessingTime;
 				using (ManagedProcessGroup ProcessGroup = new ManagedProcessGroup())
 				{
 					using (AutoResetEvent CompletedEvent = new AutoResetEvent(false))
@@ -236,6 +248,8 @@ namespace UnrealBuildTool
 											Log.TraceInformation(Line);
 										}
 
+										AllCompletedActions.Add(CompletedAction);
+
 										// Check the exit code
 										if (CompletedAction.ExitCode == 0)
 										{
@@ -276,6 +290,27 @@ namespace UnrealBuildTool
 							}
 						}
 					}
+
+					TotalProcessingTime = ProcessGroup.TotalProcessorTime;
+				}
+
+				if (bShowCompilationTimes)
+				{
+					Log.TraceInformation("");
+					Log.TraceInformation("Total CPU Time: {0} s", TotalProcessingTime.TotalSeconds);
+					Log.TraceInformation("");
+
+					if (AllCompletedActions.Count > 0)
+					{
+						Log.TraceInformation("Compilation Time Top {0}", Math.Min(20, AllCompletedActions.Count));
+						Log.TraceInformation("");
+						foreach (BuildAction Action in AllCompletedActions.OrderByDescending(x => x.ExecutionTime).Take(20))
+						{
+							string Description = $"{(Action.Inner.CommandDescription != null ? Action.Inner.CommandDescription : Action.Inner.CommandPath.GetFileName())} {Action.Inner.StatusDescription}".Trim();
+							Log.TraceInformation("{0} [ Wall Time {1:0.00} s / CPU Time {2:0.00} s ]", Description, Action.ExecutionTime.TotalSeconds, Action.ProcessorTime.TotalSeconds);
+						}
+						Log.TraceInformation("");
+					}
 				}
 
 				return bResult;
@@ -293,10 +328,13 @@ namespace UnrealBuildTool
 		{
 			try
 			{
+				DateTime StartTime = DateTime.Now;
 				using (ManagedProcess Process = new ManagedProcess(ProcessGroup, Action.Inner.CommandPath.FullName, Action.Inner.CommandArguments, Action.Inner.WorkingDirectory.FullName, null, null, ProcessPriorityClass.BelowNormal))
 				{
 					Action.LogLines.AddRange(Process.ReadAllLines());
 					Action.ExitCode = Process.ExitCode;
+					Action.ExecutionTime = DateTime.Now - StartTime;
+					Action.ProcessorTime = Process.TotalProcessorTime;
 				}
 			}
 			catch (Exception Ex)
