Fix crash when jobs are not available (i.e. MacOS)
Only show processor time when available

#rnx

diff --git a/Engine/Source/Programs/Shared/EpicGames.Core/ManagedProcess.cs b/Engine/Source/Programs/Shared/EpicGames.Core/ManagedProcess.cs
--- a/Engine/Source/Programs/Shared/EpicGames.Core/ManagedProcess.cs
+++ b/Engine/Source/Programs/Shared/EpicGames.Core/ManagedProcess.cs
@@ -155,13 +155,20 @@ namespace EpicGames.Core
 		{
 			get
 			{
-				JOBOBJECT_BASIC_ACCOUNTING_INFORMATION AccountingInformation = new JOBOBJECT_BASIC_ACCOUNTING_INFORMATION();
-				if (QueryInformationJobObject(JobHandle!, JobObjectBasicAccountingInformation, ref AccountingInformation, Marshal.SizeOf(typeof(JOBOBJECT_BASIC_ACCOUNTING_INFORMATION))) == false) 
+				if (SupportsJobObjects)
 				{
-					throw new Win32Exception();
-				}
+					JOBOBJECT_BASIC_ACCOUNTING_INFORMATION AccountingInformation = new JOBOBJECT_BASIC_ACCOUNTING_INFORMATION();
+					if (QueryInformationJobObject(JobHandle!, JobObjectBasicAccountingInformation, ref AccountingInformation, Marshal.SizeOf(typeof(JOBOBJECT_BASIC_ACCOUNTING_INFORMATION))) == false)
+					{
+						throw new Win32Exception();
+					}
 
-				return new TimeSpan((long)AccountingInformation.TotalUserTime + (long)AccountingInformation.TotalKernelTime);
+					return new TimeSpan((long)AccountingInformation.TotalUserTime + (long)AccountingInformation.TotalKernelTime);
+				}
+				else
+				{
+					return new TimeSpan();
+				}
 			}
 		}
 
@@ -981,32 +988,13 @@ namespace EpicGames.Core
 		{
 			get
 			{
-				// Prefer the job accounting over the GetProcessTimes because we need to account for children process spawned by cl-filter too.
-				long AccountingProcessorTime = 0;
-				if (AccountingProcessGroup != null && AccountingProcessGroup.TotalProcessorTime.Ticks != 0) 
+				if (AccountingProcessGroup != null)
 				{
 					return AccountingProcessGroup.TotalProcessorTime;
 				}
-
-				// Normal process accounting is better than nothing if the processgroup one is not available
-				if (FrameworkProcess == null) 
-				{
-					long creation = 0;
-					long exit = 0;
-					long kernel = 0;
-					long user = 0;
-					if (GetProcessTimes(ProcessHandle!, out creation, out exit, out kernel, out user)) 
-					{
-						return new TimeSpan(Math.Max(user + kernel, AccountingProcessorTime));
-					}
-					else 
-					{
-						throw new Win32Exception();
-					}
-				}
-				else 
+				else
 				{
-					return FrameworkProcess.TotalProcessorTime;
+					return new TimeSpan();
 				}
 			}
 		}
diff --git a/Engine/Source/Programs/UnrealBuildTool/Executors/Experimental/TaskExecutor.cs b/Engine/Source/Programs/UnrealBuildTool/Executors/Experimental/TaskExecutor.cs
--- a/Engine/Source/Programs/UnrealBuildTool/Executors/Experimental/TaskExecutor.cs
+++ b/Engine/Source/Programs/UnrealBuildTool/Executors/Experimental/TaskExecutor.cs
@@ -162,8 +162,11 @@ namespace UnrealBuildTool
 			if (bShowCompilationTimes)
 			{
 				Log.TraceInformation("");
-				Log.TraceInformation("Total CPU Time: {0} s", ProcessGroup.TotalProcessorTime.TotalSeconds);
-				Log.TraceInformation("");
+				if (ProcessGroup.TotalProcessorTime.Ticks > 0)
+				{
+					Log.TraceInformation("Total CPU Time: {0} s", ProcessGroup.TotalProcessorTime.TotalSeconds);
+					Log.TraceInformation("");
+				}
 
 				if (Tasks.Count > 0)
 				{
@@ -172,7 +175,15 @@ namespace UnrealBuildTool
 					foreach (var Pair in Tasks.OrderByDescending(x => x.Value.Result.ExecutionTime).Take(20))
 					{
 						string Description = $"{(Pair.Key.Inner.CommandDescription != null ? Pair.Key.Inner.CommandDescription : Pair.Key.Inner.CommandPath.GetFileNameWithoutExtension())} {Pair.Key.Inner.StatusDescription}".Trim();
-						Log.TraceInformation("{0} [ Wall Time {1:0.00} s / CPU Time {2:0.00} s ]", Description, Pair.Value.Result.ExecutionTime.TotalSeconds, Pair.Value.Result.ProcessorTime.TotalSeconds);
+						if (Pair.Value.Result.ProcessorTime.Ticks > 0)
+						{
+							Log.TraceInformation("{0} [ Wall Time {1:0.00} s / CPU Time {2:0.00} s ]", Description, Pair.Value.Result.ExecutionTime.TotalSeconds, Pair.Value.Result.ProcessorTime.TotalSeconds);
+						}
+						else
+						{
+							Log.TraceInformation("{0} [ Time {1:0.00} s ]", Description, Pair.Value.Result.ExecutionTime.TotalSeconds);
+						}
+						
 					}
 					Log.TraceInformation("");
 				}
diff --git a/Engine/Source/Programs/UnrealBuildTool/Executors/ParallelExecutor.cs b/Engine/Source/Programs/UnrealBuildTool/Executors/ParallelExecutor.cs
--- a/Engine/Source/Programs/UnrealBuildTool/Executors/ParallelExecutor.cs
+++ b/Engine/Source/Programs/UnrealBuildTool/Executors/ParallelExecutor.cs
@@ -297,8 +297,11 @@ namespace UnrealBuildTool
 				if (bShowCompilationTimes)
 				{
 					Log.TraceInformation("");
-					Log.TraceInformation("Total CPU Time: {0} s", TotalProcessingTime.TotalSeconds);
-					Log.TraceInformation("");
+					if (TotalProcessingTime.Ticks > 0)
+					{
+						Log.TraceInformation("Total CPU Time: {0} s", TotalProcessingTime.TotalSeconds);
+						Log.TraceInformation("");
+					}
 
 					if (AllCompletedActions.Count > 0)
 					{
@@ -307,7 +310,14 @@ namespace UnrealBuildTool
 						foreach (BuildAction Action in AllCompletedActions.OrderByDescending(x => x.ExecutionTime).Take(20))
 						{
 							string Description = $"{(Action.Inner.CommandDescription != null ? Action.Inner.CommandDescription : Action.Inner.CommandPath.GetFileName())} {Action.Inner.StatusDescription}".Trim();
-							Log.TraceInformation("{0} [ Wall Time {1:0.00} s / CPU Time {2:0.00} s ]", Description, Action.ExecutionTime.TotalSeconds, Action.ProcessorTime.TotalSeconds);
+							if (Action.ProcessorTime.Ticks > 0)
+							{
+								Log.TraceInformation("{0} [ Wall Time {1:0.00} s / CPU Time {2:0.00} s ]", Description, Action.ExecutionTime.TotalSeconds, Action.ProcessorTime.TotalSeconds);
+							}
+							else
+							{
+								Log.TraceInformation("{0} [ Time {1:0.00} s ]", Description, Action.ExecutionTime.TotalSeconds);
+							}
 						}
 						Log.TraceInformation("");
 					}
