P4 Blame tool for P4V (requires input of the line number)

diff --git a/Engine/Extras/P4VUtils/Commands/P4BlameCommand.cs b/Engine/Extras/P4VUtils/Commands/P4BlameCommand.cs
new file mode 100644
--- /dev/null
+++ b/Engine/Extras/P4VUtils/Commands/P4BlameCommand.cs
@@ -0,0 +1,67 @@
+// Copyright Epic Games, Inc. All Rights Reserved.
+
+using EpicGames.Perforce;
+using Microsoft.Extensions.Logging;
+using System.Collections.Generic;
+using System.Text.RegularExpressions;
+using System.Threading;
+using System.Threading.Tasks;
+
+namespace P4VUtils.Commands
+{
+	class P4BlameCommand : Command
+	{
+		public override string Description => "Prints out who last edited a specific line in a file";
+
+		public override CustomToolInfo CustomTool => new CustomToolInfo("P4 Blame this file using line number", "%d") 
+		{ 
+			ShowConsole = true,
+			PromptForArgument = true,
+			PromptText = "Please enter the line number:"
+		};
+
+		public override async Task<int> Execute(string[] Args, IReadOnlyDictionary<string, string> ConfigValues, ILogger Logger)
+		{
+			if (Args.Length < 3)
+			{
+				Logger.LogError("Missing file name and/or line number");
+				return 1;
+			}
+
+			string FileName = Args[1];
+			int LineNumber = 0;
+			if (!int.TryParse(Args[2], out LineNumber) || LineNumber == 0)
+			{
+				Logger.LogError("Please enter a valid line number");
+				return 1;
+			}
+
+			Logger.LogInformation("Blaming line {0} in {1}", LineNumber, FileName);
+
+			PerforceConnection Perforce = new PerforceConnection(null, null, Logger);
+
+			try
+			{
+				List<AnnotateRecord> Results = await Perforce.AnnotateAsync(FileName, AnnotateOptions.FollowIntegrations | AnnotateOptions.IgnoreWhiteSpaceChanges | AnnotateOptions.OutputUserAndDate, CancellationToken.None);
+
+				if (LineNumber >= Results.Count)
+				{
+					Logger.LogError("Please enter a valid line number (this file only has {0} lines in P4)", Results.Count);
+					return 1;
+				}
+
+				// the first object in the list is a header, so the first line is at index 1
+				AnnotateRecord LineResult = Results[LineNumber];
+				string EscapedString = Regex.Escape(LineResult.Data);
+
+				Logger.LogInformation(@"Change {1} by {2} on [{3}]:'{4}'", LineResult.LowerCL, LineResult.UserName, LineResult.Time, EscapedString);
+			}
+			catch (PerforceException Exception)
+			{
+				Logger.LogError(Exception.Message);
+			}
+
+			return 0;
+		}
+	}
+}
diff --git a/Engine/Extras/P4VUtils/Program.cs b/Engine/Extras/P4VUtils/Program.cs
--- a/Engine/Extras/P4VUtils/Program.cs
+++ b/Engine/Extras/P4VUtils/Program.cs
@@ -23,6 +23,8 @@ namespace P4VUtils
 		public bool ShowConsole { get; set; }
 		public bool RefreshUI { get; set; } = true;
 		public string Shortcut { get; set; } = "";
+		public bool PromptForArgument { get; set; } = false;
+		public string PromptText { get; set; } = "";
 
 		public CustomToolInfo(string Name, string Arguments)
 		{
@@ -47,6 +49,7 @@ namespace P4VUtils
 		{
 			["describe"] = new DescribeCommand(),
 			["findlastedit"] = new FindLastEditCommand(),
+			["p4blame"] = new P4BlameCommand(),
 			["snapshot"] = new SnapshotCommand(),
 			["copyclnum"] = new CopyCLCommand(),
 		};
@@ -302,6 +305,17 @@ namespace P4VUtils
 							ToolDef.AppendChild(Refresh);
 						}
 
+						if (CustomTool.PromptForArgument)
+						{
+							XmlElement Prompt = Document.CreateElement("Prompt");
+							{
+								XmlElement PromptText = Document.CreateElement("PromptText");
+								PromptText.InnerText = CustomTool.PromptText.Length > 0 ? CustomTool.PromptText : "Argument";
+								Prompt.AppendChild(PromptText);
+							}
+							ToolDef.AppendChild(Prompt);
+						}
+
 						XmlElement AddToContext = Document.CreateElement("AddToContext");
 						AddToContext.InnerText = CustomTool.AddToContextMenu ? "true" : "false";
 						ToolDef.AppendChild(AddToContext);
diff --git a/Engine/Source/Programs/Shared/EpicGames.Perforce/PerforceConnection.cs b/Engine/Source/Programs/Shared/EpicGames.Perforce/PerforceConnection.cs
--- a/Engine/Source/Programs/Shared/EpicGames.Perforce/PerforceConnection.cs
+++ b/Engine/Source/Programs/Shared/EpicGames.Perforce/PerforceConnection.cs
@@ -3010,5 +3010,55 @@ namespace EpicGames.Perforce
 			return CommandAsync<UndoRecord>(Arguments.ToString(), null, CancellationToken);
 		}
 		#endregion
+
+		#region p4 annotate
+
+		/// <summary>
+		/// Runs the annotate p4 command
+		/// </summary>
+		/// <param name="FileSpec">Depot path to the file</param>
+		/// <param name="Options">Options for the anotate command</param>
+		/// <param name="CancellationToken">Token used to cancel the operation</param>
+		/// <returns>List of annotate records</returns>
+		public async Task<List<AnnotateRecord>> AnnotateAsync(string FileSpec, AnnotateOptions Options, CancellationToken CancellationToken)
+		{
+			return (await TryAnnotateAsync(FileSpec, Options, CancellationToken)).Data;
+		}
+
+		/// <summary>
+		/// Runs the annotate p4 command 
+		/// </summary>
+		/// <param name="FileSpec">Depot path to the file</param>
+		/// <param name="Options">Options for the anotate command</param>
+		/// <param name="CancellationToken">Token used to cancel the operation</param>
+		/// <returns>List of annotate records</returns>
+		public Task<PerforceResponseList<AnnotateRecord>> TryAnnotateAsync(string FileSpec, AnnotateOptions Options, CancellationToken CancellationToken)
+		{
+			PerforceResponseList<AnnotateRecord> Responses = new PerforceResponseList<AnnotateRecord>();
+			StringBuilder Arguments = new StringBuilder("annotate ");
+			if ((Options & AnnotateOptions.IncludeDeletedFilesAndLines) != 0)
+			{
+				Arguments.Append(" -a");
+			}
+			if ((Options & AnnotateOptions.IgnoreWhiteSpaceChanges) != 0)
+			{
+				Arguments.Append(" -db");
+			}
+			if ((Options & AnnotateOptions.OutputUserAndDate) != 0)
+			{
+				Arguments.Append(" -u");
+			}
+			if ((Options & AnnotateOptions.FollowIntegrations) != 0)
+			{
+				Arguments.Append(" -I");
+			}
+
+			Arguments.AppendFormat(" {0}", FileSpec);
+
+			return CommandAsync<AnnotateRecord>(Arguments.ToString(), null, CancellationToken);
+		}
+
+		#endregion
+
 	}
 }
diff --git a/Engine/Source/Programs/Shared/EpicGames.Perforce/Records/AnnotateLineRecord.cs b/Engine/Source/Programs/Shared/EpicGames.Perforce/Records/AnnotateLineRecord.cs
new file mode 100644
--- /dev/null
+++ b/Engine/Source/Programs/Shared/EpicGames.Perforce/Records/AnnotateLineRecord.cs
@@ -0,0 +1,47 @@
+// Copyright Epic Games, Inc. All Rights Reserved.
+
+using System;
+
+namespace EpicGames.Perforce
+{
+	public class AnnotateLineRecord
+	{
+		[PerforceTag("depotFile")]
+		public readonly string DepotFile;
+
+		/// <summary>
+		/// The upper changelist 
+		/// </summary>
+		[PerforceTag("upper")]
+		public readonly int UpperRev;
+
+		/// <summary>
+		/// The lower changelist 
+		/// </summary>
+		[PerforceTag("lower")]
+		public readonly int LowerRev;
+
+		/// <summary>
+		/// Timestamp of this modification
+		/// </summary>
+		[PerforceTag("time", Optional = true)]
+		public readonly DateTime Time;
+
+		/// <summary>
+		/// Private constructor for serialization
+		/// </summary>
+		private AnnotateLineRecord()
+		{
+			DepotFile = null!;
+		}
+
+		/// <summary>
+		/// Format this record for display in the debugger
+		/// </summary>
+		/// <returns>Summary of this revision</returns>
+		public override string ToString()
+		{
+			return String.Format("FileLoc {0}, from rev {1} to {2} on {3}", DepotFile, LowerRev, UpperRev, Time);
+		}
+	}
+}
diff --git a/Engine/Source/Programs/Shared/EpicGames.Perforce/Records/AnnotateOptions.cs b/Engine/Source/Programs/Shared/EpicGames.Perforce/Records/AnnotateOptions.cs
new file mode 100644
--- /dev/null
+++ b/Engine/Source/Programs/Shared/EpicGames.Perforce/Records/AnnotateOptions.cs
@@ -0,0 +1,43 @@
+// Copyright Epic Games, Inc. All Rights Reserved.
+
+using System;
+
+namespace EpicGames.Perforce
+{
+	[Flags]
+	public enum AnnotateOptions
+	{
+		/// <summary>
+		/// No options specified
+		/// </summary>
+		None = 0,
+
+		/// <summary>
+		/// includes both deleted files and lines no longer present
+		/// at the head revision. In the latter case, both the starting and ending
+		///	revision for each line is displayed.
+		/// </summary>
+		IncludeDeletedFilesAndLines = 1,
+
+		/// <summary>
+		/// ignore whitespace changes
+		/// </summary>
+		IgnoreWhiteSpaceChanges = 2,
+
+		/// <summary>
+		/// follows all integrations into the file.  If a line was
+		/// introduced into the file by a merge, the source of the merge is
+		/// displayed as the changelist that introduced the line. If the source
+		/// itself was the result of an integration, that source is used instead,
+		/// and so on.
+		/// </summary>
+		FollowIntegrations = 4,
+
+		/// <summary>
+		/// directs the annotate command to output the user who
+		/// modified the line in the file and the date it was modified.
+		/// </summary>
+		OutputUserAndDate = 8,
+
+	}
+}
diff --git a/Engine/Source/Programs/Shared/EpicGames.Perforce/Records/AnnotateRecord.cs b/Engine/Source/Programs/Shared/EpicGames.Perforce/Records/AnnotateRecord.cs
new file mode 100644
--- /dev/null
+++ b/Engine/Source/Programs/Shared/EpicGames.Perforce/Records/AnnotateRecord.cs
@@ -0,0 +1,96 @@
+// Copyright Epic Games, Inc. All Rights Reserved.
+
+using System;
+using System.Collections.Generic;
+
+namespace EpicGames.Perforce
+{
+	public class AnnotateRecord
+	{
+		/// <summary>
+		/// The revision number of this file
+		/// </summary>
+		[PerforceTag("rev", Optional = true)]
+		public readonly int RevisionNumber;
+
+		/// <summary>
+		/// The changelist responsible for this revision of the file
+		/// </summary>
+		[PerforceTag("change", Optional = true)]
+		public readonly int ChangeNumber;
+
+		/// <summary>
+		/// Type of the file
+		/// </summary>
+		[PerforceTag("type", Optional = true)]
+		public readonly string Type;
+
+		/// <summary>
+		/// The upper changelist 
+		/// </summary>
+		[PerforceTag("upper", Optional = true)]
+		public readonly string UpperCL;
+
+		/// <summary>
+		/// The lower changelist 
+		/// </summary>
+		[PerforceTag("lower", Optional = true)]
+		public readonly string LowerCL;
+
+		/// <summary>
+		/// Author of the changelist
+		/// </summary>
+		[PerforceTag("user", Optional = true)]
+		public readonly string UserName;
+
+		/// <summary>
+		/// Timestamp of this modification
+		/// </summary>
+		[PerforceTag("time", Optional = true)]
+		public readonly DateTime Time;
+
+		/// <summary>
+		/// Client that submitted this changelist
+		/// </summary>
+		[PerforceTag("client", Optional = true)]
+		public readonly string ClientName;
+
+		/// <summary>
+		/// The actual line
+		/// </summary>
+		[PerforceTag("data", Optional = true)]
+		public readonly string Data;
+
+		[PerforceRecordList]
+		public List<AnnotateLineRecord> LineRecord = new List<AnnotateLineRecord>();
+
+		/// <summary>
+		/// Private constructor for serialization
+		/// </summary>
+		private AnnotateRecord()
+		{
+			Type = null!;
+			UserName = null!;
+			ClientName = null!;
+			Data = null!;
+			UpperCL = null!;
+			LowerCL = null!;
+		}
+
+		/// <summary>
+		/// Format this record for display in the debugger
+		/// </summary>
+		/// <returns>Summary of this revision</returns>
+		public override string ToString()
+		{
+			if (Data == null) // this is the header
+			{
+				return String.Format("Data for revision {0}, CL {1}", RevisionNumber, ChangeNumber);
+			}
+			else
+			{
+				return String.Format("CL {0} by {1}, on {2}, line = {3}", LowerCL, UserName, Time, Data);
+			}
+		}
+	}
+}
