Fix AsBinary() returning the incorrect payload buffer for binary fields. Was causing Horde agents to deserialize incorrect hash values.

diff --git a/Engine/Source/Programs/Shared/EpicGames.Serialization/CompactBinary.cs b/Engine/Source/Programs/Shared/EpicGames.Serialization/CompactBinary.cs
--- a/Engine/Source/Programs/Shared/EpicGames.Serialization/CompactBinary.cs
+++ b/Engine/Source/Programs/Shared/EpicGames.Serialization/CompactBinary.cs
@@ -712,7 +712,7 @@ namespace EpicGames.Serialization
 				Error = CbFieldError.None;
 
 				ulong Length = VarInt.Read(Payload.Span, out int BytesRead);
-				return Memory.Slice(BytesRead, (int)Length);
+				return Payload.Slice(BytesRead, (int)Length);
 			}
 			else
 			{
