UnrealBuildTool: Use process directly to get wall clock execution time

#rnx

diff --git a/Engine/Source/Programs/Shared/EpicGames.Core/ManagedProcess.cs b/Engine/Source/Programs/Shared/EpicGames.Core/ManagedProcess.cs
--- a/Engine/Source/Programs/Shared/EpicGames.Core/ManagedProcess.cs
+++ b/Engine/Source/Programs/Shared/EpicGames.Core/ManagedProcess.cs
@@ -301,8 +301,33 @@ namespace EpicGames.Core
 		[DllImport("kernel32.dll", SetLastError=true)]
 		static extern UInt32 WaitForSingleObject(SafeHandleZeroOrMinusOneIsInvalid hHandle, UInt32 dwMilliseconds);
 
-		[DllImport("kernel32.dll", CharSet = CharSet.Auto, SetLastError = true)]
-		public static extern bool GetProcessTimes(SafeHandleZeroOrMinusOneIsInvalid handle, out long creation, out long exit, out long kernel, out long user);
+		[DllImport("kernel32.dll", SetLastError = true)]
+		[return: MarshalAs(UnmanagedType.Bool)]
+		static extern bool GetProcessTimes(SafeHandleZeroOrMinusOneIsInvalid hProcess,
+			out System.Runtime.InteropServices.ComTypes.FILETIME lpCreationTime,
+			out System.Runtime.InteropServices.ComTypes.FILETIME lpExitTime,
+			out System.Runtime.InteropServices.ComTypes.FILETIME lpKernelTime,
+			out System.Runtime.InteropServices.ComTypes.FILETIME lpUserTime);
+
+		/// <summary>
+		/// Converts FILETIME to DateTime.
+		/// </summary>
+		/// <param name="Time">Input FILETIME structure</param>
+		/// <returns>Converted DateTime</returns>
+		static DateTime FileTimeToDateTime(System.Runtime.InteropServices.ComTypes.FILETIME Time)
+		{
+			ulong High = (ulong)Time.dwHighDateTime;
+			uint Low = (uint)Time.dwLowDateTime;
+			long FileTime = (long)((High << 32) + Low);
+			try
+			{
+				return DateTime.FromFileTimeUtc(FileTime);
+			}
+			catch (ArgumentOutOfRangeException)
+			{
+				return DateTime.MinValue;
+			}
+		}
 
 		const int ERROR_ACCESS_DENIED = 5;
 
@@ -1069,5 +1094,50 @@ namespace EpicGames.Core
 				}
 			}
 		}
+
+		/// <summary>
+		/// The creation time of the process.
+		/// </summary>
+		public DateTime StartTime
+		{
+			get
+			{
+				if (FrameworkProcess == null)
+				{
+					System.Runtime.InteropServices.ComTypes.FILETIME CreationTime;
+					if (!GetProcessTimes(ProcessHandle!, out CreationTime, out _, out _, out _))
+					{
+						throw new Win32Exception();
+					}
+					return FileTimeToDateTime(CreationTime);
+				}
+				else
+				{
+					return FrameworkProcess.StartTime;
+				}
+			}
+		}
+		/// <summary>
+		/// The exit time of the process. Throws an exception if the process has not terminated.
+		/// </summary>
+		public DateTime ExitTime
+		{
+			get
+			{
+				if (FrameworkProcess == null)
+				{
+					System.Runtime.InteropServices.ComTypes.FILETIME ExitTime;
+					if (!GetProcessTimes(ProcessHandle!, out _, out ExitTime, out _, out _))
+					{
+						throw new Win32Exception();
+					}
+					return FileTimeToDateTime(ExitTime);
+				}
+				else
+				{
+					return FrameworkProcess.ExitTime;
+				}
+			}
+		}
 	}
 }
diff --git a/Engine/Source/Programs/UnrealBuildTool/Executors/Experimental/TaskExecutor.cs b/Engine/Source/Programs/UnrealBuildTool/Executors/Experimental/TaskExecutor.cs
--- a/Engine/Source/Programs/UnrealBuildTool/Executors/Experimental/TaskExecutor.cs
+++ b/Engine/Source/Programs/UnrealBuildTool/Executors/Experimental/TaskExecutor.cs
@@ -217,7 +217,6 @@ namespace UnrealBuildTool
 				string Description = string.Empty;
 				List<string> LogLines;
 				int ExitCode;
-				DateTime StartTime = DateTime.Now;
 				TimeSpan ExecutionTime;
 				TimeSpan ProcessorTime;
 				using (ManagedProcess Process = new ManagedProcess(ProcessGroup, Action.CommandPath.FullName, Action.CommandArguments, Action.WorkingDirectory.FullName, null, null, ProcessPriorityClass.BelowNormal))
@@ -228,7 +227,7 @@ namespace UnrealBuildTool
 					LogLines = Console.OutputEncoding.GetString(StdOutStream.GetBuffer(), 0, Convert.ToInt32(StdOutStream.Length)).Split(LineEndingSplit, StringSplitOptions.RemoveEmptyEntries).ToList();
 					ExitCode = Process.ExitCode;
 					ProcessorTime = Process.TotalProcessorTime;
-					ExecutionTime = DateTime.Now - StartTime;
+					ExecutionTime = Process.ExitTime - Process.StartTime;
 				}
 
 				return new ExecuteResults(LogLines, ExitCode, ExecutionTime, ProcessorTime);
diff --git a/Engine/Source/Programs/UnrealBuildTool/Executors/ParallelExecutor.cs b/Engine/Source/Programs/UnrealBuildTool/Executors/ParallelExecutor.cs
--- a/Engine/Source/Programs/UnrealBuildTool/Executors/ParallelExecutor.cs
+++ b/Engine/Source/Programs/UnrealBuildTool/Executors/ParallelExecutor.cs
@@ -338,13 +338,12 @@ namespace UnrealBuildTool
 		{
 			try
 			{
-				DateTime StartTime = DateTime.Now;
 				using (ManagedProcess Process = new ManagedProcess(ProcessGroup, Action.Inner.CommandPath.FullName, Action.Inner.CommandArguments, Action.Inner.WorkingDirectory.FullName, null, null, ProcessPriorityClass.BelowNormal))
 				{
 					Action.LogLines.AddRange(Process.ReadAllLines());
 					Action.ExitCode = Process.ExitCode;
-					Action.ExecutionTime = DateTime.Now - StartTime;
 					Action.ProcessorTime = Process.TotalProcessorTime;
+					Action.ExecutionTime = Process.ExitTime - Process.StartTime;
 				}
 			}
 			catch (Exception Ex)
