Horde: Output additional debug info when StatsV2 file is not found.

diff --git a/Engine/Source/Programs/Shared/EpicGames.Perforce.Managed/ManagedWorkspace.cs b/Engine/Source/Programs/Shared/EpicGames.Perforce.Managed/ManagedWorkspace.cs
--- a/Engine/Source/Programs/Shared/EpicGames.Perforce.Managed/ManagedWorkspace.cs
+++ b/Engine/Source/Programs/Shared/EpicGames.Perforce.Managed/ManagedWorkspace.cs
@@ -1847,6 +1847,27 @@ namespace EpicGames.Perforce.Managed
 
 		static readonly Utf8String StatsFileName = "StatsV2.json";
 
+		static async Task LogPerforceCommandAsync(string ArgumentList, ILogger Logger)
+		{
+			Logger.LogInformation("Running command: {ArgList}", ArgumentList);
+			using (ManagedProcessGroup ChildProcessGroup = new ManagedProcessGroup())
+			{
+				using (ManagedProcess ChildProcess = new ManagedProcess(ChildProcessGroup, "p4.exe", ArgumentList, null, null, null, ProcessPriorityClass.Normal))
+				{
+					for (; ; )
+					{
+						string? Line = await ChildProcess.ReadLineAsync();
+						if (Line == null)
+						{
+							break;
+						}
+						Logger.LogInformation("{0}", Line);
+					}
+					ChildProcess.WaitForExit();
+				}
+			}
+		}
+
 		/// <summary>
 		/// Syncs a batch of files
 		/// </summary>
@@ -1919,6 +1940,13 @@ namespace EpicGames.Perforce.Managed
 							{
 								Logger.LogInformation("SyncList: {Line}", Line);
 							}
+
+							await LogPerforceCommandAsync($"protects -M -u {ClientWithFileList.UserName} {LocalFile}", Logger);
+							await LogPerforceCommandAsync($"protects -M -u {ClientWithFileList.UserName} {LocalFile.FullName.ToLower()}", Logger);
+							await LogPerforceCommandAsync($"protects -M -u {ClientWithFileList.UserName} {LocalFile.FullName.Replace('\\', '/')}", Logger);
+
+							ClientWithFileList.GlobalOptions.Add($"-Zdebug=dm=2");
+
 							for (int Idx = 0; ; Idx++)
 							{
 								Logger.LogInformation("Loop {Idx}: Check {File} - {State}", Idx, LocalFile, FileReference.Exists(LocalFile) ? "exists" : "does not exist");
