Horde: Reduce peak memory usage when building workspace snapshots. Was previously keeping references to Memory<byte> ranges within a much larger buffer, preventing them from being GC'd.

diff --git a/Engine/Source/Programs/Shared/EpicGames.Core/Utf8String.cs b/Engine/Source/Programs/Shared/EpicGames.Core/Utf8String.cs
--- a/Engine/Source/Programs/Shared/EpicGames.Core/Utf8String.cs
+++ b/Engine/Source/Programs/Shared/EpicGames.Core/Utf8String.cs
@@ -86,6 +86,17 @@ namespace EpicGames.Core
 			this.Memory = new ReadOnlyMemory<byte>(Buffer, Offset, Length);
 		}
 
+		/// <summary>
+		/// Duplicate this string
+		/// </summary>
+		/// <returns></returns>
+		public Utf8String Clone()
+		{
+			byte[] NewBuffer = new byte[Memory.Length];
+			Memory.CopyTo(NewBuffer);
+			return new Utf8String(NewBuffer);
+		}
+
 		/// <summary>
 		/// Tests two strings for equality
 		/// </summary>
diff --git a/Engine/Source/Programs/Shared/EpicGames.Perforce.Managed/ManagedWorkspace.cs b/Engine/Source/Programs/Shared/EpicGames.Perforce.Managed/ManagedWorkspace.cs
--- a/Engine/Source/Programs/Shared/EpicGames.Perforce.Managed/ManagedWorkspace.cs
+++ b/Engine/Source/Programs/Shared/EpicGames.Perforce.Managed/ManagedWorkspace.cs
@@ -1429,8 +1429,11 @@ namespace EpicGames.Perforce.Managed
 						throw new InvalidDataException($"Client path returned by Perforce ('{Record.ClientFile}') does not begin with client name ('{ClientPrefix}')");
 					}
 
+					// Duplicate the client path. If we reference into the raw record, we'll prevent all the raw P4 output from being garbage collected.
+					Utf8String ClientFile = Record.ClientFile.Clone();
+
 					// Get the client path after the initial client prefix
-					ReadOnlySpan<byte> PathSpan = Record.ClientFile.Span;
+					ReadOnlySpan<byte> PathSpan = ClientFile.Span;
 
 					// Parse out the data
 					StreamSnapshotBuilder LastStreamDirectory = Builder;
@@ -1452,7 +1455,7 @@ namespace EpicGames.Perforce.Managed
 						}
 
 						// Get the fragment text
-						Utf8String Fragment = new Utf8String(Record.ClientFile.Memory.Slice(FragmentMinIdx, FragmentMaxIdx - FragmentMinIdx));
+						Utf8String Fragment = new Utf8String(ClientFile.Memory.Slice(FragmentMinIdx, FragmentMaxIdx - FragmentMinIdx));
 
 						// If this fragment matches the same fragment from the previous iteration, take the last stream directory straight away
 						if (FragmentIdx < Fragments.Count)
@@ -1488,12 +1491,12 @@ namespace EpicGames.Perforce.Managed
 					}
 
 					Md5Hash Digest = Md5Hash.Parse(Record.Digest);
-					FileContentId ContentId = new FileContentId(Digest, Record.HeadType);
+					FileContentId ContentId = new FileContentId(Digest, Record.HeadType.Clone());
 					int Revision = (int)Utf8String.ParseUnsignedInt(Record.HaveRev);
 
 					// Add a new StreamFileInfo to the last directory object
-					Utf8String FileName = PerforceUtils.UnescapePath(Record.ClientFile.Slice(FragmentMinIdx));
-					LastStreamDirectory.NameToFile.Add(FileName, new StreamFile(Record.DepotFile, Record.FileSize, ContentId, Revision));
+					Utf8String FileName = PerforceUtils.UnescapePath(ClientFile.Slice(FragmentMinIdx));
+					LastStreamDirectory.NameToFile.Add(FileName, new StreamFile(Record.DepotFile.Clone(), Record.FileSize, ContentId, Revision));
 				};
 
 				// Create the workspace, and add records for all the files. Exclude deleted files with digest = null.
