Add Slice() extension method for List<T>, allowing use of ranges.

diff --git a/Engine/Source/Programs/Shared/EpicGames.Core/ListSegment.cs b/Engine/Source/Programs/Shared/EpicGames.Core/ListSegment.cs
new file mode 100644
--- /dev/null
+++ b/Engine/Source/Programs/Shared/EpicGames.Core/ListSegment.cs
@@ -0,0 +1,183 @@
+// Copyright Epic Games, Inc. All Rights Reserved.
+
+using System;
+using System.Collections;
+using System.Collections.Generic;
+using System.Text;
+
+namespace EpicGames.Core
+{
+	/// <summary>
+	/// Delimits a section of a list
+	/// </summary>
+	/// <typeparam name="T">The list element type</typeparam>
+	public readonly struct ListSegment<T> : IEnumerable<T>, IEnumerable, IReadOnlyCollection<T>, IReadOnlyList<T>
+	{
+		/// <summary>
+		/// Gets the original array containing the range of elements that the list segment
+		/// delimits.
+		/// </summary>
+		public IList<T> List { get; }
+
+		/// <summary>
+		/// Gets the position of the first element in the range delimited by the list segment,
+		/// relative to the start of the original list.
+		/// </summary>
+		public int Offset { get; }
+
+		/// <summary>
+		/// Gets the number of elements in the range delimited by the array segment.
+		/// </summary>
+		public int Count { get; }
+
+		/// <summary>
+		/// An empty list segment of the given element type
+		/// </summary>
+		public static ListSegment<T> Empty { get; } = new ListSegment<T>(null!, 0, 0);
+
+		/// <inheritdoc/>
+		public bool IsReadOnly => List.IsReadOnly;
+
+		/// <summary>
+		/// Initializes a new instance of the System.ArraySegment`1 structure that delimits
+		/// all the elements in the specified array.
+		/// </summary>
+		public ListSegment(IList<T> List)
+			: this(List, 0, List.Count)
+		{
+		}
+
+		/// <summary>
+		/// 
+		/// </summary>
+		/// <param name="List"></param>
+		/// <param name="Offset"></param>
+		/// <param name="Count"></param>
+		public ListSegment(IList<T> List, int Offset, int Count)
+		{
+			if (Offset < 0 || Offset > List.Count)
+			{
+				throw new ArgumentOutOfRangeException(nameof(Offset));
+			}
+			if (Count < 0 || Offset + Count > List.Count)
+			{
+				throw new ArgumentOutOfRangeException(nameof(Count));
+			}
+
+			this.List = List;
+			this.Offset = Offset;
+			this.Count = Count;
+		}
+
+		/// <summary>
+		/// 
+		/// </summary>
+		/// <param name="index"></param>
+		/// <returns></returns>
+		/// <exception cref="ArgumentOutOfRangeException">Index is not a valid index in the ListSegment</exception>
+		public T this[int Index]
+		{
+			get
+			{
+				if (Index < 0 || Index >= Count)
+				{
+					throw new ArgumentOutOfRangeException(nameof(Index));
+				}
+				return List[Index + Offset];
+			}
+			set
+			{
+				if (Index < 0 || Index >= Count)
+				{
+					throw new ArgumentOutOfRangeException(nameof(Index));
+				}
+				List[Index + Offset] = value;
+			}
+		}
+
+		/// <summary>
+		/// 
+		/// </summary>
+		/// <param name="List"></param>
+		public static implicit operator ListSegment<T>(List<T> List)
+		{
+			return new ListSegment<T>(List);
+		}
+
+		/// <summary>
+		/// Enumerator implementation
+		/// </summary>
+		struct Enumerator : IEnumerator<T>, IEnumerator
+		{
+			IList<T> List;
+			int Index;
+			int MaxIndex;
+
+			/// <inheritdoc/>
+			public T Current
+			{
+				get
+				{
+					if (Index >= MaxIndex)
+					{
+						throw new IndexOutOfRangeException();
+					}
+					return List[Index];
+				}
+			}
+
+			/// <summary>
+			/// Constructor
+			/// </summary>
+			/// <param name="Segment"></param>
+			public Enumerator(ListSegment<T> Segment)
+			{
+				this.List = Segment.List;
+				this.Index = Segment.Offset;
+				this.MaxIndex = Segment.Offset + Segment.Count;
+			}
+
+			/// <inheritdoc/>
+			void IDisposable.Dispose()
+			{
+			}
+
+			/// <inheritdoc/>
+			public bool MoveNext()
+			{
+				return ++Index < MaxIndex;
+			}
+
+			/// <inheritdoc/>
+			object? IEnumerator.Current => Current;
+
+			/// <inheritdoc/>
+			void IEnumerator.Reset() => throw new InvalidOperationException();
+		}
+
+		/// <inheritdoc/>
+		public IEnumerator<T> GetEnumerator() => new Enumerator(this);
+
+		/// <inheritdoc/>
+		IEnumerator IEnumerable.GetEnumerator() => GetEnumerator();
+	}
+
+	/// <summary>
+	/// Extension method to slice a list
+	/// </summary>
+	public static class ListSegmentExtensions
+	{
+		/// <summary>
+		/// Create a list segment
+		/// </summary>
+		/// <typeparam name="T"></typeparam>
+		/// <param name="List"></param>
+		/// <param name="Offset"></param>
+		/// <param name="Count"></param>
+		/// <returns></returns>
+		public static ListSegment<T> Slice<T>(this IList<T> List, int Offset, int Count)
+		{
+			return new ListSegment<T>(List, Offset, Count);
+		}
+	}
+}
