Horde: Separate implementation of stream snapshots from tree objects.

diff --git a/Engine/Source/Programs/Shared/EpicGames.Perforce.Managed/ManagedWorkspace.cs b/Engine/Source/Programs/Shared/EpicGames.Perforce.Managed/ManagedWorkspace.cs
--- a/Engine/Source/Programs/Shared/EpicGames.Perforce.Managed/ManagedWorkspace.cs
+++ b/Engine/Source/Programs/Shared/EpicGames.Perforce.Managed/ManagedWorkspace.cs
@@ -1394,7 +1394,7 @@ namespace EpicGames.Perforce.Managed
 		/// <param name="CancellationToken">Cancellation token</param>
 		private async Task<StreamSnapshotFromMemory> FindClientContentsAsync(PerforceClientConnection PerforceClient, int ChangeNumber, bool bFakeSync, CancellationToken CancellationToken)
 		{
-			StreamSnapshotBuilder Builder = new StreamSnapshotBuilder();
+			StreamTreeBuilder Builder = new StreamTreeBuilder();
 
 			using (Trace("FetchMetadata"))
 			using (ILoggerProgress Scope = Logger.BeginProgressScope("Fetching metadata..."))
@@ -1406,7 +1406,7 @@ namespace EpicGames.Perforce.Managed
 
 				// List of the last path fragments. Since file records that are returned are typically sorted by their position in the tree, we can save quite a lot of processing by
 				// reusing as many fragemnts as possible.
-				List<(Utf8String, StreamSnapshotBuilder)> Fragments = new List<(Utf8String, StreamSnapshotBuilder)>();
+				List<(Utf8String, StreamTreeBuilder)> Fragments = new List<(Utf8String, StreamTreeBuilder)>();
 
 				// Handler for each returned record
 				FStatIndexedRecord Record = new FStatIndexedRecord();
@@ -1436,7 +1436,7 @@ namespace EpicGames.Perforce.Managed
 					ReadOnlySpan<byte> PathSpan = ClientFile.Span;
 
 					// Parse out the data
-					StreamSnapshotBuilder LastStreamDirectory = Builder;
+					StreamTreeBuilder LastStreamDirectory = Builder;
 
 					// Try to match up as many fragments from the last file.
 					int FragmentMinIdx = ClientPrefix.Length;
@@ -1475,11 +1475,11 @@ namespace EpicGames.Perforce.Managed
 						{
 							Utf8String UnescapedFragment = PerforceUtils.UnescapePath(Fragment);
 
-							StreamSnapshotBuilder? NextStreamDirectory;
-							if (!LastStreamDirectory.NameToSubDirectory.TryGetValue(UnescapedFragment, out NextStreamDirectory))
+							StreamTreeBuilder? NextStreamDirectory;
+							if (!LastStreamDirectory.NameToTreeBuilder.TryGetValue(UnescapedFragment, out NextStreamDirectory))
 							{
-								NextStreamDirectory = new StreamSnapshotBuilder();
-								LastStreamDirectory.NameToSubDirectory.Add(UnescapedFragment, NextStreamDirectory);
+								NextStreamDirectory = new StreamTreeBuilder();
+								LastStreamDirectory.NameToTreeBuilder.Add(UnescapedFragment, NextStreamDirectory);
 							}
 							LastStreamDirectory = NextStreamDirectory;
 
diff --git a/Engine/Source/Programs/Shared/EpicGames.Perforce.Managed/StreamSnapshot.cs b/Engine/Source/Programs/Shared/EpicGames.Perforce.Managed/StreamSnapshot.cs
--- a/Engine/Source/Programs/Shared/EpicGames.Perforce.Managed/StreamSnapshot.cs
+++ b/Engine/Source/Programs/Shared/EpicGames.Perforce.Managed/StreamSnapshot.cs
@@ -19,7 +19,7 @@ namespace EpicGames.Perforce.Managed
 		/// <summary>
 		/// Empty snapshot instance
 		/// </summary>
-		public static StreamSnapshot Empty => new StreamSnapshotFromMemory(new StreamSnapshotBuilder());
+		public static StreamSnapshot Empty => new StreamSnapshotFromMemory(new StreamTreeBuilder());
 
 		/// <summary>
 		/// The root digest
diff --git a/Engine/Source/Programs/Shared/EpicGames.Perforce.Managed/StreamSnapshotBuilder.cs b/Engine/Source/Programs/Shared/EpicGames.Perforce.Managed/StreamSnapshotBuilder.cs
deleted file mode 100644
diff --git a/Engine/Source/Programs/Shared/EpicGames.Perforce.Managed/StreamSnapshotFromMemory.cs b/Engine/Source/Programs/Shared/EpicGames.Perforce.Managed/StreamSnapshotFromMemory.cs
--- a/Engine/Source/Programs/Shared/EpicGames.Perforce.Managed/StreamSnapshotFromMemory.cs
+++ b/Engine/Source/Programs/Shared/EpicGames.Perforce.Managed/StreamSnapshotFromMemory.cs
@@ -47,7 +47,7 @@ namespace EpicGames.Perforce.Managed
 		/// Constructor
 		/// </summary>
 		/// <param name="Builder"></param>
-		public StreamSnapshotFromMemory(StreamSnapshotBuilder Builder)
+		public StreamSnapshotFromMemory(StreamTreeBuilder Builder)
 		{
 			Dictionary<IoHash, CbObject> HashToTree = new Dictionary<IoHash, CbObject>();
 			this.Root = Builder.Encode(HashToTree);
diff --git a/Engine/Source/Programs/Shared/EpicGames.Perforce.Managed/StreamTree.cs b/Engine/Source/Programs/Shared/EpicGames.Perforce.Managed/StreamTree.cs
--- a/Engine/Source/Programs/Shared/EpicGames.Perforce.Managed/StreamTree.cs
+++ b/Engine/Source/Programs/Shared/EpicGames.Perforce.Managed/StreamTree.cs
@@ -172,13 +172,12 @@ namespace EpicGames.Perforce.Managed
 		/// <summary>
 		/// Map of name to file within the directory
 		/// </summary>
-		public Dictionary<Utf8String, StreamFile> NameToFile = new Dictionary<Utf8String, StreamFile>();
+		public Dictionary<Utf8String, StreamFile> NameToFile { get; } = new Dictionary<Utf8String, StreamFile>();
 
 		/// <summary>
 		/// Map of name to subdirectory
 		/// </summary>
-		public Dictionary<Utf8String, StreamTreeRef> NameToTree = new Dictionary<Utf8String, StreamTreeRef>(FileUtils.PlatformPathComparerUtf8);
-
+		public Dictionary<Utf8String, StreamTreeRef> NameToTree { get; } = new Dictionary<Utf8String, StreamTreeRef>(FileUtils.PlatformPathComparerUtf8);
 
 		#region Field names
 		static Utf8String NameField = "name";
@@ -193,6 +192,15 @@ namespace EpicGames.Perforce.Managed
 		{
 		}
 
+		/// <summary>
+		/// Default constructor
+		/// </summary>
+		public StreamTree(Dictionary<Utf8String, StreamFile> NameToFile, Dictionary<Utf8String, StreamTreeRef> NameToTree)
+		{
+			this.NameToFile = NameToFile;
+			this.NameToTree = NameToTree;
+		}
+
 		/// <summary>
 		/// Deserialize a tree from a compact binary object
 		/// </summary>
diff --git a/Engine/Source/Programs/Shared/EpicGames.Perforce.Managed/StreamTreeBuilder.cs b/Engine/Source/Programs/Shared/EpicGames.Perforce.Managed/StreamTreeBuilder.cs
new file mode 100644
--- /dev/null
+++ b/Engine/Source/Programs/Shared/EpicGames.Perforce.Managed/StreamTreeBuilder.cs
@@ -0,0 +1,114 @@
+// Copyright Epic Games, Inc. All Rights Reserved.
+
+using EpicGames.Core;
+using EpicGames.Serialization;
+using System;
+using System.Collections.Generic;
+using System.Linq;
+using System.Text;
+using System.Threading.Tasks;
+
+namespace EpicGames.Perforce.Managed
+{
+	/// <summary>
+	/// Utility class to efficiently track changes to a StreamTree object in memory
+	/// </summary>
+	public class StreamTreeBuilder
+	{
+		/// <summary>
+		/// Map from name to mutable tree
+		/// </summary>
+		public Dictionary<Utf8String, StreamFile> NameToFile { get; }
+
+		/// <summary>
+		/// Map from name to mutable tree
+		/// </summary>
+		public Dictionary<Utf8String, StreamTreeRef> NameToTree { get; }
+
+		/// <summary>
+		/// Map from name to mutable tree
+		/// </summary>
+		public Dictionary<Utf8String, StreamTreeBuilder> NameToTreeBuilder { get; } = new Dictionary<Utf8String, StreamTreeBuilder>(FileUtils.PlatformPathComparerUtf8);
+
+		/// <summary>
+		/// Tests whether the tree is empty
+		/// </summary>
+		public bool IsEmpty => NameToFile.Count == 0 && NameToTree.Count == 0 && NameToTreeBuilder.Count == 0;
+
+		/// <summary>
+		/// Constructor
+		/// </summary>
+		public StreamTreeBuilder()
+		{
+			NameToFile = new Dictionary<Utf8String, StreamFile>(FileUtils.PlatformPathComparerUtf8);
+			NameToTree = new Dictionary<Utf8String, StreamTreeRef>(FileUtils.PlatformPathComparerUtf8);
+		}
+
+		/// <summary>
+		/// Copy constructor
+		/// </summary>
+		/// <param name="Tree"></param>
+		public StreamTreeBuilder(StreamTree Tree)
+		{
+			this.NameToFile = new Dictionary<Utf8String, StreamFile>(Tree.NameToFile, FileUtils.PlatformPathComparerUtf8);
+			this.NameToTree = new Dictionary<Utf8String, StreamTreeRef>(Tree.NameToTree, FileUtils.PlatformPathComparerUtf8);
+		}
+
+		/// <summary>
+		/// Encodes the current tree state and returns a reference to it
+		/// </summary>
+		/// <param name="Objects">Dictionary of encoded objects</param>
+		/// <returns></returns>
+		public StreamTreeRef Encode(Dictionary<IoHash, CbObject> Objects)
+		{
+			StreamTreeRef? Ref = EncodeInternal(Objects);
+			if (Ref == null)
+			{
+				CbObject Object = new StreamTree().ToCbObject(Utf8String.Empty);
+				IoHash ObjectHash = Object.GetHash();
+				Objects[ObjectHash] = Object;
+				Ref = new StreamTreeRef(Utf8String.Empty, ObjectHash);
+			}
+			return Ref;
+		}
+
+		/// <summary>
+		/// Encodes the current tree state and returns a reference to it
+		/// </summary>
+		/// <param name="Objects">Dictionary of encoded objects</param>
+		/// <returns></returns>
+		public StreamTreeRef? EncodeInternal(Dictionary<IoHash, CbObject> Objects)
+		{
+			// Recursively serialize all the child items
+			foreach ((Utf8String SubTreeName, StreamTreeBuilder SubTree) in NameToTreeBuilder)
+			{
+				StreamTreeRef? SubTreeRef = SubTree.EncodeInternal(Objects);
+				if (SubTreeRef == null)
+				{
+					NameToTreeBuilder.Remove(SubTreeName);
+				}
+				else
+				{
+					NameToTree[SubTreeName] = SubTreeRef;
+				}
+			}
+			NameToTreeBuilder.Clear();
+
+			// If if's empty, don't create this tree at all
+			if (IsEmpty)
+			{
+				return null;
+			}
+
+			// Find the common base path for all items in this tree.
+			StreamTree Tree = new StreamTree(NameToFile, NameToTree);
+			Utf8String BasePath = Tree.FindBasePath();
+			CbObject PackedObject = Tree.ToCbObject(BasePath);
+
+			IoHash ObjectHash = PackedObject.GetHash();
+			Objects[ObjectHash] = PackedObject;
+
+			return new StreamTreeRef(BasePath, ObjectHash);
+		}
+	}
+}
diff --git a/Engine/Source/Programs/Shared/EpicGames.Perforce.Managed/StreamTreeReader.cs b/Engine/Source/Programs/Shared/EpicGames.Perforce.Managed/StreamTreeReader.cs
new file mode 100644
--- /dev/null
+++ b/Engine/Source/Programs/Shared/EpicGames.Perforce.Managed/StreamTreeReader.cs
@@ -0,0 +1,51 @@
+using EpicGames.Core;
+using EpicGames.Serialization;
+using System;
+using System.Collections.Generic;
+using System.IO;
+using System.Text;
+using System.Threading;
+using System.Threading.Tasks;
+
+namespace EpicGames.Perforce.Managed
+{
+	/// <summary>
+	/// Interface for 
+	/// </summary>
+	public abstract class StreamTreeReader
+	{
+		/// <summary>
+		/// Reads a node of the tree
+		/// </summary>
+		/// <param name="Ref"></param>
+		/// <returns></returns>
+		public abstract Task<StreamTree> ReadAsync(StreamTreeRef Ref);
+	}
+
+	/// <summary>
+	/// Implements a <see cref="StreamTreeReader"/> using a contiguous block of memory
+	/// </summary>
+	public class StreamTreeMemoryReader : StreamTreeReader
+	{
+		/// <summary>
+		/// Map from hash to encoded CB tree object
+		/// </summary>
+		Dictionary<IoHash, CbObject> HashToTree;
+
+		/// <summary>
+		/// Constructor
+		/// </summary>
+		/// <param name="Root"></param>
+		/// <param name="HashToTree"></param>
+		public StreamTreeMemoryReader(Dictionary<IoHash, CbObject> HashToTree)
+		{
+			this.HashToTree = HashToTree;
+		}
+
+		/// <inheritdoc/>
+		public override Task<StreamTree> ReadAsync(StreamTreeRef Ref)
+		{
+			return Task.FromResult(new StreamTree(HashToTree[Ref.Hash], Ref.Path));
+		}
+	}
+}
