Missing copyright notice.

diff --git a/Engine/Source/Programs/Shared/EpicGames.Perforce.Managed/StreamTreeReader.cs b/Engine/Source/Programs/Shared/EpicGames.Perforce.Managed/StreamTreeReader.cs
--- a/Engine/Source/Programs/Shared/EpicGames.Perforce.Managed/StreamTreeReader.cs
+++ b/Engine/Source/Programs/Shared/EpicGames.Perforce.Managed/StreamTreeReader.cs
@@ -1,3 +1,5 @@
+// Copyright Epic Games, Inc. All Rights Reserved.
+
 using EpicGames.Core;
 using EpicGames.Serialization;
 using System;
