Horde: Fix superfluous length field being written out with object data.

diff --git a/Engine/Source/Programs/Shared/EpicGames.Serialization/CompactBinaryWriter.cs b/Engine/Source/Programs/Shared/EpicGames.Serialization/CompactBinaryWriter.cs
--- a/Engine/Source/Programs/Shared/EpicGames.Serialization/CompactBinaryWriter.cs
+++ b/Engine/Source/Programs/Shared/EpicGames.Serialization/CompactBinaryWriter.cs
@@ -460,6 +460,17 @@ namespace EpicGames.Serialization
 			WriteHashPayload(Hash);
 		}
 
+		/// <summary>
+		/// Writes the payload for an object to the buffer
+		/// </summary>
+		/// <param name="Object"></param>
+		void WriteObjectPayload(CbObject Object)
+		{
+			CbField Field = Object.AsField();
+			Memory<byte> Buffer = Allocate(Field.Payload.Length);
+			Field.Payload.CopyTo(Buffer);
+		}
+
 		/// <summary>
 		/// Writes an object directly into the writer
 		/// </summary>
@@ -467,7 +478,7 @@ namespace EpicGames.Serialization
 		public void WriteObject(CbObject Object)
 		{
 			WriteField(CbFieldType.Object);
-			WriteBinaryPayload(Object.AsField().Payload.Span);
+			WriteObjectPayload(Object);
 		}
 
 		/// <summary>
@@ -478,7 +489,7 @@ namespace EpicGames.Serialization
 		public void WriteObject(Utf8String Name, CbObject Object)
 		{
 			WriteField(CbFieldType.Object, Name);
-			WriteBinaryPayload(Object.AsField().Payload.Span);
+			WriteObjectPayload(Object);
 		}
 
 		/// <summary>
