Horde: Update version number for workspace snapshots, to prevent attempting to deserialize corrupt CB objects.

diff --git a/Engine/Source/Programs/Shared/EpicGames.Perforce.Managed/StreamSnapshotFromMemory.cs b/Engine/Source/Programs/Shared/EpicGames.Perforce.Managed/StreamSnapshotFromMemory.cs
--- a/Engine/Source/Programs/Shared/EpicGames.Perforce.Managed/StreamSnapshotFromMemory.cs
+++ b/Engine/Source/Programs/Shared/EpicGames.Perforce.Managed/StreamSnapshotFromMemory.cs
@@ -20,7 +20,7 @@ namespace EpicGames.Perforce.Managed
 		/// <summary>
 		/// The current signature for saved directory objects
 		/// </summary>
-		static readonly byte[] CurrentSignature = { (byte)'W', (byte)'S', (byte)'D', 4 };
+		static readonly byte[] CurrentSignature = { (byte)'W', (byte)'S', (byte)'D', 5 };
 
 		/// <summary>
 		/// The root digest
