Horde: Additional debug logging for StatsV2 errors.

diff --git a/Engine/Source/Programs/Shared/EpicGames.Perforce.Managed/ManagedWorkspace.cs b/Engine/Source/Programs/Shared/EpicGames.Perforce.Managed/ManagedWorkspace.cs
--- a/Engine/Source/Programs/Shared/EpicGames.Perforce.Managed/ManagedWorkspace.cs
+++ b/Engine/Source/Programs/Shared/EpicGames.Perforce.Managed/ManagedWorkspace.cs
@@ -1918,6 +1918,10 @@ namespace EpicGames.Perforce.Managed
 
 				PerforceConnection ClientWithFileList = new PerforceConnection(Client);
 				ClientWithFileList.GlobalOptions.Add($"-x\"{SyncFileName}\"");
+				if (StatsFile != null)
+				{
+					ClientWithFileList.GlobalOptions.Add($"-Zdebug=dm=2");
+				}
 
 				List<SyncRecord> Records = await ClientWithFileList.SyncAsync(SyncOptions.Force | SyncOptions.FullDepotSyntax, -1, new string[0], CancellationToken);
 				if (StatsFile != null)
@@ -1947,11 +1951,9 @@ namespace EpicGames.Perforce.Managed
 								Logger.LogInformation("SyncList: {Line}", Line);
 							}
 
-							await LogPerforceCommandAsync($"protects -M -u {ClientWithFileList.UserName} {LocalFile}", Logger);
-							await LogPerforceCommandAsync($"protects -M -u {ClientWithFileList.UserName} {LocalFile.FullName.ToLower()}", Logger);
-							await LogPerforceCommandAsync($"protects -M -u {ClientWithFileList.UserName} {LocalFile.FullName.Replace('\\', '/')}", Logger);
-
-							ClientWithFileList.GlobalOptions.Add($"-Zdebug=dm=2");
+							await LogPerforceCommandAsync($"protects -M {LocalFile}", Logger);
+							await LogPerforceCommandAsync($"protects -M {LocalFile.FullName.ToLower()}", Logger);
+							await LogPerforceCommandAsync($"protects -M {LocalFile.FullName.Replace('\\', '/')}", Logger);
 
 							for (int Idx = 0; ; Idx++)
 							{
