HordeCommmon: RemoteExecution UploadList reading files and hashing async

#rnx

diff --git a/Engine/Source/Programs/Shared/EpicGames.Horde.Common/RemoteExecution/Json.cs b/Engine/Source/Programs/Shared/EpicGames.Horde.Common/RemoteExecution/Json.cs
--- a/Engine/Source/Programs/Shared/EpicGames.Horde.Common/RemoteExecution/Json.cs
+++ b/Engine/Source/Programs/Shared/EpicGames.Horde.Common/RemoteExecution/Json.cs
@@ -43,7 +43,7 @@ namespace EpicGames.Horde.Common.RemoteExecution
 			{
 				Directory.Directories.Add(await DirectoryNode.BuildAsync(UploadList));
 			}
-			return UploadList.Add(Directory);
+			return await UploadList.AddAsync(Directory);
 		}
 	}
 
@@ -67,7 +67,7 @@ namespace EpicGames.Horde.Common.RemoteExecution
 		public Dictionary<string, string> EnvVars { get; set; } = new Dictionary<string, string>();
 		public string WorkingDirectory { get; set; } = String.Empty;
 
-		public Digest Build(UploadList UploadList)
+		public async Task<Digest> BuildAsync(UploadList UploadList)
 		{
 			RpcCommand Command = new RpcCommand();
 			Command.Arguments.Add(Arguments);
@@ -79,7 +79,7 @@ namespace EpicGames.Horde.Common.RemoteExecution
 				Command.EnvironmentVariables.Add(new RpcCommand.Types.EnvironmentVariable { Name = Pair.Key, Value = Pair.Value });
 			}
 
-			return UploadList.Add(Command);
+			return await UploadList.AddAsync(Command);
 		}
 	}
 
@@ -91,11 +91,11 @@ namespace EpicGames.Horde.Common.RemoteExecution
 		public async Task<Digest> BuildAsync(UploadList UploadList, ByteString? Salt, bool DoNotCache)
 		{
 			Action NewAction = new Action();
-			NewAction.CommandDigest = Command.Build(UploadList);
+			NewAction.CommandDigest = await Command.BuildAsync(UploadList);
 			NewAction.InputRootDigest = await Workspace.BuildAsync(UploadList);
 			NewAction.Salt = Salt;
 			NewAction.DoNotCache = DoNotCache;
-			return UploadList.Add(NewAction);
+			return await UploadList.AddAsync(NewAction);
 		}
 	}
 }
diff --git a/Engine/Source/Programs/Shared/EpicGames.Horde.Common/RemoteExecution/StorageExtensions.cs b/Engine/Source/Programs/Shared/EpicGames.Horde.Common/RemoteExecution/StorageExtensions.cs
--- a/Engine/Source/Programs/Shared/EpicGames.Horde.Common/RemoteExecution/StorageExtensions.cs
+++ b/Engine/Source/Programs/Shared/EpicGames.Horde.Common/RemoteExecution/StorageExtensions.cs
@@ -27,7 +27,7 @@ namespace EpicGames.Horde.Common.RemoteExecution
 		{
 			BatchUpdateBlobsRequest.Types.Request BlobRequest = new BatchUpdateBlobsRequest.Types.Request();
 			BlobRequest.Data = ByteString.CopyFrom(Data);
-			BlobRequest.Digest = GetDigest(Data);
+			BlobRequest.Digest = await GetDigestAsync(Data);
 
 			BatchUpdateBlobsRequest Request = new BatchUpdateBlobsRequest();
 			Request.InstanceName = InstanceName;
@@ -45,7 +45,7 @@ namespace EpicGames.Horde.Common.RemoteExecution
 			Request.Digests.Add(Digest);
 
 			BatchReadBlobsResponse Response = await Storage.BatchReadBlobsAsync(Request);
-			if (Response.Responses[0].Status.Code != (int) Google.Rpc.Code.Ok)
+			if (Response.Responses[0].Status.Code != (int)Google.Rpc.Code.Ok)
 			{
 				throw new Exception($"Failed getting digest {Digest}. Code={Response.Responses[0].Status.Code} Message={Response.Responses[0].Status.Message}");
 			}
@@ -59,7 +59,7 @@ namespace EpicGames.Horde.Common.RemoteExecution
 		{
 			BatchUpdateBlobsRequest.Types.Request BlobRequest = new BatchUpdateBlobsRequest.Types.Request();
 			BlobRequest.Data = Message.ToByteString();
-			BlobRequest.Digest = GetDigest(BlobRequest.Data.ToByteArray());
+			BlobRequest.Digest = await GetDigestAsync(BlobRequest.Data);
 
 			BatchUpdateBlobsRequest Request = new BatchUpdateBlobsRequest();
 			Request.InstanceName = InstanceName;
@@ -69,45 +69,54 @@ namespace EpicGames.Horde.Common.RemoteExecution
 			//??
 			return BlobRequest.Digest;
 		}
-		
-		public static Digest GetDigest(ByteString Data)
+
+		public static async Task<Digest> GetDigestAsync(ByteString Data)
 		{
-			return GetDigest(Data.ToByteArray(), DigestFunction.Types.Value.Epiciohash);
+			return await GetDigestAsync(Data.ToByteArray(), DigestFunction.Types.Value.Epiciohash);
 		}
 
-		public static Digest GetDigest(byte[] Data)
+		public static async Task<Digest> GetDigestAsync(byte[] Data)
 		{
-			return GetDigest(Data, DigestFunction.Types.Value.Epiciohash);
+			return await GetDigestAsync(Data, DigestFunction.Types.Value.Epiciohash);
 		}
 
-		public static Digest GetDigest(ByteString Data, DigestFunction.Types.Value DigestFunction)
+		public static async Task<Digest> GetDigestAsync(ByteString Data, DigestFunction.Types.Value DigestFunction)
 		{
-			return GetDigest(Data.ToByteArray(), DigestFunction);
+			return await GetDigestAsync(Data.ToByteArray(), DigestFunction);
 		}
 
-		public static Digest GetDigest(byte[] Data, DigestFunction.Types.Value DigestFunction)
+		public static async Task<Digest> GetDigestAsync(byte[] Data, DigestFunction.Types.Value DigestFunction)
 		{
 			switch (DigestFunction)
 			{
-				case Build.Bazel.Remote.Execution.V2.DigestFunction.Types.Value.Sha256: return GetDigestSha256(Data);
-				case Build.Bazel.Remote.Execution.V2.DigestFunction.Types.Value.Epiciohash: return GetDigestIoHash(Data);
-				default: throw new Exception("Unknown digest function " + DigestFunction);
+				case Build.Bazel.Remote.Execution.V2.DigestFunction.Types.Value.Sha256: return await GetDigestSha256Async(Data);
+				case Build.Bazel.Remote.Execution.V2.DigestFunction.Types.Value.Epiciohash: return await GetDigestIoHashAsync(Data);
+				default: throw new NotImplementedException($"Unsupported digest function {DigestFunction}");
 			}
 		}
-		
-		public static Digest GetDigestSha256(byte[] Data)
+
+		public static async Task<Digest> GetDigestSha256Async(byte[] Data)
 		{
-			using (SHA256 Algorithm = SHA256.Create())
+			string Hash = await Task.Run(() =>
 			{
-				byte[] Hash = Algorithm.ComputeHash(Data);
-				return new Digest { Hash = StringUtils.FormatHexString(Hash), SizeBytes = Data.Length };
-			}
+				byte[] Result;
+				using (SHA256 Algorithm = SHA256.Create())
+				{
+					Result = Algorithm.ComputeHash(Data);
+				}
+				return StringUtils.FormatHexString(Result);
+			});
+			return new Digest { Hash = Hash, SizeBytes = Data.Length };
 		}
-		
-		public static Digest GetDigestIoHash(byte[] Data)
+
+		public static async Task<Digest> GetDigestIoHashAsync(byte[] Data)
 		{
-			IoHash Hash = IoHash.Compute(Data.AsSpan());
-			return new Digest { Hash = StringUtils.FormatHexString(Hash.Span).ToUpper(), SizeBytes = Data.Length };
+			string Hash = await Task.Run(() =>
+			{
+				IoHash Result = IoHash.Compute(Data.AsSpan());
+				return StringUtils.FormatHexString(Result.Span).ToUpper();
+			});
+			return new Digest { Hash = Hash, SizeBytes = Data.Length };
 		}
 	}
 }
diff --git a/Engine/Source/Programs/Shared/EpicGames.Horde.Common/RemoteExecution/UploadList.cs b/Engine/Source/Programs/Shared/EpicGames.Horde.Common/RemoteExecution/UploadList.cs
--- a/Engine/Source/Programs/Shared/EpicGames.Horde.Common/RemoteExecution/UploadList.cs
+++ b/Engine/Source/Programs/Shared/EpicGames.Horde.Common/RemoteExecution/UploadList.cs
@@ -2,7 +2,7 @@
 
 using Build.Bazel.Remote.Execution.V2;
 using Google.Protobuf;
-using System.Collections.Generic;
+using System.Collections.Concurrent;
 using System.IO;
 using System.Threading.Tasks;
 
@@ -13,33 +13,36 @@ namespace EpicGames.Horde.Common.RemoteExecution
 	/// </summary>
 	public class UploadList
 	{
-		public Dictionary<string, BatchUpdateBlobsRequest.Types.Request> Blobs = new Dictionary<string, BatchUpdateBlobsRequest.Types.Request>();
+		public ConcurrentDictionary<string, BatchUpdateBlobsRequest.Types.Request> Blobs = new ConcurrentDictionary<string, BatchUpdateBlobsRequest.Types.Request>();
 
-		public Digest Add<T>(IMessage<T> Message) where T : IMessage<T>
+		public async Task<Digest> AddAsync(ByteString ByteString)
 		{
-			ByteString ByteString = Message.ToByteString();
-
-			Digest Digest = StorageExtensions.GetDigest(ByteString);
+			Digest Digest = await StorageExtensions.GetDigestAsync(ByteString);
 			if (!Blobs.ContainsKey(Digest.Hash))
 			{
 				BatchUpdateBlobsRequest.Types.Request Request = new BatchUpdateBlobsRequest.Types.Request();
 				Request.Data = ByteString;
 				Request.Digest = Digest;
-				Blobs.Add(Digest.Hash, Request);
+				Blobs.TryAdd(Digest.Hash, Request);
 			}
+
 			return Digest;
 		}
 
+		public async Task<Digest> AddAsync<T>(IMessage<T> Message) where T : IMessage<T>
+		{
+			return await AddAsync(Message.ToByteString());
+		}
+
 		public async Task<Digest> AddFileAsync(string LocalFile)
 		{
-			BatchUpdateBlobsRequest.Types.Request Content = new BatchUpdateBlobsRequest.Types.Request();
+			ByteString ByteString;
 			using (FileStream Stream = File.Open(LocalFile, FileMode.Open, FileAccess.Read, FileShare.Read))
 			{
-				Content.Data = await ByteString.FromStreamAsync(Stream);
-				Content.Digest = StorageExtensions.GetDigest(Content.Data);
+				ByteString = await ByteString.FromStreamAsync(Stream);
 			}
-			Blobs[Content.Digest.Hash] = Content;
-			return Content.Digest;
+
+			return await AddAsync(ByteString);
 		}
 	}
 }
