Horde: Fix memory leak of Blake3.Hasher objects.

diff --git a/Engine/Source/Programs/Shared/EpicGames.Serialization/CompactBinary.cs b/Engine/Source/Programs/Shared/EpicGames.Serialization/CompactBinary.cs
--- a/Engine/Source/Programs/Shared/EpicGames.Serialization/CompactBinary.cs
+++ b/Engine/Source/Programs/Shared/EpicGames.Serialization/CompactBinary.cs
@@ -1205,13 +1205,15 @@ namespace EpicGames.Serialization
 		/// <returns></returns>
 		public Blake3Hash GetHash()
 		{
-			Blake3.Hasher Hasher = Blake3.Hasher.New();
-			AppendHash(Hasher);
+			using (Blake3.Hasher Hasher = Blake3.Hasher.New())
+			{
+				AppendHash(Hasher);
 
-			byte[] Hash = new byte[32];
-			Hasher.Finalize(Hash);
+				byte[] Hash = new byte[32];
+				Hasher.Finalize(Hash);
 
-			return new Blake3Hash(Hash);
+				return new Blake3Hash(Hash);
+			}
 		}
 
 		/// <summary>
@@ -1819,13 +1821,15 @@ namespace EpicGames.Serialization
 		/// <returns></returns>
 		public Blake3Hash GetHash()
 		{
-			Blake3.Hasher Hasher = Blake3.Hasher.New();
-			AppendHash(Hasher);
+			using (Blake3.Hasher Hasher = Blake3.Hasher.New())
+			{
+				AppendHash(Hasher);
 
-			byte[] Result = new byte[Blake3Hash.NumBytes];
-			Hasher.Finalize(Result);
+				byte[] Result = new byte[Blake3Hash.NumBytes];
+				Hasher.Finalize(Result);
 
-			return new Blake3Hash(Result);
+				return new Blake3Hash(Result);
+			}
 		}
 
 		/// <summary>
@@ -2040,13 +2044,15 @@ namespace EpicGames.Serialization
 		/// <returns></returns>
 		public Blake3Hash GetHash()
 		{
-			Blake3.Hasher Hasher = Blake3.Hasher.New();
-			AppendHash(Hasher);
+			using (Blake3.Hasher Hasher = Blake3.Hasher.New())
+			{
+				AppendHash(Hasher);
 
-			byte[] Data = new byte[Blake3Hash.NumBytes];
-			Hasher.Finalize(Data);
+				byte[] Data = new byte[Blake3Hash.NumBytes];
+				Hasher.Finalize(Data);
 
-			return new Blake3Hash(Data);
+				return new Blake3Hash(Data);
+			}
 		}
 
 		/// <summary>
