Horde: Fix leak of stream objects from ManagedProcess.

diff --git a/Engine/Source/Programs/Shared/EpicGames.Core/ManagedProcess.cs b/Engine/Source/Programs/Shared/EpicGames.Core/ManagedProcess.cs
--- a/Engine/Source/Programs/Shared/EpicGames.Core/ManagedProcess.cs
+++ b/Engine/Source/Programs/Shared/EpicGames.Core/ManagedProcess.cs
@@ -842,6 +842,13 @@ namespace EpicGames.Core
 				StdOutRead.Dispose();
 				StdOutRead = null;
 			}
+
+			StdIn?.Dispose();
+			StdOut?.Dispose();
+			StdOutText?.Dispose();
+			StdErr?.Dispose();
+			StdErrText?.Dispose();
+
 			if (AccountingProcessGroup != null) 
 			{
 				AccountingProcessGroup.Dispose();
