Horde: Remove usages of ConcurrentBag in async code; seems to be leaking references to buffers through TLS because we're not draining the collection.

diff --git a/Engine/Source/Programs/Horde/HordeServer/Services/PerforceLoadBalancer.cs b/Engine/Source/Programs/Horde/HordeServer/Services/PerforceLoadBalancer.cs
--- a/Engine/Source/Programs/Horde/HordeServer/Services/PerforceLoadBalancer.cs
+++ b/Engine/Source/Programs/Horde/HordeServer/Services/PerforceLoadBalancer.cs
@@ -331,7 +331,7 @@ namespace HordeServer.Services
 			Globals Globals = await DatabaseService.GetGlobalsAsync();
 
 			// Set of new server entries
-			ConcurrentBag<PerforceServerEntry> NewServers = new ConcurrentBag<PerforceServerEntry>();
+			List<PerforceServerEntry> NewServers = new List<PerforceServerEntry>();
 
 			// Update the state of all the valid servers
 			foreach (PerforceCluster Cluster in Globals.PerforceClusters)
@@ -441,7 +441,7 @@ namespace HordeServer.Services
 			}
 		}
 
-		async Task UpdateServerAsync(PerforceCluster Cluster, PerforceServer Server, ConcurrentBag<PerforceServerEntry> NewServers)
+		async Task UpdateServerAsync(PerforceCluster Cluster, PerforceServer Server, List<PerforceServerEntry> NewServers)
 		{
 			string InitialHostName = Server.ServerAndPort;
 			int Port = 1666;
diff --git a/Engine/Source/Programs/Shared/EpicGames.Perforce.Managed/ManagedWorkspace.cs b/Engine/Source/Programs/Shared/EpicGames.Perforce.Managed/ManagedWorkspace.cs
--- a/Engine/Source/Programs/Shared/EpicGames.Perforce.Managed/ManagedWorkspace.cs
+++ b/Engine/Source/Programs/Shared/EpicGames.Perforce.Managed/ManagedWorkspace.cs
@@ -2023,7 +2023,7 @@ namespace EpicGames.Perforce.Managed
 		/// </summary>
 		/// <param name="FileToMove">Information about the file to move</param>
 		/// <param name="FilesToSync">List of files to be synced. If the move fails, the file will be added to this list of files to sync.</param>
-		void MoveFileFromCache(WorkspaceFileToMove FileToMove, ConcurrentBag<WorkspaceFileToSync> FilesToSync)
+		void MoveFileFromCache(WorkspaceFileToMove FileToMove, ConcurrentQueue<WorkspaceFileToSync> FilesToSync)
 		{
 			try
 			{
@@ -2032,7 +2032,7 @@ namespace EpicGames.Perforce.Managed
 			catch(Exception Ex)
 			{
 				Logger.LogWarning(Ex, "warning: Unable to move {CacheFile} from cache to {WorkspaceFile}. Syncing instead.", FileToMove.TrackedFile.GetLocation(), FileToMove.WorkspaceFile.GetLocation());
-				FilesToSync.Add(new WorkspaceFileToSync(FileToMove.StreamFile, FileToMove.WorkspaceFile));
+				FilesToSync.Enqueue(new WorkspaceFileToSync(FileToMove.StreamFile, FileToMove.WorkspaceFile));
 				bRequiresRepair = true;
 			}
 		}
@@ -2042,7 +2042,7 @@ namespace EpicGames.Perforce.Managed
 		/// </summary>
 		/// <param name="FileToCopy">Information about the file to move</param>
 		/// <param name="FilesToSync">List of files to be synced. If the move fails, the file will be added to this list of files to sync.</param>
-		void CopyFileWithinWorkspace(WorkspaceFileToCopy FileToCopy, ConcurrentBag<WorkspaceFileToSync> FilesToSync)
+		void CopyFileWithinWorkspace(WorkspaceFileToCopy FileToCopy, ConcurrentQueue<WorkspaceFileToSync> FilesToSync)
 		{
 			try
 			{
@@ -2052,7 +2052,7 @@ namespace EpicGames.Perforce.Managed
 			catch(Exception Ex)
 			{
 				Logger.LogWarning(Ex, "warning: Unable to copy {SourceFile} to {TargetFile}. Syncing instead.", FileToCopy.SourceWorkspaceFile.GetLocation(), FileToCopy.TargetWorkspaceFile.GetLocation());
-				FilesToSync.Add(new WorkspaceFileToSync(FileToCopy.StreamFile, FileToCopy.TargetWorkspaceFile));
+				FilesToSync.Enqueue(new WorkspaceFileToSync(FileToCopy.StreamFile, FileToCopy.TargetWorkspaceFile));
 				bRequiresRepair = true;
 			}
 		}
diff --git a/Engine/Source/Programs/Shared/EpicGames.Perforce.Managed/Transactions/AddTransaction.cs b/Engine/Source/Programs/Shared/EpicGames.Perforce.Managed/Transactions/AddTransaction.cs
--- a/Engine/Source/Programs/Shared/EpicGames.Perforce.Managed/Transactions/AddTransaction.cs
+++ b/Engine/Source/Programs/Shared/EpicGames.Perforce.Managed/Transactions/AddTransaction.cs
@@ -56,8 +56,8 @@ namespace EpicGames.Perforce.Managed
 		public WorkspaceDirectoryInfo NewWorkspaceRootDir;
 		public StreamSnapshot StreamSnapshot;
 		public ConcurrentDictionary<CachedFileInfo, WorkspaceFileToMove> FilesToMove = new ConcurrentDictionary<CachedFileInfo, WorkspaceFileToMove>();
-		public ConcurrentBag<WorkspaceFileToCopy> FilesToCopy = new ConcurrentBag<WorkspaceFileToCopy>();
-		public ConcurrentBag<WorkspaceFileToSync> FilesToSync = new ConcurrentBag<WorkspaceFileToSync>();
+		public ConcurrentQueue<WorkspaceFileToCopy> FilesToCopy = new ConcurrentQueue<WorkspaceFileToCopy>();
+		public ConcurrentQueue<WorkspaceFileToSync> FilesToSync = new ConcurrentQueue<WorkspaceFileToSync>();
 
 		Dictionary<FileContentId, CachedFileInfo> ContentIdToTrackedFile;
 		Dictionary<FileContentId, WorkspaceFileInfo> ContentIdToWorkspaceFile = new Dictionary<FileContentId, WorkspaceFileInfo>();
@@ -164,7 +164,7 @@ namespace EpicGames.Perforce.Managed
 				}
 				else
 				{
-					FilesToCopy.Add(new WorkspaceFileToCopy(StreamFile, FilesToMove[TrackedFile].WorkspaceFile, WorkspaceFile));
+					FilesToCopy.Enqueue(new WorkspaceFileToCopy(StreamFile, FilesToMove[TrackedFile].WorkspaceFile, WorkspaceFile));
 				}
 			}
 			else
@@ -172,11 +172,11 @@ namespace EpicGames.Perforce.Managed
 				WorkspaceFileInfo? SourceWorkspaceFile;
 				if(ContentIdToWorkspaceFile.TryGetValue(StreamFile.ContentId, out SourceWorkspaceFile))
 				{
-					FilesToCopy.Add(new WorkspaceFileToCopy(StreamFile, SourceWorkspaceFile, WorkspaceFile));
+					FilesToCopy.Enqueue(new WorkspaceFileToCopy(StreamFile, SourceWorkspaceFile, WorkspaceFile));
 				}
 				else
 				{
-					FilesToSync.Add(new WorkspaceFileToSync(StreamFile, WorkspaceFile));
+					FilesToSync.Enqueue(new WorkspaceFileToSync(StreamFile, WorkspaceFile));
 				}
 			}
 		}
diff --git a/Engine/Source/Programs/Shared/EpicGames.Perforce.Managed/Transactions/RemoveTransaction.cs b/Engine/Source/Programs/Shared/EpicGames.Perforce.Managed/Transactions/RemoveTransaction.cs
--- a/Engine/Source/Programs/Shared/EpicGames.Perforce.Managed/Transactions/RemoveTransaction.cs
+++ b/Engine/Source/Programs/Shared/EpicGames.Perforce.Managed/Transactions/RemoveTransaction.cs
@@ -17,8 +17,8 @@ namespace EpicGames.Perforce.Managed
 		public WorkspaceDirectoryInfo NewWorkspaceRootDir;
 		public StreamSnapshot StreamSnapshot;
 		public ConcurrentDictionary<FileContentId, WorkspaceFileInfo> FilesToMove = new ConcurrentDictionary<FileContentId, WorkspaceFileInfo>();
-		public ConcurrentBag<WorkspaceFileInfo> FilesToDelete = new ConcurrentBag<WorkspaceFileInfo>();
-		public ConcurrentBag<WorkspaceDirectoryInfo> DirectoriesToDelete = new ConcurrentBag<WorkspaceDirectoryInfo>();
+		public ConcurrentQueue<WorkspaceFileInfo> FilesToDelete = new ConcurrentQueue<WorkspaceFileInfo>();
+		public ConcurrentQueue<WorkspaceDirectoryInfo> DirectoriesToDelete = new ConcurrentQueue<WorkspaceDirectoryInfo>();
 
 		public RemoveTransaction(WorkspaceDirectoryInfo WorkspaceRootDir, StreamSnapshot StreamSnapshot, Dictionary<FileContentId, CachedFileInfo> ContentIdToTrackedFile)
 		{
@@ -74,7 +74,7 @@ namespace EpicGames.Perforce.Managed
 
 		void RemoveDirectory(WorkspaceDirectoryInfo WorkspaceDir, ThreadPoolWorkQueue Queue)
 		{
-			DirectoriesToDelete.Add(WorkspaceDir);
+			DirectoriesToDelete.Enqueue(WorkspaceDir);
 
 			foreach(WorkspaceDirectoryInfo WorkspaceSubDir in WorkspaceDir.NameToSubDirectory.Values)
 			{
@@ -90,7 +90,7 @@ namespace EpicGames.Perforce.Managed
 		{
 			if(ContentIdToTrackedFile.ContainsKey(WorkspaceFile.ContentId) || !FilesToMove.TryAdd(WorkspaceFile.ContentId, WorkspaceFile))
 			{
-				FilesToDelete.Add(WorkspaceFile);
+				FilesToDelete.Enqueue(WorkspaceFile);
 			}
 		}
 	}
diff --git a/Engine/Source/Programs/Shared/EpicGames.Perforce.Managed/WorkspaceDirectoryInfo.cs b/Engine/Source/Programs/Shared/EpicGames.Perforce.Managed/WorkspaceDirectoryInfo.cs
--- a/Engine/Source/Programs/Shared/EpicGames.Perforce.Managed/WorkspaceDirectoryInfo.cs
+++ b/Engine/Source/Programs/Shared/EpicGames.Perforce.Managed/WorkspaceDirectoryInfo.cs
@@ -130,8 +130,8 @@ namespace EpicGames.Perforce.Managed
 		/// <param name="DirectoriesToDelete">Recevies an array of directories to delete</param>
 		public void Refresh(bool bRemoveUntracked, out FileInfo[] FilesToDelete, out DirectoryInfo[] DirectoriesToDelete)
 		{
-			ConcurrentBag<FileInfo> ConcurrentFilesToDelete = new ConcurrentBag<FileInfo>();
-			ConcurrentBag<DirectoryInfo> ConcurrentDirectoriesToDelete = new ConcurrentBag<DirectoryInfo>();
+			ConcurrentQueue<FileInfo> ConcurrentFilesToDelete = new ConcurrentQueue<FileInfo>();
+			ConcurrentQueue<DirectoryInfo> ConcurrentDirectoriesToDelete = new ConcurrentQueue<DirectoryInfo>();
 			using (ThreadPoolWorkQueue Queue = new ThreadPoolWorkQueue())
 			{
 				Queue.Enqueue(() => Refresh(new DirectoryInfo(GetFullName()), bRemoveUntracked, ConcurrentFilesToDelete, ConcurrentDirectoriesToDelete, Queue));
@@ -148,7 +148,7 @@ namespace EpicGames.Perforce.Managed
 		/// <param name="FilesToDelete"></param>
 		/// <param name="DirectoriesToDelete"></param>
 		/// <param name="Queue"></param>
-		void Refresh(DirectoryInfo Info, bool bRemoveUntracked, ConcurrentBag<FileInfo> FilesToDelete, ConcurrentBag<DirectoryInfo> DirectoriesToDelete, ThreadPoolWorkQueue Queue)
+		void Refresh(DirectoryInfo Info, bool bRemoveUntracked, ConcurrentQueue<FileInfo> FilesToDelete, ConcurrentQueue<DirectoryInfo> DirectoriesToDelete, ThreadPoolWorkQueue Queue)
 		{
 			// Recurse through subdirectories
 			Dictionary<Utf8String, WorkspaceDirectoryInfo> NewNameToSubDirectory = new Dictionary<Utf8String, WorkspaceDirectoryInfo>(NameToSubDirectory.Count, NameToSubDirectory.Comparer);
@@ -162,7 +162,7 @@ namespace EpicGames.Perforce.Managed
 				}
 				else if (bRemoveUntracked)
 				{
-					DirectoriesToDelete.Add(SubDirectoryInfo);
+					DirectoriesToDelete.Enqueue(SubDirectoryInfo);
 				}
 			}
 			NameToSubDirectory = NewNameToSubDirectory;
@@ -180,14 +180,14 @@ namespace EpicGames.Perforce.Managed
 					}
 					else
 					{
-						FilesToDelete.Add(File);
+						FilesToDelete.Enqueue(File);
 					}
 				}
 				else
 				{
 					if (bRemoveUntracked)
 					{
-						FilesToDelete.Add(File);
+						FilesToDelete.Enqueue(File);
 					}
 				}
 			}
@@ -211,7 +211,7 @@ namespace EpicGames.Perforce.Managed
 		/// <returns></returns>
 		public string[] FindDifferences()
 		{
-			ConcurrentBag<string> Paths = new ConcurrentBag<string>();
+			ConcurrentQueue<string> Paths = new ConcurrentQueue<string>();
 			using (ThreadPoolWorkQueue Queue = new ThreadPoolWorkQueue())
 			{
 				Queue.Enqueue(() => FindDifferences(new DirectoryInfo(GetFullName()), "/", Paths, Queue));
@@ -226,7 +226,7 @@ namespace EpicGames.Perforce.Managed
 		/// <param name="Path"></param>
 		/// <param name="Paths"></param>
 		/// <param name="Queue"></param>
-		void FindDifferences(DirectoryInfo Directory, string Path, ConcurrentBag<string> Paths, ThreadPoolWorkQueue Queue)
+		void FindDifferences(DirectoryInfo Directory, string Path, ConcurrentQueue<string> Paths, ThreadPoolWorkQueue Queue)
 		{
 			// Recurse through subdirectories
 			HashSet<Utf8String> RemainingSubDirectoryNames = new HashSet<Utf8String>(NameToSubDirectory.Keys);
@@ -239,11 +239,11 @@ namespace EpicGames.Perforce.Managed
 					Queue.Enqueue(() => StagedSubDirectory.FindDifferences(SubDirectory, String.Format("{0}{1}/", Path, SubDirectory.Name), Paths, Queue));
 					continue;
 				}
-				Paths.Add(String.Format("+{0}{1}/...", Path, SubDirectory.Name));
+				Paths.Enqueue(String.Format("+{0}{1}/...", Path, SubDirectory.Name));
 			}
 			foreach (Utf8String RemainingSubDirectoryName in RemainingSubDirectoryNames)
 			{
-				Paths.Add(String.Format("-{0}{1}/...", Path, RemainingSubDirectoryName));
+				Paths.Enqueue(String.Format("-{0}{1}/...", Path, RemainingSubDirectoryName));
 			}
 
 			// Search through files
@@ -253,11 +253,11 @@ namespace EpicGames.Perforce.Managed
 				WorkspaceFileInfo? StagedFile;
 				if (!NameToFile.TryGetValue(File.Name, out StagedFile))
 				{
-					Paths.Add(String.Format("+{0}{1}", Path, File.Name));
+					Paths.Enqueue(String.Format("+{0}{1}", Path, File.Name));
 				}
 				else if (!StagedFile.MatchesAttributes(File))
 				{
-					Paths.Add(String.Format("!{0}{1}", Path, File.Name));
+					Paths.Enqueue(String.Format("!{0}{1}", Path, File.Name));
 					RemainingFileNames.Remove(File.Name);
 				}
 				else
@@ -267,7 +267,7 @@ namespace EpicGames.Perforce.Managed
 			}
 			foreach (Utf8String RemainingFileName in RemainingFileNames)
 			{
-				Paths.Add(String.Format("-{0}{1}", Path, RemainingFileName));
+				Paths.Enqueue(String.Format("-{0}{1}", Path, RemainingFileName));
 			}
 		}
 
