UAT: Always include timestamps in the on-disk log, even if they aren't displayed on the console

diff --git a/Engine/Source/Programs/Shared/EpicGames.Core/Log.cs b/Engine/Source/Programs/Shared/EpicGames.Core/Log.cs
--- a/Engine/Source/Programs/Shared/EpicGames.Core/Log.cs
+++ b/Engine/Source/Programs/Shared/EpicGames.Core/Log.cs
@@ -886,7 +886,7 @@ namespace EpicGames.Core
 		[MethodImplAttribute(MethodImplOptions.NoInlining)]
 		private List<string> FormatMessage(int StackFramesToSkip, LogEventType Verbosity, LogFormatOptions Options, bool bForConsole, string Format, params object[] Args)
 		{
-			string TimePrefix = IncludeTimestamps ? String.Format("[{0:hh\\:mm\\:ss\\.fff}] ", Timer.Elapsed) : "";
+			string TimePrefix = (!bForConsole || IncludeTimestamps) ? String.Format("[{0:hh\\:mm\\:ss\\.fff}] ", Timer.Elapsed) : "";
 			string SourcePrefix = (bForConsole ? IncludeCallingMethodForConsole : IncludeCallingMethod) ? string.Format("{0}: ", GetSource(StackFramesToSkip)) : "";
 			string SeverityPrefix = (IncludeSeverityPrefix && ((Options & LogFormatOptions.NoSeverityPrefix) == 0)) ? GetSeverityPrefix(Verbosity) : "";
 
