Horde: Add a simple, attribute-based serializer for compact binary objects.

diff --git a/Engine/Source/Programs/Shared/EpicGames.Serialization/CompactBinarySerializer.cs b/Engine/Source/Programs/Shared/EpicGames.Serialization/CompactBinarySerializer.cs
new file mode 100644
--- /dev/null
+++ b/Engine/Source/Programs/Shared/EpicGames.Serialization/CompactBinarySerializer.cs
@@ -0,0 +1,282 @@
+// Copyright Epic Games, Inc. All Rights Reserved.
+
+using EpicGames.Core;
+using System;
+using System.Collections.Concurrent;
+using System.Collections.Generic;
+using System.Linq.Expressions;
+using System.Reflection;
+using System.Reflection.Emit;
+using System.Text;
+
+namespace EpicGames.Serialization
+{
+	/// <summary>
+	/// Attribute used to mark a property that should be serialized to compact binary
+	/// </summary>
+	public class CbFieldAttribute : Attribute
+	{
+		/// <summary>
+		/// Name of the serialized field
+		/// </summary>
+		public string? Name { get; set; }
+
+		/// <summary>
+		/// Default constructor
+		/// </summary>
+		public CbFieldAttribute()
+		{
+		}
+
+		/// <summary>
+		/// Constructor
+		/// </summary>
+		/// <param name="Name"></param>
+		public CbFieldAttribute(string Name)
+		{
+			this.Name = Name;
+		}
+	}
+
+	/// <summary>
+	/// Exception thrown when serializing cb objects
+	/// </summary>
+	public class CbException : Exception
+	{
+		/// <inheritdoc cref="Exception(string?)"/>
+		public CbException(string Message) : base(Message)
+		{
+		}
+
+		/// <inheritdoc cref="Exception(string?, Exception)"/>
+		public CbException(string Message, Exception Inner) : base(Message, Inner)
+		{
+		}
+	}
+
+	/// <summary>
+	/// Information about a reflected type
+	/// </summary>
+	class CbReflectedType
+	{
+		public delegate object CreateObjectDelegate();
+
+		public Type Type { get; }
+		public CreateObjectDelegate CreateObject { get; }
+		public Dictionary<Utf8String, PropertyInfo> NameToPropertyInfo { get; } = new Dictionary<Utf8String, PropertyInfo>();
+
+		public CbReflectedType(Type Type)
+		{
+			this.Type = Type;
+			this.CreateObject = GetCreateObjectMethod(Type);
+
+			PropertyInfo[] Properties = Type.GetProperties(BindingFlags.Public | BindingFlags.Instance);
+			foreach (PropertyInfo Property in Properties)
+			{
+				CbFieldAttribute? Attribute = Property.GetCustomAttribute<CbFieldAttribute>();
+				if (Attribute != null)
+				{
+					Utf8String Name = Attribute.Name ?? Property.Name;
+					NameToPropertyInfo.Add(Name, Property);
+				}
+			}
+
+			if (NameToPropertyInfo.Count == 0)
+			{
+				throw new NotImplementedException("Class does not have any compact binary property fields");
+			}
+		}
+
+		static CreateObjectDelegate GetCreateObjectMethod(Type Type)
+		{
+			DynamicMethod DynamicMethod = new DynamicMethod("_", Type, null);
+			ILGenerator Generator = DynamicMethod.GetILGenerator();
+
+			ConstructorInfo? Constructor = Type.GetConstructor(BindingFlags.Public | BindingFlags.NonPublic | BindingFlags.Instance, null, Type.EmptyTypes, null);
+			if (Constructor == null)
+			{
+				throw new CbException($"Unable to find default constructor for {Type}");
+			}
+
+			Generator.Emit(OpCodes.Newobj, Constructor);
+			Generator.Emit(OpCodes.Ret);
+
+			return (CreateObjectDelegate)DynamicMethod.CreateDelegate(typeof(CreateObjectDelegate));
+		}
+
+		/*
+				static WriteDelegate CreateWriteDelegate(Type Type)
+				{
+					DynamicMethod DynamicMethod = new DynamicMethod("_", Type, null);
+					ILGenerator Generator = DynamicMethod.GetILGenerator();
+
+					PropertyInfo[] TypeProperties = Type.GetProperties(BindingFlags.Public | BindingFlags.Instance);
+					foreach (PropertyInfo TypeProperty in TypeProperties)
+					{
+						CbFieldAttribute? Attribute = TypeProperty.GetCustomAttribute<CbFieldAttribute>();
+						if (Attribute != null)
+						{
+							Utf8String Name = Attribute.Name ?? TypeProperty.Name;
+							Generator.EmitCall(OpCodes.Call, TypeProperty.GetMethod!, null);
+
+							TypeInfo Info = TypeToInfo[TypeProperty.PropertyType];
+							Generator.EmitCall(OpCodes.Call, Info.WriteMethod, null);
+						}
+					}
+
+					Generator.Emit(OpCodes.Ret);
+					return (WriteDelegate)DynamicMethod.CreateDelegate(typeof(WriteDelegate));
+				}
+		*/
+	}
+
+	/// <summary>
+	/// Attribute-driven compact binary serializer
+	/// </summary>
+	public static class CbSerializer
+	{
+		/// <summary>
+		/// Cache of type reflection information
+		/// </summary>
+		static ConcurrentDictionary<Type, CbReflectedType> TypeToReflectedType = new ConcurrentDictionary<Type, CbReflectedType>();
+
+		/// <summary>
+		/// Find or add reflection info for a particular type
+		/// </summary>
+		/// <param name="Type"></param>
+		/// <returns></returns>
+		static CbReflectedType FindOrAddReflectedType(Type Type)
+		{
+			CbReflectedType? ReflectedType;
+			while (!TypeToReflectedType.TryGetValue(Type, out ReflectedType))
+			{
+				ReflectedType = new CbReflectedType(Type);
+				if(TypeToReflectedType.TryAdd(Type, ReflectedType))
+				{
+					break;
+				}
+			}
+			return ReflectedType;
+		}
+
+		/// <summary>
+		/// Serialize an object
+		/// </summary>
+		/// <typeparam name="T"></typeparam>
+		/// <param name="Object"></param>
+		/// <returns></returns>
+		public static CbObject Serialize<T>(T Object) where T : notnull
+		{
+			CbWriter Writer = new CbWriter();
+			Writer.BeginObject();
+			SerializeObject(Writer, Object);
+			Writer.EndObject();
+			return Writer.ToObject();
+		}
+
+		/// <summary>
+		/// Serialize an individual object to a <see cref="CbWriter"/>
+		/// </summary>
+		/// <param name="Writer"></param>
+		/// <param name="Object"></param>
+		static void SerializeObject(CbWriter Writer, object Object)
+		{
+			CbReflectedType ReflectedType = FindOrAddReflectedType(Object.GetType());
+
+			foreach ((Utf8String Name, PropertyInfo PropertyInfo) in ReflectedType.NameToPropertyInfo)
+			{
+				object? Value = PropertyInfo.GetValue(Object);
+				if (PropertyInfo.PropertyType == typeof(int))
+				{
+					int IntValue = (int)Value!;
+					if (IntValue != 0)
+					{
+						Writer.WriteInteger(Name, IntValue);
+					}
+				}
+				else if (PropertyInfo.PropertyType == typeof(string))
+				{
+					string? StringValue = (string?)Value;
+					if (!String.IsNullOrEmpty(StringValue))
+					{
+						Writer.WriteString(Name, StringValue);
+					}
+				}
+				else if (PropertyInfo.PropertyType == typeof(Utf8String))
+				{
+					Utf8String StringValue = (Utf8String)Value!;
+					if (StringValue.Length > 0)
+					{
+						Writer.WriteString(Name, StringValue);
+					}
+				}
+				else if (PropertyInfo.PropertyType.IsClass)
+				{
+					if (Value != null)
+					{
+						Writer.BeginObject(Name);
+						SerializeObject(Writer, Value);
+						Writer.EndObject();
+					}
+				}
+				else
+				{
+					throw new NotImplementedException();
+				}
+			}
+		}
+
+		/// <summary>
+		/// Deserialize an object from a <see cref="CbObject"/>
+		/// </summary>
+		/// <typeparam name="T"></typeparam>
+		/// <param name="Object"></param>
+		/// <returns></returns>
+		public static T Deserialize<T>(CbObject Object)
+		{
+			return (T)DeserializeObject(typeof(T), Object);
+		}
+
+		/// <summary>
+		/// Deserialize an individual class from a CbObject
+		/// </summary>
+		/// <param name="Type"></param>
+		/// <param name="Object"></param>
+		/// <returns></returns>
+		static object DeserializeObject(Type Type, CbObject Object)
+		{
+			CbReflectedType ReflectedType = FindOrAddReflectedType(Type);
+
+			object NewObject = ReflectedType.CreateObject();
+			foreach (CbField Field in Object)
+			{
+				PropertyInfo? PropertyInfo;
+				if (ReflectedType.NameToPropertyInfo.TryGetValue(Field.Name, out PropertyInfo))
+				{
+					Type PropertyType = PropertyInfo.PropertyType;
+					if (PropertyType == typeof(int))
+					{
+						PropertyInfo.SetValue(NewObject, Field.AsInt32());
+					}
+					else if (PropertyType == typeof(string))
+					{
+						PropertyInfo.SetValue(NewObject, Field.AsString());
+					}
+					else if (PropertyType == typeof(IoHash))
+					{
+						PropertyInfo.SetValue(NewObject, Field.AsHash());
+					}
+					else if (PropertyType.IsClass)
+					{
+						PropertyInfo.SetValue(NewObject, DeserializeObject(PropertyType, Field.AsObject()));
+					}
+					else
+					{
+						throw new NotImplementedException();
+					}
+				}
+			}
+			return NewObject;
+		}
+	}
+}
