Horde: Generate dynamic methods for serializing compact binary objects.

diff --git a/Engine/Source/Programs/Shared/EpicGames.Serialization/CompactBinarySerializer.cs b/Engine/Source/Programs/Shared/EpicGames.Serialization/CompactBinarySerializer.cs
--- a/Engine/Source/Programs/Shared/EpicGames.Serialization/CompactBinarySerializer.cs
+++ b/Engine/Source/Programs/Shared/EpicGames.Serialization/CompactBinarySerializer.cs
@@ -4,6 +4,7 @@ using EpicGames.Core;
 using System;
 using System.Collections.Concurrent;
 using System.Collections.Generic;
+using System.Linq;
 using System.Linq.Expressions;
 using System.Reflection;
 using System.Reflection.Emit;
@@ -55,109 +56,26 @@ namespace EpicGames.Serialization
 	}
 
 	/// <summary>
-	/// Information about a reflected type
+	/// Attribute-driven compact binary serializer
 	/// </summary>
-	class CbReflectedType
+	public static class CbSerializer
 	{
-		public delegate object CreateObjectDelegate();
-
-		public Type Type { get; }
-		public CreateObjectDelegate CreateObject { get; }
-		public Dictionary<Utf8String, PropertyInfo> NameToPropertyInfo { get; } = new Dictionary<Utf8String, PropertyInfo>();
+		delegate void WriteObjectDelegate(CbWriter Writer, object Value);
 
-		public CbReflectedType(Type Type)
+		class CbReflectedTypeInfo<T>
 		{
-			this.Type = Type;
-			this.CreateObject = GetCreateObjectMethod(Type);
-
-			PropertyInfo[] Properties = Type.GetProperties(BindingFlags.Public | BindingFlags.Instance);
-			foreach (PropertyInfo Property in Properties)
-			{
-				CbFieldAttribute? Attribute = Property.GetCustomAttribute<CbFieldAttribute>();
-				if (Attribute != null)
-				{
-					Utf8String Name = Attribute.Name ?? Property.Name;
-					NameToPropertyInfo.Add(Name, Property);
-				}
-			}
-
-			if (NameToPropertyInfo.Count == 0)
-			{
-				throw new NotImplementedException("Class does not have any compact binary property fields");
-			}
+			public static Utf8String[]? Names = null;
 		}
 
-		static CreateObjectDelegate GetCreateObjectMethod(Type Type)
+		static Dictionary<Type, MethodInfo> WritePropertyMethods = new Dictionary<Type, MethodInfo>
 		{
-			DynamicMethod DynamicMethod = new DynamicMethod("_", Type, null);
-			ILGenerator Generator = DynamicMethod.GetILGenerator();
-
-			ConstructorInfo? Constructor = Type.GetConstructor(BindingFlags.Public | BindingFlags.NonPublic | BindingFlags.Instance, null, Type.EmptyTypes, null);
-			if (Constructor == null)
-			{
-				throw new CbException($"Unable to find default constructor for {Type}");
-			}
-
-			Generator.Emit(OpCodes.Newobj, Constructor);
-			Generator.Emit(OpCodes.Ret);
+			[typeof(int)] = GetMethodInfo(() => WriteInt32(default!, default, default)),
+			[typeof(string)] = GetMethodInfo(() => WriteString(default!, default, default)),
+		};
 
-			return (CreateObjectDelegate)DynamicMethod.CreateDelegate(typeof(CreateObjectDelegate));
-		}
-
-		/*
-				static WriteDelegate CreateWriteDelegate(Type Type)
-				{
-					DynamicMethod DynamicMethod = new DynamicMethod("_", Type, null);
-					ILGenerator Generator = DynamicMethod.GetILGenerator();
-
-					PropertyInfo[] TypeProperties = Type.GetProperties(BindingFlags.Public | BindingFlags.Instance);
-					foreach (PropertyInfo TypeProperty in TypeProperties)
-					{
-						CbFieldAttribute? Attribute = TypeProperty.GetCustomAttribute<CbFieldAttribute>();
-						if (Attribute != null)
-						{
-							Utf8String Name = Attribute.Name ?? TypeProperty.Name;
-							Generator.EmitCall(OpCodes.Call, TypeProperty.GetMethod!, null);
-
-							TypeInfo Info = TypeToInfo[TypeProperty.PropertyType];
-							Generator.EmitCall(OpCodes.Call, Info.WriteMethod, null);
-						}
-					}
-
-					Generator.Emit(OpCodes.Ret);
-					return (WriteDelegate)DynamicMethod.CreateDelegate(typeof(WriteDelegate));
-				}
-		*/
-	}
-
-	/// <summary>
-	/// Attribute-driven compact binary serializer
-	/// </summary>
-	public static class CbSerializer
-	{
-		/// <summary>
-		/// Cache of type reflection information
-		/// </summary>
-		static ConcurrentDictionary<Type, CbReflectedType> TypeToReflectedType = new ConcurrentDictionary<Type, CbReflectedType>();
-
-		/// <summary>
-		/// Find or add reflection info for a particular type
-		/// </summary>
-		/// <param name="Type"></param>
-		/// <returns></returns>
-		static CbReflectedType FindOrAddReflectedType(Type Type)
-		{
-			CbReflectedType? ReflectedType;
-			while (!TypeToReflectedType.TryGetValue(Type, out ReflectedType))
-			{
-				ReflectedType = new CbReflectedType(Type);
-				if(TypeToReflectedType.TryAdd(Type, ReflectedType))
-				{
-					break;
-				}
-			}
-			return ReflectedType;
-		}
+		static ConcurrentDictionary<Type, MethodInfo> WriteObjectPropertyMethods = new ConcurrentDictionary<Type, MethodInfo>();
+		static ConcurrentDictionary<Type, MethodInfo> WriteObjectContentsMethods = new ConcurrentDictionary<Type, MethodInfo>();
+		static ConcurrentDictionary<Type, WriteObjectDelegate> WriteObjectDelegates = new ConcurrentDictionary<Type, WriteObjectDelegate>();
 
 		/// <summary>
 		/// Serialize an object
@@ -169,61 +87,17 @@ namespace EpicGames.Serialization
 		{
 			CbWriter Writer = new CbWriter();
 			Writer.BeginObject();
-			SerializeObject(Writer, Object);
-			Writer.EndObject();
-			return Writer.ToObject();
-		}
-
-		/// <summary>
-		/// Serialize an individual object to a <see cref="CbWriter"/>
-		/// </summary>
-		/// <param name="Writer"></param>
-		/// <param name="Object"></param>
-		static void SerializeObject(CbWriter Writer, object Object)
-		{
-			CbReflectedType ReflectedType = FindOrAddReflectedType(Object.GetType());
 
-			foreach ((Utf8String Name, PropertyInfo PropertyInfo) in ReflectedType.NameToPropertyInfo)
+			WriteObjectDelegate? WriteObject;
+			if (!WriteObjectDelegates.TryGetValue(typeof(T), out WriteObject))
 			{
-				object? Value = PropertyInfo.GetValue(Object);
-				if (PropertyInfo.PropertyType == typeof(int))
-				{
-					int IntValue = (int)Value!;
-					if (IntValue != 0)
-					{
-						Writer.WriteInteger(Name, IntValue);
-					}
-				}
-				else if (PropertyInfo.PropertyType == typeof(string))
-				{
-					string? StringValue = (string?)Value;
-					if (!String.IsNullOrEmpty(StringValue))
-					{
-						Writer.WriteString(Name, StringValue);
-					}
-				}
-				else if (PropertyInfo.PropertyType == typeof(Utf8String))
-				{
-					Utf8String StringValue = (Utf8String)Value!;
-					if (StringValue.Length > 0)
-					{
-						Writer.WriteString(Name, StringValue);
-					}
-				}
-				else if (PropertyInfo.PropertyType.IsClass)
-				{
-					if (Value != null)
-					{
-						Writer.BeginObject(Name);
-						SerializeObject(Writer, Value);
-						Writer.EndObject();
-					}
-				}
-				else
-				{
-					throw new NotImplementedException();
-				}
+				CreateClassContentsWriter(typeof(T));
+				WriteObject = WriteObjectDelegates[typeof(T)];
 			}
+			WriteObject(Writer, Object);
+
+			Writer.EndObject();
+			return Writer.ToObject();
 		}
 
 		/// <summary>
@@ -245,13 +119,14 @@ namespace EpicGames.Serialization
 		/// <returns></returns>
 		static object DeserializeObject(Type Type, CbObject Object)
 		{
-			CbReflectedType ReflectedType = FindOrAddReflectedType(Type);
+			(Utf8String, PropertyInfo)[] Properties = GetProperties(Type);
+			Dictionary<Utf8String, PropertyInfo> NameToPropertyInfo = Properties.ToDictionary(x => x.Item1, x => x.Item2);
 
-			object NewObject = ReflectedType.CreateObject();
+			object NewObject = Activator.CreateInstance(Type)!;
 			foreach (CbField Field in Object)
 			{
 				PropertyInfo? PropertyInfo;
-				if (ReflectedType.NameToPropertyInfo.TryGetValue(Field.Name, out PropertyInfo))
+				if (NameToPropertyInfo.TryGetValue(Field.Name, out PropertyInfo))
 				{
 					Type PropertyType = PropertyInfo.PropertyType;
 					if (PropertyType == typeof(int))
@@ -278,5 +153,149 @@ namespace EpicGames.Serialization
 			}
 			return NewObject;
 		}
+
+		static MethodInfo CreateClassPropertyWriter(Type Type, Dictionary<Type, MethodInfo> NewWritePropertyMethods, Dictionary<Type, MethodInfo> NewWriteContentsMethods)
+		{
+			MethodInfo? Method;
+			if (!WriteObjectPropertyMethods.TryGetValue(Type, out Method) && !NewWritePropertyMethods.TryGetValue(Type, out Method))
+			{
+				// Create the new method
+				DynamicMethod DynamicMethod = new DynamicMethod("_", null, new Type[] { typeof(CbWriter), typeof(Utf8String), Type });
+				Method = DynamicMethod;
+				NewWritePropertyMethods.Add(Type, Method);
+
+				// Implement the body
+				ILGenerator Generator = DynamicMethod.GetILGenerator();
+				Generator.Emit(OpCodes.Ldarg_2);
+
+				Label SkipLabel = Generator.DefineLabel();
+				Generator.Emit(OpCodes.Brfalse, SkipLabel);
+
+				Generator.Emit(OpCodes.Ldarg_0);
+				Generator.Emit(OpCodes.Ldarg_1);
+				Generator.EmitCall(OpCodes.Call, GetMethodInfo<CbWriter>(x => x.BeginObject(null!)), null);
+
+				Generator.Emit(OpCodes.Ldarg_0);
+				Generator.Emit(OpCodes.Ldarg_2);
+				Generator.EmitCall(OpCodes.Call, CreateClassContentsWriter(Type, NewWritePropertyMethods, NewWriteContentsMethods), null);
+
+				Generator.Emit(OpCodes.Ldarg_0);
+				Generator.EmitCall(OpCodes.Call, GetMethodInfo<CbWriter>(x => x.EndObject()), null);
+
+				Generator.MarkLabel(SkipLabel);
+				Generator.Emit(OpCodes.Ret);
+
+				// Add it into the map
+				WriteObjectPropertyMethods.TryAdd(Type, Method);
+			}
+			return Method;
+		}
+
+		static MethodInfo CreateClassContentsWriter(Type Type)
+		{
+			MethodInfo? Method;
+			if (!WriteObjectContentsMethods.TryGetValue(Type, out Method))
+			{
+				Dictionary<Type, MethodInfo> NewWritePropertyMethods = new Dictionary<Type, MethodInfo>();
+				Dictionary<Type, MethodInfo> NewWriteContentsMethods = new Dictionary<Type, MethodInfo>();
+				Method = CreateClassContentsWriter(Type, NewWritePropertyMethods, NewWriteContentsMethods);
+			}
+			return Method;
+		}
+
+		static MethodInfo CreateClassContentsWriter(Type Type, Dictionary<Type, MethodInfo> NewWritePropertyMethods, Dictionary<Type, MethodInfo> NewWriteContentsMethods)
+		{
+			MethodInfo? Method;
+			if (!WriteObjectContentsMethods.TryGetValue(Type, out Method) && !NewWriteContentsMethods.TryGetValue(Type, out Method))
+			{
+				// Create the method
+				DynamicMethod DynamicMethod = new DynamicMethod("_", null, new Type[] { typeof(CbWriter), typeof(object) });
+				Method = DynamicMethod;
+				NewWriteContentsMethods[Type] = DynamicMethod;
+
+				// Find the reflected properties from this type
+				(Utf8String Name, PropertyInfo Property)[] Properties = GetProperties(Type);
+
+				// Create a static type with the required reflection data
+				Type ReflectedType = typeof(CbReflectedTypeInfo<>).MakeGenericType(Type);
+				FieldInfo NamesField = ReflectedType.GetField(nameof(CbReflectedTypeInfo<object>.Names))!;
+				NamesField.SetValue(null, Properties.Select(x => x.Name).ToArray());
+
+				// Generate code for the method
+				ILGenerator Generator = DynamicMethod.GetILGenerator();
+				for (int Idx = 0; Idx < Properties.Length; Idx++)
+				{
+					PropertyInfo Property = Properties[Idx].Property;
+					Type PropertyType = Property.PropertyType;
+
+					Generator.Emit(OpCodes.Ldarg_0);
+
+					Generator.Emit(OpCodes.Ldsfld, NamesField);
+					Generator.Emit(OpCodes.Ldc_I4, Idx);
+					Generator.Emit(OpCodes.Ldelem, typeof(Utf8String));
+
+					Generator.Emit(OpCodes.Ldarg_1);
+					Generator.EmitCall(OpCodes.Call, Property.GetMethod!, null);
+
+					MethodInfo? WriteMethod;
+					if (PropertyType.IsClass)
+					{
+						WriteMethod = CreateClassPropertyWriter(PropertyType, NewWritePropertyMethods, NewWriteContentsMethods);
+					}
+					else if(!WritePropertyMethods.TryGetValue(PropertyType, out WriteMethod))
+					{
+						throw new CbException($"Unable to serialize type {PropertyType.Name}");
+					}
+					Generator.EmitCall(OpCodes.Call, WriteMethod, null);
+				}
+				Generator.Emit(OpCodes.Ret);
+
+				// Add the method to the cache
+				WriteObjectContentsMethods.TryAdd(Type, DynamicMethod);
+				WriteObjectDelegates.TryAdd(Type, (WriteObjectDelegate)DynamicMethod.CreateDelegate(typeof(WriteObjectDelegate)));
+			}
+			return Method;
+		}
+
+		static MethodInfo GetMethodInfo(Expression<Action> Expr)
+		{
+			return ((MethodCallExpression)Expr.Body).Method;
+		}
+
+		static MethodInfo GetMethodInfo<T>(Expression<Action<T>> Expr)
+		{
+			return ((MethodCallExpression)Expr.Body).Method;
+		}
+
+		static (Utf8String, PropertyInfo)[] GetProperties(Type Type)
+		{
+			List<(Utf8String, PropertyInfo)> PropertyList = new List<(Utf8String, PropertyInfo)>();
+			foreach (PropertyInfo Property in Type.GetProperties(BindingFlags.Public | BindingFlags.Instance))
+			{
+				CbFieldAttribute? Attribute = Property.GetCustomAttribute<CbFieldAttribute>();
+				if (Attribute != null)
+				{
+					Utf8String Name = Attribute.Name ?? Property.Name;
+					PropertyList.Add((Name, Property));
+				}
+			}
+			return PropertyList.ToArray();
+		}
+
+		static void WriteInt32(CbWriter Writer, Utf8String Name, int Value)
+		{
+			if (Value != 0)
+			{
+				Writer.WriteInteger(Name, Value);
+			}
+		}
+
+		static void WriteString(CbWriter Writer, Utf8String Name, string? Value)
+		{
+			if (!String.IsNullOrEmpty(Value))
+			{
+				Writer.WriteString(Name, Value);
+			}
+		}
 	}
 }
