Horde: Add IL generated serializer for reading objects. Optimized for fields ordered the same as the current class layout. Only currently supports some basic types (int, string, other objects).

diff --git a/Engine/Source/Programs/Shared/EpicGames.Serialization/CompactBinary.cs b/Engine/Source/Programs/Shared/EpicGames.Serialization/CompactBinary.cs
--- a/Engine/Source/Programs/Shared/EpicGames.Serialization/CompactBinary.cs
+++ b/Engine/Source/Programs/Shared/EpicGames.Serialization/CompactBinary.cs
@@ -721,12 +721,21 @@ namespace EpicGames.Serialization
 			}
 		}
 
+		/// <summary>
+		/// Access the field as a UTF-8 string.
+		/// </summary>
+		/// <returns></returns>
+		public Utf8String AsString()
+		{
+			return AsString(default);
+		}
+
 		/// <summary>
 		/// Access the field as a UTF-8 string. Returns the provided default on error.
 		/// </summary>
 		/// <param name="Default">Default value to return</param>
 		/// <returns></returns>
-		public Utf8String AsString(Utf8String Default = default)
+		public Utf8String AsString(Utf8String Default)
 		{
 			if (CbFieldUtils.IsString(TypeWithFlags))
 			{
@@ -768,7 +777,15 @@ namespace EpicGames.Serialization
 		/// <summary>
 		/// Access the field as an int32. Returns the provided default on error.
 		/// </summary>
-		public int AsInt32(int Default = 0)
+		public int AsInt32()
+		{
+			return AsInt32(0);
+		}
+
+		/// <summary>
+		/// Access the field as an int32. Returns the provided default on error.
+		/// </summary>
+		public int AsInt32(int Default)
 		{
 			return (int)AsInteger((ulong)Default, 31, true);
 		}
@@ -1591,6 +1608,24 @@ namespace EpicGames.Serialization
 			return new CbFieldIterator(View, Type);
 		}
 
+		/// <summary>
+		/// Check if the current value is valid
+		/// </summary>
+		/// <returns></returns>
+		public bool IsValid()
+		{
+			return Current.GetType() != CbFieldType.None;
+		}
+
+		/// <summary>
+		/// Accessor for the current value
+		/// </summary>
+		/// <returns></returns>
+		public CbField GetCurrent()
+		{
+			return Current;
+		}
+
 		/// <summary>
 		/// Copy the field range into a buffer of exactly GetRangeSize() bytes.
 		/// </summary>
@@ -1687,7 +1722,7 @@ namespace EpicGames.Serialization
 		/// <param name="Iterator"></param>
 		public static implicit operator bool(CbFieldIterator Iterator)
 		{
-			return Iterator.Current.GetType() != CbFieldType.None;
+			return Iterator.IsValid();
 		}
 
 		/// <summary>
diff --git a/Engine/Source/Programs/Shared/EpicGames.Serialization/CompactBinarySerializer.cs b/Engine/Source/Programs/Shared/EpicGames.Serialization/CompactBinarySerializer.cs
--- a/Engine/Source/Programs/Shared/EpicGames.Serialization/CompactBinarySerializer.cs
+++ b/Engine/Source/Programs/Shared/EpicGames.Serialization/CompactBinarySerializer.cs
@@ -8,7 +8,6 @@ using System.Linq;
 using System.Linq.Expressions;
 using System.Reflection;
 using System.Reflection.Emit;
-using System.Text;
 
 namespace EpicGames.Serialization
 {
@@ -60,11 +59,29 @@ namespace EpicGames.Serialization
 	/// </summary>
 	public static class CbSerializer
 	{
+		delegate object ReadObjectDelegate(CbField Value);
 		delegate void WriteObjectDelegate(CbWriter Writer, object Value);
 
 		class CbReflectedTypeInfo<T>
 		{
 			public static Utf8String[]? Names = null;
+			public static PropertyInfo[]? Properties = null;
+
+			public static bool MatchName(CbField Field, int Idx)
+			{
+				return Field.Name == Names![Idx];
+			}
+		}
+
+		static Dictionary<Type, MethodInfo> ReadPropertyMethods = new Dictionary<Type, MethodInfo>
+		{
+			[typeof(int)] = GetMethodInfo<CbField>(x => x.AsInt32()),
+			[typeof(string)] = GetMethodInfo<CbField>(x => GetString(x))
+		};
+
+		static string GetString(CbField Field)
+		{
+			return Field.AsString().ToString();
 		}
 
 		static Dictionary<Type, MethodInfo> WritePropertyMethods = new Dictionary<Type, MethodInfo>
@@ -73,8 +90,12 @@ namespace EpicGames.Serialization
 			[typeof(string)] = GetMethodInfo(() => WriteString(default!, default, default)),
 		};
 
+		static ConcurrentDictionary<Type, MethodInfo> ReadObjectContentsMethods = new ConcurrentDictionary<Type, MethodInfo>();
+
 		static ConcurrentDictionary<Type, MethodInfo> WriteObjectPropertyMethods = new ConcurrentDictionary<Type, MethodInfo>();
 		static ConcurrentDictionary<Type, MethodInfo> WriteObjectContentsMethods = new ConcurrentDictionary<Type, MethodInfo>();
+
+		static ConcurrentDictionary<Type, ReadObjectDelegate> ReadObjectDelegates = new ConcurrentDictionary<Type, ReadObjectDelegate>();
 		static ConcurrentDictionary<Type, WriteObjectDelegate> WriteObjectDelegates = new ConcurrentDictionary<Type, WriteObjectDelegate>();
 
 		/// <summary>
@@ -91,7 +112,7 @@ namespace EpicGames.Serialization
 			WriteObjectDelegate? WriteObject;
 			if (!WriteObjectDelegates.TryGetValue(typeof(T), out WriteObject))
 			{
-				CreateClassContentsWriter(typeof(T));
+				CreateObjectContentsWriter(typeof(T));
 				WriteObject = WriteObjectDelegates[typeof(T)];
 			}
 			WriteObject(Writer, Object);
@@ -108,53 +129,16 @@ namespace EpicGames.Serialization
 		/// <returns></returns>
 		public static T Deserialize<T>(CbObject Object)
 		{
-			return (T)DeserializeObject(typeof(T), Object);
-		}
-
-		/// <summary>
-		/// Deserialize an individual class from a CbObject
-		/// </summary>
-		/// <param name="Type"></param>
-		/// <param name="Object"></param>
-		/// <returns></returns>
-		static object DeserializeObject(Type Type, CbObject Object)
-		{
-			(Utf8String, PropertyInfo)[] Properties = GetProperties(Type);
-			Dictionary<Utf8String, PropertyInfo> NameToPropertyInfo = Properties.ToDictionary(x => x.Item1, x => x.Item2);
-
-			object NewObject = Activator.CreateInstance(Type)!;
-			foreach (CbField Field in Object)
+			ReadObjectDelegate? ReadObject;
+			if (!ReadObjectDelegates.TryGetValue(typeof(T), out ReadObject))
 			{
-				PropertyInfo? PropertyInfo;
-				if (NameToPropertyInfo.TryGetValue(Field.Name, out PropertyInfo))
-				{
-					Type PropertyType = PropertyInfo.PropertyType;
-					if (PropertyType == typeof(int))
-					{
-						PropertyInfo.SetValue(NewObject, Field.AsInt32());
-					}
-					else if (PropertyType == typeof(string))
-					{
-						PropertyInfo.SetValue(NewObject, Field.AsString());
-					}
-					else if (PropertyType == typeof(IoHash))
-					{
-						PropertyInfo.SetValue(NewObject, Field.AsHash());
-					}
-					else if (PropertyType.IsClass)
-					{
-						PropertyInfo.SetValue(NewObject, DeserializeObject(PropertyType, Field.AsObject()));
-					}
-					else
-					{
-						throw new NotImplementedException();
-					}
-				}
+				CreateObjectReader(typeof(T));
+				ReadObject = ReadObjectDelegates[typeof(T)];
 			}
-			return NewObject;
+			return (T)ReadObject(Object.AsField());
 		}
 
-		static MethodInfo CreateClassPropertyWriter(Type Type, Dictionary<Type, MethodInfo> NewWritePropertyMethods, Dictionary<Type, MethodInfo> NewWriteContentsMethods)
+		static MethodInfo CreateObjectPropertyWriter(Type Type, Dictionary<Type, MethodInfo> NewWritePropertyMethods, Dictionary<Type, MethodInfo> NewWriteContentsMethods)
 		{
 			MethodInfo? Method;
 			if (!WriteObjectPropertyMethods.TryGetValue(Type, out Method) && !NewWritePropertyMethods.TryGetValue(Type, out Method))
@@ -177,7 +161,7 @@ namespace EpicGames.Serialization
 
 				Generator.Emit(OpCodes.Ldarg_0);
 				Generator.Emit(OpCodes.Ldarg_2);
-				Generator.EmitCall(OpCodes.Call, CreateClassContentsWriter(Type, NewWritePropertyMethods, NewWriteContentsMethods), null);
+				Generator.EmitCall(OpCodes.Call, CreateObjectContentsWriter(Type, NewWritePropertyMethods, NewWriteContentsMethods), null);
 
 				Generator.Emit(OpCodes.Ldarg_0);
 				Generator.EmitCall(OpCodes.Call, GetMethodInfo<CbWriter>(x => x.EndObject()), null);
@@ -191,19 +175,19 @@ namespace EpicGames.Serialization
 			return Method;
 		}
 
-		static MethodInfo CreateClassContentsWriter(Type Type)
+		static MethodInfo CreateObjectContentsWriter(Type Type)
 		{
 			MethodInfo? Method;
 			if (!WriteObjectContentsMethods.TryGetValue(Type, out Method))
 			{
 				Dictionary<Type, MethodInfo> NewWritePropertyMethods = new Dictionary<Type, MethodInfo>();
 				Dictionary<Type, MethodInfo> NewWriteContentsMethods = new Dictionary<Type, MethodInfo>();
-				Method = CreateClassContentsWriter(Type, NewWritePropertyMethods, NewWriteContentsMethods);
+				Method = CreateObjectContentsWriter(Type, NewWritePropertyMethods, NewWriteContentsMethods);
 			}
 			return Method;
 		}
 
-		static MethodInfo CreateClassContentsWriter(Type Type, Dictionary<Type, MethodInfo> NewWritePropertyMethods, Dictionary<Type, MethodInfo> NewWriteContentsMethods)
+		static MethodInfo CreateObjectContentsWriter(Type Type, Dictionary<Type, MethodInfo> NewWritePropertyMethods, Dictionary<Type, MethodInfo> NewWriteContentsMethods)
 		{
 			MethodInfo? Method;
 			if (!WriteObjectContentsMethods.TryGetValue(Type, out Method) && !NewWriteContentsMethods.TryGetValue(Type, out Method))
@@ -238,13 +222,16 @@ namespace EpicGames.Serialization
 					Generator.EmitCall(OpCodes.Call, Property.GetMethod!, null);
 
 					MethodInfo? WriteMethod;
-					if (PropertyType.IsClass)
-					{
-						WriteMethod = CreateClassPropertyWriter(PropertyType, NewWritePropertyMethods, NewWriteContentsMethods);
-					}
-					else if(!WritePropertyMethods.TryGetValue(PropertyType, out WriteMethod))
+					if (!WritePropertyMethods.TryGetValue(PropertyType, out WriteMethod))
 					{
-						throw new CbException($"Unable to serialize type {PropertyType.Name}");
+						if (PropertyType.IsClass)
+						{
+							WriteMethod = CreateObjectPropertyWriter(PropertyType, NewWritePropertyMethods, NewWriteContentsMethods);
+						}
+						else
+						{
+							throw new CbException($"Unable to serialize type {PropertyType.Name}");
+						}
 					}
 					Generator.EmitCall(OpCodes.Call, WriteMethod, null);
 				}
@@ -257,6 +244,159 @@ namespace EpicGames.Serialization
 			return Method;
 		}
 
+		static MethodInfo CreateObjectReader(Type Type)
+		{
+			MethodInfo? Method;
+			if(!ReadObjectContentsMethods.TryGetValue(Type, out Method))
+			{
+				Dictionary<Type, MethodInfo> NewReadContentsMethods = new Dictionary<Type, MethodInfo>();
+				Method = CreateObjectReader(Type, NewReadContentsMethods);
+			}
+			return Method;
+		}
+
+		static MethodInfo CreateObjectReader(Type Type, Dictionary<Type, MethodInfo> NewReadMethods)
+		{
+			MethodInfo? Method;
+			if (!ReadObjectContentsMethods.TryGetValue(Type, out Method) && !NewReadMethods.TryGetValue(Type, out Method))
+			{
+				// Create the method
+				DynamicMethod DynamicMethod = new DynamicMethod("_", typeof(object), new Type[] { typeof(CbField) });
+				Method = DynamicMethod;
+				NewReadMethods[Type] = DynamicMethod;
+
+				// Generate the IL for it
+				ILGenerator Generator = DynamicMethod.GetILGenerator();
+
+				// Construct the object
+				ConstructorInfo? Constructor = Type.GetConstructor(BindingFlags.Public | BindingFlags.NonPublic | BindingFlags.Instance, null, Type.EmptyTypes, null);
+				if (Constructor == null)
+				{
+					throw new CbException($"Unable to find default constructor for {Type}");
+				}
+
+				// Find the reflected properties from this type
+				(Utf8String Name, PropertyInfo Property)[] Properties = GetProperties(Type);
+
+				// Create a static type with the required reflection data
+				Type ReflectedType = typeof(CbReflectedTypeInfo<>).MakeGenericType(Type);
+				FieldInfo NamesField = ReflectedType.GetField(nameof(CbReflectedTypeInfo<object>.Names))!;
+				NamesField.SetValue(null, Properties.Select(x => x.Name).ToArray());
+				MethodInfo MatchNameMethod = ReflectedType.GetMethod(nameof(CbReflectedTypeInfo<object>.MatchName))!;
+
+				// NewObjectLocal = new Type()
+				LocalBuilder NewObjectLocal = Generator.DeclareLocal(typeof(object));
+				Generator.Emit(OpCodes.Newobj, Constructor);
+				Generator.Emit(OpCodes.Stloc, NewObjectLocal);
+
+				// Stack(0) = CbField.CreateIterator()
+				Generator.Emit(OpCodes.Ldarg_0);
+				Generator.EmitCall(OpCodes.Call, GetMethodInfo<CbField>(x => x.CreateIterator()), null);
+
+				// CbFieldIterator IteratorLocal = Stack(0)
+				LocalBuilder IteratorLocal = Generator.DeclareLocal(typeof(CbFieldIterator));
+				Generator.Emit(OpCodes.Dup);
+				Generator.Emit(OpCodes.Stloc, IteratorLocal);
+
+				// if(!Stack.Pop().IsValid()) goto ReturnLabel
+				Generator.EmitCall(OpCodes.Call, GetMethodInfo<CbFieldIterator>(x => x.IsValid()), null);
+				Label ReturnLabel = Generator.DefineLabel();
+				Generator.Emit(OpCodes.Brfalse, ReturnLabel);
+
+				// NamesLocal = CbReflectedTypeInfo<Type>.Names
+				LocalBuilder NamesLocal = Generator.DeclareLocal(typeof(Utf8String[]));
+				Generator.Emit(OpCodes.Ldsfld, NamesField);
+				Generator.Emit(OpCodes.Stloc, NamesLocal);
+
+				// IterationLoopLabel:
+				Label IterationLoopLabel = Generator.DefineLabel();
+				Generator.MarkLabel(IterationLoopLabel);
+
+				// bool MatchLocal = false
+				LocalBuilder MatchLocal = Generator.DeclareLocal(typeof(bool));
+				Generator.Emit(OpCodes.Ldc_I4_0);
+				Generator.Emit(OpCodes.Stloc, MatchLocal);
+
+				// Stack(0) = IteratorLocal.GetCurrent()
+				Generator.Emit(OpCodes.Ldloc, IteratorLocal);
+				Generator.EmitCall(OpCodes.Call, GetMethodInfo<CbFieldIterator>(x => x.GetCurrent()), null);
+
+				// Try to parse each of the properties in order. If fields are ordered correctly, we will parse the object in a single pass. Otherwise we can loop and start again.
+				LocalBuilder FieldLocal = Generator.DeclareLocal(typeof(CbField));
+				for (int Idx = 0; Idx < Properties.Length; Idx++)
+				{
+					PropertyInfo Property = Properties[Idx].Property;
+
+					MethodInfo? ReadMethod;
+					if (!ReadPropertyMethods.TryGetValue(Property.PropertyType, out ReadMethod))
+					{
+						if (Property.PropertyType.IsClass)
+						{
+							ReadMethod = CreateObjectReader(Property.PropertyType, NewReadMethods);
+						}
+						else
+						{
+							throw new CbException($"Unable to serialize type {Property.PropertyType.Name}");
+						}
+					}
+
+					// if(!CbReflectedTypeInfo<Type>.MatchName(Stack(0), Idx)) goto SkipPropertyLabel
+					Label SkipPropertyLabel = Generator.DefineLabel();
+					Generator.Emit(OpCodes.Dup); // Current CbField
+					Generator.Emit(OpCodes.Ldc_I4, Idx);
+					Generator.Emit(OpCodes.Call, MatchNameMethod);
+					Generator.Emit(OpCodes.Brfalse, SkipPropertyLabel);
+
+					// FieldLocal = Stack.Pop()
+					Generator.Emit(OpCodes.Stloc, FieldLocal);
+
+					// Property.SetMethod(NewObjectLocal, ReadMethod(FieldLocal))
+					Generator.Emit(OpCodes.Ldloc, NewObjectLocal);
+					Generator.Emit(OpCodes.Ldloc, FieldLocal);
+					Generator.EmitCall(OpCodes.Call, ReadMethod, null);
+					Generator.EmitCall(OpCodes.Call, Property.SetMethod!, null);
+
+					// if(!IteratorLocal.MoveNext()) goto ReturnLabel
+					Generator.Emit(OpCodes.Ldloc, IteratorLocal);
+					Generator.EmitCall(OpCodes.Call, GetMethodInfo<CbFieldIterator>(x => x.MoveNext()), null);
+					Generator.Emit(OpCodes.Brfalse, ReturnLabel);
+
+					// MatchLocal = true
+					Generator.Emit(OpCodes.Ldc_I4_1);
+					Generator.Emit(OpCodes.Stloc, MatchLocal);
+
+					// Stack(0) = IteratorLocal.GetCurrent()
+					Generator.Emit(OpCodes.Ldloc, IteratorLocal);
+					Generator.EmitCall(OpCodes.Call, GetMethodInfo<CbFieldIterator>(x => x.GetCurrent()), null);
+
+					// SkipPropertyLabel:
+					Generator.MarkLabel(SkipPropertyLabel);
+				}
+
+				// Stack.Pop()
+				Generator.Emit(OpCodes.Pop); // Current CbField
+
+				// if(MatchLocal) goto IterationLoopLabel
+				Generator.Emit(OpCodes.Ldloc, MatchLocal);
+				Generator.Emit(OpCodes.Brtrue, IterationLoopLabel);
+
+				// if(IteratorLocal.MoveNext()) goto IterationLoopLabel
+				Generator.Emit(OpCodes.Ldloc, IteratorLocal);
+				Generator.EmitCall(OpCodes.Call, GetMethodInfo<CbFieldIterator>(x => x.MoveNext()), null);
+				Generator.Emit(OpCodes.Brtrue, IterationLoopLabel);
+
+				// return NewObjectLocal
+				Generator.MarkLabel(ReturnLabel);
+				Generator.Emit(OpCodes.Ldloc, NewObjectLocal);
+				Generator.Emit(OpCodes.Ret);
+
+				// Add the new methods
+				ReadObjectContentsMethods.TryAdd(Type, Method);
+				ReadObjectDelegates.TryAdd(Type, (ReadObjectDelegate)DynamicMethod.CreateDelegate(typeof(ReadObjectDelegate)));
+			}
+			return Method;
+		}
+
 		static MethodInfo GetMethodInfo(Expression<Action> Expr)
 		{
 			return ((MethodCallExpression)Expr.Body).Method;
