Horde: Add support for reading and writing list types to compact binary.

diff --git a/Engine/Source/Programs/Shared/EpicGames.Serialization/CompactBinarySerializer.cs b/Engine/Source/Programs/Shared/EpicGames.Serialization/CompactBinarySerializer.cs
--- a/Engine/Source/Programs/Shared/EpicGames.Serialization/CompactBinarySerializer.cs
+++ b/Engine/Source/Programs/Shared/EpicGames.Serialization/CompactBinarySerializer.cs
@@ -2,8 +2,10 @@
 
 using EpicGames.Core;
 using System;
+using System.Collections;
 using System.Collections.Concurrent;
 using System.Collections.Generic;
+using System.Diagnostics.CodeAnalysis;
 using System.Linq;
 using System.Linq.Expressions;
 using System.Reflection;
@@ -62,7 +64,7 @@ namespace EpicGames.Serialization
 		delegate object ReadObjectDelegate(CbField Value);
 		delegate void WriteObjectDelegate(CbWriter Writer, object Value);
 
-		class CbReflectedTypeInfo<T>
+		class CbReflectedTypeInfo<T> where T : class
 		{
 			public static Utf8String[]? Names = null;
 			public static PropertyInfo[]? Properties = null;
@@ -71,23 +73,50 @@ namespace EpicGames.Serialization
 			{
 				return Field.Name == Names![Idx];
 			}
+
+			public static object ReadList(CbField Field)
+			{
+				List<T> List = new List<T>();
+				foreach (CbField ElementField in Field)
+				{
+					List.Add(Deserialize<T>(ElementField.AsObject()));
+				}
+				return List;
+			}
+
+			public static void WriteEnumerable(CbWriter Writer, Utf8String Name, object Sequence)
+			{
+				using (IEnumerator<T> Enumerator = ((IEnumerable<T>)Sequence).GetEnumerator())
+				{
+					if (Enumerator.MoveNext())
+					{
+						Writer.BeginArray(Name);
+						for (; ; )
+						{
+							Serialize<T>(Writer, Enumerator.Current);
+							if (!Enumerator.MoveNext())
+							{
+								break;
+							}
+						}
+						Writer.EndArray();
+					}
+				}
+			}
 		}
 
 		static Dictionary<Type, MethodInfo> ReadPropertyMethods = new Dictionary<Type, MethodInfo>
 		{
 			[typeof(int)] = GetMethodInfo<CbField>(x => x.AsInt32()),
-			[typeof(string)] = GetMethodInfo<CbField>(x => GetString(x))
+			[typeof(Utf8String)] = GetMethodInfo<CbField>(x => x.AsString()),
+			[typeof(IoHash)] = GetMethodInfo<CbField>(x => x.AsHash())
 		};
 
-		static string GetString(CbField Field)
-		{
-			return Field.AsString().ToString();
-		}
-
 		static Dictionary<Type, MethodInfo> WritePropertyMethods = new Dictionary<Type, MethodInfo>
 		{
 			[typeof(int)] = GetMethodInfo(() => WriteInt32(default!, default, default)),
-			[typeof(string)] = GetMethodInfo(() => WriteString(default!, default, default)),
+			[typeof(Utf8String)] = GetMethodInfo(() => WriteString(default!, default, default)),
+			[typeof(IoHash)] = GetMethodInfo<CbWriter>(x => x.WriteHash(default, default))
 		};
 
 		static ConcurrentDictionary<Type, MethodInfo> ReadObjectContentsMethods = new ConcurrentDictionary<Type, MethodInfo>();
@@ -107,6 +136,18 @@ namespace EpicGames.Serialization
 		public static CbObject Serialize<T>(T Object) where T : notnull
 		{
 			CbWriter Writer = new CbWriter();
+			Serialize<T>(Writer, Object);
+			return Writer.ToObject();
+		}
+
+		/// <summary>
+		/// Serialize an object to a given writer
+		/// </summary>
+		/// <typeparam name="T"></typeparam>
+		/// <param name="Writer"></param>
+		/// <param name="Object"></param>
+		public static void Serialize<T>(CbWriter Writer, T Object) where T : notnull
+		{
 			Writer.BeginObject();
 
 			WriteObjectDelegate? WriteObject;
@@ -118,7 +159,6 @@ namespace EpicGames.Serialization
 			WriteObject(Writer, Object);
 
 			Writer.EndObject();
-			return Writer.ToObject();
 		}
 
 		/// <summary>
@@ -143,31 +183,41 @@ namespace EpicGames.Serialization
 			MethodInfo? Method;
 			if (!WriteObjectPropertyMethods.TryGetValue(Type, out Method) && !NewWritePropertyMethods.TryGetValue(Type, out Method))
 			{
-				// Create the new method
-				DynamicMethod DynamicMethod = new DynamicMethod("_", null, new Type[] { typeof(CbWriter), typeof(Utf8String), Type });
-				Method = DynamicMethod;
-				NewWritePropertyMethods.Add(Type, Method);
+				// Check if it's a list implementation
+				Type? EnumerableType;
+				if (TryGetInterfaceType(Type, typeof(IEnumerable<>), out EnumerableType))
+				{
+					Type ReflectedType = typeof(CbReflectedTypeInfo<>).MakeGenericType(EnumerableType.GenericTypeArguments[0]);
+					Method = ReflectedType.GetMethod(nameof(CbReflectedTypeInfo<object>.WriteEnumerable))!;
+				}
+				else
+				{
+					// Create the new method
+					DynamicMethod DynamicMethod = new DynamicMethod("_", null, new Type[] { typeof(CbWriter), typeof(Utf8String), Type });
+					Method = DynamicMethod;
+					NewWritePropertyMethods.Add(Type, Method);
 
-				// Implement the body
-				ILGenerator Generator = DynamicMethod.GetILGenerator();
-				Generator.Emit(OpCodes.Ldarg_2);
+					// Implement the body
+					ILGenerator Generator = DynamicMethod.GetILGenerator();
+					Generator.Emit(OpCodes.Ldarg_2);
 
-				Label SkipLabel = Generator.DefineLabel();
-				Generator.Emit(OpCodes.Brfalse, SkipLabel);
+					Label SkipLabel = Generator.DefineLabel();
+					Generator.Emit(OpCodes.Brfalse, SkipLabel);
 
-				Generator.Emit(OpCodes.Ldarg_0);
-				Generator.Emit(OpCodes.Ldarg_1);
-				Generator.EmitCall(OpCodes.Call, GetMethodInfo<CbWriter>(x => x.BeginObject(null!)), null);
+					Generator.Emit(OpCodes.Ldarg_0);
+					Generator.Emit(OpCodes.Ldarg_1);
+					Generator.EmitCall(OpCodes.Call, GetMethodInfo<CbWriter>(x => x.BeginObject(null!)), null);
 
-				Generator.Emit(OpCodes.Ldarg_0);
-				Generator.Emit(OpCodes.Ldarg_2);
-				Generator.EmitCall(OpCodes.Call, CreateObjectContentsWriter(Type, NewWritePropertyMethods, NewWriteContentsMethods), null);
+					Generator.Emit(OpCodes.Ldarg_0);
+					Generator.Emit(OpCodes.Ldarg_2);
+					Generator.EmitCall(OpCodes.Call, CreateObjectContentsWriter(Type, NewWritePropertyMethods, NewWriteContentsMethods), null);
 
-				Generator.Emit(OpCodes.Ldarg_0);
-				Generator.EmitCall(OpCodes.Call, GetMethodInfo<CbWriter>(x => x.EndObject()), null);
+					Generator.Emit(OpCodes.Ldarg_0);
+					Generator.EmitCall(OpCodes.Call, GetMethodInfo<CbWriter>(x => x.EndObject()), null);
 
-				Generator.MarkLabel(SkipLabel);
-				Generator.Emit(OpCodes.Ret);
+					Generator.MarkLabel(SkipLabel);
+					Generator.Emit(OpCodes.Ret);
+				}
 
 				// Add it into the map
 				WriteObjectPropertyMethods.TryAdd(Type, Method);
@@ -260,139 +310,147 @@ namespace EpicGames.Serialization
 			MethodInfo? Method;
 			if (!ReadObjectContentsMethods.TryGetValue(Type, out Method) && !NewReadMethods.TryGetValue(Type, out Method))
 			{
-				// Create the method
-				DynamicMethod DynamicMethod = new DynamicMethod("_", typeof(object), new Type[] { typeof(CbField) });
-				Method = DynamicMethod;
-				NewReadMethods[Type] = DynamicMethod;
-
-				// Generate the IL for it
-				ILGenerator Generator = DynamicMethod.GetILGenerator();
-
-				// Construct the object
-				ConstructorInfo? Constructor = Type.GetConstructor(BindingFlags.Public | BindingFlags.NonPublic | BindingFlags.Instance, null, Type.EmptyTypes, null);
-				if (Constructor == null)
+				if (TryGetInterfaceType(Type, typeof(IList<>), out Type? EnumerableType))
 				{
-					throw new CbException($"Unable to find default constructor for {Type}");
+					Type ReflectedType = typeof(CbReflectedTypeInfo<>).MakeGenericType(EnumerableType.GenericTypeArguments[0]);
+					Method = ReflectedType.GetMethod(nameof(CbReflectedTypeInfo<object>.ReadList))!;
 				}
-
-				// Find the reflected properties from this type
-				(Utf8String Name, PropertyInfo Property)[] Properties = GetProperties(Type);
-
-				// Create a static type with the required reflection data
-				Type ReflectedType = typeof(CbReflectedTypeInfo<>).MakeGenericType(Type);
-				FieldInfo NamesField = ReflectedType.GetField(nameof(CbReflectedTypeInfo<object>.Names))!;
-				NamesField.SetValue(null, Properties.Select(x => x.Name).ToArray());
-				MethodInfo MatchNameMethod = ReflectedType.GetMethod(nameof(CbReflectedTypeInfo<object>.MatchName))!;
-
-				// NewObjectLocal = new Type()
-				LocalBuilder NewObjectLocal = Generator.DeclareLocal(typeof(object));
-				Generator.Emit(OpCodes.Newobj, Constructor);
-				Generator.Emit(OpCodes.Stloc, NewObjectLocal);
-
-				// Stack(0) = CbField.CreateIterator()
-				Generator.Emit(OpCodes.Ldarg_0);
-				Generator.EmitCall(OpCodes.Call, GetMethodInfo<CbField>(x => x.CreateIterator()), null);
-
-				// CbFieldIterator IteratorLocal = Stack(0)
-				LocalBuilder IteratorLocal = Generator.DeclareLocal(typeof(CbFieldIterator));
-				Generator.Emit(OpCodes.Dup);
-				Generator.Emit(OpCodes.Stloc, IteratorLocal);
-
-				// if(!Stack.Pop().IsValid()) goto ReturnLabel
-				Generator.EmitCall(OpCodes.Call, GetMethodInfo<CbFieldIterator>(x => x.IsValid()), null);
-				Label ReturnLabel = Generator.DefineLabel();
-				Generator.Emit(OpCodes.Brfalse, ReturnLabel);
-
-				// NamesLocal = CbReflectedTypeInfo<Type>.Names
-				LocalBuilder NamesLocal = Generator.DeclareLocal(typeof(Utf8String[]));
-				Generator.Emit(OpCodes.Ldsfld, NamesField);
-				Generator.Emit(OpCodes.Stloc, NamesLocal);
-
-				// IterationLoopLabel:
-				Label IterationLoopLabel = Generator.DefineLabel();
-				Generator.MarkLabel(IterationLoopLabel);
-
-				// bool MatchLocal = false
-				LocalBuilder MatchLocal = Generator.DeclareLocal(typeof(bool));
-				Generator.Emit(OpCodes.Ldc_I4_0);
-				Generator.Emit(OpCodes.Stloc, MatchLocal);
-
-				// Stack(0) = IteratorLocal.GetCurrent()
-				Generator.Emit(OpCodes.Ldloc, IteratorLocal);
-				Generator.EmitCall(OpCodes.Call, GetMethodInfo<CbFieldIterator>(x => x.GetCurrent()), null);
-
-				// Try to parse each of the properties in order. If fields are ordered correctly, we will parse the object in a single pass. Otherwise we can loop and start again.
-				LocalBuilder FieldLocal = Generator.DeclareLocal(typeof(CbField));
-				for (int Idx = 0; Idx < Properties.Length; Idx++)
+				else
 				{
-					PropertyInfo Property = Properties[Idx].Property;
+					// Create the method
+					DynamicMethod DynamicMethod = new DynamicMethod("_", typeof(object), new Type[] { typeof(CbField) });
+					Method = DynamicMethod;
+					NewReadMethods[Type] = DynamicMethod;
 
-					MethodInfo? ReadMethod;
-					if (!ReadPropertyMethods.TryGetValue(Property.PropertyType, out ReadMethod))
+					// Generate the IL for it
+					ILGenerator Generator = DynamicMethod.GetILGenerator();
+
+					// Construct the object
+					ConstructorInfo? Constructor = Type.GetConstructor(BindingFlags.Public | BindingFlags.NonPublic | BindingFlags.Instance, null, Type.EmptyTypes, null);
+					if (Constructor == null)
 					{
-						if (Property.PropertyType.IsClass)
-						{
-							ReadMethod = CreateObjectReader(Property.PropertyType, NewReadMethods);
-						}
-						else
-						{
-							throw new CbException($"Unable to serialize type {Property.PropertyType.Name}");
-						}
+						throw new CbException($"Unable to find default constructor for {Type}");
 					}
 
-					// if(!CbReflectedTypeInfo<Type>.MatchName(Stack(0), Idx)) goto SkipPropertyLabel
-					Label SkipPropertyLabel = Generator.DefineLabel();
-					Generator.Emit(OpCodes.Dup); // Current CbField
-					Generator.Emit(OpCodes.Ldc_I4, Idx);
-					Generator.Emit(OpCodes.Call, MatchNameMethod);
-					Generator.Emit(OpCodes.Brfalse, SkipPropertyLabel);
+					// Find the reflected properties from this type
+					(Utf8String Name, PropertyInfo Property)[] Properties = GetProperties(Type);
 
-					// FieldLocal = Stack.Pop()
-					Generator.Emit(OpCodes.Stloc, FieldLocal);
+					// Create a static type with the required reflection data
+					Type ReflectedType = typeof(CbReflectedTypeInfo<>).MakeGenericType(Type);
+					FieldInfo NamesField = ReflectedType.GetField(nameof(CbReflectedTypeInfo<object>.Names))!;
+					NamesField.SetValue(null, Properties.Select(x => x.Name).ToArray());
+					MethodInfo MatchNameMethod = ReflectedType.GetMethod(nameof(CbReflectedTypeInfo<object>.MatchName))!;
 
-					// Property.SetMethod(NewObjectLocal, ReadMethod(FieldLocal))
-					Generator.Emit(OpCodes.Ldloc, NewObjectLocal);
-					Generator.Emit(OpCodes.Ldloc, FieldLocal);
-					Generator.EmitCall(OpCodes.Call, ReadMethod, null);
-					Generator.EmitCall(OpCodes.Call, Property.SetMethod!, null);
+					// NewObjectLocal = new Type()
+					LocalBuilder NewObjectLocal = Generator.DeclareLocal(typeof(object));
+					Generator.Emit(OpCodes.Newobj, Constructor);
+					Generator.Emit(OpCodes.Stloc, NewObjectLocal);
 
-					// if(!IteratorLocal.MoveNext()) goto ReturnLabel
-					Generator.Emit(OpCodes.Ldloc, IteratorLocal);
-					Generator.EmitCall(OpCodes.Call, GetMethodInfo<CbFieldIterator>(x => x.MoveNext()), null);
+					// Stack(0) = CbField.CreateIterator()
+					Generator.Emit(OpCodes.Ldarg_0);
+					Generator.EmitCall(OpCodes.Call, GetMethodInfo<CbField>(x => x.CreateIterator()), null);
+
+					// CbFieldIterator IteratorLocal = Stack(0)
+					LocalBuilder IteratorLocal = Generator.DeclareLocal(typeof(CbFieldIterator));
+					Generator.Emit(OpCodes.Dup);
+					Generator.Emit(OpCodes.Stloc, IteratorLocal);
+
+					// if(!Stack.Pop().IsValid()) goto ReturnLabel
+					Generator.EmitCall(OpCodes.Call, GetMethodInfo<CbFieldIterator>(x => x.IsValid()), null);
+					Label ReturnLabel = Generator.DefineLabel();
 					Generator.Emit(OpCodes.Brfalse, ReturnLabel);
 
-					// MatchLocal = true
-					Generator.Emit(OpCodes.Ldc_I4_1);
+					// NamesLocal = CbReflectedTypeInfo<Type>.Names
+					LocalBuilder NamesLocal = Generator.DeclareLocal(typeof(Utf8String[]));
+					Generator.Emit(OpCodes.Ldsfld, NamesField);
+					Generator.Emit(OpCodes.Stloc, NamesLocal);
+
+					// IterationLoopLabel:
+					Label IterationLoopLabel = Generator.DefineLabel();
+					Generator.MarkLabel(IterationLoopLabel);
+
+					// bool MatchLocal = false
+					LocalBuilder MatchLocal = Generator.DeclareLocal(typeof(bool));
+					Generator.Emit(OpCodes.Ldc_I4_0);
 					Generator.Emit(OpCodes.Stloc, MatchLocal);
 
 					// Stack(0) = IteratorLocal.GetCurrent()
 					Generator.Emit(OpCodes.Ldloc, IteratorLocal);
 					Generator.EmitCall(OpCodes.Call, GetMethodInfo<CbFieldIterator>(x => x.GetCurrent()), null);
 
-					// SkipPropertyLabel:
-					Generator.MarkLabel(SkipPropertyLabel);
-				}
+					// Try to parse each of the properties in order. If fields are ordered correctly, we will parse the object in a single pass. Otherwise we can loop and start again.
+					LocalBuilder FieldLocal = Generator.DeclareLocal(typeof(CbField));
+					for (int Idx = 0; Idx < Properties.Length; Idx++)
+					{
+						PropertyInfo Property = Properties[Idx].Property;
+
+						MethodInfo? ReadMethod;
+						if (!ReadPropertyMethods.TryGetValue(Property.PropertyType, out ReadMethod))
+						{
+							if (Property.PropertyType.IsClass)
+							{
+								ReadMethod = CreateObjectReader(Property.PropertyType, NewReadMethods);
+							}
+							else
+							{
+								throw new CbException($"Unable to serialize type {Property.PropertyType.Name}");
+							}
+						}
 
-				// Stack.Pop()
-				Generator.Emit(OpCodes.Pop); // Current CbField
+						// if(!CbReflectedTypeInfo<Type>.MatchName(Stack(0), Idx)) goto SkipPropertyLabel
+						Label SkipPropertyLabel = Generator.DefineLabel();
+						Generator.Emit(OpCodes.Dup); // Current CbField
+						Generator.Emit(OpCodes.Ldc_I4, Idx);
+						Generator.Emit(OpCodes.Call, MatchNameMethod);
+						Generator.Emit(OpCodes.Brfalse, SkipPropertyLabel);
+
+						// FieldLocal = Stack.Pop()
+						Generator.Emit(OpCodes.Stloc, FieldLocal);
+
+						// Property.SetMethod(NewObjectLocal, ReadMethod(FieldLocal))
+						Generator.Emit(OpCodes.Ldloc, NewObjectLocal);
+						Generator.Emit(OpCodes.Ldloc, FieldLocal);
+						Generator.EmitCall(OpCodes.Call, ReadMethod, null);
+						Generator.EmitCall(OpCodes.Call, Property.SetMethod!, null);
+
+						// if(!IteratorLocal.MoveNext()) goto ReturnLabel
+						Generator.Emit(OpCodes.Ldloc, IteratorLocal);
+						Generator.EmitCall(OpCodes.Call, GetMethodInfo<CbFieldIterator>(x => x.MoveNext()), null);
+						Generator.Emit(OpCodes.Brfalse, ReturnLabel);
+
+						// MatchLocal = true
+						Generator.Emit(OpCodes.Ldc_I4_1);
+						Generator.Emit(OpCodes.Stloc, MatchLocal);
+
+						// Stack(0) = IteratorLocal.GetCurrent()
+						Generator.Emit(OpCodes.Ldloc, IteratorLocal);
+						Generator.EmitCall(OpCodes.Call, GetMethodInfo<CbFieldIterator>(x => x.GetCurrent()), null);
+
+						// SkipPropertyLabel:
+						Generator.MarkLabel(SkipPropertyLabel);
+					}
 
-				// if(MatchLocal) goto IterationLoopLabel
-				Generator.Emit(OpCodes.Ldloc, MatchLocal);
-				Generator.Emit(OpCodes.Brtrue, IterationLoopLabel);
+					// Stack.Pop()
+					Generator.Emit(OpCodes.Pop); // Current CbField
 
-				// if(IteratorLocal.MoveNext()) goto IterationLoopLabel
-				Generator.Emit(OpCodes.Ldloc, IteratorLocal);
-				Generator.EmitCall(OpCodes.Call, GetMethodInfo<CbFieldIterator>(x => x.MoveNext()), null);
-				Generator.Emit(OpCodes.Brtrue, IterationLoopLabel);
+					// if(MatchLocal) goto IterationLoopLabel
+					Generator.Emit(OpCodes.Ldloc, MatchLocal);
+					Generator.Emit(OpCodes.Brtrue, IterationLoopLabel);
 
-				// return NewObjectLocal
-				Generator.MarkLabel(ReturnLabel);
-				Generator.Emit(OpCodes.Ldloc, NewObjectLocal);
-				Generator.Emit(OpCodes.Ret);
+					// if(IteratorLocal.MoveNext()) goto IterationLoopLabel
+					Generator.Emit(OpCodes.Ldloc, IteratorLocal);
+					Generator.EmitCall(OpCodes.Call, GetMethodInfo<CbFieldIterator>(x => x.MoveNext()), null);
+					Generator.Emit(OpCodes.Brtrue, IterationLoopLabel);
+
+					// return NewObjectLocal
+					Generator.MarkLabel(ReturnLabel);
+					Generator.Emit(OpCodes.Ldloc, NewObjectLocal);
+					Generator.Emit(OpCodes.Ret);
 
-				// Add the new methods
-				ReadObjectContentsMethods.TryAdd(Type, Method);
-				ReadObjectDelegates.TryAdd(Type, (ReadObjectDelegate)DynamicMethod.CreateDelegate(typeof(ReadObjectDelegate)));
+					// Add the new methods
+					ReadObjectContentsMethods.TryAdd(Type, Method);
+					ReadObjectDelegates.TryAdd(Type, (ReadObjectDelegate)DynamicMethod.CreateDelegate(typeof(ReadObjectDelegate)));
+				}
 			}
 			return Method;
 		}
@@ -419,9 +477,35 @@ namespace EpicGames.Serialization
 					PropertyList.Add((Name, Property));
 				}
 			}
+			if (PropertyList.Count == 0)
+			{
+				throw new CbException($"{Type.Name} does not have any properties marked with [CbField]");
+			}
 			return PropertyList.ToArray();
 		}
 
+		static bool TryGetInterfaceType(Type Type, Type GenericType, [NotNullWhen(true)] out Type? OutCollectionType)
+		{
+			foreach (Type InterfaceType in Type.GetInterfaces())
+			{
+				if (InterfaceType.GetGenericTypeDefinition() == GenericType)
+				{
+					OutCollectionType = InterfaceType;
+					return true;
+				}
+			}
+			OutCollectionType = null;
+			return false;
+		}
+
+		static void WriteInt32Value(CbWriter Writer, int Value)
+		{
+			if (Value != 0)
+			{
+				Writer.WriteIntegerValue(Value);
+			}
+		}
+
 		static void WriteInt32(CbWriter Writer, Utf8String Name, int Value)
 		{
 			if (Value != 0)
@@ -430,9 +514,17 @@ namespace EpicGames.Serialization
 			}
 		}
 
-		static void WriteString(CbWriter Writer, Utf8String Name, string? Value)
+		static void WriteStringValue(CbWriter Writer, Utf8String Value)
+		{
+			if (Value.Length > 0)
+			{
+				Writer.WriteStringValue(Value);
+			}
+		}
+
+		static void WriteString(CbWriter Writer, Utf8String Name, Utf8String Value)
 		{
-			if (!String.IsNullOrEmpty(Value))
+			if (Value.Length > 0)
 			{
 				Writer.WriteString(Name, Value);
 			}
