Horde: Add a framework for being able to expand CB serialization with custom converters.

CbSerializer.RegisterConverterFactory() allows registering factories that can create custom CbConverter instances for different types.

Move tagged-property IL generation into its own converter factory (CbDefaultConverterFactory). Add custom converters for list and dictionary types.

diff --git a/Engine/Source/Programs/Shared/EpicGames.Serialization/CompactBinarySerializer.cs b/Engine/Source/Programs/Shared/EpicGames.Serialization/CompactBinarySerializer.cs
--- a/Engine/Source/Programs/Shared/EpicGames.Serialization/CompactBinarySerializer.cs
+++ b/Engine/Source/Programs/Shared/EpicGames.Serialization/CompactBinarySerializer.cs
@@ -1,6 +1,7 @@
 // Copyright Epic Games, Inc. All Rights Reserved.
 
 using EpicGames.Core;
+using EpicGames.Serialization.Converters;
 using System;
 using System.Collections;
 using System.Collections.Concurrent;
@@ -61,473 +62,157 @@ namespace EpicGames.Serialization
 	/// </summary>
 	public static class CbSerializer
 	{
-		delegate object ReadObjectDelegate(CbField Value);
-		delegate void WriteObjectDelegate(CbWriter Writer, object Value);
-
-		class CbReflectedTypeInfo<T> where T : class
-		{
-			public static Utf8String[]? Names = null;
-			public static PropertyInfo[]? Properties = null;
-
-			public static bool MatchName(CbField Field, int Idx)
-			{
-				return Field.Name == Names![Idx];
-			}
-
-			public static object ReadList(CbField Field)
-			{
-				List<T> List = new List<T>();
-				foreach (CbField ElementField in Field)
-				{
-					List.Add(Deserialize<T>(ElementField.AsObject()));
-				}
-				return List;
-			}
-
-			public static void WriteEnumerable(CbWriter Writer, Utf8String Name, object Sequence)
-			{
-				using (IEnumerator<T> Enumerator = ((IEnumerable<T>)Sequence).GetEnumerator())
-				{
-					if (Enumerator.MoveNext())
-					{
-						Writer.BeginArray(Name);
-						for (; ; )
-						{
-							Serialize<T>(Writer, Enumerator.Current);
-							if (!Enumerator.MoveNext())
-							{
-								break;
-							}
-						}
-						Writer.EndArray();
-					}
-				}
-			}
-		}
-
-		static Dictionary<Type, MethodInfo> ReadPropertyMethods = new Dictionary<Type, MethodInfo>
-		{
-			[typeof(int)] = GetMethodInfo<CbField>(x => x.AsInt32()),
-			[typeof(Utf8String)] = GetMethodInfo<CbField>(x => x.AsString()),
-			[typeof(IoHash)] = GetMethodInfo<CbField>(x => x.AsHash())
-		};
-
-		static Dictionary<Type, MethodInfo> WritePropertyMethods = new Dictionary<Type, MethodInfo>
-		{
-			[typeof(int)] = GetMethodInfo(() => WriteInt32(default!, default, default)),
-			[typeof(Utf8String)] = GetMethodInfo(() => WriteString(default!, default, default)),
-			[typeof(IoHash)] = GetMethodInfo<CbWriter>(x => x.WriteHash(default, default))
-		};
-
-		static ConcurrentDictionary<Type, MethodInfo> ReadObjectContentsMethods = new ConcurrentDictionary<Type, MethodInfo>();
-
-		static ConcurrentDictionary<Type, MethodInfo> WriteObjectPropertyMethods = new ConcurrentDictionary<Type, MethodInfo>();
-		static ConcurrentDictionary<Type, MethodInfo> WriteObjectContentsMethods = new ConcurrentDictionary<Type, MethodInfo>();
-
-		static ConcurrentDictionary<Type, ReadObjectDelegate> ReadObjectDelegates = new ConcurrentDictionary<Type, ReadObjectDelegate>();
-		static ConcurrentDictionary<Type, WriteObjectDelegate> WriteObjectDelegates = new ConcurrentDictionary<Type, WriteObjectDelegate>();
-
 		/// <summary>
-		/// Serialize an object
+		/// Object used for locking access to shared objects
 		/// </summary>
-		/// <typeparam name="T"></typeparam>
-		/// <param name="Object"></param>
-		/// <returns></returns>
-		public static CbObject Serialize<T>(T Object) where T : notnull
-		{
-			CbWriter Writer = new CbWriter();
-			Serialize<T>(Writer, Object);
-			return Writer.ToObject();
-		}
+		static object LockObject = new object();
 
 		/// <summary>
-		/// Serialize an object to a given writer
+		/// Cache of type to converter
 		/// </summary>
-		/// <typeparam name="T"></typeparam>
-		/// <param name="Writer"></param>
-		/// <param name="Object"></param>
-		public static void Serialize<T>(CbWriter Writer, T Object) where T : notnull
-		{
-			Writer.BeginObject();
-
-			WriteObjectDelegate? WriteObject;
-			if (!WriteObjectDelegates.TryGetValue(typeof(T), out WriteObject))
-			{
-				CreateObjectContentsWriter(typeof(T));
-				WriteObject = WriteObjectDelegates[typeof(T)];
-			}
-			WriteObject(Writer, Object);
-
-			Writer.EndObject();
-		}
+		public static Dictionary<Type, CbConverter> TypeToConverter = new Dictionary<Type, CbConverter>();
 
 		/// <summary>
-		/// Deserialize an object from a <see cref="CbObject"/>
+		/// List of converter factories. Must be 
 		/// </summary>
-		/// <typeparam name="T"></typeparam>
-		/// <param name="Object"></param>
-		/// <returns></returns>
-		public static T Deserialize<T>(CbObject Object)
+		public static List<CbConverterFactory> ConverterFactories = new List<CbConverterFactory>
 		{
-			ReadObjectDelegate? ReadObject;
-			if (!ReadObjectDelegates.TryGetValue(typeof(T), out ReadObject))
-			{
-				CreateObjectReader(typeof(T));
-				ReadObject = ReadObjectDelegates[typeof(T)];
-			}
-			return (T)ReadObject(Object.AsField());
-		}
-
-		static MethodInfo CreateObjectPropertyWriter(Type Type, Dictionary<Type, MethodInfo> NewWritePropertyMethods, Dictionary<Type, MethodInfo> NewWriteContentsMethods)
+			new CbDefaultConverterFactory(),
+			new CbListConverterFactory(),
+			new CbDictionaryConverterFactory()
+		};
+	
+		/// <summary>
+		/// Class used to cache values returned by CreateConverterInfo without having to do a dictionary lookup.
+		/// </summary>
+		/// <typeparam name="T">The type to be converted</typeparam>
+		class CbConverterCache<T>
 		{
-			MethodInfo? Method;
-			if (!WriteObjectPropertyMethods.TryGetValue(Type, out Method) && !NewWritePropertyMethods.TryGetValue(Type, out Method))
-			{
-				// Check if it's a list implementation
-				Type? EnumerableType;
-				if (TryGetInterfaceType(Type, typeof(IEnumerable<>), out EnumerableType))
-				{
-					Type ReflectedType = typeof(CbReflectedTypeInfo<>).MakeGenericType(EnumerableType.GenericTypeArguments[0]);
-					Method = ReflectedType.GetMethod(nameof(CbReflectedTypeInfo<object>.WriteEnumerable))!;
-				}
-				else
-				{
-					// Create the new method
-					DynamicMethod DynamicMethod = new DynamicMethod("_", null, new Type[] { typeof(CbWriter), typeof(Utf8String), Type });
-					Method = DynamicMethod;
-					NewWritePropertyMethods.Add(Type, Method);
-
-					// Implement the body
-					ILGenerator Generator = DynamicMethod.GetILGenerator();
-					Generator.Emit(OpCodes.Ldarg_2);
+			/// <summary>
+			/// The converter instance
+			/// </summary>
+			public static CbConverter<T> Instance { get; } = CreateConverter();
 
-					Label SkipLabel = Generator.DefineLabel();
-					Generator.Emit(OpCodes.Brfalse, SkipLabel);
-
-					Generator.Emit(OpCodes.Ldarg_0);
-					Generator.Emit(OpCodes.Ldarg_1);
-					Generator.EmitCall(OpCodes.Call, GetMethodInfo<CbWriter>(x => x.BeginObject(null!)), null);
-
-					Generator.Emit(OpCodes.Ldarg_0);
-					Generator.Emit(OpCodes.Ldarg_2);
-					Generator.EmitCall(OpCodes.Call, CreateObjectContentsWriter(Type, NewWritePropertyMethods, NewWriteContentsMethods), null);
-
-					Generator.Emit(OpCodes.Ldarg_0);
-					Generator.EmitCall(OpCodes.Call, GetMethodInfo<CbWriter>(x => x.EndObject()), null);
-
-					Generator.MarkLabel(SkipLabel);
-					Generator.Emit(OpCodes.Ret);
-				}
-
-				// Add it into the map
-				WriteObjectPropertyMethods.TryAdd(Type, Method);
-			}
-			return Method;
-		}
-
-		static MethodInfo CreateObjectContentsWriter(Type Type)
-		{
-			MethodInfo? Method;
-			if (!WriteObjectContentsMethods.TryGetValue(Type, out Method))
+			/// <summary>
+			/// Create a typed converter
+			/// </summary>
+			/// <returns></returns>
+			static CbConverter<T> CreateConverter()
 			{
-				Dictionary<Type, MethodInfo> NewWritePropertyMethods = new Dictionary<Type, MethodInfo>();
-				Dictionary<Type, MethodInfo> NewWriteContentsMethods = new Dictionary<Type, MethodInfo>();
-				Method = CreateObjectContentsWriter(Type, NewWritePropertyMethods, NewWriteContentsMethods);
+				CbConverter Converter = GetConverter(typeof(T));
+				return (Converter as CbConverter<T>) ?? new CbConverterWrapper(Converter);
 			}
-			return Method;
-		}
 
-		static MethodInfo CreateObjectContentsWriter(Type Type, Dictionary<Type, MethodInfo> NewWritePropertyMethods, Dictionary<Type, MethodInfo> NewWriteContentsMethods)
-		{
-			MethodInfo? Method;
-			if (!WriteObjectContentsMethods.TryGetValue(Type, out Method) && !NewWriteContentsMethods.TryGetValue(Type, out Method))
+			/// <summary>
+			/// Wrapper class to convert an untyped converter into a typed one
+			/// </summary>
+			class CbConverterWrapper : CbConverter<T>
 			{
-				// Create the method
-				DynamicMethod DynamicMethod = new DynamicMethod("_", null, new Type[] { typeof(CbWriter), typeof(object) });
-				Method = DynamicMethod;
-				NewWriteContentsMethods[Type] = DynamicMethod;
-
-				// Find the reflected properties from this type
-				(Utf8String Name, PropertyInfo Property)[] Properties = GetProperties(Type);
-
-				// Create a static type with the required reflection data
-				Type ReflectedType = typeof(CbReflectedTypeInfo<>).MakeGenericType(Type);
-				FieldInfo NamesField = ReflectedType.GetField(nameof(CbReflectedTypeInfo<object>.Names))!;
-				NamesField.SetValue(null, Properties.Select(x => x.Name).ToArray());
-
-				// Generate code for the method
-				ILGenerator Generator = DynamicMethod.GetILGenerator();
-				for (int Idx = 0; Idx < Properties.Length; Idx++)
-				{
-					PropertyInfo Property = Properties[Idx].Property;
-					Type PropertyType = Property.PropertyType;
+				CbConverter Inner;
 
-					Generator.Emit(OpCodes.Ldarg_0);
+				public CbConverterWrapper(CbConverter Inner) => this.Inner = Inner;
 
-					Generator.Emit(OpCodes.Ldsfld, NamesField);
-					Generator.Emit(OpCodes.Ldc_I4, Idx);
-					Generator.Emit(OpCodes.Ldelem, typeof(Utf8String));
+				/// <inheritdoc/>
+				public override T Read(CbField Field) => (T)Inner.ReadObject(Field)!;
 
-					Generator.Emit(OpCodes.Ldarg_1);
-					Generator.EmitCall(OpCodes.Call, Property.GetMethod!, null);
-
-					MethodInfo? WriteMethod;
-					if (!WritePropertyMethods.TryGetValue(PropertyType, out WriteMethod))
-					{
-						if (PropertyType.IsClass)
-						{
-							WriteMethod = CreateObjectPropertyWriter(PropertyType, NewWritePropertyMethods, NewWriteContentsMethods);
-						}
-						else
-						{
-							throw new CbException($"Unable to serialize type {PropertyType.Name}");
-						}
-					}
-					Generator.EmitCall(OpCodes.Call, WriteMethod, null);
-				}
-				Generator.Emit(OpCodes.Ret);
-
-				// Add the method to the cache
-				WriteObjectContentsMethods.TryAdd(Type, DynamicMethod);
-				WriteObjectDelegates.TryAdd(Type, (WriteObjectDelegate)DynamicMethod.CreateDelegate(typeof(WriteObjectDelegate)));
-			}
-			return Method;
-		}
+				/// <inheritdoc/>
+				public override void Write(CbWriter Writer, T Value) => Inner.WriteObject(Writer, Value);
 
-		static MethodInfo CreateObjectReader(Type Type)
-		{
-			MethodInfo? Method;
-			if(!ReadObjectContentsMethods.TryGetValue(Type, out Method))
-			{
-				Dictionary<Type, MethodInfo> NewReadContentsMethods = new Dictionary<Type, MethodInfo>();
-				Method = CreateObjectReader(Type, NewReadContentsMethods);
+				/// <inheritdoc/>
+				public override void WriteNamed(CbWriter Writer, Utf8String Name, T Value) => Inner.WriteNamedObject(Writer, Name, Value);
 			}
-			return Method;
 		}
 
-		static MethodInfo CreateObjectReader(Type Type, Dictionary<Type, MethodInfo> NewReadMethods)
+		/// <summary>
+		/// Gets the converter for a particular type
+		/// </summary>
+		/// <param name="Type"></param>
+		/// <returns></returns>
+		public static CbConverter GetConverter(Type Type)
 		{
-			MethodInfo? Method;
-			if (!ReadObjectContentsMethods.TryGetValue(Type, out Method) && !NewReadMethods.TryGetValue(Type, out Method))
+			CbConverter? Converter;
+			lock (LockObject)
 			{
-				if (TryGetInterfaceType(Type, typeof(IList<>), out Type? EnumerableType))
+				if (!TypeToConverter.TryGetValue(Type, out Converter))
 				{
-					Type ReflectedType = typeof(CbReflectedTypeInfo<>).MakeGenericType(EnumerableType.GenericTypeArguments[0]);
-					Method = ReflectedType.GetMethod(nameof(CbReflectedTypeInfo<object>.ReadList))!;
-				}
-				else
-				{
-					// Create the method
-					DynamicMethod DynamicMethod = new DynamicMethod("_", typeof(object), new Type[] { typeof(CbField) });
-					Method = DynamicMethod;
-					NewReadMethods[Type] = DynamicMethod;
-
-					// Generate the IL for it
-					ILGenerator Generator = DynamicMethod.GetILGenerator();
-
-					// Construct the object
-					ConstructorInfo? Constructor = Type.GetConstructor(BindingFlags.Public | BindingFlags.NonPublic | BindingFlags.Instance, null, Type.EmptyTypes, null);
-					if (Constructor == null)
+					CbConverterAttribute? ConverterAttribute = Type.GetCustomAttribute<CbConverterAttribute>();
+					if (ConverterAttribute != null)
 					{
-						throw new CbException($"Unable to find default constructor for {Type}");
+						Converter = (CbConverter?)Activator.CreateInstance(ConverterAttribute.ConverterType)!;
 					}
-
-					// Find the reflected properties from this type
-					(Utf8String Name, PropertyInfo Property)[] Properties = GetProperties(Type);
-
-					// Create a static type with the required reflection data
-					Type ReflectedType = typeof(CbReflectedTypeInfo<>).MakeGenericType(Type);
-					FieldInfo NamesField = ReflectedType.GetField(nameof(CbReflectedTypeInfo<object>.Names))!;
-					NamesField.SetValue(null, Properties.Select(x => x.Name).ToArray());
-					MethodInfo MatchNameMethod = ReflectedType.GetMethod(nameof(CbReflectedTypeInfo<object>.MatchName))!;
-
-					// NewObjectLocal = new Type()
-					LocalBuilder NewObjectLocal = Generator.DeclareLocal(typeof(object));
-					Generator.Emit(OpCodes.Newobj, Constructor);
-					Generator.Emit(OpCodes.Stloc, NewObjectLocal);
-
-					// Stack(0) = CbField.CreateIterator()
-					Generator.Emit(OpCodes.Ldarg_0);
-					Generator.EmitCall(OpCodes.Call, GetMethodInfo<CbField>(x => x.CreateIterator()), null);
-
-					// CbFieldIterator IteratorLocal = Stack(0)
-					LocalBuilder IteratorLocal = Generator.DeclareLocal(typeof(CbFieldIterator));
-					Generator.Emit(OpCodes.Dup);
-					Generator.Emit(OpCodes.Stloc, IteratorLocal);
-
-					// if(!Stack.Pop().IsValid()) goto ReturnLabel
-					Generator.EmitCall(OpCodes.Call, GetMethodInfo<CbFieldIterator>(x => x.IsValid()), null);
-					Label ReturnLabel = Generator.DefineLabel();
-					Generator.Emit(OpCodes.Brfalse, ReturnLabel);
-
-					// NamesLocal = CbReflectedTypeInfo<Type>.Names
-					LocalBuilder NamesLocal = Generator.DeclareLocal(typeof(Utf8String[]));
-					Generator.Emit(OpCodes.Ldsfld, NamesField);
-					Generator.Emit(OpCodes.Stloc, NamesLocal);
-
-					// IterationLoopLabel:
-					Label IterationLoopLabel = Generator.DefineLabel();
-					Generator.MarkLabel(IterationLoopLabel);
-
-					// bool MatchLocal = false
-					LocalBuilder MatchLocal = Generator.DeclareLocal(typeof(bool));
-					Generator.Emit(OpCodes.Ldc_I4_0);
-					Generator.Emit(OpCodes.Stloc, MatchLocal);
-
-					// Stack(0) = IteratorLocal.GetCurrent()
-					Generator.Emit(OpCodes.Ldloc, IteratorLocal);
-					Generator.EmitCall(OpCodes.Call, GetMethodInfo<CbFieldIterator>(x => x.GetCurrent()), null);
-
-					// Try to parse each of the properties in order. If fields are ordered correctly, we will parse the object in a single pass. Otherwise we can loop and start again.
-					LocalBuilder FieldLocal = Generator.DeclareLocal(typeof(CbField));
-					for (int Idx = 0; Idx < Properties.Length; Idx++)
+					else
 					{
-						PropertyInfo Property = Properties[Idx].Property;
-
-						MethodInfo? ReadMethod;
-						if (!ReadPropertyMethods.TryGetValue(Property.PropertyType, out ReadMethod))
+						for (int Idx = ConverterFactories.Count - 1; Idx >= 0 && Converter == null; Idx--)
 						{
-							if (Property.PropertyType.IsClass)
-							{
-								ReadMethod = CreateObjectReader(Property.PropertyType, NewReadMethods);
-							}
-							else
-							{
-								throw new CbException($"Unable to serialize type {Property.PropertyType.Name}");
-							}
+							Converter = ConverterFactories[Idx].CreateConverter(Type);
 						}
 
-						// if(!CbReflectedTypeInfo<Type>.MatchName(Stack(0), Idx)) goto SkipPropertyLabel
-						Label SkipPropertyLabel = Generator.DefineLabel();
-						Generator.Emit(OpCodes.Dup); // Current CbField
-						Generator.Emit(OpCodes.Ldc_I4, Idx);
-						Generator.Emit(OpCodes.Call, MatchNameMethod);
-						Generator.Emit(OpCodes.Brfalse, SkipPropertyLabel);
-
-						// FieldLocal = Stack.Pop()
-						Generator.Emit(OpCodes.Stloc, FieldLocal);
-
-						// Property.SetMethod(NewObjectLocal, ReadMethod(FieldLocal))
-						Generator.Emit(OpCodes.Ldloc, NewObjectLocal);
-						Generator.Emit(OpCodes.Ldloc, FieldLocal);
-						Generator.EmitCall(OpCodes.Call, ReadMethod, null);
-						Generator.EmitCall(OpCodes.Call, Property.SetMethod!, null);
-
-						// if(!IteratorLocal.MoveNext()) goto ReturnLabel
-						Generator.Emit(OpCodes.Ldloc, IteratorLocal);
-						Generator.EmitCall(OpCodes.Call, GetMethodInfo<CbFieldIterator>(x => x.MoveNext()), null);
-						Generator.Emit(OpCodes.Brfalse, ReturnLabel);
-
-						// MatchLocal = true
-						Generator.Emit(OpCodes.Ldc_I4_1);
-						Generator.Emit(OpCodes.Stloc, MatchLocal);
-
-						// Stack(0) = IteratorLocal.GetCurrent()
-						Generator.Emit(OpCodes.Ldloc, IteratorLocal);
-						Generator.EmitCall(OpCodes.Call, GetMethodInfo<CbFieldIterator>(x => x.GetCurrent()), null);
-
-						// SkipPropertyLabel:
-						Generator.MarkLabel(SkipPropertyLabel);
+						if (Converter == null)
+						{
+							throw new CbException($"Unable to create converter for {Type.Name}");
+						}
 					}
-
-					// Stack.Pop()
-					Generator.Emit(OpCodes.Pop); // Current CbField
-
-					// if(MatchLocal) goto IterationLoopLabel
-					Generator.Emit(OpCodes.Ldloc, MatchLocal);
-					Generator.Emit(OpCodes.Brtrue, IterationLoopLabel);
-
-					// if(IteratorLocal.MoveNext()) goto IterationLoopLabel
-					Generator.Emit(OpCodes.Ldloc, IteratorLocal);
-					Generator.EmitCall(OpCodes.Call, GetMethodInfo<CbFieldIterator>(x => x.MoveNext()), null);
-					Generator.Emit(OpCodes.Brtrue, IterationLoopLabel);
-
-					// return NewObjectLocal
-					Generator.MarkLabel(ReturnLabel);
-					Generator.Emit(OpCodes.Ldloc, NewObjectLocal);
-					Generator.Emit(OpCodes.Ret);
-
-					// Add the new methods
-					ReadObjectContentsMethods.TryAdd(Type, Method);
-					ReadObjectDelegates.TryAdd(Type, (ReadObjectDelegate)DynamicMethod.CreateDelegate(typeof(ReadObjectDelegate)));
-				}
-			}
-			return Method;
-		}
-
-		static MethodInfo GetMethodInfo(Expression<Action> Expr)
-		{
-			return ((MethodCallExpression)Expr.Body).Method;
-		}
-
-		static MethodInfo GetMethodInfo<T>(Expression<Action<T>> Expr)
-		{
-			return ((MethodCallExpression)Expr.Body).Method;
-		}
-
-		static (Utf8String, PropertyInfo)[] GetProperties(Type Type)
-		{
-			List<(Utf8String, PropertyInfo)> PropertyList = new List<(Utf8String, PropertyInfo)>();
-			foreach (PropertyInfo Property in Type.GetProperties(BindingFlags.Public | BindingFlags.Instance))
-			{
-				CbFieldAttribute? Attribute = Property.GetCustomAttribute<CbFieldAttribute>();
-				if (Attribute != null)
-				{
-					Utf8String Name = Attribute.Name ?? Property.Name;
-					PropertyList.Add((Name, Property));
+					TypeToConverter.Add(Type, Converter!);
 				}
 			}
-			if (PropertyList.Count == 0)
-			{
-				throw new CbException($"{Type.Name} does not have any properties marked with [CbField]");
-			}
-			return PropertyList.ToArray();
+			return Converter;
 		}
 
-		static bool TryGetInterfaceType(Type Type, Type GenericType, [NotNullWhen(true)] out Type? OutCollectionType)
+		/// <summary>
+		/// Gets the converter for a given type
+		/// </summary>
+		/// <typeparam name="T">Type to retrieve the converter for</typeparam>
+		/// <returns></returns>
+		public static CbConverter<T> GetConverter<T>()
 		{
-			foreach (Type InterfaceType in Type.GetInterfaces())
-			{
-				if (InterfaceType.GetGenericTypeDefinition() == GenericType)
-				{
-					OutCollectionType = InterfaceType;
-					return true;
-				}
-			}
-			OutCollectionType = null;
-			return false;
+			return CbConverterCache<T>.Instance;
 		}
 
-		static void WriteInt32Value(CbWriter Writer, int Value)
+		/// <summary>
+		/// Serialize an object
+		/// </summary>
+		/// <typeparam name="T"></typeparam>
+		/// <param name="Value"></param>
+		/// <returns></returns>
+		public static CbObject Serialize<T>(T Value)
 		{
-			if (Value != 0)
-			{
-				Writer.WriteIntegerValue(Value);
-			}
+			CbWriter Writer = new CbWriter();
+			GetConverter<T>().Write(Writer, Value);
+			return Writer.ToObject();
 		}
 
-		static void WriteInt32(CbWriter Writer, Utf8String Name, int Value)
+		/// <summary>
+		/// Serialize a property to a given writer
+		/// </summary>
+		/// <typeparam name="T"></typeparam>
+		/// <param name="Writer"></param>
+		/// <param name="Value"></param>
+		public static void Serialize<T>(CbWriter Writer, T Value)
 		{
-			if (Value != 0)
-			{
-				Writer.WriteInteger(Name, Value);
-			}
+			GetConverter<T>().Write(Writer, Value);
 		}
 
-		static void WriteStringValue(CbWriter Writer, Utf8String Value)
+		/// <summary>
+		/// Serialize a named property to the given writer
+		/// </summary>
+		/// <typeparam name="T"></typeparam>
+		/// <param name="Writer"></param>
+		/// <param name="Name"></param>
+		/// <param name="Value"></param>
+		public static void Serialize<T>(CbWriter Writer, Utf8String Name, T Value)
 		{
-			if (Value.Length > 0)
-			{
-				Writer.WriteStringValue(Value);
-			}
+			GetConverter<T>().WriteNamed(Writer, Name, Value);
 		}
 
-		static void WriteString(CbWriter Writer, Utf8String Name, Utf8String Value)
+		/// <summary>
+		/// Deserialize an object from a <see cref="CbObject"/>
+		/// </summary>
+		/// <typeparam name="T"></typeparam>
+		/// <param name="Field"></param>
+		/// <returns></returns>
+		public static T Deserialize<T>(CbField Field)
 		{
-			if (Value.Length > 0)
-			{
-				Writer.WriteString(Name, Value);
-			}
+			return GetConverter<T>().Read(Field);
 		}
 	}
 }
diff --git a/Engine/Source/Programs/Shared/EpicGames.Serialization/Converters/CbConverter.cs b/Engine/Source/Programs/Shared/EpicGames.Serialization/Converters/CbConverter.cs
new file mode 100644
--- /dev/null
+++ b/Engine/Source/Programs/Shared/EpicGames.Serialization/Converters/CbConverter.cs
@@ -0,0 +1,108 @@
+// Copyright Epic Games, Inc. All Rights Reserved.
+
+using EpicGames.Core;
+using System;
+using System.Collections.Generic;
+using System.Text;
+
+namespace EpicGames.Serialization.Converters
+{
+	/// <summary>
+	/// Attribute declaring the converter to use for a type
+	/// </summary>
+	[AttributeUsage(AttributeTargets.Class)]
+	public class CbConverterAttribute : Attribute
+	{
+		/// <summary>
+		/// The converter type
+		/// </summary>
+		public Type ConverterType { get; }
+
+		/// <summary>
+		/// Constructor
+		/// </summary>
+		/// <param name="ConverterType"></param>
+		public CbConverterAttribute(Type ConverterType)
+		{
+			this.ConverterType = ConverterType;
+		}
+	}
+
+	/// <summary>
+	/// Base class for all converters. Deriving from <see cref="CbConverter{T}"/> is more efficient
+	/// </summary>
+	public abstract class CbConverter
+	{
+		/// <summary>
+		/// Reads an object from a field
+		/// </summary>
+		/// <param name="Field"></param>
+		/// <returns></returns>
+		public abstract object? ReadObject(CbField Field);
+
+		/// <summary>
+		/// Writes an object to a field
+		/// </summary>
+		/// <param name="Writer"></param>
+		/// <param name="Value"></param>
+		public abstract void WriteObject(CbWriter Writer, object? Value);
+
+		/// <summary>
+		/// Writes an object to a named field, if not equal to the default value
+		/// </summary>
+		/// <param name="Writer"></param>
+		/// <param name="Name"></param>
+		/// <param name="Value"></param>
+		public abstract void WriteNamedObject(CbWriter Writer, Utf8String Name, object? Value);
+	}
+
+	/// <summary>
+	/// Converter for a particular type
+	/// </summary>
+	public abstract class CbConverter<T> : CbConverter
+	{
+		/// <summary>
+		/// Reads an object from a field
+		/// </summary>
+		/// <param name="Field"></param>
+		/// <returns></returns>
+		public abstract T Read(CbField Field);
+
+		/// <summary>
+		/// Writes an object to a field
+		/// </summary>
+		/// <param name="Writer"></param>
+		/// <param name="Value"></param>
+		public abstract void Write(CbWriter Writer, T Value);
+
+		/// <summary>
+		/// Writes an object to a named field, if not equal to the default value
+		/// </summary>
+		/// <param name="Writer"></param>
+		/// <param name="Name"></param>
+		/// <param name="Value"></param>
+		public abstract void WriteNamed(CbWriter Writer, Utf8String Name, T Value);
+
+		/// <inheritdoc/>
+		public sealed override object? ReadObject(CbField Field) => Read(Field);
+
+		/// <inheritdoc/>
+		public sealed override void WriteObject(CbWriter Writer, object? Value) => Write(Writer, (T)Value!);
+
+		/// <inheritdoc/>
+		public sealed override void WriteNamedObject(CbWriter Writer, Utf8String Name, object? Value) => WriteNamed(Writer, Name, (T)Value!);
+	}
+
+	/// <summary>
+	/// Base class for converter factories
+	/// </summary>
+	public abstract class CbConverterFactory
+	{
+		/// <summary>
+		/// Create a converter for the given type
+		/// </summary>
+		/// <param name="Type">The type to create a converter for</param>
+		/// <returns>The converter instance, or null if this factory does not support the given type</returns>
+		public abstract CbConverter? CreateConverter(Type Type);
+	}
+}
diff --git a/Engine/Source/Programs/Shared/EpicGames.Serialization/Converters/CbDefaultConverter.cs b/Engine/Source/Programs/Shared/EpicGames.Serialization/Converters/CbDefaultConverter.cs
new file mode 100644
--- /dev/null
+++ b/Engine/Source/Programs/Shared/EpicGames.Serialization/Converters/CbDefaultConverter.cs
@@ -0,0 +1,481 @@
+// Copyright Epic Games, Inc. All Rights Reserved.
+
+using EpicGames.Core;
+using System;
+using System.Collections.Generic;
+using System.Linq;
+using System.Linq.Expressions;
+using System.Reflection;
+using System.Reflection.Emit;
+using System.Text;
+
+namespace EpicGames.Serialization.Converters
+{
+	/// <summary>
+	/// Default converter implementation. Calls delegates to perform read/write operations.
+	/// </summary>
+	/// <typeparam name="T"></typeparam>
+	class CbDefaultConverter<T> : CbConverter<T>
+	{
+		public delegate T ReadDelegate(CbField Field);
+		public delegate void WriteDelegate(CbWriter Writer, T Value);
+		public delegate void WriteNamedDelegate(CbWriter Writer, Utf8String Name, T Value);
+
+		readonly ReadDelegate ReadFunc;
+		readonly WriteDelegate WriteFunc;
+		readonly WriteNamedDelegate WriteNamedFunc;
+
+		public CbDefaultConverter(ReadDelegate ReadFunc, WriteDelegate WriteFunc, WriteNamedDelegate WriteNamedFunc)
+		{
+			this.ReadFunc = ReadFunc;
+			this.WriteFunc = WriteFunc;
+			this.WriteNamedFunc = WriteNamedFunc;
+		}
+
+		public CbDefaultConverter(Expression<Func<CbField, T>> Read, Expression<Action<CbWriter, T>> Write, Expression<Action<CbWriter, Utf8String, T>> WriteNamed)
+		{
+			Func<CbField, T> ReadFunc = Read.Compile();
+			this.ReadFunc = f => ReadFunc(f);
+
+			Action<CbWriter, T> WriteFunc = Write.Compile();
+			this.WriteFunc = (w, t) => WriteFunc(w, t);
+
+			Action<CbWriter, Utf8String, T> WriteNamedFunc = WriteNamed.Compile();
+			this.WriteNamedFunc = (w, n, t) => WriteNamedFunc(w, n, t);
+		}
+
+		public CbDefaultConverter(DynamicMethod ReadMethod, DynamicMethod WriteMethod, DynamicMethod WriteNamedMethod)
+		{
+			this.ReadFunc = (ReadDelegate)ReadMethod.CreateDelegate(typeof(ReadDelegate));
+			this.WriteFunc = (WriteDelegate)WriteMethod.CreateDelegate(typeof(WriteDelegate));
+			this.WriteNamedFunc = (WriteNamedDelegate)WriteNamedMethod.CreateDelegate(typeof(WriteNamedDelegate));
+		}
+
+		public override T Read(CbField Field) => ReadFunc(Field);
+		public override void Write(CbWriter Writer, T Value) => WriteFunc(Writer, Value);
+		public override void WriteNamed(CbWriter Writer, Utf8String Name, T Value) => WriteNamedFunc(Writer, Name, Value);
+	}
+
+	/// <summary>
+	/// Converter which generates dynamic serialization methods for classes based on attribute markup
+	/// </summary>
+	class CbDefaultConverterFactory : CbConverterFactory
+	{
+		static object LockObject = new object();
+
+		delegate void SkipDefaultDelegate(ILGenerator Generator, Label SkipLabel);
+
+		class CbConverterInfo
+		{
+			public CbConverter Converter { get; }
+			public MethodInfo ReadMethod { get; }
+			public MethodInfo WriteMethod { get; }
+			public MethodInfo WriteNamedMethod { get; }
+			public SkipDefaultDelegate? SkipDefault { get; }
+
+			public CbConverterInfo(CbConverter Converter, MethodInfo ReadMethod, MethodInfo WriteMethod, MethodInfo WriteNamedMethod, SkipDefaultDelegate? SkipDefault)
+			{
+				this.Converter = Converter;
+				this.ReadMethod = ReadMethod;
+				this.WriteMethod = WriteMethod;
+				this.WriteNamedMethod = WriteNamedMethod;
+				this.SkipDefault = SkipDefault;
+			}
+		}
+
+		// Implementation of CbConverterInfo for types that already have a registered converter
+		class CbCustomConverterInfo<T> : CbConverterInfo
+		{
+			static CbConverter<T> StaticConverter = CbSerializer.GetConverter<T>();
+
+			static T Read(CbField Field) => StaticConverter.Read(Field);
+			static void Write(CbWriter Writer, T Value) => StaticConverter.Write(Writer, Value);
+			static void WriteNamed(CbWriter Writer, Utf8String Name, T Value) => StaticConverter.WriteNamed(Writer, Name, Value);
+
+			public CbCustomConverterInfo()
+				: base(StaticConverter, GetMethodInfo<CbField>(x => Read(x)), GetMethodInfo(() => Write(default!, default!)), GetMethodInfo(() => WriteNamed(default!, default!, default!)), null)
+			{
+			}
+		}
+
+		static Dictionary<Type, CbConverterInfo> TypeToConverterInfo = new Dictionary<Type, CbConverterInfo>(new KeyValuePair<Type, CbConverterInfo>[]
+		{
+			GetPodConverterInfo(x => x.AsString(), (w, v) => w.WriteStringValue(v), (w, n, v) => w.WriteString(n, v), null),
+			GetPodConverterInfo(x => x.AsHash(), (w, v) => w.WriteHashValue(v), (w, n, v) => w.WriteHash(n, v), null)
+			// TODO: other basic types
+		});
+
+		static KeyValuePair<Type, CbConverterInfo> GetPodConverterInfo<T>(Expression<Func<CbField, T>> Read, Expression<Action<CbWriter, T>> Write, Expression<Action<CbWriter, Utf8String, T>> WriteNamed, SkipDefaultDelegate? SkipDefault)
+		{
+			CbDefaultConverter<T> Converter = new CbDefaultConverter<T>(Read, Write, WriteNamed);
+
+			MethodInfo ReadMethod = ((MethodCallExpression)Read.Body).Method;
+			MethodInfo WriteMethod = ((MethodCallExpression)Write.Body).Method;
+			MethodInfo WriteNamedMethod = ((MethodCallExpression)WriteNamed.Body).Method;
+
+			return new KeyValuePair<Type, CbConverterInfo>(typeof(T), new CbConverterInfo(Converter, ReadMethod, WriteMethod, WriteNamedMethod, SkipDefault));
+		}
+
+		static CbConverterInfo GetConverterInfo(Type Type)
+		{
+			lock (LockObject)
+			{
+				CbConverterInfo? ConverterInfo;
+				if (!TypeToConverterInfo.TryGetValue(Type, out ConverterInfo))
+				{
+					CbConverter Converter = CbSerializer.GetConverter(Type);
+					if (!TypeToConverterInfo.TryGetValue(Type, out ConverterInfo)) // CreateConverter above may have added it
+					{
+						Type GenericType = typeof(CbCustomConverterInfo<>).MakeGenericType(Type);
+						ConverterInfo = (CbConverterInfo)Activator.CreateInstance(GenericType)!;
+						TypeToConverterInfo.Add(Type, ConverterInfo);
+					}
+				}
+				return ConverterInfo;
+			}
+		}
+
+		public override CbConverter CreateConverter(Type Type)
+		{
+			lock (LockObject)
+			{
+				CbConverterInfo? ConverterInfo;
+				if (!TypeToConverterInfo.TryGetValue(Type, out ConverterInfo))
+				{
+					DynamicMethod ReadMethod = CreateObjectReader(Type);
+					DynamicMethod WriteMethod = CreateObjectWriter(Type);
+					DynamicMethod WriteNamedMethod = CreateNamedObjectWriter(Type);
+
+					Type GenericType = typeof(CbDefaultConverter<>).MakeGenericType(Type);
+					CbConverter Converter = (CbConverter)Activator.CreateInstance(GenericType, ReadMethod, WriteMethod, WriteNamedMethod)!;
+
+					ConverterInfo = new CbConverterInfo(Converter, ReadMethod, WriteMethod, WriteNamedMethod, SkipIfNullOrZero);
+					TypeToConverterInfo.Add(Type, ConverterInfo);
+				}
+				return ConverterInfo.Converter;
+			}
+		}
+
+		static void SkipIfNullOrZero(ILGenerator Generator, Label SkipLabel)
+		{
+			Generator.Emit(OpCodes.Dup);
+			Generator.Emit(OpCodes.Brfalse, SkipLabel);
+		}
+
+		class CbReflectedTypeInfo<T> where T : class
+		{
+			public static Utf8String[]? Names = null;
+			public static PropertyInfo[]? Properties = null;
+
+			public static bool MatchName(CbField Field, int Idx)
+			{
+				return Field.Name == Names![Idx];
+			}
+		}
+
+		static Dictionary<Type, DynamicMethod> WriteMethods = new Dictionary<Type, DynamicMethod>();
+
+		static DynamicMethod CreateObjectWriter(Type Type)
+		{
+			DynamicMethod? DynamicMethod;
+			if (!WriteMethods.TryGetValue(Type, out DynamicMethod))
+			{
+				// Create the new method
+				DynamicMethod = new DynamicMethod("_", null, new Type[] { typeof(CbWriter), Type });
+				WriteMethods.Add(Type, DynamicMethod);
+
+				// Implement the body
+				ILGenerator Generator = DynamicMethod.GetILGenerator();
+				Generator.Emit(OpCodes.Ldarg_1);
+
+				Label SkipLabel = Generator.DefineLabel();
+				Generator.Emit(OpCodes.Brfalse, SkipLabel);
+
+				Generator.Emit(OpCodes.Ldarg_0);
+				Generator.EmitCall(OpCodes.Call, GetMethodInfo<CbWriter>(x => x.BeginObject()), null);
+
+				Generator.Emit(OpCodes.Ldarg_0);
+				Generator.Emit(OpCodes.Ldarg_1);
+				Generator.EmitCall(OpCodes.Call, CreateObjectContentsWriter(Type), null);
+
+				Generator.Emit(OpCodes.Ldarg_0);
+				Generator.EmitCall(OpCodes.Call, GetMethodInfo<CbWriter>(x => x.EndObject()), null);
+
+				Generator.MarkLabel(SkipLabel);
+				Generator.Emit(OpCodes.Ret);
+			}
+			return DynamicMethod;
+		}
+
+		static Dictionary<Type, DynamicMethod> WriteNamedMethods = new Dictionary<Type, DynamicMethod>();
+
+		static DynamicMethod CreateNamedObjectWriter(Type Type)
+		{
+			DynamicMethod? DynamicMethod;
+			if (!WriteNamedMethods.TryGetValue(Type, out DynamicMethod))
+			{
+				// Create the new method
+				DynamicMethod = new DynamicMethod("_", null, new Type[] { typeof(CbWriter), typeof(Utf8String), Type });
+				WriteNamedMethods.Add(Type, DynamicMethod);
+
+				// Implement the body
+				ILGenerator Generator = DynamicMethod.GetILGenerator();
+				Generator.Emit(OpCodes.Ldarg_2);
+
+				Label SkipLabel = Generator.DefineLabel();
+				Generator.Emit(OpCodes.Brfalse, SkipLabel);
+
+				Generator.Emit(OpCodes.Ldarg_0);
+				Generator.Emit(OpCodes.Ldarg_1);
+				Generator.EmitCall(OpCodes.Call, GetMethodInfo<CbWriter>(x => x.BeginObject(null!)), null);
+
+				Generator.Emit(OpCodes.Ldarg_0);
+				Generator.Emit(OpCodes.Ldarg_2);
+				Generator.EmitCall(OpCodes.Call, CreateObjectContentsWriter(Type), null);
+
+				Generator.Emit(OpCodes.Ldarg_0);
+				Generator.EmitCall(OpCodes.Call, GetMethodInfo<CbWriter>(x => x.EndObject()), null);
+
+				Generator.MarkLabel(SkipLabel);
+				Generator.Emit(OpCodes.Ret);
+			}
+			return DynamicMethod;
+		}
+
+		static Dictionary<Type, DynamicMethod> WriteObjectProperties = new Dictionary<Type, DynamicMethod>();
+
+		static DynamicMethod CreateObjectContentsWriter(Type Type)
+		{
+			DynamicMethod? DynamicMethod;
+			if (!WriteObjectProperties.TryGetValue(Type, out DynamicMethod))
+			{
+				// Create the method
+				DynamicMethod = new DynamicMethod("_", null, new Type[] { typeof(CbWriter), Type });
+				WriteObjectProperties.Add(Type, DynamicMethod);
+
+				// Find the reflected properties from this type
+				(Utf8String Name, PropertyInfo Property)[] Properties = GetProperties(Type);
+
+				// Create a static type with the required reflection data
+				Type ReflectedType = typeof(CbReflectedTypeInfo<>).MakeGenericType(Type);
+				FieldInfo NamesField = ReflectedType.GetField(nameof(CbReflectedTypeInfo<object>.Names))!;
+				NamesField.SetValue(null, Properties.Select(x => x.Name).ToArray());
+
+				// Generate code for the method
+				ILGenerator Generator = DynamicMethod.GetILGenerator();
+				for (int Idx = 0; Idx < Properties.Length; Idx++)
+				{
+					PropertyInfo Property = Properties[Idx].Property;
+					Type PropertyType = Property.PropertyType;
+
+					// Get the field value
+					Generator.Emit(OpCodes.Ldarg_1);
+					Generator.EmitCall(OpCodes.Call, Property.GetMethod!, null);
+
+					Label SkipLabel = Generator.DefineLabel();
+
+					MethodInfo WriteMethod;
+					if (WriteNamedMethods.TryGetValue(PropertyType, out DynamicMethod? DynamicWriteMethod))
+					{
+						Generator.Emit(OpCodes.Dup);
+						Generator.Emit(OpCodes.Brfalse, SkipLabel);
+						WriteMethod = DynamicWriteMethod;
+					}
+					else
+					{
+						CbConverterInfo ConverterInfo = GetConverterInfo(PropertyType);
+						ConverterInfo.SkipDefault?.Invoke(Generator, SkipLabel);
+						WriteMethod = ConverterInfo.WriteNamedMethod;
+					}
+
+					// Store the variable in a local
+					LocalBuilder Local = Generator.DeclareLocal(PropertyType);
+					Generator.Emit(OpCodes.Dup);
+					Generator.Emit(OpCodes.Stloc, Local);
+
+					// Call the writer
+					Generator.Emit(OpCodes.Ldarg_0);
+
+					Generator.Emit(OpCodes.Ldsfld, NamesField);
+					Generator.Emit(OpCodes.Ldc_I4, Idx);
+					Generator.Emit(OpCodes.Ldelem, typeof(Utf8String));
+
+					Generator.Emit(OpCodes.Ldloc, Local);
+					Generator.EmitCall(OpCodes.Call, WriteMethod, null);
+
+					// Remove the duplicated value from the top of the stack
+					Generator.MarkLabel(SkipLabel);
+					Generator.Emit(OpCodes.Pop);
+				}
+				Generator.Emit(OpCodes.Ret);
+			}
+			return DynamicMethod;
+		}
+
+		static Dictionary<Type, DynamicMethod> ReadObjectMethods = new Dictionary<Type, DynamicMethod>();
+
+		static DynamicMethod CreateObjectReader(Type Type)
+		{
+			DynamicMethod? DynamicMethod;
+			if (!ReadObjectMethods.TryGetValue(Type, out DynamicMethod))
+			{
+				// Create the method
+				DynamicMethod = new DynamicMethod("_", Type, new Type[] { typeof(CbField) });
+				ReadObjectMethods.Add(Type, DynamicMethod);
+
+				// Generate the IL for it
+				ILGenerator Generator = DynamicMethod.GetILGenerator();
+
+				// Construct the object
+				ConstructorInfo? Constructor = Type.GetConstructor(BindingFlags.Public | BindingFlags.NonPublic | BindingFlags.Instance, null, Type.EmptyTypes, null);
+				if (Constructor == null)
+				{
+					throw new CbException($"Unable to find default constructor for {Type}");
+				}
+
+				// Find the reflected properties from this type
+				(Utf8String Name, PropertyInfo Property)[] Properties = GetProperties(Type);
+
+				// Create a static type with the required reflection data
+				Type ReflectedType = typeof(CbReflectedTypeInfo<>).MakeGenericType(Type);
+				FieldInfo NamesField = ReflectedType.GetField(nameof(CbReflectedTypeInfo<object>.Names))!;
+				NamesField.SetValue(null, Properties.Select(x => x.Name).ToArray());
+				MethodInfo MatchNameMethod = ReflectedType.GetMethod(nameof(CbReflectedTypeInfo<object>.MatchName))!;
+
+				// NewObjectLocal = new Type()
+				LocalBuilder NewObjectLocal = Generator.DeclareLocal(typeof(object));
+				Generator.Emit(OpCodes.Newobj, Constructor);
+				Generator.Emit(OpCodes.Stloc, NewObjectLocal);
+
+				// Stack(0) = CbField.CreateIterator()
+				Generator.Emit(OpCodes.Ldarg_0);
+				Generator.EmitCall(OpCodes.Call, GetMethodInfo<CbField>(x => x.CreateIterator()), null);
+
+				// CbFieldIterator IteratorLocal = Stack(0)
+				LocalBuilder IteratorLocal = Generator.DeclareLocal(typeof(CbFieldIterator));
+				Generator.Emit(OpCodes.Dup);
+				Generator.Emit(OpCodes.Stloc, IteratorLocal);
+
+				// if(!Stack.Pop().IsValid()) goto ReturnLabel
+				Generator.EmitCall(OpCodes.Call, GetMethodInfo<CbFieldIterator>(x => x.IsValid()), null);
+				Label ReturnLabel = Generator.DefineLabel();
+				Generator.Emit(OpCodes.Brfalse, ReturnLabel);
+
+				// NamesLocal = CbReflectedTypeInfo<Type>.Names
+				LocalBuilder NamesLocal = Generator.DeclareLocal(typeof(Utf8String[]));
+				Generator.Emit(OpCodes.Ldsfld, NamesField);
+				Generator.Emit(OpCodes.Stloc, NamesLocal);
+
+				// IterationLoopLabel:
+				Label IterationLoopLabel = Generator.DefineLabel();
+				Generator.MarkLabel(IterationLoopLabel);
+
+				// bool MatchLocal = false
+				LocalBuilder MatchLocal = Generator.DeclareLocal(typeof(bool));
+				Generator.Emit(OpCodes.Ldc_I4_0);
+				Generator.Emit(OpCodes.Stloc, MatchLocal);
+
+				// Stack(0) = IteratorLocal.GetCurrent()
+				Generator.Emit(OpCodes.Ldloc, IteratorLocal);
+				Generator.EmitCall(OpCodes.Call, GetMethodInfo<CbFieldIterator>(x => x.GetCurrent()), null);
+
+				// Try to parse each of the properties in order. If fields are ordered correctly, we will parse the object in a single pass. Otherwise we can loop and start again.
+				LocalBuilder FieldLocal = Generator.DeclareLocal(typeof(CbField));
+				for (int Idx = 0; Idx < Properties.Length; Idx++)
+				{
+					PropertyInfo Property = Properties[Idx].Property;
+
+					// Get the read method for this property type
+					MethodInfo ReadMethod;
+					if (ReadObjectMethods.TryGetValue(Property.PropertyType, out DynamicMethod? DynamicReadMethod))
+					{
+						ReadMethod = DynamicReadMethod;
+					}
+					else
+					{
+						ReadMethod = GetConverterInfo(Property.PropertyType).ReadMethod;
+					}
+
+					// if(!CbReflectedTypeInfo<Type>.MatchName(Stack(0), Idx)) goto SkipPropertyLabel
+					Label SkipPropertyLabel = Generator.DefineLabel();
+					Generator.Emit(OpCodes.Dup); // Current CbField
+					Generator.Emit(OpCodes.Ldc_I4, Idx);
+					Generator.Emit(OpCodes.Call, MatchNameMethod);
+					Generator.Emit(OpCodes.Brfalse, SkipPropertyLabel);
+
+					// FieldLocal = Stack.Pop()
+					Generator.Emit(OpCodes.Stloc, FieldLocal);
+
+					// Property.SetMethod(NewObjectLocal, ReadMethod(FieldLocal))
+					Generator.Emit(OpCodes.Ldloc, NewObjectLocal);
+					Generator.Emit(OpCodes.Ldloc, FieldLocal);
+					Generator.EmitCall(OpCodes.Call, ReadMethod, null);
+					Generator.EmitCall(OpCodes.Call, Property.SetMethod!, null);
+
+					// if(!IteratorLocal.MoveNext()) goto ReturnLabel
+					Generator.Emit(OpCodes.Ldloc, IteratorLocal);
+					Generator.EmitCall(OpCodes.Call, GetMethodInfo<CbFieldIterator>(x => x.MoveNext()), null);
+					Generator.Emit(OpCodes.Brfalse, ReturnLabel);
+
+					// MatchLocal = true
+					Generator.Emit(OpCodes.Ldc_I4_1);
+					Generator.Emit(OpCodes.Stloc, MatchLocal);
+
+					// Stack(0) = IteratorLocal.GetCurrent()
+					Generator.Emit(OpCodes.Ldloc, IteratorLocal);
+					Generator.EmitCall(OpCodes.Call, GetMethodInfo<CbFieldIterator>(x => x.GetCurrent()), null);
+
+					// SkipPropertyLabel:
+					Generator.MarkLabel(SkipPropertyLabel);
+				}
+
+				// Stack.Pop()
+				Generator.Emit(OpCodes.Pop); // Current CbField
+
+				// if(MatchLocal) goto IterationLoopLabel
+				Generator.Emit(OpCodes.Ldloc, MatchLocal);
+				Generator.Emit(OpCodes.Brtrue, IterationLoopLabel);
+
+				// if(IteratorLocal.MoveNext()) goto IterationLoopLabel
+				Generator.Emit(OpCodes.Ldloc, IteratorLocal);
+				Generator.EmitCall(OpCodes.Call, GetMethodInfo<CbFieldIterator>(x => x.MoveNext()), null);
+				Generator.Emit(OpCodes.Brtrue, IterationLoopLabel);
+
+				// return NewObjectLocal
+				Generator.MarkLabel(ReturnLabel);
+				Generator.Emit(OpCodes.Ldloc, NewObjectLocal);
+				Generator.Emit(OpCodes.Ret);
+			}
+			return DynamicMethod;
+		}
+
+		static MethodInfo GetMethodInfo(Expression<Action> Expr)
+		{
+			return ((MethodCallExpression)Expr.Body).Method;
+		}
+
+		static MethodInfo GetMethodInfo<T>(Expression<Action<T>> Expr)
+		{
+			return ((MethodCallExpression)Expr.Body).Method;
+		}
+
+		static (Utf8String, PropertyInfo)[] GetProperties(Type Type)
+		{
+			List<(Utf8String, PropertyInfo)> PropertyList = new List<(Utf8String, PropertyInfo)>();
+			foreach (PropertyInfo Property in Type.GetProperties(BindingFlags.Public | BindingFlags.Instance))
+			{
+				CbFieldAttribute? Attribute = Property.GetCustomAttribute<CbFieldAttribute>();
+				if (Attribute != null)
+				{
+					Utf8String Name = Attribute.Name ?? Property.Name;
+					PropertyList.Add((Name, Property));
+				}
+			}
+			if (PropertyList.Count == 0)
+			{
+				throw new CbException($"{Type.Name} does not have any properties marked with [CbField]");
+			}
+			return PropertyList.ToArray();
+		}
+	}
+}
diff --git a/Engine/Source/Programs/Shared/EpicGames.Serialization/Converters/CbDictionaryConverter.cs b/Engine/Source/Programs/Shared/EpicGames.Serialization/Converters/CbDictionaryConverter.cs
new file mode 100644
--- /dev/null
+++ b/Engine/Source/Programs/Shared/EpicGames.Serialization/Converters/CbDictionaryConverter.cs
@@ -0,0 +1,97 @@
+// Copyright Epic Games, Inc. All Rights Reserved.
+
+using EpicGames.Core;
+using System;
+using System.Collections;
+using System.Collections.Concurrent;
+using System.Collections.Generic;
+using System.Diagnostics.CodeAnalysis;
+using System.Linq;
+using System.Linq.Expressions;
+using System.Reflection;
+using System.Reflection.Emit;
+
+namespace EpicGames.Serialization.Converters
+{
+	/// <summary>
+	/// Converter for dictionary types
+	/// </summary>
+	class CbDictionaryConverter<TKey, TValue> : CbConverter<Dictionary<TKey, TValue>> where TKey : notnull
+	{
+		/// <inheritdoc/>
+		public override Dictionary<TKey, TValue> Read(CbField Field)
+		{
+			Dictionary<TKey, TValue> Dictionary = new Dictionary<TKey, TValue>();
+			foreach (CbField Element in Field)
+			{
+				IEnumerator<CbField> Enumerator = Element.AsArray().GetEnumerator();
+
+				if (!Enumerator.MoveNext())
+				{
+					throw new CbException("Missing key for dictionary entry");
+				}
+				TKey Key = CbSerializer.Deserialize<TKey>(Enumerator.Current);
+
+				if (!Enumerator.MoveNext())
+				{
+					throw new CbException("Missing value for dictionary entry");
+				}
+				TValue Value = CbSerializer.Deserialize<TValue>(Enumerator.Current);
+
+				Dictionary.Add(Key, Value);
+			}
+			return Dictionary;
+		}
+
+		/// <inheritdoc/>
+		public override void Write(CbWriter Writer, Dictionary<TKey, TValue> Value)
+		{
+			if (Value.Count > 0)
+			{
+				Writer.BeginUniformArray(CbFieldType.Array);
+				foreach (KeyValuePair<TKey, TValue> Pair in Value)
+				{
+					Writer.BeginArray();
+					CbSerializer.Serialize(Writer, Pair.Key);
+					CbSerializer.Serialize(Writer, Pair.Value);
+					Writer.EndArray();
+				}
+				Writer.EndUniformArray();
+			}
+		}
+
+		/// <inheritdoc/>
+		public override void WriteNamed(CbWriter Writer, Utf8String Name, Dictionary<TKey, TValue> Value)
+		{
+			if (Value.Count > 0)
+			{
+				Writer.BeginUniformArray(Name, CbFieldType.Array);
+				foreach (KeyValuePair<TKey, TValue> Pair in Value)
+				{
+					Writer.BeginArray();
+					CbSerializer.Serialize(Writer, Pair.Key);
+					CbSerializer.Serialize(Writer, Pair.Value);
+					Writer.EndArray();
+				}
+				Writer.EndUniformArray();
+			}
+		}
+	}
+
+	/// <summary>
+	/// Factory for CbDictionaryConverter
+	/// </summary>
+	class CbDictionaryConverterFactory : CbConverterFactory
+	{
+		/// <inheritdoc/>
+		public override CbConverter? CreateConverter(Type Type)
+		{
+			if (Type.IsGenericType && Type.GetGenericTypeDefinition() == typeof(Dictionary<,>))
+			{
+				Type ConverterType = typeof(CbDictionaryConverter<,>).MakeGenericType(Type.GenericTypeArguments);
+				return (CbConverter)Activator.CreateInstance(ConverterType)!;
+			}
+			return null;
+		}
+	}
+}
diff --git a/Engine/Source/Programs/Shared/EpicGames.Serialization/Converters/CbListConverter.cs b/Engine/Source/Programs/Shared/EpicGames.Serialization/Converters/CbListConverter.cs
new file mode 100644
--- /dev/null
+++ b/Engine/Source/Programs/Shared/EpicGames.Serialization/Converters/CbListConverter.cs
@@ -0,0 +1,124 @@
+// Copyright Epic Games, Inc. All Rights Reserved.
+
+using EpicGames.Core;
+using System;
+using System.Collections;
+using System.Collections.Concurrent;
+using System.Collections.Generic;
+using System.Diagnostics.CodeAnalysis;
+using System.Linq;
+using System.Linq.Expressions;
+using System.Reflection;
+using System.Reflection.Emit;
+
+namespace EpicGames.Serialization.Converters
+{
+	/// <summary>
+	/// Converter for list types
+	/// </summary>
+	/// <typeparam name="T"></typeparam>
+	class CbListConverter<T> : CbConverter<List<T>>
+	{
+		/// <inheritdoc/>
+		public override List<T> Read(CbField Field)
+		{
+			List<T> List = new List<T>();
+			foreach (CbField ElementField in Field)
+			{
+				List.Add(CbSerializer.Deserialize<T>(ElementField));
+			}
+			return List;
+		}
+
+		/// <inheritdoc/>
+		public override void Write(CbWriter Writer, List<T> List)
+		{
+			Writer.BeginArray();
+			foreach (T Element in List)
+			{
+				CbSerializer.Serialize<T>(Writer, Element);
+			}
+			Writer.EndArray();
+		}
+
+		/// <inheritdoc/>
+		public override void WriteNamed(CbWriter Writer, Utf8String Name, List<T> List)
+		{
+			if (List.Count > 0)
+			{
+				Writer.BeginArray(Name);
+				foreach (T Element in List)
+				{
+					CbSerializer.Serialize<T>(Writer, Element);
+				}
+				Writer.EndArray();
+			}
+		}
+	}
+
+	/// <summary>
+	/// Specialization for serializing string lists
+	/// </summary>
+	sealed class CbStringListConverter : CbConverter<List<Utf8String>>
+	{
+		/// <inheritdoc/>
+		public override List<Utf8String> Read(CbField Field)
+		{
+			List<Utf8String> List = new List<Utf8String>();
+			foreach (CbField ElementField in Field)
+			{
+				List.Add(ElementField.AsString());
+			}
+			return List;
+		}
+
+		/// <inheritdoc/>
+		public override void Write(CbWriter Writer, List<Utf8String> List)
+		{
+			Writer.BeginUniformArray(CbFieldType.String);
+			foreach (Utf8String String in List)
+			{
+				Writer.WriteStringValue(String);
+			}
+			Writer.EndUniformArray();
+		}
+
+		/// <inheritdoc/>
+		public override void WriteNamed(CbWriter Writer, Utf8String Name, List<Utf8String> List)
+		{
+			if (List.Count > 0)
+			{
+				Writer.BeginUniformArray(Name, CbFieldType.String);
+				foreach (Utf8String String in List)
+				{
+					Writer.WriteStringValue(String);
+				}
+				Writer.EndUniformArray();
+			}
+		}
+	}
+
+	/// <summary>
+	/// Factory for CbListConverter
+	/// </summary>
+	class CbListConverterFactory : CbConverterFactory
+	{
+		/// <inheritdoc/>
+		public override CbConverter? CreateConverter(Type Type)
+		{
+			if (Type.IsGenericType && Type.GetGenericTypeDefinition() == typeof(List<>))
+			{
+				if (Type.GenericTypeArguments[0] == typeof(Utf8String))
+				{
+					return new CbStringListConverter();
+				}
+				else
+				{
+					Type ConverterType = typeof(CbListConverter<>).MakeGenericType(Type.GenericTypeArguments);
+					return (CbConverter)Activator.CreateInstance(ConverterType)!;
+				}
+			}
+			return null;
+		}
+	}
+}
