Horde: Add compact-binary version of remote execution structures.

diff --git a/Engine/Source/Programs/Shared/EpicGames.Horde/Compute/Action.cs b/Engine/Source/Programs/Shared/EpicGames.Horde/Compute/Action.cs
new file mode 100644
--- /dev/null
+++ b/Engine/Source/Programs/Shared/EpicGames.Horde/Compute/Action.cs
@@ -0,0 +1,45 @@
+// Copyright Epic Games, Inc. All Rights Reserved.
+
+using EpicGames.Core;
+using EpicGames.Serialization;
+using System;
+using System.Collections.Generic;
+using System.Text;
+
+namespace EpicGames.Horde.Compute
+{
+	/// <summary>
+	/// Describes an action to be executed in a particular workspace
+	/// </summary>
+	public class Action
+	{
+		/// <summary>
+		/// Hash of a <see cref="Command"/> object encoded as a CbObject and stored in the CAS
+		/// </summary>
+		[CbField("c")]
+		public IoHash CommandHash { get; set; }
+
+		/// <summary>
+		/// Hash of a <see cref="DirectoryTree"/> object encoded as a CbObject and stored in the CAS
+		/// </summary>
+		[CbField("s")]
+		public IoHash SandboxHash { get; set; }
+
+		/// <summary>
+		/// Requirements for the agent to execute the work
+		/// </summary>
+		[CbField("r")]
+		public Requirements? Requirements { get; set; }
+
+		/// <summary>
+		/// Constructor
+		/// </summary>
+		/// <param name="CommandHash"></param>
+		/// <param name="SandboxHash"></param>
+		public Action(IoHash CommandHash, IoHash SandboxHash)
+		{
+			this.CommandHash = CommandHash;
+			this.SandboxHash = SandboxHash;
+		}
+	}
+}
diff --git a/Engine/Source/Programs/Shared/EpicGames.Horde/Compute/ActionResult.cs b/Engine/Source/Programs/Shared/EpicGames.Horde/Compute/ActionResult.cs
new file mode 100644
--- /dev/null
+++ b/Engine/Source/Programs/Shared/EpicGames.Horde/Compute/ActionResult.cs
@@ -0,0 +1,49 @@
+// Copyright Epic Games, Inc. All Rights Reserved.
+
+using EpicGames.Core;
+using EpicGames.Serialization;
+using System;
+using System.Collections.Generic;
+using System.Text;
+
+namespace EpicGames.Horde.Compute
+{
+	/// <summary>
+	/// Describes an action to be executed in a particular workspace
+	/// </summary>
+	public class ActionResult
+	{
+		/// <summary>
+		/// Exit code of the process
+		/// </summary>
+		[CbField("e")]
+		public int ExitCode { get; set; }
+
+		/// <summary>
+		/// Hash of the stdout data
+		/// </summary>
+		[CbField("so")]
+		public IoHash? StdOutHash { get; set; }
+
+		/// <summary>
+		/// Hash of the stderr data
+		/// </summary>
+		[CbField("se")]
+		public IoHash? StdErrHash { get; set; }
+
+		/// <summary>
+		/// Hash of the output <see cref="DirectoryTree"/> stored in the CAS. Only items matching the output paths in the action will be included.
+		/// </summary>
+		[CbField("o")]
+		public IoHash? OutputHash { get; set; }
+
+		/// <summary>
+		/// Constructor
+		/// </summary>
+		/// <param name="ExitCode">Exit code of the process</param>
+		public ActionResult(int ExitCode)
+		{
+			this.ExitCode = ExitCode;
+		}
+	}
+}
diff --git a/Engine/Source/Programs/Shared/EpicGames.Horde/Compute/Command.cs b/Engine/Source/Programs/Shared/EpicGames.Horde/Compute/Command.cs
new file mode 100644
--- /dev/null
+++ b/Engine/Source/Programs/Shared/EpicGames.Horde/Compute/Command.cs
@@ -0,0 +1,64 @@
+// Copyright Epic Games, Inc. All Rights Reserved.
+
+using EpicGames.Core;
+using EpicGames.Serialization;
+using System;
+using System.Collections.Generic;
+using System.Text;
+
+namespace EpicGames.Horde.Compute
+{
+	/// <summary>
+	/// Command to be executed on a remote agent
+	/// </summary>
+	public class Command
+	{
+		/// <summary>
+		/// The executable to run
+		/// </summary>
+		[CbField("e")]
+		public Utf8String Executable { get; set; }
+
+		/// <summary>
+		/// List of command line arguments for the process to run.
+		/// </summary>
+		[CbField("a")]
+		public List<Utf8String> Arguments { get; set; } = new List<Utf8String>();
+
+		/// <summary>
+		/// Environment variables to set for the child process
+		/// </summary>
+		[CbField("v")]
+		public Dictionary<Utf8String, Utf8String> EnvVars { get; set; } = new Dictionary<Utf8String, Utf8String>();
+
+		/// <summary>
+		/// Path to the working directory within the workspace
+		/// </summary>
+		[CbField("d")]
+		public Utf8String WorkingDirectory { get; set; } = Utf8String.Empty;
+
+		/// <summary>
+		/// List of output paths to be preserved on completion of the action. These may be files or directories.
+		/// </summary>
+		[CbField("p")]
+		public List<Utf8String> OutputPaths { get; set; } = new List<Utf8String>();
+
+		/// <summary>
+		/// Private constructor for serialization
+		/// </summary>
+		private Command()
+		{
+		}
+
+		/// <summary>
+		/// Constructor
+		/// </summary>
+		/// <param name="Executable">The executable to run</param>
+		/// <param name="Arguments">List of command line arguments.</param>
+		public Command(Utf8String Executable, List<Utf8String> Arguments)
+		{
+			this.Executable = Executable;
+			this.Arguments = Arguments;
+		}
+	}
+}
diff --git a/Engine/Source/Programs/Shared/EpicGames.Horde/Compute/DirectoryTree.cs b/Engine/Source/Programs/Shared/EpicGames.Horde/Compute/DirectoryTree.cs
new file mode 100644
--- /dev/null
+++ b/Engine/Source/Programs/Shared/EpicGames.Horde/Compute/DirectoryTree.cs
@@ -0,0 +1,96 @@
+// Copyright Epic Games, Inc. All Rights Reserved.
+
+using EpicGames.Core;
+using EpicGames.Serialization;
+using System;
+using System.Collections.Generic;
+using System.Text;
+
+namespace EpicGames.Horde.Compute
+{
+	/// <summary>
+	/// Stores information about a directory in an action's workspace
+	/// </summary>
+	public class DirectoryTree
+	{
+		/// <summary>
+		/// List of files within the directory, sorted by name
+		/// </summary>
+		[CbField("f")]
+		public List<FileNode> Files { get; set; } = new List<FileNode>();
+
+		/// <summary>
+		/// Sub-directories within this directory, sorted by name
+		/// </summary>
+		[CbField("d")]
+		public List<DirectoryNode> Directories { get; set; } = new List<DirectoryNode>();
+	}
+
+	/// <summary>
+	/// Reference to a named file within a <see cref="DirectoryTree"/>
+	/// </summary>
+	public class FileNode
+	{
+		/// <summary>
+		/// Name of the directory
+		/// </summary>
+		[CbField("n")]
+		public Utf8String Name { get; set; }
+
+		/// <summary>
+		/// Hash of bulk data for the file stored in CAS.
+		/// </summary>
+		[CbField("h")]
+		public IoHash Hash { get; set; }
+
+		/// <summary>
+		/// Private constructor for serialization
+		/// </summary>
+		private FileNode()
+		{
+		}
+
+		/// <summary>
+		/// Constructor
+		/// </summary>
+		public FileNode(Utf8String Name, IoHash Hash)
+		{
+			this.Name = Name;
+			this.Hash = Hash;
+		}
+	}
+
+	/// <summary>
+	/// Reference to a named directory within a <see cref="DirectoryTree"/>
+	/// </summary>
+	public class DirectoryNode
+	{
+		/// <summary>
+		/// Name of the directory
+		/// </summary>
+		[CbField("n")]
+		public Utf8String Name { get; set; }
+
+		/// <summary>
+		/// Hash of a <see cref="CbObject"/> encoded <see cref="DirectoryTree"/> stored in CAS.
+		/// </summary>
+		[CbField("h")]
+		public IoHash Hash { get; set; }
+
+		/// <summary>
+		/// Private constructor for serialization
+		/// </summary>
+		private DirectoryNode()
+		{
+		}
+
+		/// <summary>
+		/// Constructor
+		/// </summary>
+		public DirectoryNode(Utf8String Name, IoHash Hash)
+		{
+			this.Name = Name;
+			this.Hash = Hash;
+		}
+	}
+}
diff --git a/Engine/Source/Programs/Shared/EpicGames.Horde/Compute/Requirements.cs b/Engine/Source/Programs/Shared/EpicGames.Horde/Compute/Requirements.cs
new file mode 100644
--- /dev/null
+++ b/Engine/Source/Programs/Shared/EpicGames.Horde/Compute/Requirements.cs
@@ -0,0 +1,34 @@
+// Copyright Epic Games, Inc. All Rights Reserved.
+
+using EpicGames.Core;
+using EpicGames.Serialization;
+using System;
+using System.Collections.Generic;
+using System.Text;
+
+namespace EpicGames.Horde.Compute
+{
+	/// <summary>
+	/// Stores information about a directory in an action's workspace
+	/// </summary>
+	public class Requirements
+	{
+		/// <summary>
+		/// Condition string to be evaluated against the machine spec, eg. cpu-cores >= 10 && ram.mb >= 200 && pool == 'worker'
+		/// </summary>
+		[CbField("c")]
+		public Utf8String Condition { get; set; }
+
+		/// <summary>
+		/// Resources used by the process
+		/// </summary>
+		[CbField("r")]
+		public Dictionary<Utf8String, int> Resources { get; set; } = new Dictionary<Utf8String, int>();
+
+		/// <summary>
+		/// Whether we require exclusive access to the device
+		/// </summary>
+		[CbField("e")]
+		public bool Exclusive { get; set; }
+	}
+}
diff --git a/Engine/Source/Programs/Shared/EpicGames.Horde/EpicGames.Horde.csproj b/Engine/Source/Programs/Shared/EpicGames.Horde/EpicGames.Horde.csproj
new file mode 100644
--- /dev/null
+++ b/Engine/Source/Programs/Shared/EpicGames.Horde/EpicGames.Horde.csproj
@@ -0,0 +1,44 @@
+<Project Sdk="Microsoft.NET.Sdk">
+  <PropertyGroup>
+    <TargetFramework>netcoreapp3.1</TargetFramework>
+    <Configuration Condition=" '$(Configuration)' == '' ">Development</Configuration>
+    <OutputType>Library</OutputType>
+    <Deterministic>true</Deterministic>
+    <RootNamespace>EpicGames.Horde</RootNamespace>
+    <GenerateAssemblyInfo>false</GenerateAssemblyInfo>
+    <GenerateTargetFrameworkAttribute>false</GenerateTargetFrameworkAttribute>
+    <Configurations>Debug;Release;Development</Configurations>
+    <Nullable>enable</Nullable>
+  </PropertyGroup>
+  <PropertyGroup Condition=" '$(Configuration)|$(Platform)' == 'Debug|AnyCPU' ">
+    <DocumentationFile></DocumentationFile>
+  </PropertyGroup>
+  <PropertyGroup Condition="'$(Configuration)|$(Platform)' == 'Development|AnyCPU'">
+    <DocumentationFile></DocumentationFile>
+    <Optimize>false</Optimize>
+  </PropertyGroup>
+  <PropertyGroup Condition="'$(Configuration)|$(Platform)' == 'Release|AnyCPU'">
+    <DocumentationFile></DocumentationFile>
+    <Optimize>true</Optimize>
+  </PropertyGroup>
+  <ItemGroup>
+    <Compile Include="..\..\DotNETCommon\MetaData.cs">
+      <Link>Properties\MetaData.cs</Link>
+    </Compile>
+  </ItemGroup>
+  <ItemGroup>
+    <PackageReference Include="Google.Protobuf" Version="3.12.3" />
+    <PackageReference Include="Grpc.Core.Api" Version="2.35.0" />
+    <PackageReference Include="Grpc.Tools" Version="2.30.0">
+      <PrivateAssets>all</PrivateAssets>
+      <IncludeAssets>runtime; build; native; contentfiles; analyzers; buildtransitive</IncludeAssets>
+    </PackageReference>
+    <PackageReference Include="Microsoft.CodeAnalysis.CSharp.Workspaces" Version="3.10.0" />
+    <PackageReference Include="Microsoft.Extensions.Logging" Version="5.0.0" />
+  </ItemGroup>
+  <ItemGroup>
+    <ProjectReference Include="..\EpicGames.Core\EpicGames.Core.csproj" />
+    <ProjectReference Include="..\EpicGames.IoHash\EpicGames.IoHash.csproj" />
+    <ProjectReference Include="..\EpicGames.Serialization\EpicGames.Serialization.csproj" />
+  </ItemGroup>
+</Project>
\ No newline at end of file
