Horde: Rename Action -> ComputeTask.

diff --git a/Engine/Source/Programs/Shared/EpicGames.Horde/Compute/Action.cs b/Engine/Source/Programs/Shared/EpicGames.Horde/Compute/Action.cs
deleted file mode 100644
diff --git a/Engine/Source/Programs/Shared/EpicGames.Horde/Compute/ActionResult.cs b/Engine/Source/Programs/Shared/EpicGames.Horde/Compute/ActionResult.cs
deleted file mode 100644
diff --git a/Engine/Source/Programs/Shared/EpicGames.Horde/Compute/ComputeTask.cs b/Engine/Source/Programs/Shared/EpicGames.Horde/Compute/ComputeTask.cs
new file mode 100644
--- /dev/null
+++ b/Engine/Source/Programs/Shared/EpicGames.Horde/Compute/ComputeTask.cs
@@ -0,0 +1,45 @@
+// Copyright Epic Games, Inc. All Rights Reserved.
+
+using EpicGames.Core;
+using EpicGames.Serialization;
+using System;
+using System.Collections.Generic;
+using System.Text;
+
+namespace EpicGames.Horde.Compute
+{
+	/// <summary>
+	/// Describes an action to be executed in a particular workspace
+	/// </summary>
+	public class ComputeTask
+	{
+		/// <summary>
+		/// Hash of a <see cref="Command"/> object encoded as a CbObject and stored in the CAS
+		/// </summary>
+		[CbField("c")]
+		public IoHash CommandHash { get; set; }
+
+		/// <summary>
+		/// Hash of a <see cref="DirectoryTree"/> object encoded as a CbObject and stored in the CAS
+		/// </summary>
+		[CbField("s")]
+		public IoHash SandboxHash { get; set; }
+
+		/// <summary>
+		/// Requirements for the agent to execute the work
+		/// </summary>
+		[CbField("r")]
+		public Requirements? Requirements { get; set; }
+
+		/// <summary>
+		/// Constructor
+		/// </summary>
+		/// <param name="CommandHash"></param>
+		/// <param name="SandboxHash"></param>
+		public ComputeTask(IoHash CommandHash, IoHash SandboxHash)
+		{
+			this.CommandHash = CommandHash;
+			this.SandboxHash = SandboxHash;
+		}
+	}
+}
diff --git a/Engine/Source/Programs/Shared/EpicGames.Horde/Compute/ComputeTaskResult.cs b/Engine/Source/Programs/Shared/EpicGames.Horde/Compute/ComputeTaskResult.cs
new file mode 100644
--- /dev/null
+++ b/Engine/Source/Programs/Shared/EpicGames.Horde/Compute/ComputeTaskResult.cs
@@ -0,0 +1,49 @@
+// Copyright Epic Games, Inc. All Rights Reserved.
+
+using EpicGames.Core;
+using EpicGames.Serialization;
+using System;
+using System.Collections.Generic;
+using System.Text;
+
+namespace EpicGames.Horde.Compute
+{
+	/// <summary>
+	/// Describes an action to be executed in a particular workspace
+	/// </summary>
+	public class ComputeTaskResult
+	{
+		/// <summary>
+		/// Exit code of the process
+		/// </summary>
+		[CbField("e")]
+		public int ExitCode { get; set; }
+
+		/// <summary>
+		/// Hash of the stdout data
+		/// </summary>
+		[CbField("so")]
+		public IoHash? StdOutHash { get; set; }
+
+		/// <summary>
+		/// Hash of the stderr data
+		/// </summary>
+		[CbField("se")]
+		public IoHash? StdErrHash { get; set; }
+
+		/// <summary>
+		/// Hash of the output <see cref="DirectoryTree"/> stored in the CAS. Only items matching the output paths in the action will be included.
+		/// </summary>
+		[CbField("o")]
+		public IoHash? OutputHash { get; set; }
+
+		/// <summary>
+		/// Constructor
+		/// </summary>
+		/// <param name="ExitCode">Exit code of the process</param>
+		public ComputeTaskResult(int ExitCode)
+		{
+			this.ExitCode = ExitCode;
+		}
+	}
+}
