Horde: Merge command with ComputeTask, and store requirements as a separate hashed object. This makes it easier to group tasks by requirements for scheduling.

diff --git a/Engine/Source/Programs/Shared/EpicGames.Horde/Compute/Command.cs b/Engine/Source/Programs/Shared/EpicGames.Horde/Compute/Command.cs
deleted file mode 100644
diff --git a/Engine/Source/Programs/Shared/EpicGames.Horde/Compute/ComputeTask.cs b/Engine/Source/Programs/Shared/EpicGames.Horde/Compute/ComputeTask.cs
--- a/Engine/Source/Programs/Shared/EpicGames.Horde/Compute/ComputeTask.cs
+++ b/Engine/Source/Programs/Shared/EpicGames.Horde/Compute/ComputeTask.cs
@@ -14,10 +14,28 @@ namespace EpicGames.Horde.Compute
 	public class ComputeTask
 	{
 		/// <summary>
-		/// Hash of a <see cref="Command"/> object encoded as a CbObject and stored in the CAS
+		/// The executable to run
 		/// </summary>
-		[CbField("c")]
-		public IoHash CommandHash { get; set; }
+		[CbField("e")]
+		public Utf8String Executable { get; set; }
+
+		/// <summary>
+		/// List of command line arguments for the process to run.
+		/// </summary>
+		[CbField("a")]
+		public List<Utf8String> Arguments { get; set; } = new List<Utf8String>();
+
+		/// <summary>
+		/// Environment variables to set for the child process
+		/// </summary>
+		[CbField("v")]
+		public Dictionary<Utf8String, Utf8String> EnvVars { get; set; } = new Dictionary<Utf8String, Utf8String>();
+
+		/// <summary>
+		/// Path to the working directory within the workspace
+		/// </summary>
+		[CbField("w")]
+		public Utf8String WorkingDirectory { get; set; }
 
 		/// <summary>
 		/// Hash of a <see cref="DirectoryTree"/> object encoded as a CbObject and stored in the CAS
@@ -29,16 +47,26 @@ namespace EpicGames.Horde.Compute
 		/// Requirements for the agent to execute the work
 		/// </summary>
 		[CbField("r")]
-		public Requirements? Requirements { get; set; }
+		public IoHash RequirementsHash { get; set; } = IoHash.Zero;
+
+		/// <summary>
+		/// List of output paths to be captured on completion of the action. These may be files or directories.
+		/// </summary>
+		[CbField("o")]
+		public List<Utf8String> OutputPaths { get; set; } = new List<Utf8String>();
 
 		/// <summary>
 		/// Constructor
 		/// </summary>
-		/// <param name="CommandHash"></param>
-		/// <param name="SandboxHash"></param>
-		public ComputeTask(IoHash CommandHash, IoHash SandboxHash)
+		/// <param name="Executable">The executable to run</param>
+		/// <param name="Arguments">Arguments for the executable to run</param>
+		/// <param name="WorkingDirectory">Working directory for execution</param>
+		/// <param name="SandboxHash">Hash of the sandbox</param>
+		public ComputeTask(Utf8String Executable, List<Utf8String> Arguments, Utf8String WorkingDirectory, IoHash SandboxHash)
 		{
-			this.CommandHash = CommandHash;
+			this.Executable = Executable;
+			this.Arguments = Arguments;
+			this.WorkingDirectory = WorkingDirectory;
 			this.SandboxHash = SandboxHash;
 		}
 	}
