Horde: Add a constructor for the Requirements class.

diff --git a/Engine/Source/Programs/Shared/EpicGames.Horde/Compute/Requirements.cs b/Engine/Source/Programs/Shared/EpicGames.Horde/Compute/Requirements.cs
--- a/Engine/Source/Programs/Shared/EpicGames.Horde/Compute/Requirements.cs
+++ b/Engine/Source/Programs/Shared/EpicGames.Horde/Compute/Requirements.cs
@@ -30,5 +30,21 @@ namespace EpicGames.Horde.Compute
 		/// </summary>
 		[CbField("e")]
 		public bool Exclusive { get; set; }
+
+		/// <summary>
+		/// Default constructor
+		/// </summary>
+		public Requirements()
+		{
+		}
+
+		/// <summary>
+		/// Construct a requirements object with a condition
+		/// </summary>
+		/// <param name="Condition"></param>
+		public Requirements(Utf8String Condition)
+		{
+			this.Condition = Condition;
+		}
 	}
 }
