Horde: Add common definitions for the ref table, blob store, and cbobject/iohash wrapper types.

diff --git a/Engine/Source/Programs/Shared/EpicGames.Horde/Common/CbObjectWrapper.cs b/Engine/Source/Programs/Shared/EpicGames.Horde/Common/CbObjectWrapper.cs
new file mode 100644
--- /dev/null
+++ b/Engine/Source/Programs/Shared/EpicGames.Horde/Common/CbObjectWrapper.cs
@@ -0,0 +1,50 @@
+// Copyright Epic Games, Inc. All Rights Reserved.
+
+using EpicGames.Core;
+using EpicGames.Serialization;
+using Google.Protobuf;
+using System;
+using System.Collections.Generic;
+using System.Text;
+using System.Threading.Tasks;
+
+namespace EpicGames.Horde.Common
+{
+	/// <summary>
+	/// Utility methods for CbObjectWrapper
+	/// </summary>
+	public partial class CbObjectWrapper
+	{
+		/// <summary>
+		/// Parse the buffer as a CbObject
+		/// </summary>
+		public CbObject Object => new CbObject(Data.ToByteArray());
+
+		/// <summary>
+		/// Constructor
+		/// </summary>
+		/// <param name="CbObject"></param>
+		public CbObjectWrapper(CbObject CbObject)
+		{
+			this.Data = ByteString.CopyFrom(CbObject.GetView().ToArray());
+		}
+
+		/// <summary>
+		/// Convert a wrapper to regular CbObject
+		/// </summary>
+		/// <param name="ObjectWrapper"></param>
+		public static implicit operator CbObject(CbObjectWrapper ObjectWrapper)
+		{
+			return new CbObject(ObjectWrapper.Data.ToByteArray());
+		}
+
+		/// <summary>
+		/// Convert a CbObject to a wrapper type
+		/// </summary>
+		/// <param name="Object"></param>
+		public static implicit operator CbObjectWrapper(CbObject Object)
+		{
+			return new CbObjectWrapper(Object);
+		}
+	}
+}
diff --git a/Engine/Source/Programs/Shared/EpicGames.Horde/Common/IoHashWrapper.cs b/Engine/Source/Programs/Shared/EpicGames.Horde/Common/IoHashWrapper.cs
new file mode 100644
--- /dev/null
+++ b/Engine/Source/Programs/Shared/EpicGames.Horde/Common/IoHashWrapper.cs
@@ -0,0 +1,44 @@
+// Copyright Epic Games, Inc. All Rights Reserved.
+
+using EpicGames.Core;
+using Google.Protobuf;
+using System;
+using System.Collections.Generic;
+using System.Text;
+using System.Threading.Tasks;
+
+namespace EpicGames.Horde.Common
+{
+	/// <summary>
+	/// Methods for IoHashWrapper
+	/// </summary>
+	public partial class IoHashWrapper
+	{
+		/// <summary>
+		/// Constructor
+		/// </summary>
+		/// <param name="Hash"></param>
+		public IoHashWrapper(IoHash Hash)
+		{
+			this.Hash = ByteString.CopyFrom(Hash.Memory.ToArray());
+		}
+
+		/// <summary>
+		/// Convert from an IoHashWrapper to an IoHash
+		/// </summary>
+		/// <param name="Hash"></param>
+		public static implicit operator IoHash(IoHashWrapper Hash)
+		{
+			return new IoHash(Hash.ToByteArray());
+		}
+
+		/// <summary>
+		/// Convert from an IoHash to an IoHashWrapper
+		/// </summary>
+		/// <param name="Hash"></param>
+		public static implicit operator IoHashWrapper(IoHash Hash)
+		{
+			return new IoHashWrapper(Hash);
+		}
+	}
+}
diff --git a/Engine/Source/Programs/Shared/EpicGames.Horde/EpicGames.Horde.csproj b/Engine/Source/Programs/Shared/EpicGames.Horde/EpicGames.Horde.csproj
--- a/Engine/Source/Programs/Shared/EpicGames.Horde/EpicGames.Horde.csproj
+++ b/Engine/Source/Programs/Shared/EpicGames.Horde/EpicGames.Horde.csproj
@@ -41,4 +41,10 @@
     <ProjectReference Include="..\EpicGames.IoHash\EpicGames.IoHash.csproj" />
     <ProjectReference Include="..\EpicGames.Serialization\EpicGames.Serialization.csproj" />
   </ItemGroup>
+  <ItemGroup>
+    <Protobuf Include="Protos\horde\common\cbobject_wrapper.proto" ProtoRoot="Protos" />
+    <Protobuf Include="Protos\horde\common\iohash_wrapper.proto" ProtoRoot="Protos" />
+    <Protobuf Include="Protos\horde\storage\blob_store.proto" ProtoRoot="Protos" />
+    <Protobuf Include="Protos\horde\storage\ref_table.proto" ProtoRoot="Protos" />
+  </ItemGroup>
 </Project>
\ No newline at end of file
