Missing files.

diff --git a/Engine/Source/Programs/Horde/HordeServer/HordeServer.csproj b/Engine/Source/Programs/Horde/HordeServer/HordeServer.csproj
--- a/Engine/Source/Programs/Horde/HordeServer/HordeServer.csproj
+++ b/Engine/Source/Programs/Horde/HordeServer/HordeServer.csproj
@@ -58,6 +58,8 @@
   </ItemGroup>
 
   <ItemGroup>
+    <Protobuf Include="..\..\Shared\EpicGames.Horde\Protos\horde\storage\ref_table.proto" ProtoRoot="..\..\Shared\EpicGames.Horde\Protos" GrpcServices="Server" />
+    <Protobuf Include="..\..\Shared\EpicGames.Horde\Protos\horde\storage\blob_store.proto" ProtoRoot="..\..\Shared\EpicGames.Horde\Protos" GrpcServices="Server" />
     <Protobuf Include="..\..\Shared\EpicGames.Horde.Common\Protos\google\bytestream\bytestream.proto" ProtoRoot="..\..\Shared\EpicGames.Horde.Common\Protos" GrpcServices="Server" />
     <Protobuf Include="..\..\Shared\EpicGames.Horde.Common\Protos\horde\action_rpc.proto" ProtoRoot="..\..\Shared\EpicGames.Horde.Common\Protos" GrpcServices="Server" />
     <Protobuf Include="..\..\Shared\EpicGames.Horde.Common\Protos\horde\horde_rpc.proto" ProtoRoot="..\..\Shared\EpicGames.Horde.Common\Protos" GrpcServices="Server" />
@@ -69,6 +71,7 @@
   <ItemGroup>
     <ProjectReference Include="..\..\Shared\EpicGames.Core\EpicGames.Core.csproj" />
     <ProjectReference Include="..\..\Shared\EpicGames.Horde.Common\EpicGames.Horde.Common.csproj" />
+    <ProjectReference Include="..\..\Shared\EpicGames.Horde\EpicGames.Horde.csproj" />
     <ProjectReference Include="..\..\Shared\EpicGames.IoHash\EpicGames.IoHash.csproj" />
     <ProjectReference Include="..\..\Shared\EpicGames.Perforce\EpicGames.Perforce.csproj" />
     <ProjectReference Include="..\..\Shared\EpicGames.Serialization\EpicGames.Serialization.csproj" />
diff --git a/Engine/Source/Programs/Horde/HordeServer/Utilities/RedisSerializer.cs b/Engine/Source/Programs/Horde/HordeServer/Utilities/RedisSerializer.cs
--- a/Engine/Source/Programs/Horde/HordeServer/Utilities/RedisSerializer.cs
+++ b/Engine/Source/Programs/Horde/HordeServer/Utilities/RedisSerializer.cs
@@ -14,6 +14,7 @@ namespace HordeServer.Utilities
 	/// <summary>
 	/// Attribute specifying the converter type to use for a class
 	/// </summary>
+	[AttributeUsage(AttributeTargets.Class)]
 	public class RedisConverterAttribute : Attribute
 	{
 		/// <summary>
diff --git a/Engine/Source/Programs/Shared/EpicGames.Horde/Storage/BlobStoreMessages.cs b/Engine/Source/Programs/Shared/EpicGames.Horde/Storage/BlobStoreMessages.cs
new file mode 100644
--- /dev/null
+++ b/Engine/Source/Programs/Shared/EpicGames.Horde/Storage/BlobStoreMessages.cs
@@ -0,0 +1,47 @@
+// Copyright Epic Games, Inc. All Rights Reserved.
+
+using EpicGames.Core;
+using Google.Protobuf;
+using System;
+using System.Collections.Generic;
+using System.Text;
+using System.Threading.Tasks;
+
+namespace EpicGames.Horde.Storage
+{
+	/// <summary>
+	/// Additional methods for AddBlobRequest
+	/// </summary>
+	public partial class AddBlobRequest
+	{
+		/// <summary>
+		/// Constructor
+		/// </summary>
+		/// <param name="Hash"></param>
+		/// <param name="Data"></param>
+		public AddBlobRequest(IoHash Hash, ReadOnlyMemory<byte> Data)
+		{
+			this.Hash = Hash;
+			this.Data = ByteString.CopyFrom(Data.ToArray());
+		}
+	}
+
+	/// <summary>
+	/// Additional methods for GetBlobResponse
+	/// </summary>
+	partial class GetBlobResponse
+	{
+		/// <summary>
+		/// Constructor
+		/// </summary>
+		/// <param name="Data"></param>
+		public GetBlobResponse(byte[]? Data)
+		{
+			if (Data != null)
+			{
+				this.Data = ByteString.CopyFrom(Data);
+				this.Exists = true;
+			}
+		}
+	}
+}
diff --git a/Engine/Source/Programs/Shared/EpicGames.Horde/Storage/RefTableMessages.cs b/Engine/Source/Programs/Shared/EpicGames.Horde/Storage/RefTableMessages.cs
new file mode 100644
--- /dev/null
+++ b/Engine/Source/Programs/Shared/EpicGames.Horde/Storage/RefTableMessages.cs
@@ -0,0 +1,92 @@
+// Copyright Epic Games, Inc. All Rights Reserved.
+
+using EpicGames.Core;
+using EpicGames.Horde.Common;
+using EpicGames.Serialization;
+using Google.Protobuf;
+using Google.Protobuf.WellKnownTypes;
+using System;
+using System.Collections.Generic;
+using System.Linq;
+using System.Text;
+using System.Threading.Tasks;
+
+namespace EpicGames.Horde.Storage
+{
+	/// <summary>
+	/// Request to update or add a reference
+	/// </summary>
+	partial class SetRefRequest
+	{
+		public SetRefRequest(string NamespaceId, string BucketId, string Name, CbObject Object)
+		{
+			this.NamespaceId = NamespaceId;
+			this.BucketId = BucketId;
+			this.Name = Name;
+			this.Object = Object;
+		}
+	}
+
+	/// <summary>
+	/// Response from adding a reference
+	/// </summary>
+	partial class SetRefResponse
+	{
+		public SetRefResponse(IEnumerable<IoHash> MissingHashes)
+		{
+			this.MissingHashes.Add(MissingHashes.Select(x => (IoHashWrapper)x));
+		}
+	}
+
+	/// <summary>
+	/// Request to get a reference
+	/// </summary>
+	partial class GetRefRequest
+	{
+		public GetRefRequest(string NamespaceId, string BucketId, string Name)
+		{
+			this.NamespaceId = NamespaceId;
+			this.BucketId = BucketId;
+			this.Name = Name;
+		}
+	}
+
+	/// <summary>
+	/// Response from querying a reference
+	/// </summary>
+	partial class GetRefResponse
+	{
+		public GetRefResponse(CbObject Object, DateTime LastAccessTime, bool Finalized)
+		{
+			this.Object = Object;
+			this.LastAccessTime = Timestamp.FromDateTime(LastAccessTime);
+			this.Finalized = Finalized;
+		}
+	}
+
+	/// <summary>
+	/// Request to update the last access time of a reference
+	/// </summary>
+	partial class TouchRequest
+	{
+		public TouchRequest(string NamespaceId, string BucketId, string Name)
+		{
+			this.NamespaceId = NamespaceId;
+			this.BucketId = BucketId;
+			this.Name = Name;
+		}
+	}
+
+	/// <summary>
+	/// Request to delete a reference
+	/// </summary>
+	partial class DeleteRequest
+	{
+		public DeleteRequest(string NamespaceId, string BucketId, string Name)
+		{
+			this.NamespaceId = NamespaceId;
+			this.BucketId = BucketId;
+			this.Name = Name;
+		}
+	}
+}
