Make relative, with fewer allocations

diff --git a/Engine/Source/Programs/Shared/EpicGames.Core/FileSystemReference.cs b/Engine/Source/Programs/Shared/EpicGames.Core/FileSystemReference.cs
--- a/Engine/Source/Programs/Shared/EpicGames.Core/FileSystemReference.cs
+++ b/Engine/Source/Programs/Shared/EpicGames.Core/FileSystemReference.cs
@@ -5,6 +5,7 @@ using System.Collections.Generic;
 using System.IO;
 using System.Linq;
 using System.Text;
+using System.Threading;
 using System.Threading.Tasks;
 
 namespace EpicGames.Core
@@ -50,10 +51,11 @@ namespace EpicGames.Core
 		/// Create a full path by concatenating multiple strings
 		/// </summary>
 		/// <returns></returns>
+		static ThreadLocal<StringBuilder> CombineStringsStringBuilder = new ThreadLocal<StringBuilder>(() => new StringBuilder(260));
 		static protected string CombineStrings(DirectoryReference BaseDirectory, params string[] Fragments)
 		{
 			// Get the initial string to append to, and strip any root directory suffix from it
-			StringBuilder NewFullName = new StringBuilder(BaseDirectory.FullName);
+			StringBuilder NewFullName = CombineStringsStringBuilder.Value!.Clear().Append(BaseDirectory.FullName);
 			if (NewFullName.Length > 0 && NewFullName[NewFullName.Length - 1] == Path.DirectorySeparatorChar)
 			{
 				NewFullName.Remove(NewFullName.Length - 1, 1);
@@ -243,7 +245,15 @@ namespace EpicGames.Core
 		/// </summary>
 		/// <param name="Directory">The directory to create a relative path from</param>
 		/// <returns>A relative path from the given directory</returns>
+		static ThreadLocal<StringBuilder> MakeRelativeToStringBuilder = new ThreadLocal<StringBuilder>(() => new StringBuilder(260));
 		public string MakeRelativeTo(DirectoryReference Directory)
+		{
+			StringBuilder Result = MakeRelativeToStringBuilder.Value!.Clear();
+			WriteRelativeTo(Directory, ref Result);
+			return Result.ToString();
+		}
+
+		public void WriteRelativeTo(DirectoryReference Directory, ref StringBuilder Result)
 		{
 			// Find how much of the path is common between the two paths. This length does not include a trailing directory separator character.
 			int CommonDirectoryLength = -1;
@@ -254,7 +264,8 @@ namespace EpicGames.Core
 					// The two paths are identical. Just return the "." character.
 					if (Idx == Directory.FullName.Length)
 					{
-						return ".";
+						Result.Append('.');
+						return;
 					}
 
 					// Check if we're finishing on a complete directory name
@@ -290,11 +301,11 @@ namespace EpicGames.Core
 			// If there's no relative path, just return the absolute path
 			if (CommonDirectoryLength == -1)
 			{
-				return FullName;
+				Result.Append(FullName);
+				return;
 			}
 
 			// Append all the '..' separators to get back to the common directory, then the rest of the string to reach the target item
-			StringBuilder Result = new StringBuilder();
 			for (int Idx = CommonDirectoryLength + 1; Idx < Directory.FullName.Length; Idx++)
 			{
 				// Move up a directory
@@ -311,7 +322,6 @@ namespace EpicGames.Core
 			{
 				Result.Append(FullName, CommonDirectoryLength + 1, FullName.Length - CommonDirectoryLength - 1);
 			}
-			return Result.ToString();
 		}
 
 		public string ToNormalizedPath()
