Horede: Make fileSize Optional for the Perforce PrintRecord

Perforce is sometimes omitting that field which causes pre-flights to fail, and it's not used by any code right now.

diff --git a/Engine/Source/Programs/Shared/EpicGames.Perforce/Records/PrintRecord.cs b/Engine/Source/Programs/Shared/EpicGames.Perforce/Records/PrintRecord.cs
--- a/Engine/Source/Programs/Shared/EpicGames.Perforce/Records/PrintRecord.cs
+++ b/Engine/Source/Programs/Shared/EpicGames.Perforce/Records/PrintRecord.cs
@@ -52,7 +52,7 @@ namespace EpicGames.Perforce
 		/// <summary>
 		/// Size of the file in bytes
 		/// </summary>
-		[PerforceTag("fileSize")]
+		[PerforceTag("fileSize", Optional = true)]
 		public long FileSize;
 
 		/// <summary>
