Horde: Support serialization of bool values to compact binary objects through attribute-driven serialization.

diff --git a/Engine/Source/Programs/Shared/EpicGames.Serialization/CompactBinary.cs b/Engine/Source/Programs/Shared/EpicGames.Serialization/CompactBinary.cs
--- a/Engine/Source/Programs/Shared/EpicGames.Serialization/CompactBinary.cs
+++ b/Engine/Source/Programs/Shared/EpicGames.Serialization/CompactBinary.cs
@@ -937,12 +937,18 @@ namespace EpicGames.Serialization
 			}
 		}
 
+		/// <summary>
+		/// Access the field as a bool. Returns the provided default on error.
+		/// </summary>
+		/// <returns>Value of the field</returns>
+		public bool AsBool() => AsBool(false);
+
 		/// <summary>
 		/// Access the field as a bool. Returns the provided default on error.
 		/// </summary>
 		/// <param name="Default">Default value</param>
 		/// <returns>Value of the field</returns>
-		public bool AsBool(bool Default = false)
+		public bool AsBool(bool Default)
 		{
 			switch (GetType())
 			{
diff --git a/Engine/Source/Programs/Shared/EpicGames.Serialization/CompactBinaryWriter.cs b/Engine/Source/Programs/Shared/EpicGames.Serialization/CompactBinaryWriter.cs
--- a/Engine/Source/Programs/Shared/EpicGames.Serialization/CompactBinaryWriter.cs
+++ b/Engine/Source/Programs/Shared/EpicGames.Serialization/CompactBinaryWriter.cs
@@ -308,6 +308,25 @@ namespace EpicGames.Serialization
 			WriteField(CbFieldType.Null, Name);
 		}
 
+		/// <summary>
+		/// Writes a boolean value
+		/// </summary>
+		/// <param name="Value"></param>
+		public void WriteBoolValue(bool Value)
+		{
+			WriteField(Value? CbFieldType.BoolTrue : CbFieldType.BoolFalse);
+		}
+
+		/// <summary>
+		/// Writes a boolean value
+		/// </summary>
+		/// <param name="Name">Name of the field</param>
+		/// <param name="Value"></param>
+		public void WriteBool(Utf8String Name, bool Value)
+		{
+			WriteField(Value ? CbFieldType.BoolTrue : CbFieldType.BoolFalse, Name);
+		}
+
 		/// <summary>
 		/// Writes the payload for an integer
 		/// </summary>
diff --git a/Engine/Source/Programs/Shared/EpicGames.Serialization/Converters/CbDefaultConverter.cs b/Engine/Source/Programs/Shared/EpicGames.Serialization/Converters/CbDefaultConverter.cs
--- a/Engine/Source/Programs/Shared/EpicGames.Serialization/Converters/CbDefaultConverter.cs
+++ b/Engine/Source/Programs/Shared/EpicGames.Serialization/Converters/CbDefaultConverter.cs
@@ -100,6 +100,7 @@ namespace EpicGames.Serialization.Converters
 
 		static Dictionary<Type, CbConverterInfo> TypeToConverterInfo = new Dictionary<Type, CbConverterInfo>(new KeyValuePair<Type, CbConverterInfo>[]
 		{
+			GetPodConverterInfo(x => x.AsBool(), (w, v) => w.WriteBoolValue(v), (w, n, v) => w.WriteBool(n, v), null),
 			GetPodConverterInfo(x => x.AsString(), (w, v) => w.WriteStringValue(v), (w, n, v) => w.WriteString(n, v), null),
 			GetPodConverterInfo(x => x.AsHash(), (w, v) => w.WriteHashValue(v), (w, n, v) => w.WriteHash(n, v), null)
 			// TODO: other basic types
