Horde: Allow slicing lists starting at a given offset.

diff --git a/Engine/Source/Programs/Shared/EpicGames.Core/ListSegment.cs b/Engine/Source/Programs/Shared/EpicGames.Core/ListSegment.cs
--- a/Engine/Source/Programs/Shared/EpicGames.Core/ListSegment.cs
+++ b/Engine/Source/Programs/Shared/EpicGames.Core/ListSegment.cs
@@ -167,6 +167,18 @@ namespace EpicGames.Core
 	/// </summary>
 	public static class ListSegmentExtensions
 	{
+		/// <summary>
+		/// Create a list segment
+		/// </summary>
+		/// <typeparam name="T"></typeparam>
+		/// <param name="List"></param>
+		/// <param name="Offset"></param>
+		/// <returns></returns>
+		public static ListSegment<T> Slice<T>(this IList<T> List, int Offset)
+		{
+			return new ListSegment<T>(List, Offset, List.Count - Offset);
+		}
+
 		/// <summary>
 		/// Create a list segment
 		/// </summary>
