Horde: Remove unnecessary dependency on LibGit2Sharp from EpicGames.Perforce.Managed.

diff --git a/Engine/Source/Programs/Shared/EpicGames.Perforce.Managed/EpicGames.Perforce.Managed.csproj b/Engine/Source/Programs/Shared/EpicGames.Perforce.Managed/EpicGames.Perforce.Managed.csproj
--- a/Engine/Source/Programs/Shared/EpicGames.Perforce.Managed/EpicGames.Perforce.Managed.csproj
+++ b/Engine/Source/Programs/Shared/EpicGames.Perforce.Managed/EpicGames.Perforce.Managed.csproj
@@ -29,8 +29,5 @@
       <Link>Properties\MetaData.cs</Link>
     </Compile>
   </ItemGroup>
-  <ItemGroup>
-    <PackageReference Include="LibGit2Sharp" Version="0.26.2" />
-  </ItemGroup>
 
 </Project>
