Move DotNETCommon/BuildUtilities to Shared/EpicGames.Build

Also moves DotNETCommon/Metadata.cs to Shared/Metadata.cs

diff --git a/Engine/Build/BatchFiles/BuildUAT.bat b/Engine/Build/BatchFiles/BuildUAT.bat
--- a/Engine/Build/BatchFiles/BuildUAT.bat
+++ b/Engine/Build/BatchFiles/BuildUAT.bat
@@ -16,15 +16,15 @@ rem ## Change the CWD to /Engine/Source.
 pushd "%~dp0..\..\Source"
 if not exist ..\Build\BatchFiles\BuildUAT.bat goto Error_BatchFileInWrongLocation
 
-rem Check to see if the files in the AutomationTool, BuildUtilities, EpicGames.Core, or UnrealBuildTool directory have changed.
+rem Check to see if the files in the AutomationTool, EpicGames.Build, EpicGames.Core, or UnrealBuildTool directory have changed.
 rem find ".cs" files to only lines that match those names - excludes lines that will change for uninteresting reasons, like free space
 md ..\Intermediate\Build >nul 2>nul
 
 dir /s^
  Programs\Shared\EpicGames.Core\*.cs^
  Programs\Shared\EpicGames.Core\*.csproj^
- Programs\DotNETCommon\BuildUtilities\*.cs^
- Programs\DotNETCommon\BuildUtilities\*.csproj^
+ Programs\Shared\EpicGames.Build\*.cs^
+ Programs\Shared\EpicGames.Build\*.csproj^
  Programs\UnrealBuildTool\*.cs^
  Programs\UnrealBuildTool\*.csproj^
  | find ".cs" >..\Intermediate\Build\AutomationToolFiles.txt
diff --git a/Engine/Build/BatchFiles/BuildUAT.sh b/Engine/Build/BatchFiles/BuildUAT.sh
--- a/Engine/Build/BatchFiles/BuildUAT.sh
+++ b/Engine/Build/BatchFiles/BuildUAT.sh
@@ -19,7 +19,7 @@ if [ ! -f ../Build/BatchFiles/BuildUAT.sh ]; then
   exit 1
 fi
 
-# Check to see if the files in the AutomationTool, BuildUtilities, EpicGames.Core, or UnrealBuildTool
+# Check to see if the files in the AutomationTool, EpicGames.Build, EpicGames.Core, or UnrealBuildTool
 # directory have changed.
 
 mkdir -p ../Intermediate/Build
@@ -35,7 +35,7 @@ elif [ ! -f ../Binaries/DotNET/AutomationTool/AutomationTool.dll ]; then
 
 elif [ -f ../Intermediate/Build/AutomationToolLastBuildTime ]; then
   UPDATED_DEP_FILES="$(find \
-    Programs/Shared/EpicGames.Core Programs/DotNETCommon/BuildUtilities Programs/UnrealBuildTool \
+    Programs/Shared/EpicGames.Core Programs/Shared/EpicGames.Build Programs/UnrealBuildTool \
     -type f \
     \( -iname \*.cs -or -iname \*.csproj \) \
     -newer ../Intermediate/Build/AutomationToolLastBuildTime)"
diff --git a/Engine/Source/Programs/AutomationTool/Android/Android.Automation.csproj b/Engine/Source/Programs/AutomationTool/Android/Android.Automation.csproj
--- a/Engine/Source/Programs/AutomationTool/Android/Android.Automation.csproj
+++ b/Engine/Source/Programs/AutomationTool/Android/Android.Automation.csproj
@@ -1,4 +1,4 @@
-﻿<Project Sdk="Microsoft.NET.Sdk">
+<Project Sdk="Microsoft.NET.Sdk">
   <PropertyGroup>
     <TargetFramework>netcoreapp3.1</TargetFramework>
     <Configuration Condition=" '$(Configuration)' == '' ">Development</Configuration>
@@ -34,7 +34,7 @@
     </Reference>
   </ItemGroup>
   <ItemGroup>
-    <Compile Include="..\..\DotNETCommon\MetaData.cs">
+    <Compile Include="..\..\Shared\MetaData.cs">
       <Link>Properties\MetaData.cs</Link>
     </Compile>
   </ItemGroup>
diff --git a/Engine/Source/Programs/AutomationTool/AutomationTool.csproj b/Engine/Source/Programs/AutomationTool/AutomationTool.csproj
--- a/Engine/Source/Programs/AutomationTool/AutomationTool.csproj
+++ b/Engine/Source/Programs/AutomationTool/AutomationTool.csproj
@@ -72,7 +72,7 @@
   </Choose>
 
   <ItemGroup>
-    <Compile Include="..\DotNETCommon\MetaData.cs">
+    <Compile Include="..\Shared\MetaData.cs">
       <Link>Properties\MetaData.cs</Link>
     </Compile>
     <Compile Include="CompileScriptModules.cs" />
@@ -108,8 +108,8 @@
   
   <ItemGroup>
     <ProjectReference Include="..\Shared\EpicGames.Core\EpicGames.Core.csproj" PrivateAssets="all" />
+    <ProjectReference Include="..\Shared\EpicGames.Build\EpicGames.Build.csproj" PrivateAssets="all" />
     <ProjectReference Include="..\UnrealBuildTool\UnrealBuildTool.csproj" PrivateAssets="all" />
-    <ProjectReference Include="..\DotNETCommon\BuildUtilities\BuildUtilities.csproj" PrivateAssets="all" />
   </ItemGroup>
   
   <ItemGroup>
diff --git a/Engine/Source/Programs/AutomationTool/AutomationUtils/AutomationUtils.Automation.csproj b/Engine/Source/Programs/AutomationTool/AutomationUtils/AutomationUtils.Automation.csproj
--- a/Engine/Source/Programs/AutomationTool/AutomationUtils/AutomationUtils.Automation.csproj
+++ b/Engine/Source/Programs/AutomationTool/AutomationUtils/AutomationUtils.Automation.csproj
@@ -39,8 +39,8 @@
 
   <ItemGroup>
     <ProjectReference Include="..\..\Shared\EpicGames.Core\EpicGames.Core.csproj" PrivateAssets="All"><Private>true</Private></ProjectReference>
+    <ProjectReference Include="..\..\Shared\EpicGames.Build\EpicGames.Build.csproj" PrivateAssets="All"> <Private>true</Private> </ProjectReference>
     <ProjectReference Include="..\..\UnrealBuildTool\UnrealBuildTool.csproj" PrivateAssets="All"><Private>true</Private></ProjectReference>
-    <ProjectReference Include="..\..\DotNETCommon\BuildUtilities\BuildUtilities.csproj" PrivateAssets="All"> <Private>true</Private> </ProjectReference>
   </ItemGroup>
   <ItemGroup>
     <PackageReference Include="System.Data.DataSetExtensions" Version="4.5.0" />
diff --git a/Engine/Source/Programs/AutomationTool/BuildGraph/BuildGraph.Automation.csproj b/Engine/Source/Programs/AutomationTool/BuildGraph/BuildGraph.Automation.csproj
--- a/Engine/Source/Programs/AutomationTool/BuildGraph/BuildGraph.Automation.csproj
+++ b/Engine/Source/Programs/AutomationTool/BuildGraph/BuildGraph.Automation.csproj
@@ -37,7 +37,7 @@
     </Reference>
   </ItemGroup>
   <ItemGroup>
-    <Compile Include="..\..\DotNETCommon\MetaData.cs">
+    <Compile Include="..\..\Shared\MetaData.cs">
       <Link>Properties\MetaData.cs</Link>
     </Compile>
   </ItemGroup>
diff --git a/Engine/Source/Programs/AutomationTool/HoloLens/HoloLens.Automation.csproj b/Engine/Source/Programs/AutomationTool/HoloLens/HoloLens.Automation.csproj
--- a/Engine/Source/Programs/AutomationTool/HoloLens/HoloLens.Automation.csproj
+++ b/Engine/Source/Programs/AutomationTool/HoloLens/HoloLens.Automation.csproj
@@ -55,7 +55,7 @@
   <Import Project="Sdk.targets" Sdk="Microsoft.NET.Sdk" Condition="'$(IsWindows)'!='true'" />
 
   <ItemGroup>
-    <Compile Include="..\..\DotNETCommon\MetaData.cs">
+    <Compile Include="..\..\Shared\MetaData.cs">
       <Link>Properties\MetaData.cs</Link>
     </Compile>
   </ItemGroup>
diff --git a/Engine/Source/Programs/AutomationTool/IOS/IOS.Automation.csproj b/Engine/Source/Programs/AutomationTool/IOS/IOS.Automation.csproj
--- a/Engine/Source/Programs/AutomationTool/IOS/IOS.Automation.csproj
+++ b/Engine/Source/Programs/AutomationTool/IOS/IOS.Automation.csproj
@@ -44,7 +44,7 @@
   </ItemGroup>
 
   <ItemGroup>
-    <Compile Include="..\..\DotNETCommon\MetaData.cs">
+    <Compile Include="..\..\Shared\MetaData.cs">
       <Link>Properties\MetaData.cs</Link>
     </Compile>
   </ItemGroup>
diff --git a/Engine/Source/Programs/AutomationTool/Linux/Linux.Automation.csproj b/Engine/Source/Programs/AutomationTool/Linux/Linux.Automation.csproj
--- a/Engine/Source/Programs/AutomationTool/Linux/Linux.Automation.csproj
+++ b/Engine/Source/Programs/AutomationTool/Linux/Linux.Automation.csproj
@@ -1,4 +1,4 @@
-﻿<Project Sdk="Microsoft.NET.Sdk">
+<Project Sdk="Microsoft.NET.Sdk">
   <PropertyGroup>
     <TargetFramework>netcoreapp3.1</TargetFramework>
     <Configuration Condition=" '$(Configuration)' == '' ">Development</Configuration>
@@ -29,7 +29,7 @@
   </PropertyGroup>
 
   <ItemGroup>
-    <Compile Include="..\..\DotNETCommon\MetaData.cs">
+    <Compile Include="..\..\Shared\MetaData.cs">
       <Link>Properties\MetaData.cs</Link>
     </Compile>
   </ItemGroup>
diff --git a/Engine/Source/Programs/AutomationTool/Localization/Localization.Automation.csproj b/Engine/Source/Programs/AutomationTool/Localization/Localization.Automation.csproj
--- a/Engine/Source/Programs/AutomationTool/Localization/Localization.Automation.csproj
+++ b/Engine/Source/Programs/AutomationTool/Localization/Localization.Automation.csproj
@@ -1,4 +1,4 @@
-﻿<Project Sdk="Microsoft.NET.Sdk">
+<Project Sdk="Microsoft.NET.Sdk">
   <PropertyGroup>
     <TargetFramework>netcoreapp3.1</TargetFramework>
     <Configuration Condition=" '$(Configuration)' == '' ">Development</Configuration>
@@ -29,7 +29,7 @@
   </PropertyGroup>
 
   <ItemGroup>
-    <Compile Include="..\..\DotNETCommon\MetaData.cs">
+    <Compile Include="..\..\Shared\MetaData.cs">
       <Link>Properties\MetaData.cs</Link>
     </Compile>
   </ItemGroup>
diff --git a/Engine/Source/Programs/AutomationTool/Lumin/Lumin.Automation.csproj b/Engine/Source/Programs/AutomationTool/Lumin/Lumin.Automation.csproj
--- a/Engine/Source/Programs/AutomationTool/Lumin/Lumin.Automation.csproj
+++ b/Engine/Source/Programs/AutomationTool/Lumin/Lumin.Automation.csproj
@@ -35,7 +35,7 @@
   </ItemGroup>
 
   <ItemGroup>
-    <Compile Include="..\..\DotNETCommon\MetaData.cs">
+    <Compile Include="..\..\Shared\MetaData.cs">
       <Link>Properties\MetaData.cs</Link>
     </Compile>
   </ItemGroup>
diff --git a/Engine/Source/Programs/AutomationTool/Mac/Mac.Automation.csproj b/Engine/Source/Programs/AutomationTool/Mac/Mac.Automation.csproj
--- a/Engine/Source/Programs/AutomationTool/Mac/Mac.Automation.csproj
+++ b/Engine/Source/Programs/AutomationTool/Mac/Mac.Automation.csproj
@@ -1,4 +1,4 @@
-﻿<Project Sdk="Microsoft.NET.Sdk">
+<Project Sdk="Microsoft.NET.Sdk">
   <PropertyGroup>
     <TargetFramework>netcoreapp3.1</TargetFramework>
     <Configuration Condition=" '$(Configuration)' == '' ">Development</Configuration>
@@ -28,7 +28,7 @@
     <DebugType>full</DebugType>
   </PropertyGroup>
   <ItemGroup>
-    <Compile Include="..\..\DotNETCommon\MetaData.cs">
+    <Compile Include="..\..\Shared\MetaData.cs">
       <Link>Properties\MetaData.cs</Link>
     </Compile>
   </ItemGroup>
diff --git a/Engine/Source/Programs/AutomationTool/OneSkyLocalization/OneSkyLocalization.Automation.csproj b/Engine/Source/Programs/AutomationTool/OneSkyLocalization/OneSkyLocalization.Automation.csproj
--- a/Engine/Source/Programs/AutomationTool/OneSkyLocalization/OneSkyLocalization.Automation.csproj
+++ b/Engine/Source/Programs/AutomationTool/OneSkyLocalization/OneSkyLocalization.Automation.csproj
@@ -1,4 +1,4 @@
-﻿<Project Sdk="Microsoft.NET.Sdk">
+<Project Sdk="Microsoft.NET.Sdk">
   <PropertyGroup>
     <TargetFramework>netcoreapp3.1</TargetFramework>
     <Configuration Condition=" '$(Configuration)' == '' ">Development</Configuration>
@@ -36,7 +36,7 @@
   </ItemGroup>
 
   <ItemGroup>
-    <Compile Include="..\..\DotNETCommon\MetaData.cs">
+    <Compile Include="..\..\Shared\MetaData.cs">
       <Link>Properties\MetaData.cs</Link>
     </Compile>
   </ItemGroup>
diff --git a/Engine/Source/Programs/AutomationTool/Scripts/AutomationScripts.Automation.csproj b/Engine/Source/Programs/AutomationTool/Scripts/AutomationScripts.Automation.csproj
--- a/Engine/Source/Programs/AutomationTool/Scripts/AutomationScripts.Automation.csproj
+++ b/Engine/Source/Programs/AutomationTool/Scripts/AutomationScripts.Automation.csproj
@@ -40,7 +40,7 @@
   </ItemGroup>
 
   <ItemGroup>
-    <Compile Include="..\..\DotNETCommon\MetaData.cs">
+    <Compile Include="..\..\Shared\MetaData.cs">
       <Link>Properties\MetaData.cs</Link>
     </Compile>
   </ItemGroup>
diff --git a/Engine/Source/Programs/AutomationTool/TVOS/TVOS.Automation.csproj b/Engine/Source/Programs/AutomationTool/TVOS/TVOS.Automation.csproj
--- a/Engine/Source/Programs/AutomationTool/TVOS/TVOS.Automation.csproj
+++ b/Engine/Source/Programs/AutomationTool/TVOS/TVOS.Automation.csproj
@@ -1,4 +1,4 @@
-﻿<Project Sdk="Microsoft.NET.Sdk">
+<Project Sdk="Microsoft.NET.Sdk">
   <PropertyGroup>
     <TargetFramework>netcoreapp3.1</TargetFramework>
     <Configuration Condition=" '$(Configuration)' == '' ">Development</Configuration>
@@ -35,7 +35,7 @@
   </ItemGroup>
 
   <ItemGroup>
-    <Compile Include="..\..\DotNETCommon\MetaData.cs">
+    <Compile Include="..\..\Shared\MetaData.cs">
       <Link>Properties\MetaData.cs</Link>
     </Compile>
   </ItemGroup>
diff --git a/Engine/Source/Programs/AutomationTool/Turnkey/Turnkey.Automation.csproj b/Engine/Source/Programs/AutomationTool/Turnkey/Turnkey.Automation.csproj
--- a/Engine/Source/Programs/AutomationTool/Turnkey/Turnkey.Automation.csproj
+++ b/Engine/Source/Programs/AutomationTool/Turnkey/Turnkey.Automation.csproj
@@ -54,12 +54,12 @@
   <Import Project="Sdk.targets" Sdk="Microsoft.NET.Sdk" Condition="'$(IsWindows)'!='true'"/>
 
   <ItemGroup>
-    <Compile Include="..\..\DotNETCommon\MetaData.cs">
+    <Compile Include="..\..\Shared\MetaData.cs">
       <Link>Properties\MetaData.cs</Link>
     </Compile>
   </ItemGroup>
   <ItemGroup>
-    <ProjectReference Include="..\..\DotNETCommon\BuildUtilities\BuildUtilities.csproj" PrivateAssets="All"><Private>false</Private></ProjectReference>
+    <ProjectReference Include="..\..\Shared\EpicGames.Build\EpicGames.Build.csproj" PrivateAssets="All"><Private>false</Private></ProjectReference>
     <ProjectReference Include="..\..\Shared\EpicGames.Core\EpicGames.Core.csproj" PrivateAssets="All"><Private>false</Private></ProjectReference>
     <ProjectReference Include="..\..\UnrealBuildTool\UnrealBuildTool.csproj" PrivateAssets="All"><Private>false</Private></ProjectReference>
     <ProjectReference Include="..\AutomationUtils\AutomationUtils.Automation.csproj" PrivateAssets="All"><Private>false</Private></ProjectReference>
diff --git a/Engine/Source/Programs/AutomationTool/Win/Win.Automation.csproj b/Engine/Source/Programs/AutomationTool/Win/Win.Automation.csproj
--- a/Engine/Source/Programs/AutomationTool/Win/Win.Automation.csproj
+++ b/Engine/Source/Programs/AutomationTool/Win/Win.Automation.csproj
@@ -1,4 +1,4 @@
-﻿<Project Sdk="Microsoft.NET.Sdk">
+<Project Sdk="Microsoft.NET.Sdk">
   <PropertyGroup>
     <TargetFramework>netcoreapp3.1</TargetFramework>
     <Configuration Condition=" '$(Configuration)' == '' ">Development</Configuration>
@@ -35,7 +35,7 @@
   </ItemGroup>
 
   <ItemGroup>
-    <Compile Include="..\..\DotNETCommon\MetaData.cs">
+    <Compile Include="..\..\Shared\MetaData.cs">
       <Link>Properties\MetaData.cs</Link>
     </Compile>
   </ItemGroup>
diff --git a/Engine/Source/Programs/AutomationToolLauncher/AutomationToolLauncher.csproj b/Engine/Source/Programs/AutomationToolLauncher/AutomationToolLauncher.csproj
--- a/Engine/Source/Programs/AutomationToolLauncher/AutomationToolLauncher.csproj
+++ b/Engine/Source/Programs/AutomationToolLauncher/AutomationToolLauncher.csproj
@@ -1,4 +1,4 @@
-﻿<Project Sdk="Microsoft.NET.Sdk">
+<Project Sdk="Microsoft.NET.Sdk">
   <PropertyGroup>
     <TargetFramework>netcoreapp3.1</TargetFramework>
     <Configuration Condition=" '$(Configuration)' == '' ">Development</Configuration>
@@ -26,7 +26,7 @@
   </PropertyGroup>
 
   <ItemGroup>
-    <Compile Include="..\DotNETCommon\MetaData.cs">
+    <Compile Include="..\Shared\MetaData.cs">
       <Link>Properties\MetaData.cs</Link>
     </Compile>
   </ItemGroup>
diff --git a/Engine/Source/Programs/DotNETCommon/BuildUtilities/Automation/CommandLine.cs b/Engine/Source/Programs/DotNETCommon/BuildUtilities/Automation/CommandLine.cs
deleted file mode 100644
diff --git a/Engine/Source/Programs/DotNETCommon/BuildUtilities/Automation/ExitCode.cs b/Engine/Source/Programs/DotNETCommon/BuildUtilities/Automation/ExitCode.cs
deleted file mode 100644
diff --git a/Engine/Source/Programs/DotNETCommon/BuildUtilities/Automation/FormatCommandLine.cs b/Engine/Source/Programs/DotNETCommon/BuildUtilities/Automation/FormatCommandLine.cs
deleted file mode 100644
diff --git a/Engine/Source/Programs/DotNETCommon/BuildUtilities/Automation/ProcessSingleton.cs b/Engine/Source/Programs/DotNETCommon/BuildUtilities/Automation/ProcessSingleton.cs
deleted file mode 100644
diff --git a/Engine/Source/Programs/DotNETCommon/BuildUtilities/BuildUtilities.csproj b/Engine/Source/Programs/DotNETCommon/BuildUtilities/BuildUtilities.csproj
deleted file mode 100644
diff --git a/Engine/Source/Programs/DotNETCommon/BuildUtilities/Platform/Mac/ApplePlatformSDK.cs b/Engine/Source/Programs/DotNETCommon/BuildUtilities/Platform/Mac/ApplePlatformSDK.cs
deleted file mode 100644
diff --git a/Engine/Source/Programs/DotNETCommon/BuildUtilities/System/DirectoryItem.cs b/Engine/Source/Programs/DotNETCommon/BuildUtilities/System/DirectoryItem.cs
deleted file mode 100644
diff --git a/Engine/Source/Programs/DotNETCommon/BuildUtilities/System/DirectoryLookupCache.cs b/Engine/Source/Programs/DotNETCommon/BuildUtilities/System/DirectoryLookupCache.cs
deleted file mode 100644
diff --git a/Engine/Source/Programs/DotNETCommon/BuildUtilities/System/EnumeratePlugins.cs b/Engine/Source/Programs/DotNETCommon/BuildUtilities/System/EnumeratePlugins.cs
deleted file mode 100644
diff --git a/Engine/Source/Programs/DotNETCommon/BuildUtilities/System/FileItem.cs b/Engine/Source/Programs/DotNETCommon/BuildUtilities/System/FileItem.cs
deleted file mode 100644
diff --git a/Engine/Source/Programs/DotNETCommon/BuildUtilities/System/NativeProjectsBase.cs b/Engine/Source/Programs/DotNETCommon/BuildUtilities/System/NativeProjectsBase.cs
deleted file mode 100644
diff --git a/Engine/Source/Programs/DotNETCommon/BuildUtilities/System/Rules.cs b/Engine/Source/Programs/DotNETCommon/BuildUtilities/System/Rules.cs
deleted file mode 100644
diff --git a/Engine/Source/Programs/DotNETCommon/BuildUtilities/System/StartupTraceListener.cs b/Engine/Source/Programs/DotNETCommon/BuildUtilities/System/StartupTraceListener.cs
deleted file mode 100644
diff --git a/Engine/Source/Programs/DotNETCommon/BuildUtilities/UEBuildPlatformSDK.cs b/Engine/Source/Programs/DotNETCommon/BuildUtilities/UEBuildPlatformSDK.cs
deleted file mode 100644
diff --git a/Engine/Source/Programs/DotNETCommon/BuildUtilities/Unreal.cs b/Engine/Source/Programs/DotNETCommon/BuildUtilities/Unreal.cs
deleted file mode 100644
diff --git a/Engine/Source/Programs/DotNETCommon/MetaData.cs b/Engine/Source/Programs/DotNETCommon/MetaData.cs
deleted file mode 100644
diff --git a/Engine/Source/Programs/Shared/EpicGames.Build/Automation/CommandLine.cs b/Engine/Source/Programs/Shared/EpicGames.Build/Automation/CommandLine.cs
new file mode 100644
--- /dev/null
+++ b/Engine/Source/Programs/Shared/EpicGames.Build/Automation/CommandLine.cs
@@ -0,0 +1,95 @@
+﻿// Copyright Epic Games, Inc. All Rights Reserved.
+
+using System;
+using System.Collections.Generic;
+using System.Text;
+
+namespace UnrealBuildBase
+{
+	/// <summary>
+	/// Command to execute info.
+	/// </summary>
+	public class CommandInfo
+	{
+		public string CommandName;
+		public List<string> Arguments = new List<string>();
+
+		public CommandInfo(string CommandNameIn)
+		{
+			CommandName = CommandNameIn;
+		}
+
+		public override string ToString()
+		{
+			string Result = CommandName;
+			Result += "(";
+			for (int Index = 0; Index < Arguments.Count; ++Index)
+			{
+				if (Index > 0)
+				{
+					Result += ", ";
+				}
+				Result += Arguments[Index];
+			}
+			Result += ")";
+			return Result;
+		}
+	}
+
+	public partial class ParsedCommandLine
+	{
+		public Dictionary<string, string> GlobalParameters
+		{
+			get;
+			private set;
+		}
+
+		public ParsedCommandLine(Dictionary<string, string> GlobalParametersIn)
+		{
+			GlobalParameters = new Dictionary<string, string>(GlobalParametersIn, StringComparer.InvariantCultureIgnoreCase);
+		}
+
+		HashSet<string> GlobalParameterValues = new HashSet<string>(StringComparer.InvariantCultureIgnoreCase);
+		public void SetGlobal(string Parameter)
+		{
+			if (!GlobalParameters.ContainsKey(Parameter))
+            {
+				throw new Exception($"unknown global parameter {Parameter}");
+            }
+			GlobalParameterValues.Add(Parameter);
+		}
+		public bool TrySetGlobal(string Parameter)
+		{
+			if (!GlobalParameters.ContainsKey(Parameter))
+            {
+				return false;
+            }
+			GlobalParameterValues.Add(Parameter);
+			return true;
+		}
+		public bool IsSetGlobal(string Parameter)
+        {
+			if (!GlobalParameters.ContainsKey(Parameter))
+            {
+				throw new Exception($"unknown global parameter {Parameter}");
+            }
+			return GlobalParameterValues.Contains(Parameter);
+        }
+
+		Dictionary<string, object?> UncheckedParameters = new Dictionary<string, object?>(StringComparer.InvariantCultureIgnoreCase);
+		public void SetUnchecked(string Name, object? Value = null)
+        {
+			UncheckedParameters[Name] = Value;
+        }
+		public bool IsSetUnchecked(string Name)
+        {
+			return UncheckedParameters.ContainsKey(Name);
+        }
+		public object? GetValueUnchecked(string Name)
+        {
+            return UncheckedParameters.TryGetValue(Name, out object? Out) ? Out : null;
+        }
+
+		public List<CommandInfo> CommandsToExecute = new List<CommandInfo>();
+	}
+}
diff --git a/Engine/Source/Programs/Shared/EpicGames.Build/Automation/ExitCode.cs b/Engine/Source/Programs/Shared/EpicGames.Build/Automation/ExitCode.cs
new file mode 100644
--- /dev/null
+++ b/Engine/Source/Programs/Shared/EpicGames.Build/Automation/ExitCode.cs
@@ -0,0 +1,67 @@
+// Copyright Epic Games, Inc. All Rights Reserved.
+
+using System;
+using System.Collections.Generic;
+using System.Text;
+
+namespace UnrealBuildBase
+{
+	// NOTE: this needs to be kept in sync with EditorAnalytics.h and iPhonePackager.cs
+	public enum ExitCode
+    {
+        Error_UATNotFound = -1,
+        Success = 0,
+        Error_Unknown = 1,
+        Error_Arguments = 2,
+        Error_UnknownCommand = 3,
+        Error_SDKNotFound = 10,
+        Error_ProvisionNotFound = 11,
+        Error_CertificateNotFound = 12,
+        Error_ProvisionAndCertificateNotFound = 13,
+        Error_InfoPListNotFound = 14,
+        Error_KeyNotFoundInPList = 15,
+        Error_ProvisionExpired = 16,
+        Error_CertificateExpired = 17,
+        Error_CertificateProvisionMismatch = 18,
+        Error_CodeUnsupported = 19,
+        Error_PluginsUnsupported = 20,
+        Error_UnknownCookFailure = 25,
+        Error_UnknownDeployFailure = 26,
+        Error_UnknownBuildFailure = 27,
+        Error_UnknownPackageFailure = 28,
+        Error_UnknownLaunchFailure = 29,
+        Error_StageMissingFile = 30,
+        Error_FailedToCreateIPA = 31,
+        Error_FailedToCodeSign = 32,
+        Error_DeviceBackupFailed = 33,
+        Error_AppUninstallFailed = 34,
+        Error_AppInstallFailed = 35,
+        Error_AppNotFound = 36,
+        Error_StubNotSignedCorrectly = 37,
+        Error_IPAMissingInfoPList = 38,
+        Error_DeleteFile = 39,
+        Error_DeleteDirectory = 40,
+        Error_CreateDirectory = 41,
+        Error_CopyFile = 42,
+        Error_OnlyOneObbFileSupported = 50,
+        Error_FailureGettingPackageInfo = 51,
+        Error_OnlyOneTargetConfigurationSupported = 52,
+        Error_ObbNotFound = 53,
+        Error_AndroidBuildToolsPathNotFound = 54,
+        Error_NoApkSuitableForArchitecture = 55,
+        Error_FilesInstallFailed = 56,
+        Error_RemoteCertificatesNotFound = 57,
+        Error_LauncherFailed = 100,
+        Error_UATLaunchFailure = 101,
+        Error_FailedToDeleteStagingDirectory = 102,
+        Error_MissingExecutable = 103,
+        Error_DeviceNotSetupForDevelopment = 150,
+        Error_DeviceOSNewerThanSDK = 151,
+		Error_TestFailure = 152,
+		Error_SymbolizedSONotFound = 153,
+		Error_LicenseNotAccepted = 154,
+		Error_AndroidOBBError = 155,
+		Error_SDKInstallFailed = 200,
+		Error_DeviceUpdateFailed = 201,
+	};
+}
diff --git a/Engine/Source/Programs/Shared/EpicGames.Build/Automation/FormatCommandLine.cs b/Engine/Source/Programs/Shared/EpicGames.Build/Automation/FormatCommandLine.cs
new file mode 100644
--- /dev/null
+++ b/Engine/Source/Programs/Shared/EpicGames.Build/Automation/FormatCommandLine.cs
@@ -0,0 +1,56 @@
+// Copyright Epic Games, Inc. All Rights Reserved.
+
+using System;
+using System.Collections.Generic;
+using System.Text;
+
+namespace UnrealBuildBase
+{
+	public partial class CommandLine
+	{
+		/// <summary>
+		/// Converts a list of arguments to a string where each argument is separated with a space character.
+		/// </summary>
+		/// <param name="Args">Arguments</param>
+		/// <returns>Single string containing all arguments separated with a space.</returns>
+		public static string FormatCommandLine(IEnumerable<string> Arguments)
+		{
+			StringBuilder Result = new StringBuilder();
+			foreach(string Argument in Arguments)
+			{
+				if(Result.Length > 0)
+				{
+					Result.Append(" ");
+				}
+				Result.Append(FormatArgumentForCommandLine(Argument));
+			}
+			return Result.ToString();
+		}
+
+		/// <summary>
+		/// Format a single argument for passing on the command line, inserting quotes as necessary.
+		/// </summary>
+		/// <param name="Argument">The argument to quote</param>
+		/// <returns>The argument, with quotes if necessary</returns>
+		public static string FormatArgumentForCommandLine(string Argument)
+		{
+			// Check if the argument contains a space. If not, we can just pass it directly.
+			int SpaceIdx = Argument.IndexOf(' ');
+			if(SpaceIdx == -1)
+			{
+				return Argument;
+			}
+
+			// If it does have a space, and it's formatted as an option (ie. -Something=), try to insert quotes after the equals character
+			int EqualsIdx = Argument.IndexOf('=');
+			if(Argument.StartsWith("-") && EqualsIdx != -1 && EqualsIdx < SpaceIdx)
+			{
+				return String.Format("{0}=\"{1}\"", Argument.Substring(0, EqualsIdx), Argument.Substring(EqualsIdx + 1));
+			}
+			else
+			{
+				return String.Format("\"{0}\"", Argument);
+			}
+		}
+	}
+}
diff --git a/Engine/Source/Programs/Shared/EpicGames.Build/Automation/ProcessSingleton.cs b/Engine/Source/Programs/Shared/EpicGames.Build/Automation/ProcessSingleton.cs
new file mode 100644
--- /dev/null
+++ b/Engine/Source/Programs/Shared/EpicGames.Build/Automation/ProcessSingleton.cs
@@ -0,0 +1,86 @@
+﻿// Copyright Epic Games, Inc. All Rights Reserved.
+
+using EpicGames.Core;
+using System;
+using System.Collections.Generic;
+using System.IO;
+using System.Reflection;
+using System.Text;
+using System.Threading;
+
+#nullable disable
+
+namespace UnrealBuildBase
+{
+	public class ProcessSingleton
+	{
+		public static string GetUniqueMutexForPath(string Name, string UniquePath)
+		{
+			// generate a md5 hash of the path, as GetHashCode is not guaranteed to generate a stable hash
+			return string.Format("Global\\{0}_{1}", Name, ContentHash.MD5(UniquePath.ToUpperInvariant()));
+		}
+
+		/// <summary>
+		/// Returns true/false based on whether this is the only instance
+		/// running (checked at startup).
+		/// </summary>
+		public static bool IsSoleInstance { get; private set; }
+
+		/// <summary>
+		/// Runs the specified delegate checking if this is the only instance of the application.
+		/// </summary>
+		/// <param name="Main"></param>
+		/// <param name="Param"></param>
+		public static ExitCode RunSingleInstance(Func<ExitCode> Main, bool bWaitForUATMutex)
+		{
+			bool AllowMultipleInsances = (Environment.GetEnvironmentVariable("uebp_UATMutexNoWait") == "1");
+	
+            string EntryAssemblyLocation = Assembly.GetEntryAssembly().GetOriginalLocation();
+
+			string MutexName = GetUniqueMutexForPath(Path.GetFileNameWithoutExtension(EntryAssemblyLocation), EntryAssemblyLocation);
+			using (Mutex SingleInstanceMutex = new Mutex(true, MutexName, out bool bCreatedMutex))
+			{
+				IsSoleInstance = bCreatedMutex;
+
+				if (!IsSoleInstance && AllowMultipleInsances == false)
+				{
+					if (bWaitForUATMutex)
+					{
+						Log.TraceWarning("Another instance of UAT at '{0}' is running, and the -WaitForUATMutex parameter has been used. Waiting for other UAT to finish...", EntryAssemblyLocation);
+						int Seconds = 0;
+						while (WaitMutexNoExceptions(SingleInstanceMutex, 15 * 1000) == false)
+						{
+							Seconds += 15;
+							Log.TraceInformation("Still waiting for Mutex. {0} seconds has passed...", Seconds);
+						}
+					}
+					else
+					{
+						throw new Exception($"A conflicting instance of AutomationTool is already running. Current location: {EntryAssemblyLocation}. A process manager may be used to determine the conflicting process and what tool may have launched it");
+					}
+				}
+
+				ExitCode Result = Main();
+
+				if (IsSoleInstance)
+				{
+					SingleInstanceMutex.ReleaseMutex();
+				}
+
+				return Result;
+			}
+		}
+
+		static bool WaitMutexNoExceptions(Mutex Mutex, int TimeoutMs)
+		{
+			try
+			{
+				return Mutex.WaitOne(15 * 1000);
+			}
+			catch (AbandonedMutexException)
+			{
+				return true;
+			}
+		}
+	}
+}
diff --git a/Engine/Source/Programs/Shared/EpicGames.Build/EpicGames.Build.csproj b/Engine/Source/Programs/Shared/EpicGames.Build/EpicGames.Build.csproj
new file mode 100644
--- /dev/null
+++ b/Engine/Source/Programs/Shared/EpicGames.Build/EpicGames.Build.csproj
@@ -0,0 +1,37 @@
+﻿<Project Sdk="Microsoft.NET.Sdk">
+  <PropertyGroup>
+    <TargetFramework>netcoreapp3.1</TargetFramework>
+    <Configuration Condition=" '$(Configuration)' == '' ">Development</Configuration>
+    <OutputType>Library</OutputType>
+    <RootNamespace>UnrealBuildBase</RootNamespace>
+    <GenerateAssemblyInfo>false</GenerateAssemblyInfo>
+    <GenerateTargetFrameworkAttribute>false</GenerateTargetFrameworkAttribute>
+    <Configurations>Debug;Release;Development</Configurations>
+    <Nullable>enable</Nullable>
+  </PropertyGroup>
+  <PropertyGroup Condition=" '$(Configuration)|$(Platform)' == 'Debug|AnyCPU' ">
+    <DocumentationFile></DocumentationFile>
+  </PropertyGroup>
+  <PropertyGroup Condition="'$(Configuration)|$(Platform)' == 'Development|AnyCPU'">
+    <DocumentationFile></DocumentationFile>
+    <Optimize>true</Optimize>
+  </PropertyGroup>
+  <PropertyGroup Condition="'$(Configuration)|$(Platform)' == 'Release|AnyCPU'">
+    <DocumentationFile></DocumentationFile>
+    <Optimize>true</Optimize>
+  </PropertyGroup>
+  
+  <ItemGroup>
+    <Compile Include="..\MetaData.cs">
+      <Link>Properties\MetaData.cs</Link>
+    </Compile>
+  </ItemGroup>
+  <ItemGroup>
+    <ProjectReference Include="..\..\Shared\EpicGames.Core\EpicGames.Core.csproj" />
+  </ItemGroup>
+  <ItemGroup>
+    <PackageReference Include="Microsoft.CSharp" Version="4.7.0" />
+    <PackageReference Include="Microsoft.Win32.Registry" Version="4.7.0" />
+    <PackageReference Include="System.Data.DataSetExtensions" Version="4.5.0" />
+  </ItemGroup>
+</Project>
\ No newline at end of file
diff --git a/Engine/Source/Programs/Shared/EpicGames.Build/Platform/Mac/ApplePlatformSDK.cs b/Engine/Source/Programs/Shared/EpicGames.Build/Platform/Mac/ApplePlatformSDK.cs
new file mode 100644
--- /dev/null
+++ b/Engine/Source/Programs/Shared/EpicGames.Build/Platform/Mac/ApplePlatformSDK.cs
@@ -0,0 +1,173 @@
+// Copyright Epic Games, Inc. All Rights Reserved.
+
+using System;
+using System.Diagnostics;
+using System.IO;
+using System.Text;
+using System.Text.RegularExpressions;
+using EpicGames.Core;
+using Microsoft.VisualBasic.CompilerServices;
+using Microsoft.Win32;
+
+namespace UnrealBuildBase
+{
+	static public class ApplePlatformSDK
+	{
+		public static readonly string? InstalledSDKVersion = GetInstalledSDKVersion();
+
+		private static string? GetInstalledSDKVersion()
+		{
+			// get xcode version on Mac
+			if (RuntimePlatform.IsMac)
+			{
+				string Output = RunLocalProcessAndReturnStdOut("sh", "-c 'xcodebuild -version'");
+				Match Result = Regex.Match(Output, @"Xcode (\S*)");
+				return Result.Success ? Result.Groups[1].Value : null;
+			}
+
+			if (RuntimePlatform.IsWindows)
+			{
+				// otherwise, get iTunes "Version"
+				string? DllPath =
+					Registry.GetValue(
+						"HKEY_LOCAL_MACHINE\\SOFTWARE\\Wow6432Node\\Apple Inc.\\Apple Mobile Device Support\\Shared",
+						"iTunesMobileDeviceDLL", null) as string;
+				if (string.IsNullOrEmpty(DllPath) || !File.Exists(DllPath))
+				{
+					DllPath = Registry.GetValue(
+						"HKEY_LOCAL_MACHINE\\SOFTWARE\\Wow6432Node\\Apple Inc.\\Apple Mobile Device Support\\Shared",
+						"MobileDeviceDLL", null) as string;
+					if (string.IsNullOrEmpty(DllPath) || !File.Exists(DllPath))
+					{
+						// iTunes >= 12.7 doesn't have a key specifying the 32-bit DLL but it does have a ASMapiInterfaceDLL key and MobileDevice.dll is in usually in the same directory
+						DllPath = Registry.GetValue(
+							"HKEY_LOCAL_MACHINE\\SOFTWARE\\Wow6432Node\\Apple Inc.\\Apple Mobile Device Support\\Shared",
+							"ASMapiInterfaceDLL", null) as string;
+						DllPath = String.IsNullOrEmpty(DllPath)
+							? null
+							: DllPath.Substring(0, DllPath.LastIndexOf('\\') + 1) + "MobileDevice.dll";
+
+						if (string.IsNullOrEmpty(DllPath) || !File.Exists(DllPath))
+						{
+							DllPath = FindWindowsStoreITunesDLL();
+						}
+					}
+				}
+
+				if (!string.IsNullOrEmpty(DllPath) && File.Exists(DllPath))
+				{
+					return FileVersionInfo.GetVersionInfo(DllPath).FileVersion;
+				}
+			}
+
+			return null;
+		}
+
+		private static string? FindWindowsStoreITunesDLL()
+		{
+			string? InstallPath = null;
+
+			string PackagesKeyName = "Software\\Classes\\Local Settings\\Software\\Microsoft\\Windows\\CurrentVersion\\AppModel\\PackageRepository\\Packages";
+
+			RegistryKey PackagesKey = Registry.LocalMachine.OpenSubKey(PackagesKeyName);
+			if (PackagesKey != null)
+			{
+				string[] PackageSubKeyNames = PackagesKey.GetSubKeyNames();
+
+				foreach (string PackageSubKeyName in PackageSubKeyNames)
+				{
+					if (PackageSubKeyName.Contains("AppleInc.iTunes") && (PackageSubKeyName.Contains("_x64") || PackageSubKeyName.Contains("_x86")))
+					{
+						string FullPackageSubKeyName = PackagesKeyName + "\\" + PackageSubKeyName;
+
+						RegistryKey iTunesKey = Registry.LocalMachine.OpenSubKey(FullPackageSubKeyName);
+						if (iTunesKey != null)
+						{
+							InstallPath = (string)iTunesKey.GetValue("Path") + "\\AMDS32\\MobileDevice.dll";
+							break;
+						}
+					}
+				}
+			}
+
+			return InstallPath;
+		}
+		
+		/// <summary>
+		/// Runs a command line process, and returns simple StdOut output. This doesn't handle errors or return codes
+		/// </summary>
+		/// <returns>The entire StdOut generated from the process as a single trimmed string</returns>
+		/// <param name="Command">Command to run</param>
+		/// <param name="Args">Arguments to Command</param>
+		private static string RunLocalProcessAndReturnStdOut(string Command, string Args)
+		{
+			int ExitCode;
+			return RunLocalProcessAndReturnStdOut(Command, Args, out ExitCode);	
+		}
+
+		/// <summary>
+		/// Runs a command line process, and returns simple StdOut output.
+		/// </summary>
+		/// <returns>The entire StdOut generated from the process as a single trimmed string</returns>
+		/// <param name="Command">Command to run</param>
+		/// <param name="Args">Arguments to Command</param>
+		/// <param name="ExitCode">The return code from the process after it exits</param>
+		/// <param name="LogOutput">Whether to also log standard output and standard error</param>
+		private static string RunLocalProcessAndReturnStdOut(string Command, string? Args, out int ExitCode, bool LogOutput = false)
+		{
+			// Process Arguments follow windows conventions in .NET Core
+			// Which means single quotes ' are not considered quotes.
+			// see [REDACTED_PROFILE_URL]/runtime/issues/29857
+			// also see UE-102580
+			// for rules see https://docs.microsoft.com/en-us/cpp/cpp/main-function-command-line-args
+			Args = Args?.Replace('\'', '\"');
+
+			ProcessStartInfo StartInfo = new ProcessStartInfo(Command, Args);
+			StartInfo.UseShellExecute = false;
+			StartInfo.RedirectStandardInput = true;
+			StartInfo.RedirectStandardOutput = true;
+			StartInfo.RedirectStandardError = true;
+			StartInfo.CreateNoWindow = true;
+			StartInfo.StandardOutputEncoding = Encoding.UTF8;
+
+			string FullOutput = "";
+			string ErrorOutput = "";
+			using (Process LocalProcess = Process.Start(StartInfo))
+			{
+				StreamReader OutputReader = LocalProcess.StandardOutput;
+				// trim off any extraneous new lines, helpful for those one-line outputs
+				FullOutput = OutputReader.ReadToEnd().Trim();
+
+				StreamReader ErrorReader = LocalProcess.StandardError;
+				// trim off any extraneous new lines, helpful for those one-line outputs
+				ErrorOutput = ErrorReader.ReadToEnd().Trim();
+				if (LogOutput)
+				{
+					if(FullOutput.Length > 0)
+					{
+						Log.TraceInformation(FullOutput);
+					}
+
+					if (ErrorOutput.Length > 0)
+					{
+						Log.TraceError(ErrorOutput);
+					}
+				}
+
+				LocalProcess.WaitForExit();
+				ExitCode = LocalProcess.ExitCode;
+			}
+
+			// trim off any extraneous new lines, helpful for those one-line outputs
+			if (ErrorOutput.Length > 0)
+			{
+				if (FullOutput.Length > 0)
+				{
+					FullOutput += Environment.NewLine;
+				}
+				FullOutput += ErrorOutput;
+			}
+			return FullOutput;
+		}
+	}
+}
\ No newline at end of file
diff --git a/Engine/Source/Programs/Shared/EpicGames.Build/System/DirectoryItem.cs b/Engine/Source/Programs/Shared/EpicGames.Build/System/DirectoryItem.cs
new file mode 100644
--- /dev/null
+++ b/Engine/Source/Programs/Shared/EpicGames.Build/System/DirectoryItem.cs
@@ -0,0 +1,359 @@
+// Copyright Epic Games, Inc. All Rights Reserved.
+
+using System;
+using System.Collections.Concurrent;
+using System.Collections.Generic;
+using System.Diagnostics.CodeAnalysis;
+using System.IO;
+using System.Linq;
+using System.Text;
+using System.Threading.Tasks;
+using EpicGames.Core;
+
+namespace UnrealBuildBase
+{
+	/// <summary>
+	/// Stores the state of a directory. May or may not exist.
+	/// </summary>
+	public class DirectoryItem
+	{
+		/// <summary>
+		/// Full path to the directory on disk
+		/// </summary>
+		public readonly DirectoryReference Location;
+
+		/// <summary>
+		/// Cached value for whether the directory exists
+		/// </summary>
+		DirectoryInfo Info;
+
+		/// <summary>
+		/// Cached map of name to subdirectory item
+		/// </summary>
+		Dictionary<string, DirectoryItem>? Directories;
+
+		/// <summary>
+		/// Cached map of name to file
+		/// </summary>
+		Dictionary<string, FileItem>? Files;
+
+		/// <summary>
+		/// Global map of location to item
+		/// </summary>
+		static ConcurrentDictionary<DirectoryReference, DirectoryItem> LocationToItem = new ConcurrentDictionary<DirectoryReference, DirectoryItem>();
+
+		/// <summary>
+		/// Constructor
+		/// </summary>
+		/// <param name="Location">Path to this directory</param>
+		/// <param name="Info">Information about this directory</param>
+		private DirectoryItem(DirectoryReference Location, DirectoryInfo Info)
+		{
+			this.Location = Location;
+			this.Info = Info;
+		}
+
+		/// <summary>
+		/// The name of this directory
+		/// </summary>
+		public string Name
+		{
+			get { return Info.Name; }
+		}
+
+		/// <summary>
+		/// The full name of this directory
+		/// </summary>
+		public string FullName
+		{
+			get { return Location.FullName; }
+		}
+
+		/// <summary>
+		/// Whether the directory exists or not
+		/// </summary>
+		public bool Exists
+		{
+			get { return Info.Exists; }
+		}
+
+		/// <summary>
+		/// The last write time of the file.
+		/// </summary>
+		public DateTime LastWriteTimeUtc
+		{
+			get { return Info.LastWriteTimeUtc; }
+		}
+
+		/// <summary>
+		/// Gets the parent directory item
+		/// </summary>
+		public DirectoryItem? GetParentDirectoryItem()
+		{
+			if(Info.Parent == null)
+			{
+				return null;
+			}
+			else
+			{
+				return GetItemByDirectoryInfo(Info.Parent);
+			}
+		}
+
+		/// <summary>
+		/// Gets a new directory item by combining the existing directory item with the given path fragments
+		/// </summary>
+		/// <param name="BaseDirectory">Base directory to append path fragments to</param>
+		/// <param name="Fragments">The path fragments to append</param>
+		/// <returns>Directory item corresponding to the combined path</returns>
+		public static DirectoryItem Combine(DirectoryItem BaseDirectory, params string[] Fragments)
+		{
+			return DirectoryItem.GetItemByDirectoryReference(DirectoryReference.Combine(BaseDirectory.Location, Fragments));
+		}
+
+		/// <summary>
+		/// Finds or creates a directory item from its location
+		/// </summary>
+		/// <param name="Location">Path to the directory</param>
+		/// <returns>The directory item for this location</returns>
+		public static DirectoryItem GetItemByPath(string Location)
+		{
+			return GetItemByDirectoryReference(new DirectoryReference(Location));
+		}
+
+		/// <summary>
+		/// Finds or creates a directory item from its location
+		/// </summary>
+		/// <param name="Location">Path to the directory</param>
+		/// <returns>The directory item for this location</returns>
+		public static DirectoryItem GetItemByDirectoryReference(DirectoryReference Location)
+		{
+			DirectoryItem? Result;
+			if(!LocationToItem.TryGetValue(Location, out Result))
+			{
+				DirectoryItem NewItem = new DirectoryItem(Location, new DirectoryInfo(Location.FullName));
+				if(LocationToItem.TryAdd(Location, NewItem))
+				{
+					Result = NewItem;
+				}
+				else
+				{
+					Result = LocationToItem[Location];
+				}
+			}
+			return Result;
+		}
+
+		/// <summary>
+		/// Finds or creates a directory item from a DirectoryInfo object
+		/// </summary>
+		/// <param name="Info">Path to the directory</param>
+		/// <returns>The directory item for this location</returns>
+		public static DirectoryItem GetItemByDirectoryInfo(DirectoryInfo Info)
+		{
+			DirectoryReference Location = new DirectoryReference(Info);
+
+			DirectoryItem? Result;
+			if(!LocationToItem.TryGetValue(Location, out Result))
+			{
+				DirectoryItem NewItem = new DirectoryItem(Location, Info);
+				if(LocationToItem.TryAdd(Location, NewItem))
+				{
+					Result = NewItem;
+				}
+				else
+				{
+					Result = LocationToItem[Location];
+				}
+			}
+			return Result;
+		}
+
+		/// <summary>
+		/// Reset the contents of the directory and allow them to be fetched again
+		/// </summary>
+		public void ResetCachedInfo()
+		{
+			Info = new DirectoryInfo(Info.FullName);
+
+			Dictionary<string, DirectoryItem>? PrevDirectories = Directories;
+			if(PrevDirectories != null)
+			{
+				foreach(DirectoryItem SubDirectory in PrevDirectories.Values)
+				{
+					SubDirectory.ResetCachedInfo();
+				}
+				Directories = null;
+			}
+
+			Dictionary<string, FileItem>? PrevFiles = Files;
+			if(PrevFiles != null)
+			{
+				foreach(FileItem File in PrevFiles.Values)
+				{
+					File.ResetCachedInfo();
+				}
+				Files = null;
+			}
+		}
+
+		/// <summary>
+		/// Resets all cached directory info. Significantly reduces performance; do not use unless strictly necessary.
+		/// </summary>
+		public static void ResetAllCachedInfo_SLOW()
+		{
+			foreach(DirectoryItem Item in LocationToItem.Values)
+			{
+				Item.Info = new DirectoryInfo(Item.Info.FullName);
+				Item.Directories = null;
+				Item.Files = null;
+			}
+			FileItem.ResetAllCachedInfo_SLOW();
+		}
+
+		/// <summary>
+		/// Caches the subdirectories of this directories
+		/// </summary>
+		public void CacheDirectories()
+		{
+			if(Directories == null)
+			{
+				Dictionary<string, DirectoryItem> NewDirectories = new Dictionary<string, DirectoryItem>(DirectoryReference.Comparer);
+				if(Info.Exists)
+				{
+					foreach(DirectoryInfo SubDirectoryInfo in Info.EnumerateDirectories())
+					{
+						if(SubDirectoryInfo.Name.Length == 1 && SubDirectoryInfo.Name[0] == '.')
+						{
+							continue;
+						}
+						else if(SubDirectoryInfo.Name.Length == 2 && SubDirectoryInfo.Name[0] == '.' && SubDirectoryInfo.Name[1] == '.')
+						{
+							continue;
+						}
+						else
+						{
+							NewDirectories[SubDirectoryInfo.Name] = DirectoryItem.GetItemByDirectoryInfo(SubDirectoryInfo);
+						}
+					}
+				}
+				Directories = NewDirectories;
+			}
+		}
+
+		/// <summary>
+		/// Enumerates all the subdirectories
+		/// </summary>
+		/// <returns>Sequence of subdirectory items</returns>
+		public IEnumerable<DirectoryItem> EnumerateDirectories()
+		{
+			CacheDirectories();
+			return Directories!.Values;
+		}
+
+		/// <summary>
+		/// Attempts to get a sub-directory by name
+		/// </summary>
+		/// <param name="Name">Name of the directory</param>
+		/// <param name="OutDirectory">If successful receives the matching directory item with this name</param>
+		/// <returns>True if the file exists, false otherwise</returns>
+		public bool TryGetDirectory(string Name, [NotNullWhen(true)] out DirectoryItem? OutDirectory)
+		{
+			if(Name.Length > 0 && Name[0] == '.')
+			{
+				if(Name.Length == 1)
+				{
+					OutDirectory = this;
+					return true;
+				}
+				else if(Name.Length == 2 && Name[1] == '.')
+				{
+					OutDirectory = GetParentDirectoryItem();
+					return OutDirectory != null;
+				}
+			}
+
+			CacheDirectories();
+			return Directories!.TryGetValue(Name, out OutDirectory);
+		}
+
+		/// <summary>
+		/// Caches the files in this directory
+		/// </summary>
+		public void CacheFiles()
+		{
+			if(Files == null)
+			{
+				Dictionary<string, FileItem> NewFiles = new Dictionary<string, FileItem>(FileReference.Comparer);
+				if(Info.Exists)
+				{
+					foreach(FileInfo FileInfo in Info.EnumerateFiles())
+					{
+						FileItem FileItem = FileItem.GetItemByFileInfo(FileInfo);
+						FileItem.UpdateCachedDirectory(this);
+						NewFiles[FileInfo.Name] = FileItem;
+					}
+				}
+				Files = NewFiles;
+			}
+		}
+
+		/// <summary>
+		/// Enumerates all the files
+		/// </summary>
+		/// <returns>Sequence of FileItems</returns>
+		public IEnumerable<FileItem> EnumerateFiles()
+		{
+			CacheFiles();
+			return Files!.Values;
+		}
+
+		/// <summary>
+		/// Attempts to get a file from this directory by name. Unlike creating a file item and checking whether it exists, this will
+		/// not create a permanent FileItem object if it does not exist.
+		/// </summary>
+		/// <param name="Name">Name of the file</param>
+		/// <param name="OutFile">If successful receives the matching file item with this name</param>
+		/// <returns>True if the file exists, false otherwise</returns>
+		public bool TryGetFile(string Name, [NotNullWhen(true)] out FileItem? OutFile)
+		{
+			CacheFiles();
+			return Files!.TryGetValue(Name, out OutFile);
+		}
+
+		/// <summary>
+		/// Formats this object as a string for debugging
+		/// </summary>
+		/// <returns>Location of the directory</returns>
+		public override string ToString()
+		{
+			return Location.FullName;
+		}
+	}
+
+	/// <summary>
+	/// Helper functions for serialization
+	/// </summary>
+	public static class DirectoryItemExtensionMethods
+	{
+		/// <summary>
+		/// Read a directory item from a binary archive
+		/// </summary>
+		/// <param name="Reader">Reader to serialize data from</param>
+		/// <returns>Instance of the serialized directory item</returns>
+		public static DirectoryItem ReadDirectoryItem(this BinaryArchiveReader Reader)
+		{
+			return Reader.ReadObjectReference<DirectoryItem>(() => DirectoryItem.GetItemByDirectoryReference(Reader.ReadDirectoryReferenceNotNull()));
+		}
+
+		/// <summary>
+		/// Write a directory item to a binary archive
+		/// </summary>
+		/// <param name="Writer">Writer to serialize data to</param>
+		/// <param name="DirectoryItem">Directory item to write</param>
+		public static void WriteDirectoryItem(this BinaryArchiveWriter Writer, DirectoryItem DirectoryItem)
+		{
+			Writer.WriteObjectReference<DirectoryItem>(DirectoryItem, () => Writer.WriteDirectoryReference(DirectoryItem.Location));
+		}
+	}
+}
diff --git a/Engine/Source/Programs/Shared/EpicGames.Build/System/DirectoryLookupCache.cs b/Engine/Source/Programs/Shared/EpicGames.Build/System/DirectoryLookupCache.cs
new file mode 100644
--- /dev/null
+++ b/Engine/Source/Programs/Shared/EpicGames.Build/System/DirectoryLookupCache.cs
@@ -0,0 +1,38 @@
+// Copyright Epic Games, Inc. All Rights Reserved.
+
+using System;
+using System.IO;
+using System.Collections.Generic;
+using EpicGames.Core;
+using System.Linq;
+
+namespace UnrealBuildBase
+{
+	static public class DirectoryLookupCache
+	{
+		static public bool FileExists(FileReference File)
+		{
+			return FileItem.GetItemByFileReference(File).Exists;
+		}
+
+		static public bool DirectoryExists(DirectoryReference Directory)
+		{
+			return DirectoryItem.GetItemByDirectoryReference(Directory).Exists;
+		}
+
+		static public IEnumerable<FileReference> EnumerateFiles(DirectoryReference Directory)
+		{
+			return DirectoryItem.GetItemByDirectoryReference(Directory).EnumerateFiles().Select(x => x.Location);
+		}
+
+		static public IEnumerable<DirectoryReference> EnumerateDirectories(DirectoryReference Directory)
+		{
+			return DirectoryItem.GetItemByDirectoryReference(Directory).EnumerateDirectories().Select(x => x.Location);
+		}
+
+        static public void InvalidateCachedDirectory(DirectoryReference Directory)
+        {
+			DirectoryItem.GetItemByDirectoryReference(Directory).ResetCachedInfo();
+        }
+	}
+}
diff --git a/Engine/Source/Programs/Shared/EpicGames.Build/System/EnumeratePlugins.cs b/Engine/Source/Programs/Shared/EpicGames.Build/System/EnumeratePlugins.cs
new file mode 100644
--- /dev/null
+++ b/Engine/Source/Programs/Shared/EpicGames.Build/System/EnumeratePlugins.cs
@@ -0,0 +1,103 @@
+﻿// Copyright Epic Games, Inc. All Rights Reserved.
+
+using EpicGames.Core;
+using System;
+using System.Collections.Generic;
+using System.Linq;
+using System.Text;
+
+namespace UnrealBuildBase
+{
+	public class PluginsBase
+	{
+		/// <summary>
+		/// Cache of plugin filenames under each directory
+		/// </summary>
+		static Dictionary<DirectoryReference, List<FileReference>> PluginFileCache = new Dictionary<DirectoryReference, List<FileReference>>();
+
+		/// <summary>
+		/// Invalidate cached plugin data so that we can pickup new things
+		/// Warning: Will make subsequent plugin lookups and directory scans slow until the caches are repopulated
+		/// </summary>
+		public static void InvalidateCache_SLOW()
+		{
+			PluginFileCache = new Dictionary<DirectoryReference, List<FileReference>>();
+			DirectoryItem.ResetAllCachedInfo_SLOW();
+		}
+
+		/// <summary>
+		/// Enumerates all the plugin files available to the given project
+		/// </summary>
+		/// <param name="ProjectFile">Path to the project file</param>
+		/// <returns>List of project files</returns>
+		public static IEnumerable<FileReference> EnumeratePlugins(FileReference? ProjectFile)
+		{
+			List<DirectoryReference> BaseDirs = new List<DirectoryReference>();
+			BaseDirs.AddRange(Unreal.GetExtensionDirs(Unreal.EngineDirectory, "Plugins"));
+			if(ProjectFile != null)
+			{
+				BaseDirs.AddRange(Unreal.GetExtensionDirs(ProjectFile.Directory, "Plugins"));
+				BaseDirs.AddRange(Unreal.GetExtensionDirs(ProjectFile.Directory, "Mods"));
+			}
+			return BaseDirs.SelectMany(x => EnumeratePlugins(x)).ToList();
+		}
+
+		/// <summary>
+		/// Find paths to all the plugins under a given parent directory (recursively)
+		/// </summary>
+		/// <param name="ParentDirectory">Parent directory to look in. Plugins will be found in any *subfolders* of this directory.</param>
+		public static IEnumerable<FileReference> EnumeratePlugins(DirectoryReference ParentDirectory)
+		{
+			List<FileReference>? FileNames;
+			if (!PluginFileCache.TryGetValue(ParentDirectory, out FileNames))
+			{
+				FileNames = new List<FileReference>();
+
+				DirectoryItem ParentDirectoryItem = DirectoryItem.GetItemByDirectoryReference(ParentDirectory);
+				if (ParentDirectoryItem.Exists)
+				{
+					using(ThreadPoolWorkQueue Queue = new ThreadPoolWorkQueue())
+					{
+						EnumeratePluginsInternal(ParentDirectoryItem, FileNames, Queue);
+					}
+				}
+
+				// Sort the filenames to ensure that the plugin order is deterministic; otherwise response files will change with each build.
+				FileNames = FileNames.OrderBy(x => x.FullName, StringComparer.OrdinalIgnoreCase).ToList();
+
+				PluginFileCache.Add(ParentDirectory, FileNames);
+			}
+			return FileNames;
+		}
+
+		/// <summary>
+		/// Find paths to all the plugins under a given parent directory (recursively)
+		/// </summary>
+		/// <param name="ParentDirectory">Parent directory to look in. Plugins will be found in any *subfolders* of this directory.</param>
+		/// <param name="FileNames">List of filenames. Will have all the discovered .uplugin files appended to it.</param>
+		/// <param name="Queue">Queue for tasks to be executed</param>
+		static void EnumeratePluginsInternal(DirectoryItem ParentDirectory, List<FileReference> FileNames, ThreadPoolWorkQueue Queue)
+		{
+			foreach (DirectoryItem ChildDirectory in ParentDirectory.EnumerateDirectories())
+			{
+				bool bSearchSubDirectories = true;
+				foreach (FileItem PluginFile in ChildDirectory.EnumerateFiles())
+				{
+					if(PluginFile.HasExtension(".uplugin"))
+					{
+						lock(FileNames)
+						{
+							FileNames.Add(PluginFile.Location);
+						}
+						bSearchSubDirectories = false;
+					}
+				}
+
+				if (bSearchSubDirectories)
+				{
+					Queue.Enqueue(() => EnumeratePluginsInternal(ChildDirectory, FileNames, Queue));
+				}
+			}
+		}
+	}
+}
diff --git a/Engine/Source/Programs/Shared/EpicGames.Build/System/FileItem.cs b/Engine/Source/Programs/Shared/EpicGames.Build/System/FileItem.cs
new file mode 100644
--- /dev/null
+++ b/Engine/Source/Programs/Shared/EpicGames.Build/System/FileItem.cs
@@ -0,0 +1,347 @@
+// Copyright Epic Games, Inc. All Rights Reserved.
+
+using System;
+using System.Collections.Generic;
+using System.Text;
+using System.IO;
+using System.Diagnostics;
+using System.Threading;
+using System.Runtime.Serialization;
+using EpicGames.Core;
+using System.Collections.Concurrent;
+
+namespace UnrealBuildBase
+{
+	/// <summary>
+	/// Represents a file on disk that is used as an input or output of a build action. FileItem instances are unique for a given path. Use FileItem.GetItemByFileReference 
+	/// to get the FileItem for a specific path.
+	/// </summary>
+	public class FileItem
+	{
+		/// <summary>
+		/// The directory containing this file
+		/// </summary>
+		DirectoryItem? CachedDirectory;
+
+		/// <summary>
+		/// Location of this file
+		/// </summary>
+		public readonly FileReference Location;
+
+		/// <summary>
+		/// The information about the file.
+		/// </summary>
+		FileInfo Info;
+
+		/// <summary>
+		/// A case-insensitive dictionary that's used to map each unique file name to a single FileItem object.
+		/// </summary>
+		static ConcurrentDictionary<FileReference, FileItem> UniqueSourceFileMap = new ConcurrentDictionary<FileReference, FileItem>();
+
+		/// <summary>
+		/// Constructor
+		/// </summary>
+		/// <param name="Location">Location of the file</param>
+		/// <param name="Info">File info</param>
+		private FileItem(FileReference Location, FileInfo Info)
+		{
+			this.Location = Location;
+			this.Info = Info;
+		}
+
+		/// <summary>
+		/// Name of this file
+		/// </summary>
+		public string Name
+		{
+			get { return Info.Name; }
+		}
+
+		/// <summary>
+		/// Full name of this file
+		/// </summary>
+		public string FullName
+		{
+			get { return Location.FullName; }
+		}
+
+		/// <summary>
+		/// Accessor for the absolute path to the file
+		/// </summary>
+		public string AbsolutePath
+		{
+			get { return Location.FullName; }
+		}
+
+		/// <summary>
+		/// Gets the directory that this file is in
+		/// </summary>
+		public DirectoryItem Directory
+		{
+			get
+			{
+				if(CachedDirectory == null)
+				{
+					CachedDirectory = DirectoryItem.GetItemByDirectoryReference(Location.Directory);
+				}
+				return CachedDirectory;
+			}
+		}
+
+		/// <summary>
+		/// Whether the file exists.
+		/// </summary>
+		public bool Exists
+		{
+			get { return Info.Exists; }
+		}
+
+		/// <summary>
+		/// Size of the file if it exists, otherwise -1
+		/// </summary>
+		public long Length
+		{
+			get { return Info.Length; }
+		}
+
+		/// <summary>
+		/// The attributes for this file
+		/// </summary>
+		public FileAttributes Attributes
+		{
+			get { return Info.Attributes; }
+		}
+
+		/// <summary>
+		/// The last write time of the file.
+		/// </summary>
+		public DateTime LastWriteTimeUtc
+		{
+			get { return Info.LastWriteTimeUtc; }
+		}
+
+		/// <summary>
+		/// Determines if the file has the given extension
+		/// </summary>
+		/// <param name="Extension">The extension to check for</param>
+		/// <returns>True if the file has the given extension, false otherwise</returns>
+		public bool HasExtension(string Extension)
+		{
+			return Location.HasExtension(Extension);
+		}
+
+		/// <summary>
+		/// Gets the directory containing this file
+		/// </summary>
+		/// <returns>DirectoryItem for the directory containing this file</returns>
+		public DirectoryItem GetDirectoryItem()
+		{
+			return Directory;
+		}
+
+		/// <summary>
+		/// Updates the cached directory for this file. Used by DirectoryItem when enumerating files, to avoid having to look this up later.
+		/// </summary>
+		/// <param name="Directory">The directory that this file is in</param>
+		public void UpdateCachedDirectory(DirectoryItem Directory)
+		{
+			Debug.Assert(Directory.Location == Location.Directory);
+			CachedDirectory = Directory;
+		}
+
+		/// <summary>
+		/// Gets a FileItem corresponding to the given path
+		/// </summary>
+		/// <param name="FilePath">Path for the FileItem</param>
+		/// <returns>The FileItem that represents the given file path.</returns>
+		public static FileItem GetItemByPath(string FilePath)
+		{
+			return GetItemByFileReference(new FileReference(FilePath));
+		}
+
+		/// <summary>
+		/// Gets a FileItem for a given path
+		/// </summary>
+		/// <param name="Info">Information about the file</param>
+		/// <returns>The FileItem that represents the given a full file path.</returns>
+		public static FileItem GetItemByFileInfo(FileInfo Info)
+		{
+			FileReference Location = new FileReference(Info);
+
+			FileItem? Result;
+			if (!UniqueSourceFileMap.TryGetValue(Location, out Result))
+			{
+				FileItem NewFileItem = new FileItem(Location, Info);
+				if(UniqueSourceFileMap.TryAdd(Location, NewFileItem))
+				{
+					Result = NewFileItem;
+				}
+				else
+				{
+					Result = UniqueSourceFileMap[Location];
+				}
+			}
+			return Result;
+		}
+
+		/// <summary>
+		/// Gets a FileItem for a given path
+		/// </summary>
+		/// <param name="Location">Location of the file</param>
+		/// <returns>The FileItem that represents the given a full file path.</returns>
+		public static FileItem GetItemByFileReference(FileReference Location)
+		{
+			FileItem? Result;
+			if (!UniqueSourceFileMap.TryGetValue(Location, out Result))
+			{
+				FileItem NewFileItem = new FileItem(Location, Location.ToFileInfo());
+				if(UniqueSourceFileMap.TryAdd(Location, NewFileItem))
+				{
+					Result = NewFileItem;
+				}
+				else
+				{
+					Result = UniqueSourceFileMap[Location];
+				}
+			}
+			return Result;
+		}
+
+		/// <summary>
+		/// Deletes the file.
+		/// </summary>
+		public void Delete()
+		{
+			Debug.Assert(Exists);
+
+			int MaxRetryCount = 3;
+			int DeleteTryCount = 0;
+			bool bFileDeletedSuccessfully = false;
+			do
+			{
+				// If this isn't the first time through, sleep a little before trying again
+				if (DeleteTryCount > 0)
+				{
+					Thread.Sleep(1000);
+				}
+				DeleteTryCount++;
+				try
+				{
+					// Delete the destination file if it exists
+					FileInfo DeletedFileInfo = new FileInfo(AbsolutePath);
+					if (DeletedFileInfo.Exists)
+					{
+						DeletedFileInfo.IsReadOnly = false;
+						DeletedFileInfo.Delete();
+					}
+					// Success!
+					bFileDeletedSuccessfully = true;
+				}
+				catch (Exception Ex)
+				{
+					Log.TraceInformation("Failed to delete file '" + AbsolutePath + "'");
+					Log.TraceInformation("    Exception: " + Ex.Message);
+					if (DeleteTryCount < MaxRetryCount)
+					{
+						Log.TraceInformation("Attempting to retry...");
+					}
+					else
+					{
+						Log.TraceInformation("ERROR: Exhausted all retries!");
+					}
+				}
+			}
+			while (!bFileDeletedSuccessfully && (DeleteTryCount < MaxRetryCount));
+		}
+
+		/// <summary>
+		/// Resets the cached file info
+		/// </summary>
+		public void ResetCachedInfo()
+		{
+			Info = Location.ToFileInfo();
+		}
+
+		/// <summary>
+		/// Resets all cached file info. Significantly reduces performance; do not use unless strictly necessary.
+		/// </summary>
+		public static void ResetAllCachedInfo_SLOW()
+		{
+			foreach(FileItem Item in UniqueSourceFileMap.Values)
+			{
+				Item.ResetCachedInfo();
+			}
+		}
+
+		/// <summary>
+		/// Return the path to this FileItem to debugging
+		/// </summary>
+		/// <returns>Absolute path to this file item</returns>
+		public override string ToString()
+		{
+			return AbsolutePath;
+		}
+	}
+
+	/// <summary>
+	/// Helper functions for serialization
+	/// </summary>
+	public static class FileItemExtensionMethods
+	{
+		/// <summary>
+		/// Read a file item from a binary archive
+		/// </summary>
+		/// <param name="Reader">Reader to serialize data from</param>
+		/// <returns>Instance of the serialized file item</returns>
+		public static FileItem ReadFileItem(this BinaryArchiveReader Reader)
+		{
+			return Reader.ReadObjectReference<FileItem>(() => FileItem.GetItemByFileReference(Reader.ReadFileReference()));
+		}
+
+		/// <summary>
+		/// Write a file item to a binary archive
+		/// </summary>
+		/// <param name="Writer">Writer to serialize data to</param>
+		/// <param name="FileItem">File item to write</param>
+		public static void WriteFileItem(this BinaryArchiveWriter Writer, FileItem? FileItem)
+		{
+			Writer.WriteObjectReference<FileItem>(FileItem!, () => Writer.WriteFileReference(FileItem!.Location));
+		}
+
+		/// <summary>
+		/// Read a file item as a DirectoryItem and name. This is slower than reading it directly, but results in a significantly smaller archive
+		/// where most files are in the same directories.
+		/// </summary>
+		/// <param name="Reader">Archive to read from</param>
+		/// <returns>FileItem read from the archive</returns>
+		static FileItem ReadCompactFileItemData(this BinaryArchiveReader Reader)
+		{
+			DirectoryItem Directory = Reader.ReadDirectoryItem();
+			string Name = Reader.ReadString();
+
+			FileItem FileItem = FileItem.GetItemByFileReference(FileReference.Combine(Directory.Location, Name));
+            FileItem.UpdateCachedDirectory(Directory);
+            return FileItem;
+		}
+
+		/// <summary>
+		/// Read a file item in a format which de-duplicates directory names.
+		/// </summary>
+		/// <param name="Reader">Reader to serialize data from</param>
+		/// <returns>Instance of the serialized file item</returns>
+		public static FileItem ReadCompactFileItem(this BinaryArchiveReader Reader)
+		{
+			return Reader.ReadObjectReference<FileItem>(() => ReadCompactFileItemData(Reader));
+		}
+
+		/// <summary>
+		/// Writes a file item in a format which de-duplicates directory names.
+		/// </summary>
+		/// <param name="Writer">Writer to serialize data to</param>
+		/// <param name="FileItem">File item to write</param>
+		public static void WriteCompactFileItem(this BinaryArchiveWriter Writer, FileItem FileItem)
+		{
+			Writer.WriteObjectReference<FileItem>(FileItem, () => { Writer.WriteDirectoryItem(FileItem.GetDirectoryItem()); Writer.WriteString(FileItem.Name); });
+		}
+	}
+}
diff --git a/Engine/Source/Programs/Shared/EpicGames.Build/System/NativeProjectsBase.cs b/Engine/Source/Programs/Shared/EpicGames.Build/System/NativeProjectsBase.cs
new file mode 100644
--- /dev/null
+++ b/Engine/Source/Programs/Shared/EpicGames.Build/System/NativeProjectsBase.cs
@@ -0,0 +1,168 @@
+// Copyright Epic Games, Inc. All Rights Reserved.
+
+using System;
+using System.Collections.Generic;
+using System.Diagnostics.CodeAnalysis;
+using System.IO;
+using System.Linq;
+using System.Text;
+using System.Threading.Tasks;
+using EpicGames.Core;
+
+namespace UnrealBuildBase
+{
+	/// <summary>
+	/// Utility functions for querying native projects (ie. those found via a .uprojectdirs query)
+	/// </summary>
+	public class NativeProjectsBase
+	{
+		/// <summary>
+		/// Object used for synchronizing access to static fields
+		/// </summary>
+		protected static object LockObject = new object();
+
+		/// <summary>
+		/// The native project base directories
+		/// </summary>
+		static HashSet<DirectoryReference>? CachedBaseDirectories;
+
+		/// <summary>
+		/// Cached list of project files within all the base directories
+		/// </summary>
+		static HashSet<FileReference>? CachedProjectFiles;
+
+		/// <summary>
+		/// Clear our cached properties. Generally only needed if your script has modified local files...
+		/// </summary>
+		public static void ClearCacheBase()
+		{
+			CachedBaseDirectories = null;
+			CachedProjectFiles = null;
+		}
+
+		/// <summary>
+		/// Retrieve the list of base directories for native projects
+		/// </summary>
+		public static IEnumerable<DirectoryReference> EnumerateBaseDirectories()
+		{
+			if(CachedBaseDirectories == null)
+			{
+				lock(LockObject)
+				{
+					if(CachedBaseDirectories == null)
+					{
+						HashSet<DirectoryReference> BaseDirs = new HashSet<DirectoryReference>();
+						foreach (FileReference RootFile in DirectoryLookupCache.EnumerateFiles(Unreal.RootDirectory))
+						{
+							if(RootFile.HasExtension(".uprojectdirs"))
+							{
+								foreach(string Line in File.ReadAllLines(RootFile.FullName))
+								{
+									string TrimLine = Line.Trim();
+									if(!TrimLine.StartsWith(";"))
+									{
+										DirectoryReference BaseProjectDir = DirectoryReference.Combine(Unreal.RootDirectory, TrimLine);
+										if(BaseProjectDir.IsUnderDirectory(Unreal.RootDirectory))
+										{
+											BaseDirs.Add(BaseProjectDir);
+										}
+										else
+										{
+											Log.TraceWarning("Project search path '{0}' referenced by '{1}' is not under '{2}', ignoring.", TrimLine, RootFile, Unreal.RootDirectory);
+										}
+									}
+								}
+							}
+						}
+						CachedBaseDirectories = BaseDirs;
+					}
+				}
+			}
+			return CachedBaseDirectories;
+		}
+
+		/// <summary>
+		/// Returns a list of all the projects
+		/// </summary>
+		/// <returns>List of projects</returns>
+		public static IEnumerable<FileReference> EnumerateProjectFiles()
+		{
+			if(CachedProjectFiles == null)
+			{
+				lock(LockObject)
+				{
+					if(CachedProjectFiles == null)
+					{
+						HashSet<FileReference> ProjectFiles = new HashSet<FileReference>();
+						foreach(DirectoryReference BaseDirectory in EnumerateBaseDirectories())
+						{
+							if(DirectoryLookupCache.DirectoryExists(BaseDirectory))
+							{
+								foreach(DirectoryReference SubDirectory in DirectoryLookupCache.EnumerateDirectories(BaseDirectory))
+								{
+									// Ignore system directories (specifically for Mac, since temporary folders are created here)
+									string DirectoryName = SubDirectory.GetDirectoryName();
+									if (DirectoryName.StartsWith(".", StringComparison.Ordinal))
+									{
+										continue;
+									}
+
+									foreach(FileReference File in DirectoryLookupCache.EnumerateFiles(SubDirectory))
+									{
+										if(File.HasExtension(".uproject"))
+										{
+											ProjectFiles.Add(File);
+										}
+									}
+								}
+							}
+						}
+						CachedProjectFiles = ProjectFiles;
+					}
+				}
+			}
+			return CachedProjectFiles;
+		}
+
+		/// <summary>
+		/// Finds all target files under a given folder, and add them to the target name to project file map
+		/// </summary>
+		/// <param name="Directory">Directory to search</param>
+		/// <param name="TargetNameToProjectFile">Map from target name to project file</param>
+		/// <param name="ProjectFile">The project file for this directory</param>
+		protected static void FindTargetFiles(DirectoryReference Directory, Dictionary<string, FileReference> TargetNameToProjectFile, FileReference ProjectFile)
+		{
+			// Search for all target files within this directory
+			bool bSearchSubFolders = true;
+			foreach (FileReference File in DirectoryLookupCache.EnumerateFiles(Directory))
+			{
+				if (File.HasExtension(".target.cs"))
+				{
+					string TargetName = Path.GetFileNameWithoutExtension(File.GetFileNameWithoutExtension());
+					TargetNameToProjectFile[TargetName] = ProjectFile;
+					bSearchSubFolders = false;
+				}
+			}
+
+			// If we didn't find anything, recurse through the subfolders
+			if(bSearchSubFolders)
+			{
+				foreach(DirectoryReference SubDirectory in DirectoryLookupCache.EnumerateDirectories(Directory))
+				{
+					FindTargetFiles(SubDirectory, TargetNameToProjectFile, ProjectFile);
+				}
+			}
+		}
+
+		/// <summary>
+		/// Checks if a given project is a native project
+		/// </summary>
+		/// <param name="ProjectFile">The project file to check</param>
+		/// <returns>True if the given project is a native project</returns>
+		public static bool IsNativeProject(FileReference ProjectFile)
+		{
+			EnumerateProjectFiles();
+			return CachedProjectFiles!.Contains(ProjectFile);
+		}
+	}
+}
diff --git a/Engine/Source/Programs/Shared/EpicGames.Build/System/Rules.cs b/Engine/Source/Programs/Shared/EpicGames.Build/System/Rules.cs
new file mode 100644
--- /dev/null
+++ b/Engine/Source/Programs/Shared/EpicGames.Build/System/Rules.cs
@@ -0,0 +1,301 @@
+// Copyright Epic Games, Inc. All Rights Reserved.
+
+using EpicGames.Core;
+using System;
+using System.Collections.Generic;
+using System.Linq;
+using System.Text;
+
+namespace UnrealBuildBase
+{
+	public class Rules
+	{
+		/// <summary>
+		/// Enum for types of rules files. Should match extensions in RulesFileExtensions.
+		/// </summary>
+		public enum RulesFileType
+		{
+			/// <summary>
+			/// .build.cs files
+			/// </summary>
+			Module,
+
+			/// <summary>
+			/// .target.cs files
+			/// </summary>
+			Target,
+
+			/// <summary>
+			/// .automation.csproj files
+			/// </summary>
+			AutomationModule
+		}
+
+		/// <summary>
+		/// Cached list of rules files in each directory of each type
+		/// </summary>
+		class RulesFileCache
+		{
+			public List<FileReference> ModuleRules = new List<FileReference>();
+			public List<FileReference> TargetRules = new List<FileReference>();
+			public List<FileReference> AutomationModules = new List<FileReference>();
+		}
+
+		/// Map of root folders to a cached list of all UBT-related source files in that folder or any of its sub-folders.
+		/// We cache these file names so we can avoid searching for them later on.
+		static Dictionary<DirectoryReference, RulesFileCache> RootFolderToRulesFileCache = new Dictionary<DirectoryReference, RulesFileCache>();
+
+		/// <summary>
+		/// 
+		/// </summary>
+		/// <param name="RulesFileType"></param>
+		/// <param name="GameFolders"></param>
+		/// <param name="ForeignPlugins"></param>
+		/// <param name="AdditionalSearchPaths"></param>
+		/// <param name="bIncludeEngine"></param>
+		/// <param name="bIncludeTempTargets">Whether to include targets generated by UAT to accomodate content-only projects that need to be compiled to include plugins</param>
+		/// <returns></returns>
+		public static List<FileReference> FindAllRulesSourceFiles(RulesFileType RulesFileType, List<DirectoryReference> GameFolders, List<FileReference> ForeignPlugins, List<DirectoryReference> AdditionalSearchPaths, bool bIncludeEngine = true, bool bIncludeTempTargets = true)
+		{
+			List<DirectoryReference> Folders = new List<DirectoryReference>();
+
+			// Add all engine source (including third party source)
+			if (bIncludeEngine)
+			{
+				Folders.AddRange(Unreal.GetExtensionDirs(Unreal.EngineDirectory, "Source"));
+			}
+
+			// @todo plugin: Disallow modules from including plugin modules as dependency modules? (except when the module is part of that plugin)
+
+			// Get all the root folders for plugins
+			List<DirectoryReference> RootFolders = new List<DirectoryReference>();
+			if (bIncludeEngine)
+			{
+				RootFolders.AddRange(Unreal.GetExtensionDirs(Unreal.EngineDirectory));
+			}
+			if (GameFolders != null)
+			{
+				RootFolders.AddRange(GameFolders.SelectMany(x => Unreal.GetExtensionDirs(x)));
+			}
+
+			// Find all the plugin source directories
+			foreach (DirectoryReference RootFolder in RootFolders)
+			{
+				DirectoryReference PluginsFolder = DirectoryReference.Combine(RootFolder, "Plugins");
+				foreach (FileReference PluginFile in PluginsBase.EnumeratePlugins(PluginsFolder))
+				{
+					Folders.Add(DirectoryReference.Combine(PluginFile.Directory, "Source"));
+				}
+			}
+
+			// Add all the extra plugin folders
+			if (ForeignPlugins != null)
+			{
+				foreach (FileReference ForeignPlugin in ForeignPlugins)
+				{
+					Folders.Add(DirectoryReference.Combine(ForeignPlugin.Directory, "Source"));
+				}
+			}
+
+			// Add in the game folders to search
+			if (GameFolders != null)
+			{
+				foreach (DirectoryReference GameFolder in GameFolders)
+				{
+					Folders.AddRange(Unreal.GetExtensionDirs(GameFolder, "Source"));
+
+					if (bIncludeTempTargets)
+					{
+						DirectoryReference GameIntermediateSourceFolder = DirectoryReference.Combine(GameFolder, "Intermediate", "Source");
+						Folders.Add(GameIntermediateSourceFolder);
+					}
+				}
+			}
+
+			// Process the additional search path, if sent in
+			if (AdditionalSearchPaths != null)
+			{
+				foreach (DirectoryReference AdditionalSearchPath in AdditionalSearchPaths)
+				{
+					if (AdditionalSearchPath != null)
+					{
+						if (DirectoryReference.Exists(AdditionalSearchPath))
+						{
+							Folders.Add(AdditionalSearchPath);
+						}
+						else
+						{
+							throw new Exception($"Couldn't find AdditionalSearchPath for rules source files '{AdditionalSearchPath}'");
+						}
+					}
+				}
+			}
+
+			return FindAllRulesFiles(Folders, RulesFileType);
+		}
+
+		/// <summary>
+		/// Invalidate the cache for the givcen directory
+		/// </summary>
+		/// <param name="DirectoryPath">Directory to invalidate</param>
+        public static void InvalidateRulesFileCache(string DirectoryPath)
+        {
+            DirectoryReference Directory = new DirectoryReference(DirectoryPath);
+            RootFolderToRulesFileCache.Remove(Directory);
+            DirectoryLookupCache.InvalidateCachedDirectory(Directory);
+        }
+
+		/// <summary>
+		/// Prefetch multiple directories in parallel
+		/// </summary>
+		/// <param name="Directories">The directories to cache</param>
+		public static void PrefetchRulesFiles(IEnumerable<DirectoryReference> Directories)
+		{
+			ThreadPoolWorkQueue? Queue = null;
+			try
+			{
+				foreach(DirectoryReference Directory in Directories)
+				{
+					if(!RootFolderToRulesFileCache.ContainsKey(Directory))
+					{
+						RulesFileCache Cache = new RulesFileCache();
+						RootFolderToRulesFileCache[Directory] = Cache;
+
+						if(Queue == null)
+						{
+							Queue = new ThreadPoolWorkQueue();
+						}
+
+						DirectoryItem DirectoryItem = DirectoryItem.GetItemByDirectoryReference(Directory);
+						Queue.Enqueue(() => FindAllRulesFilesRecursively(DirectoryItem, Cache, Queue));
+					}
+				}
+			}
+			finally
+			{
+				if(Queue != null)
+				{
+					Queue.Dispose();
+					Queue = null;
+				}
+			}
+		}
+
+		/// <summary>
+		/// Finds all the rules of the given type under a given directory
+		/// </summary>
+		/// <param name="Directory">Directory to search</param>
+		/// <param name="Type">Type of rules to return</param>
+		/// <returns>List of rules files of the given type</returns>
+		public static IReadOnlyList<FileReference> FindAllRulesFiles(DirectoryReference Directory, RulesFileType Type)
+		{
+			return FindAllRulesFiles(new List<DirectoryReference> { Directory }, Type);
+		}
+
+		/// <summary>
+		/// Finds all the rules of the given type under a given directories
+		/// </summary>
+		/// <param name="Directories">Directories to search</param>
+		/// <param name="Type">Type of rules to return</param>
+		/// <returns>List of rules files of the given type</returns>
+		private static List<FileReference> FindAllRulesFiles(IEnumerable<DirectoryReference> Directories, RulesFileType Type)
+		{
+			List<(DirectoryReference, RulesFileCache)> Caches = new List<(DirectoryReference, RulesFileCache)>(Directories.Count());
+			using (ThreadPoolWorkQueue Queue = new ThreadPoolWorkQueue())
+			{
+				foreach (DirectoryReference Directory in Directories)
+				{
+					RulesFileCache? Cache;
+					if (!RootFolderToRulesFileCache.TryGetValue(Directory, out Cache))
+					{
+						Cache = new RulesFileCache();
+						Queue.Enqueue(() => FindAllRulesFilesRecursively(DirectoryItem.GetItemByDirectoryReference(Directory), Cache, Queue));
+					}
+					Caches.Add((Directory, Cache));
+				}
+			}
+
+			List<FileReference> Files = new List<FileReference>();
+			
+			foreach ((DirectoryReference Directory, RulesFileCache Cache) in Caches)
+			{
+				if (!RootFolderToRulesFileCache.ContainsKey(Directory))
+				{
+					Cache.ModuleRules.Sort((A, B) => A.FullName.CompareTo(B.FullName));
+					Cache.TargetRules.Sort((A, B) => A.FullName.CompareTo(B.FullName));
+					Cache.AutomationModules.Sort((A, B) => A.FullName.CompareTo(B.FullName));
+					RootFolderToRulesFileCache[Directory] = Cache;
+				}
+
+				// Get the list of files of the type we're looking for
+				if (Type == RulesFileType.Module)
+				{
+					Files.AddRange(Cache.ModuleRules);
+				}
+				else if (Type == RulesFileType.Target)
+				{
+					Files.AddRange(Cache.TargetRules);
+				}
+				else if (Type == RulesFileType.AutomationModule)
+				{
+					Files.AddRange(Cache.AutomationModules);
+				}
+				else
+				{
+					throw new Exception($"Unhandled rules type: {Type}");
+				}
+			}
+
+			return Files;
+		}
+
+		/// <summary>
+		/// Search through a directory tree for any rules files
+		/// </summary>
+		/// <param name="Directory">The root directory to search from</param>
+		/// <param name="Cache">Receives all the discovered rules files</param>
+		/// <param name="Queue">Queue for adding additional tasks to</param>
+		private static void FindAllRulesFilesRecursively(DirectoryItem Directory, RulesFileCache Cache, ThreadPoolWorkQueue Queue)
+		{
+			// Scan all the files in this directory
+			bool bSearchSubFolders = true;
+			foreach (FileItem File in Directory.EnumerateFiles())
+			{
+				if (File.HasExtension(".build.cs"))
+				{
+					lock(Cache.ModuleRules)
+					{
+						Cache.ModuleRules.Add(File.Location);
+					}
+					bSearchSubFolders = false;
+				}
+				else if (File.HasExtension(".target.cs"))
+				{
+					lock(Cache.TargetRules)
+					{
+						Cache.TargetRules.Add(File.Location);
+					}
+				}
+				else if (File.HasExtension(".automation.csproj"))
+				{
+					lock(Cache.AutomationModules)
+					{
+						Cache.AutomationModules.Add(File.Location);
+					}
+					bSearchSubFolders = false;
+				}
+			}
+
+			// If we didn't find anything to stop the search, search all the subdirectories too
+			if (bSearchSubFolders)
+			{
+				foreach (DirectoryItem SubDirectory in Directory.EnumerateDirectories())
+				{
+					Queue.Enqueue(() => FindAllRulesFilesRecursively(SubDirectory, Cache, Queue));
+				}
+			}
+		}
+
+
+	}
+}
diff --git a/Engine/Source/Programs/Shared/EpicGames.Build/System/StartupTraceListener.cs b/Engine/Source/Programs/Shared/EpicGames.Build/System/StartupTraceListener.cs
new file mode 100644
--- /dev/null
+++ b/Engine/Source/Programs/Shared/EpicGames.Build/System/StartupTraceListener.cs
@@ -0,0 +1,66 @@
+// Copyright Epic Games, Inc. All Rights Reserved.
+
+using System;
+using System.Collections.Generic;
+using System.Linq;
+using System.Text;
+using System.IO;
+using System.CodeDom.Compiler;
+using Microsoft.CSharp;
+using Microsoft.Win32;
+using System.Reflection;
+using System.Diagnostics;
+using EpicGames.Core;
+
+namespace UnrealBuildBase
+{
+	/// <summary>
+	/// Captures all log output during startup until a log file writer has been created
+	/// </summary>
+	public class StartupTraceListener : TraceListener
+	{
+		StringBuilder Buffer = new StringBuilder();
+
+		/// <summary>
+		/// Constructor
+		/// </summary>
+		public StartupTraceListener()
+		{
+		}
+
+		/// <summary>
+		/// Copy the contents of the buffered output to another trace listener
+		/// </summary>
+		/// <param name="Other">The trace listener to receive the buffered output</param>
+		public void CopyTo(TraceListener Other)
+		{
+			foreach(string Line in Buffer.ToString().Split("\n"))
+			{
+				Other.WriteLine(Line);
+			}
+		}
+
+		/// <summary>
+		/// Write a message to the buffer
+		/// </summary>
+		/// <param name="Message">The message to write</param>
+		public override void Write(string Message)
+		{
+			if(NeedIndent)
+			{
+				WriteIndent();
+			}
+			Buffer.Append(Message);
+		}
+
+		/// <summary>
+		/// Write a message to the buffer, followed by a newline
+		/// </summary>
+		/// <param name="Message">The message to write</param>
+		public override void WriteLine(string Message)
+		{
+			Write(Message);
+			Buffer.Append("\n");
+		}
+	}
+}
diff --git a/Engine/Source/Programs/Shared/EpicGames.Build/UEBuildPlatformSDK.cs b/Engine/Source/Programs/Shared/EpicGames.Build/UEBuildPlatformSDK.cs
new file mode 100644
--- /dev/null
+++ b/Engine/Source/Programs/Shared/EpicGames.Build/UEBuildPlatformSDK.cs
@@ -0,0 +1,1292 @@
+// Copyright Epic Games, Inc. All Rights Reserved.
+
+using System;
+using System.Collections.Generic;
+using System.Diagnostics;
+using System.IO;
+using System.Linq;
+using System.Text;
+
+#nullable disable
+
+namespace EpicGames.Core
+{
+	/// <summary>
+	/// SDK installation status
+	/// </summary>
+	public enum SDKStatus
+	{
+		/// <summary>
+		/// Desired SDK is installed and set up.
+		/// </summary>
+		Valid,
+
+		/// <summary>
+		/// Could not find the desired SDK, SDK setup failed, etc.		
+		/// </summary>
+		Invalid,
+	};
+
+	/// <summary>
+	/// SDK for a platform
+	/// </summary>
+	abstract public class UEBuildPlatformSDK
+	{
+		// Public SDK handling, not specific to AutoSDK
+
+		public UEBuildPlatformSDK()
+		{
+		}
+
+		#region Global SDK Registration
+
+		/// <summary>
+		/// Registers the SDK for a given platform (as a string, but equivalent to UnrealTargetPlatform)
+		/// </summary>
+		/// <param name="SDK">SDK object</param>
+		/// <param name="PlatformName">Platform name for this SDK</param>
+		public static void RegisterSDKForPlatform(UEBuildPlatformSDK SDK, string PlatformName, bool bIsSdkAllowedOnHost)
+		{
+			// verify that neither platform or sdk were added before
+			if (SDKRegistry.Count(x => x.Key == PlatformName || x.Value == SDK) > 0)
+			{
+				throw new Exception(string.Format("Re-registering SDK for {0}. All Platforms must have a unique SDK object", PlatformName));
+			}
+
+			SDKRegistry.Add(PlatformName, SDK);
+
+			SDK.Init(PlatformName, bIsSdkAllowedOnHost);
+		}
+
+		private void Init(string InPlatformName, bool bInIsSdkAllowedOnHost)
+		{
+			PlatformName = InPlatformName;
+			bIsSdkAllowedOnHost = bInIsSdkAllowedOnHost;
+
+			// if the parent set up autosdk, the env vars will be wrong, but we can still get the manual SDK version from before it was setup
+			if (HasParentProcessSetupAutoSDK(out CachedManualSDKVersion))
+			{
+				// we pass along __None to indicate the parent didn't have a manual sdk installed
+				if (CachedManualSDKVersion == "__None")
+				{
+					CachedManualSDKVersion = null;
+				}
+			}
+			else
+			{
+				// if there was no parent, get the SDK version before we run AutoSDK to get the manual version
+				CachedManualSDKVersion = GetInstalledSDKVersion();
+			}
+		}
+
+		#endregion
+
+		#region Main Public Interface/Utilties
+
+		/// <summary>
+		/// Retrieves a previously registered SDK for a given platform
+		/// </summary>
+		/// <param name="PlatformName">String name of the platform (equivalent to UnrealTargetPlatform)</param>
+		/// <returns></returns>
+		public static UEBuildPlatformSDK GetSDKForPlatform(string PlatformName)
+		{
+			UEBuildPlatformSDK SDK;
+			SDKRegistry.TryGetValue(PlatformName, out SDK);
+
+			return SDK;
+		}
+
+		/// <summary>
+		/// Gets the set of all known SDKs
+		/// </summary>
+		public static UEBuildPlatformSDK[] AllSDKs
+		{
+			get	{ return SDKRegistry.Values.ToArray(); }
+		}
+
+		// String name of the platform (will match an UnrealTargetPlatform)
+		public string PlatformName;
+
+		// True if this Sdk is allowed to be used by this host - if not, we can skip a lot 
+		public bool bIsSdkAllowedOnHost;
+
+		public string GetInstalledVersion(out bool bIsAutoSDK)
+		{
+			bIsAutoSDK = HasSetupAutoSDK();
+			return GetInstalledSDKVersion();
+		}
+
+		public string GetInstalledVersion()
+		{
+			return GetInstalledSDKVersion();
+		}
+
+		public void GetInstalledVersions(out string ManualSDKVersion, out string AutoSDKVersion)
+		{
+			// if we support AutoSDKs, then return both versions
+			if (PlatformSupportsAutoSDKs())
+			{
+				AutoSDKVersion = (HasRequiredAutoSDKInstalled() == SDKStatus.Valid) ? GetInstalledSDKVersion() : null;
+//				AutoSDKVersion = GetInstalledSDKVersion();
+			}
+			else
+			{
+				AutoSDKVersion = null;
+				if (CachedManualSDKVersion != GetInstalledSDKVersion())
+				{
+					throw new Exception("Manual SDK version changed, this is not supported yet");
+				}
+			}
+
+			ManualSDKVersion = CachedManualSDKVersion;
+		}
+
+
+		public virtual bool IsVersionValid(string Version, bool bForAutoSDK)
+		{
+			return IsVersionValidInternal(Version, bForAutoSDK);
+		}
+		public virtual bool IsSoftwareVersionValid(string Version)
+		{
+			return IsSoftwareVersionValidInternal(Version);
+		}
+
+		public void ReactivateAutoSDK()
+		{
+			// @todo turnkey: this needs to force re-doing it, as it is likely a no-op, need to investigate what to clear out
+			ManageAndValidateSDK();
+		}
+
+		#endregion
+
+		#region Platform Overrides
+
+		/// <summary>
+		/// Returns the installed SDK version, used to determine if up to date or not (
+		/// </summary>
+		/// <returns></returns>
+		public abstract string GetInstalledSDKVersion();
+
+		/// <summary>
+		/// Return the SDK version that the platform wants to use (AutoSDK dir must match this, full SDKs can be in a valid range)
+		/// </summary>
+		/// <returns></returns>
+		public abstract string GetMainVersion();
+
+		/// <summary>
+		/// Returns a platform-specific version string that can be retrieved from anywhere without needing to typecast the SDK object
+		/// </summary>
+		/// <param name="VersionType">Descriptor of the version you want</param>
+		/// <returns>Version string, or empty string if not known/supported</returns>
+		public virtual string GetPlatformSpecificVersion(string VersionType)
+		{
+			return "";
+		}
+
+		/// <summary>
+		/// Return a list of full SDK versions that are already installed and can be quickly switched to without running any installers.
+		/// </summary>
+		/// <returns></returns>
+		public virtual string[] GetAllInstalledSDKVersions()
+		{
+			return new string[] { };
+		}
+
+		/// <summary>
+		/// Switch to another version of the SDK than what GetInstalledSDKVersion() returns. This will be one of the versions returned from GetAllInstalledSDKVersions()
+		/// </summary>
+		/// <param name="Version">String name of the version to switch to</param>
+		/// <param name="bSwitchForThisProcessOnly">If true, only switch for this process (usually via process environment variable). If false, switch permanently (usually via system-wide environment variable)</param>
+		/// <returns>True if successful</returns>
+		public virtual bool SwitchToAlternateSDK(string Version, bool bSwitchForThisProcessOnly)
+		{
+			return false;
+		}
+
+		/// <summary>
+		/// Gets the valid string range of Sdk versions. TryConvertVersionToInt() will need to succeed to make this usable for range checks
+		/// </summary>
+		/// <param name="MinVersion">Smallest version allowed</param>
+		/// <param name="MaxVersion">Largest version allowed (inclusive)</param>
+		public abstract void GetValidVersionRange(out string MinVersion, out string MaxVersion);
+
+
+		/// <summary>
+		/// Gets the valid string range of software/flash versions. TryConvertVersionToInt() will need to succeed to make this usable for range checks
+		/// </summary>
+		/// <param name="MinVersion">Smallest version allowed, or null if no minmum (in other words, 0 - MaxVersion)</param>
+		/// <param name="MaxVersion">Largest version allowed (inclusive), or null if no maximum (in other words, MinVersion - infinity)y</param>
+		public abstract void GetValidSoftwareVersionRange(out string MinVersion, out string MaxVersion);
+
+		/// <summary>
+		/// For a platform that doesn't use properly named AutoSDK directories, the directory name may not be convertible to an integer,
+		/// and IsVersionValid checks could fail when checking AutoSDK version for an exact match. GetMainVersion() would return the 
+		/// proper, integer-convertible version number of the SDK inside of the directory returned by GetAutoSDKDirectoryForMainVersion()
+		/// </summary>
+		/// <returns></returns>
+		public virtual string GetAutoSDKDirectoryForMainVersion()
+		{
+			return GetMainVersion();
+		}
+
+
+		/// <summary>
+		/// Gets the valid (integer) range of Sdk versions. Must be an integer to easily check a range vs a particular version
+		/// </summary>
+		/// <param name="MinVersion">Smallest version allowed</param>
+		/// <param name="MaxVersion">Largest version allowed (inclusive)</param>
+		/// <returns>True if the versions are valid, false if the platform is unable to convert its versions into an integer</returns>
+		public virtual void GetValidVersionRange(out UInt64 MinVersion, out UInt64 MaxVersion)
+		{
+			string MinVersionString, MaxVersionString;
+			GetValidVersionRange(out MinVersionString, out MaxVersionString);
+
+			// failures to convert here are bad
+			if (!TryConvertVersionToInt(MinVersionString, out MinVersion) || !TryConvertVersionToInt(MaxVersionString, out MaxVersion))
+			{
+				throw new Exception(string.Format("Unable to convert Min and Max valid versions to integers in {0} (Versions are {1} - {2})", GetType().Name, MinVersionString, MaxVersionString));
+			}
+		}
+		public virtual void GetValidSoftwareVersionRange(out UInt64 MinVersion, out UInt64 MaxVersion)
+		{
+			string MinVersionString, MaxVersionString;
+			GetValidSoftwareVersionRange(out MinVersionString, out MaxVersionString);
+
+			MinVersion = UInt64.MinValue;
+			MaxVersion = UInt64.MaxValue - 1; // MaxValue is always bad
+
+			// failures to convert here are bad
+			if ((MinVersionString != null && !TryConvertVersionToInt(MinVersionString, out MinVersion)) || 
+				(MaxVersionString != null && !TryConvertVersionToInt(MaxVersionString, out MaxVersion)))
+			{
+				throw new Exception(string.Format("Unable to convert Min and Max valid Software versions to integers in {0} (Versions are {1} - {2})", GetType().Name, MinVersionString, MaxVersionString));
+			}
+		}
+
+
+		// Let platform override behavior to determine if a version is a valid (useful for non-numeric versions)
+		protected virtual bool IsVersionValidInternal(string Version, bool bForAutoSDK)
+		{
+			// we could have null if no SDK is installed at all, etc, which is always a failure
+			if (Version == null)
+			{
+				return false;
+			}
+
+			// AutoSDK must match the desired version exactly, since that is the only one we will use
+			if (bForAutoSDK)
+			{
+				// if integer version checking failed, then we can detect valid autosdk if the version matches the autosdk directory by name
+				return string.Compare(Version, GetAutoSDKDirectoryForMainVersion(), true) == 0;
+			}
+
+			// convert it to an integer
+			UInt64 IntVersion;
+			if (!TryConvertVersionToInt(Version, out IntVersion))
+			{
+				return false;
+			}
+
+			UInt64 DesiredVersion;
+			if (!TryConvertVersionToInt(GetMainVersion(), out DesiredVersion))
+			{
+				return false;
+			}
+
+			// short circuit range check if the Version is the desired version already
+			if (IntVersion == DesiredVersion)
+			{
+				return true;
+			}
+
+			// get numeric range
+			UInt64 MinVersion, MaxVersion;
+			GetValidVersionRange(out MinVersion, out MaxVersion);
+			return IntVersion >= MinVersion && IntVersion <= MaxVersion;
+		}
+
+		protected virtual bool IsSoftwareVersionValidInternal(string Version)
+		{
+			// we could have null if no SDK is installed at all, etc, which is always a failure
+			if (Version == null)
+			{
+				return false;
+			}
+
+			// convert it to an integer
+			UInt64 IntVersion;
+			if (!TryConvertVersionToInt(Version, out IntVersion))
+			{
+				return false;
+			}
+
+			// get numeric range
+			UInt64 MinVersion, MaxVersion;
+			GetValidSoftwareVersionRange(out MinVersion, out MaxVersion);
+			return IntVersion >= MinVersion && IntVersion <= MaxVersion;
+		}
+
+
+		/// <summary>
+		/// Only the platform can convert a version string into an integer that is usable for comparison
+		/// </summary>
+		/// <param name="StringValue">Version that comes from the installed SDK or a Turnkey manifest or the like</param>
+		/// <param name="OutValue">The integer version of StringValue, can be used to compare against a valid range</param>
+		/// <returns>If the StringValue was able to be be converted to an integer</returns>
+		public virtual bool TryConvertVersionToInt(string StringValue, out UInt64 OutValue)
+		{
+			// @todo turnkey make this abstract?
+			OutValue = 0;
+			return false;
+		}
+
+		/// <summary>
+		/// Allow the platform SDK to override the name it will use in AutoSDK, but default to the platform name
+		/// </summary>
+		/// <returns>The name of the directory to use inside the AutoSDK system</returns>
+		public virtual string GetAutoSDKPlatformName()
+		{
+			return PlatformName;
+		}
+
+		#endregion
+
+		#region Print SDK Info
+		private static bool bHasShownTurnkey = false;
+		public virtual SDKStatus PrintSDKInfoAndReturnValidity(LogEventType Verbosity = LogEventType.Console, LogFormatOptions Options = LogFormatOptions.None,
+			LogEventType ErrorVerbosity = LogEventType.Error, LogFormatOptions ErrorOptions = LogFormatOptions.None)
+		{
+			string ManualSDKVersion, AutoSDKVersion;
+			GetInstalledVersions(out ManualSDKVersion, out AutoSDKVersion);
+
+			SDKStatus Validity = SDKStatus.Valid;
+
+			bool HasValidAutoSDK = HasSetupAutoSDK();
+		
+			if (HasValidAutoSDK)
+			{
+				string PlatformSDKRoot = GetPathToPlatformAutoSDKs();
+
+				UInt64 Ver;
+
+				if (TryConvertVersionToInt(AutoSDKVersion, out Ver))
+				{
+					Log.WriteLine(Verbosity, Options, "{0} using Auto SDK from: {1} 0x{2:X}", PlatformName, Path.Combine(PlatformSDKRoot, AutoSDKVersion), Ver);
+				}
+				else
+				{
+					HasValidAutoSDK = false;
+				}
+			}
+
+			if (HasValidAutoSDK == false)
+			{
+				if (HasRequiredManualSDK() == SDKStatus.Valid)
+				{
+					Log.WriteLine(Verbosity, Options, "{0} using Manual SDK {1}", PlatformName, ManualSDKVersion);
+				}
+				else
+				{
+					Validity = SDKStatus.Invalid;
+
+					string MinVersionString, MaxVersionString;
+					GetValidVersionRange(out MinVersionString, out MaxVersionString);
+
+					StringBuilder Msg = new StringBuilder();
+					Msg.AppendFormat("Unable to find a valid SDK for {0}.", PlatformName);
+					if (ManualSDKVersion != null)
+					{
+						Msg.AppendFormat(" Found Version: {0}.", ManualSDKVersion);
+					}
+
+					if (MinVersionString != MaxVersionString)
+					{
+						Msg.AppendLine(" Must be between {0} and {1}", MinVersionString, MaxVersionString);
+					}
+					else
+					{
+						Msg.AppendLine(" Must be {0}", MinVersionString);
+					}
+
+					if (!bHasShownTurnkey)
+					{
+						Msg.AppendLine("  If your Studio has it set up, you can run this command to find the SDK to install:");
+						Msg.AppendLine("    RunUAT Turnkey -command=InstallSdk -platform={0} -BestAvailable", PlatformName);
+
+						if ((ErrorOptions & LogFormatOptions.NoConsoleOutput) == LogFormatOptions.None)
+						{
+							bHasShownTurnkey = true;
+						}
+					}
+
+					// always print errors to the screen
+					Log.WriteLine(ErrorVerbosity, ErrorOptions, Msg.ToString());
+				}
+			}
+
+			return Validity;
+		}
+
+		#endregion
+
+
+
+
+
+		#region Private/Protected general functionality
+
+		// this is the SDK version that was set before activating AutoSDK, since AutoSDK may remove ability to retrieve the Manual SDK version
+		protected string CachedManualSDKVersion;
+		private static Dictionary<string, UEBuildPlatformSDK> SDKRegistry = new Dictionary<string, UEBuildPlatformSDK>();
+
+		#endregion
+
+		// AutoSDKs handling portion
+
+		#region protected AutoSDKs Utility
+
+		/// <summary>
+		/// Name of the file that holds currently install SDK version string
+		/// </summary>
+		protected const string CurrentlyInstalledSDKStringManifest = "CurrentlyInstalled.txt";
+
+		/// <summary>
+		/// name of the file that holds the last succesfully run SDK setup script version
+		/// </summary>
+		protected const string LastRunScriptVersionManifest = "CurrentlyInstalled.Version.txt";
+
+		/// <summary>
+		/// Filename of the script version file in each AutoSDK directory. Changing the contents of this file will force reinstallation of an AutoSDK.
+		/// </summary>
+		private const string ScriptVersionFilename = "Version.txt";
+
+		/// <summary>
+		/// Name of the file that holds environment variables of current SDK
+		/// </summary>
+		protected const string SDKEnvironmentVarsFile = "OutputEnvVars.txt";
+
+		protected const string SDKRootEnvVar = "UE_SDKS_ROOT";
+
+		protected const string AutoSetupEnvVar = "AutoSDKSetup";
+
+
+		private static string GetAutoSDKHostPlatform()
+		{
+			if (RuntimePlatform.IsWindows)
+			{
+				return "Win64";
+			}
+			else if (RuntimePlatform.IsMac)
+			{
+				return "Mac";
+			}
+			else if (RuntimePlatform.IsLinux)
+			{
+				return "Linux";
+			}
+			throw new Exception("Unknown host platform!");
+		}
+
+		/// <summary>
+		/// Whether platform supports switching SDKs during runtime
+		/// </summary>
+		/// <returns>true if supports</returns>
+		protected virtual bool PlatformSupportsAutoSDKs()
+		{
+			return false;
+		}
+
+		static private bool bCheckedAutoSDKRootEnvVar = false;
+		static private bool bAutoSDKSystemEnabled = false;
+		static private bool HasAutoSDKSystemEnabled()
+		{
+			if (!bCheckedAutoSDKRootEnvVar)
+			{
+				string SDKRoot = Environment.GetEnvironmentVariable(SDKRootEnvVar);
+				if (SDKRoot != null)
+				{
+					bAutoSDKSystemEnabled = true;
+				}
+				bCheckedAutoSDKRootEnvVar = true;
+			}
+			return bAutoSDKSystemEnabled;
+		}
+
+		// Whether AutoSDK setup is safe. AutoSDKs will damage manual installs on some platforms.
+		protected bool IsAutoSDKSafe()
+		{
+			return !IsAutoSDKDestructive() || !HasAnyManualInstall();
+		}
+
+		/// <summary>
+		/// Gets the version number of the SDK setup script itself.  The version in the base should ALWAYS be the master revision from the last refactor.
+		/// If you need to force a rebuild for a given platform, modify the version file.
+		/// </summary>
+		/// <returns>Setup script version</returns>
+		private string GetRequiredScriptVersionString()
+		{
+			const string UnspecifiedVersion = "UnspecifiedScriptVersion";
+
+			string VersionFilename = Path.Combine(GetPathToPlatformAutoSDKs(), GetAutoSDKDirectoryForMainVersion(), ScriptVersionFilename);
+			if (File.Exists(VersionFilename))
+			{
+				return File.ReadAllLines(VersionFilename).FirstOrDefault() ?? UnspecifiedVersion;
+			}
+			else
+			{
+				return UnspecifiedVersion;
+			}
+		}
+
+		/// <summary>
+		/// Returns path to platform SDKs
+		/// </summary>
+		/// <returns>Valid SDK string</returns>
+		protected string GetPathToPlatformAutoSDKs()
+		{
+			string SDKPath = "";
+			string SDKRoot = Environment.GetEnvironmentVariable(SDKRootEnvVar);
+			if (SDKRoot != null)
+			{
+				if (SDKRoot != "")
+				{
+					SDKPath = Path.Combine(SDKRoot, "Host" + GetAutoSDKHostPlatform(), GetAutoSDKPlatformName());
+				}
+			}
+			return SDKPath;
+		}
+
+		/// <summary>
+		/// Returns path to platform SDKs
+		/// </summary>
+		/// <returns>Valid SDK string</returns>
+		public static bool TryGetHostPlatformAutoSDKDir(out DirectoryReference OutPlatformDir)
+		{
+			string SDKRoot = Environment.GetEnvironmentVariable(SDKRootEnvVar);
+			if (String.IsNullOrEmpty(SDKRoot))
+			{
+				OutPlatformDir = null;
+				return false;
+			}
+			else
+			{
+				OutPlatformDir = DirectoryReference.Combine(new DirectoryReference(SDKRoot), "Host" + GetAutoSDKHostPlatform());
+				return true;
+			}
+		}
+
+		/// <summary>
+		/// Because most ManualSDK determination depends on reading env vars, if this process is spawned by a process that ALREADY set up
+		/// AutoSDKs then all the SDK env vars will exist, and we will spuriously detect a Manual SDK. (children inherit the environment of the parent process).
+		/// Therefore we write out an env variable to set in the command file (OutputEnvVars.txt) such that child processes can determine if their manual SDK detection
+		/// is bogus.  Make it platform specific so that platforms can be in different states.
+		/// </summary>
+		protected string GetPlatformAutoSDKSetupEnvVar()
+		{
+			return GetAutoSDKPlatformName() + AutoSetupEnvVar;
+		}
+
+		/// <summary>
+		/// Gets currently installed version
+		/// </summary>
+		/// <param name="PlatformSDKRoot">absolute path to platform SDK root</param>
+		/// <param name="OutInstalledSDKVersionString">version string as currently installed</param>
+		/// <returns>true if was able to read it</returns>
+		protected bool GetCurrentlyInstalledSDKString(string PlatformSDKRoot, out string OutInstalledSDKVersionString)
+		{
+			if (Directory.Exists(PlatformSDKRoot))
+			{
+				string VersionFilename = Path.Combine(PlatformSDKRoot, CurrentlyInstalledSDKStringManifest);
+				if (File.Exists(VersionFilename))
+				{
+					using (StreamReader Reader = new StreamReader(VersionFilename))
+					{
+						string Version = Reader.ReadLine();
+						string Type = Reader.ReadLine();
+
+						// don't allow ManualSDK installs to count as an AutoSDK install version.
+						if (Type != null && Type == "AutoSDK")
+						{
+							if (Version != null)
+							{
+								OutInstalledSDKVersionString = Version;
+								return true;
+							}
+						}
+					}
+				}
+			}
+
+			OutInstalledSDKVersionString = "";
+			return false;
+		}
+
+		/// <summary>
+		/// Gets the version of the last successfully run setup script.
+		/// </summary>
+		/// <param name="PlatformSDKRoot">absolute path to platform SDK root</param>
+		/// <param name="OutLastRunScriptVersion">version string</param>
+		/// <returns>true if was able to read it</returns>
+		protected bool GetLastRunScriptVersionString(string PlatformSDKRoot, out string OutLastRunScriptVersion)
+		{
+			if (Directory.Exists(PlatformSDKRoot))
+			{
+				string VersionFilename = Path.Combine(PlatformSDKRoot, LastRunScriptVersionManifest);
+				if (File.Exists(VersionFilename))
+				{
+					using (StreamReader Reader = new StreamReader(VersionFilename))
+					{
+						string Version = Reader.ReadLine();
+						if (Version != null)
+						{
+							OutLastRunScriptVersion = Version;
+							return true;
+						}
+					}
+				}
+			}
+
+			OutLastRunScriptVersion = "";
+			return false;
+		}
+
+		/// <summary>
+		/// Sets currently installed version
+		/// </summary>
+		/// <param name="InstalledSDKVersionString">SDK version string to set</param>
+		/// <returns>true if was able to set it</returns>
+		protected bool SetCurrentlyInstalledAutoSDKString(String InstalledSDKVersionString)
+		{
+			String PlatformSDKRoot = GetPathToPlatformAutoSDKs();
+			if (Directory.Exists(PlatformSDKRoot))
+			{
+				string VersionFilename = Path.Combine(PlatformSDKRoot, CurrentlyInstalledSDKStringManifest);
+				if (File.Exists(VersionFilename))
+				{
+					File.Delete(VersionFilename);
+				}
+
+				using (StreamWriter Writer = File.CreateText(VersionFilename))
+				{
+					Writer.WriteLine(InstalledSDKVersionString);
+					Writer.WriteLine("AutoSDK");
+					return true;
+				}
+			}
+
+			return false;
+		}
+
+		protected void SetupManualSDK()
+		{
+			if (PlatformSupportsAutoSDKs() && HasAutoSDKSystemEnabled())
+			{
+				String InstalledSDKVersionString = GetAutoSDKDirectoryForMainVersion();
+				String PlatformSDKRoot = GetPathToPlatformAutoSDKs();
+                if (!Directory.Exists(PlatformSDKRoot))
+                {
+                    Directory.CreateDirectory(PlatformSDKRoot);
+                }
+
+				{
+					string VersionFilename = Path.Combine(PlatformSDKRoot, CurrentlyInstalledSDKStringManifest);
+					if (File.Exists(VersionFilename))
+					{
+						File.Delete(VersionFilename);
+					}
+
+					string EnvVarFile = Path.Combine(PlatformSDKRoot, SDKEnvironmentVarsFile);
+					if (File.Exists(EnvVarFile))
+					{
+						File.Delete(EnvVarFile);
+					}
+
+					using (StreamWriter Writer = File.CreateText(VersionFilename))
+					{
+						Writer.WriteLine(InstalledSDKVersionString);
+						Writer.WriteLine("ManualSDK");
+					}
+				}
+			}
+		}
+
+		protected bool SetLastRunAutoSDKScriptVersion(string LastRunScriptVersion)
+		{
+			String PlatformSDKRoot = GetPathToPlatformAutoSDKs();
+			if (Directory.Exists(PlatformSDKRoot))
+			{
+				string VersionFilename = Path.Combine(PlatformSDKRoot, LastRunScriptVersionManifest);
+				if (File.Exists(VersionFilename))
+				{
+					File.Delete(VersionFilename);
+				}
+
+				using (StreamWriter Writer = File.CreateText(VersionFilename))
+				{
+					Writer.WriteLine(LastRunScriptVersion);
+					return true;
+				}
+			}
+			return false;
+		}
+
+		/// <summary>
+		/// Returns Hook names as needed by the platform
+		/// (e.g. can be overridden with custom executables or scripts)
+		/// </summary>
+		/// <param name="Hook">Hook type</param>
+		protected virtual string GetHookExecutableName(SDKHookType Hook)
+		{
+			if (RuntimePlatform.IsWindows)
+			{
+				if (Hook == SDKHookType.Uninstall)
+				{
+					return "unsetup.bat";
+				}
+				else
+				{
+					return "setup.bat";
+				}
+			}
+			else
+			{
+				if (Hook == SDKHookType.Uninstall)
+				{
+					return "unsetup.sh";
+				}
+				else
+				{
+					return "setup.sh";
+				}
+			}
+		}
+
+		/// <summary>
+		/// Whether the hook must be run with administrator privileges.
+		/// </summary>
+		/// <param name="Hook">Hook for which to check the required privileges.</param>
+		/// <returns>true if the hook must be run with administrator privileges.</returns>
+		protected virtual bool DoesHookRequireAdmin(SDKHookType Hook)
+		{
+			return true;
+		}
+
+		private void LogAutoSDKHook(object Sender, DataReceivedEventArgs Args)
+		{
+			if (Args.Data != null)
+			{
+				LogFormatOptions Options = Log.OutputLevel >= LogEventType.Verbose ? LogFormatOptions.None : LogFormatOptions.NoConsoleOutput;
+				Log.WriteLine(LogEventType.Log, Options, Args.Data);
+			}
+		}
+
+		/// <summary>
+		/// Runs install/uninstall hooks for SDK
+		/// </summary>
+		/// <param name="PlatformSDKRoot">absolute path to platform SDK root</param>
+		/// <param name="SDKVersionString">version string to run for (can be empty!)</param>
+		/// <param name="Hook">which one of hooks to run</param>
+		/// <param name="bHookCanBeNonExistent">whether a non-existing hook means failure</param>
+		/// <returns>true if succeeded</returns>
+		protected virtual bool RunAutoSDKHooks(string PlatformSDKRoot, string SDKVersionString, SDKHookType Hook, bool bHookCanBeNonExistent = true)
+		{
+			if (!IsAutoSDKSafe())
+			{
+				Log.TraceLog(GetAutoSDKPlatformName() + " attempted to run SDK hook which could have damaged manual SDK install!");
+				return false;
+			}
+			if (SDKVersionString != "")
+			{
+				string SDKDirectory = Path.Combine(PlatformSDKRoot, SDKVersionString);
+				string HookExe = Path.Combine(SDKDirectory, GetHookExecutableName(Hook));
+
+				if (File.Exists(HookExe))
+				{
+					Log.TraceLog("Running {0} hook {1}", Hook, HookExe);
+
+					// run it
+					Process HookProcess = new Process();
+					HookProcess.StartInfo.WorkingDirectory = SDKDirectory;
+					HookProcess.StartInfo.FileName = HookExe;
+					HookProcess.StartInfo.Arguments = "";
+					HookProcess.StartInfo.WindowStyle = ProcessWindowStyle.Hidden;
+
+					bool bHookRequiresAdmin = DoesHookRequireAdmin(Hook);
+					if (bHookRequiresAdmin)
+					{
+						// installers may require administrator access to succeed. so run as an admin.
+						HookProcess.StartInfo.Verb = "runas";
+					}
+					else
+					{
+						HookProcess.StartInfo.UseShellExecute = false;
+						HookProcess.StartInfo.RedirectStandardOutput = true;
+						HookProcess.StartInfo.RedirectStandardError = true;
+						HookProcess.OutputDataReceived += LogAutoSDKHook;
+						HookProcess.ErrorDataReceived += LogAutoSDKHook;
+					}
+
+					//using (ScopedTimer HookTimer = new ScopedTimer("Time to run hook: ", LogEventType.Log))
+					{
+						HookProcess.Start();
+						if (!bHookRequiresAdmin)
+						{
+							HookProcess.BeginOutputReadLine();
+							HookProcess.BeginErrorReadLine();
+						}
+						HookProcess.WaitForExit();
+					}
+
+					if (HookProcess.ExitCode != 0)
+					{
+						Log.TraceLog("Hook exited uncleanly (returned {0}), considering it failed.", HookProcess.ExitCode);
+						return false;
+					}
+
+					return true;
+				}
+				else
+				{
+					Log.TraceLog("File {0} does not exist", HookExe);
+				}
+			}
+			else
+			{
+				Log.TraceLog("Version string is blank for {0}. Can't determine {1} hook.", PlatformSDKRoot, Hook.ToString());
+			}
+
+			return bHookCanBeNonExistent;
+		}
+
+		/// <summary>
+		/// Loads environment variables from SDK
+		/// If any commands are added or removed the handling needs to be duplicated in
+		/// TargetPlatformManagerModule.cpp
+		/// </summary>
+		/// <param name="PlatformSDKRoot">absolute path to platform SDK</param>
+		/// <returns>true if succeeded</returns>
+		protected bool SetupEnvironmentFromAutoSDK(string PlatformSDKRoot)
+		{
+			string EnvVarFile = Path.Combine(PlatformSDKRoot, SDKEnvironmentVarsFile);
+			if (File.Exists(EnvVarFile))
+			{
+				using (StreamReader Reader = new StreamReader(EnvVarFile))
+				{
+					List<string> PathAdds = new List<string>();
+					List<string> PathRemoves = new List<string>();
+
+					List<string> EnvVarNames = new List<string>();
+					List<string> EnvVarValues = new List<string>();
+
+					bool bNeedsToWriteAutoSetupEnvVar = true;
+					String PlatformSetupEnvVar = GetPlatformAutoSDKSetupEnvVar();
+					for (; ; )
+					{
+						string VariableString = Reader.ReadLine();
+						if (VariableString == null)
+						{
+							break;
+						}
+
+						string[] Parts = VariableString.Split('=');
+						if (Parts.Length != 2)
+						{
+							Log.TraceLog("Incorrect environment variable declaration:");
+							Log.TraceLog(VariableString);
+							return false;
+						}
+
+						if (String.Compare(Parts[0], "strippath", true) == 0)
+						{
+							PathRemoves.Add(Parts[1]);
+						}
+						else if (String.Compare(Parts[0], "addpath", true) == 0)
+						{
+							PathAdds.Add(Parts[1]);
+						}
+						else
+						{
+							if (String.Compare(Parts[0], PlatformSetupEnvVar) == 0)
+							{
+								bNeedsToWriteAutoSetupEnvVar = false;
+							}
+							// convenience for setup.bat writers.  Trim any accidental whitespace from variable names/values.
+							EnvVarNames.Add(Parts[0].Trim());
+							EnvVarValues.Add(Parts[1].Trim());
+						}
+					}
+
+					// don't actually set anything until we successfully validate and read all values in.
+					// we don't want to set a few vars, return a failure, and then have a platform try to
+					// build against a manually installed SDK with half-set env vars.
+					for (int i = 0; i < EnvVarNames.Count; ++i)
+					{
+						string EnvVarName = EnvVarNames[i];
+						string EnvVarValue = EnvVarValues[i];
+						Log.TraceVerbose("Setting variable '{0}' to '{1}'", EnvVarName, EnvVarValue);
+						Environment.SetEnvironmentVariable(EnvVarName, EnvVarValue);
+					}
+
+
+                    // actually perform the PATH stripping / adding.
+                    String OrigPathVar = Environment.GetEnvironmentVariable("PATH");
+                    String PathDelimiter = RuntimePlatform.IsWindows ? ";" : ":";
+                    String[] PathVars = { };
+                    if (!String.IsNullOrEmpty(OrigPathVar))
+                    {
+                        PathVars = OrigPathVar.Split(PathDelimiter.ToCharArray());
+                    }
+                    else
+                    {
+                        Log.TraceVerbose("Path environment variable is null during AutoSDK");
+                    }
+
+					List<String> ModifiedPathVars = new List<string>();
+					ModifiedPathVars.AddRange(PathVars);
+
+					// perform removes first, in case they overlap with any adds.
+					foreach (String PathRemove in PathRemoves)
+					{
+						foreach (String PathVar in PathVars)
+						{
+							if (PathVar.IndexOf(PathRemove, StringComparison.OrdinalIgnoreCase) >= 0)
+							{
+								Log.TraceVerbose("Removing Path: '{0}'", PathVar);
+								ModifiedPathVars.Remove(PathVar);
+							}
+						}
+					}
+
+					// remove all the of ADDs so that if this function is executed multiple times, the paths will be guaranteed to be in the same order after each run.
+					// If we did not do this, a 'remove' that matched some, but not all, of our 'adds' would cause the order to change.
+					foreach (String PathAdd in PathAdds)
+					{
+						foreach (String PathVar in PathVars)
+						{
+							if (String.Compare(PathAdd, PathVar, true) == 0)
+							{
+								Log.TraceVerbose("Removing Path: '{0}'", PathVar);
+								ModifiedPathVars.Remove(PathVar);
+							}
+						}
+					}
+
+					// perform adds, but don't add duplicates
+					foreach (String PathAdd in PathAdds)
+					{
+						if (!ModifiedPathVars.Contains(PathAdd))
+						{
+							Log.TraceVerbose("Adding Path: '{0}'", PathAdd);
+							ModifiedPathVars.Add(PathAdd);
+						}
+					}
+
+					String ModifiedPath = String.Join(PathDelimiter, ModifiedPathVars);
+					Environment.SetEnvironmentVariable("PATH", ModifiedPath);
+
+					Reader.Close();
+
+					// write out environment variable command so any process using this commandfile will mark itself as having had autosdks set up.
+					// avoids child processes spuriously detecting manualsdks.
+					if (bNeedsToWriteAutoSetupEnvVar)
+					{
+						// write out the manual sdk version since child processes won't be able to detect manual with AutoSDK messing up env vars
+						using (StreamWriter Writer = File.AppendText(EnvVarFile))
+						{
+							Writer.WriteLine("{0}={1}", PlatformSetupEnvVar, "1");
+						}
+						// set the variable in the local environment in case this process spawns any others.
+						Environment.SetEnvironmentVariable(PlatformSetupEnvVar, "1");
+					}
+
+					// make sure we know that we've modified the local environment, invalidating manual installs for this run.
+					bLocalProcessSetupAutoSDK = true;
+
+					// tell any child processes what our manual version was before setting up autosdk
+					string ValueToWrite = CachedManualSDKVersion != null ? CachedManualSDKVersion : "__None";
+					Environment.SetEnvironmentVariable(GetPlatformAutoSDKSetupEnvVar(), ValueToWrite);
+
+					return true;
+				}
+			}
+			else
+			{
+				Log.TraceLog("Cannot set up environment for {1} because command file {2} does not exist.", PlatformSDKRoot, EnvVarFile);
+			}
+
+			return false;
+		}
+
+		protected void InvalidateCurrentlyInstalledAutoSDK()
+		{
+			String PlatformSDKRoot = GetPathToPlatformAutoSDKs();
+			if (Directory.Exists(PlatformSDKRoot))
+			{
+				string SDKFilename = Path.Combine(PlatformSDKRoot, CurrentlyInstalledSDKStringManifest);
+				if (File.Exists(SDKFilename))
+				{
+					File.Delete(SDKFilename);
+				}
+
+				string VersionFilename = Path.Combine(PlatformSDKRoot, LastRunScriptVersionManifest);
+				if (File.Exists(VersionFilename))
+				{
+					File.Delete(VersionFilename);
+				}
+
+				string EnvVarFile = Path.Combine(PlatformSDKRoot, SDKEnvironmentVarsFile);
+				if (File.Exists(EnvVarFile))
+				{
+					File.Delete(EnvVarFile);
+				}
+			}
+		}
+
+		/// <summary>
+		/// Currently installed AutoSDK is written out to a text file in a known location.
+		/// This function just compares the file's contents with the current requirements.
+		/// </summary>
+		public SDKStatus HasRequiredAutoSDKInstalled()
+		{
+			if (PlatformSupportsAutoSDKs() && HasAutoSDKSystemEnabled())
+			{
+				string AutoSDKRoot = GetPathToPlatformAutoSDKs();
+				if (AutoSDKRoot != "")
+				{
+					// check script version so script fixes can be propagated without touching every build machine's CurrentlyInstalled file manually.
+					bool bScriptVersionMatches = false;
+					string CurrentScriptVersionString;
+					if (GetLastRunScriptVersionString(AutoSDKRoot, out CurrentScriptVersionString) && CurrentScriptVersionString == GetRequiredScriptVersionString())
+					{
+						bScriptVersionMatches = true;
+					}
+
+					// check to make sure OutputEnvVars doesn't need regenerating
+					string EnvVarFile = Path.Combine(AutoSDKRoot, SDKEnvironmentVarsFile);
+					bool bEnvVarFileExists = File.Exists(EnvVarFile);
+
+					string CurrentSDKString;
+					if (bEnvVarFileExists && GetCurrentlyInstalledSDKString(AutoSDKRoot, out CurrentSDKString) && CurrentSDKString == GetAutoSDKDirectoryForMainVersion() && bScriptVersionMatches)
+					{
+						return SDKStatus.Valid;
+					}
+					return SDKStatus.Invalid;
+				}
+			}
+			return SDKStatus.Invalid;
+		}
+
+		// This tracks if we have already checked the sdk installation.
+		private Int32 SDKCheckStatus = -1;
+
+		// true if we've ever overridden the process's environment with AutoSDK data.  After that, manual installs cannot be considered valid ever again.
+		private bool bLocalProcessSetupAutoSDK = false;
+
+		protected bool HasSetupAutoSDK()
+		{
+			return bLocalProcessSetupAutoSDK || HasParentProcessSetupAutoSDK(out _);
+		}
+
+		protected bool HasParentProcessSetupAutoSDK(out string OutAutoSDKSetupValue)
+		{
+			bool bParentProcessSetupAutoSDK = false;
+			String AutoSDKSetupVarName = GetPlatformAutoSDKSetupEnvVar();
+			OutAutoSDKSetupValue = Environment.GetEnvironmentVariable(AutoSDKSetupVarName);
+			if (!String.IsNullOrEmpty(OutAutoSDKSetupValue))
+			{
+				bParentProcessSetupAutoSDK = true;
+			}
+			return bParentProcessSetupAutoSDK;
+		}
+
+		public SDKStatus HasRequiredManualSDK()
+		{
+// 			if (HasSetupAutoSDK())
+// 			{
+// 				return SDKStatus.Invalid;
+// 			}
+//
+//			// manual installs are always invalid if we have modified the process's environment for AutoSDKs
+			return HasRequiredManualSDKInternal();
+		}
+
+		// for platforms with destructive AutoSDK.  Report if any manual sdk is installed that may be damaged by an autosdk.
+		protected virtual bool HasAnyManualInstall()
+		{
+			return false;
+		}
+
+		// tells us if the user has a valid manual install.
+		protected virtual SDKStatus HasRequiredManualSDKInternal()
+		{
+			string ManualSDKVersion;
+			GetInstalledVersions(out ManualSDKVersion, out _);
+
+			return IsVersionValid(ManualSDKVersion, bForAutoSDK:false) ? SDKStatus.Valid : SDKStatus.Invalid;
+		}
+
+		// some platforms will fail if there is a manual install that is the WRONG manual install.
+		protected virtual bool AllowInvalidManualInstall()
+		{
+			return true;
+		}
+
+		// platforms can choose if they prefer a correct the the AutoSDK install over the manual install.
+		protected virtual bool PreferAutoSDK()
+		{
+			return true;
+		}
+
+		// some platforms don't support parallel SDK installs.  AutoSDK on these platforms will
+		// actively damage an existing manual install by overwriting files in it.  AutoSDK must NOT
+		// run any setup if a manual install exists in this case.
+		protected virtual bool IsAutoSDKDestructive()
+		{
+			return false;
+		}
+
+		/// <summary>
+		/// Runs batch files if necessary to set up required AutoSDK.
+		/// AutoSDKs are SDKs that have not been setup through a formal installer, but rather come from
+		/// a source control directory, or other local copy.
+		/// </summary>
+		private void SetupAutoSDK()
+		{
+			if (IsAutoSDKSafe() && PlatformSupportsAutoSDKs() && HasAutoSDKSystemEnabled())
+			{
+				// run installation for autosdk if necessary.
+				if (HasRequiredAutoSDKInstalled() == SDKStatus.Invalid)
+				{
+					//reset check status so any checking sdk status after the attempted setup will do a real check again.
+					SDKCheckStatus = -1;
+
+					string AutoSDKRoot = GetPathToPlatformAutoSDKs();
+					string CurrentSDKString;
+					GetCurrentlyInstalledSDKString(AutoSDKRoot, out CurrentSDKString);
+
+					// switch over (note that version string can be empty)
+					if (!RunAutoSDKHooks(AutoSDKRoot, CurrentSDKString, SDKHookType.Uninstall))
+					{
+						Log.TraceLog("Failed to uninstall currently installed SDK {0}", CurrentSDKString);
+						InvalidateCurrentlyInstalledAutoSDK();
+						return;
+					}
+					// delete Manifest file to avoid multiple uninstalls
+					InvalidateCurrentlyInstalledAutoSDK();
+
+					if (!RunAutoSDKHooks(AutoSDKRoot, GetAutoSDKDirectoryForMainVersion(), SDKHookType.Install, false))
+					{
+						Log.TraceLog("Failed to install required SDK {0}.  Attemping to uninstall", GetAutoSDKDirectoryForMainVersion());
+						RunAutoSDKHooks(AutoSDKRoot, GetAutoSDKDirectoryForMainVersion(), SDKHookType.Uninstall, false);
+						return;
+					}
+
+					string EnvVarFile = Path.Combine(AutoSDKRoot, SDKEnvironmentVarsFile);
+					if (!File.Exists(EnvVarFile))
+					{
+						Log.TraceLog("Installation of required SDK {0}.  Did not generate Environment file {1}", GetAutoSDKDirectoryForMainVersion(), EnvVarFile);
+						RunAutoSDKHooks(AutoSDKRoot, GetAutoSDKDirectoryForMainVersion(), SDKHookType.Uninstall, false);
+						return;
+					}
+
+					SetCurrentlyInstalledAutoSDKString(GetAutoSDKDirectoryForMainVersion());
+					SetLastRunAutoSDKScriptVersion(GetRequiredScriptVersionString());
+				}
+
+				// fixup process environment to match autosdk
+				SetupEnvironmentFromAutoSDK();
+			}
+		}
+
+		/// <summary>
+		/// Allows the platform to optionally returns a path to the internal SDK
+		/// </summary>
+		/// <returns>Valid path to the internal SDK, null otherwise</returns>
+		public virtual string GetInternalSDKPath()
+		{
+			return null;
+		}
+
+		#endregion
+
+		#region public AutoSDKs Utility
+
+		/// <summary>
+		/// Enum describing types of hooks a platform SDK can have
+		/// </summary>
+		public enum SDKHookType
+		{
+			Install,
+			Uninstall
+		};
+
+		/* Whether or not we should try to automatically switch SDKs when asked to validate the platform's SDK state. */
+		public static bool bAllowAutoSDKSwitching = true;
+
+		public SDKStatus SetupEnvironmentFromAutoSDK()
+		{
+			string PlatformSDKRoot = GetPathToPlatformAutoSDKs();
+
+			// load environment variables from current SDK
+			if (!SetupEnvironmentFromAutoSDK(PlatformSDKRoot))
+			{
+				Log.TraceLog("Failed to load environment from required SDK {0}", GetAutoSDKDirectoryForMainVersion());
+				InvalidateCurrentlyInstalledAutoSDK();
+				return SDKStatus.Invalid;
+			}
+			return SDKStatus.Valid;
+		}
+
+		/// <summary>
+		/// Whether the required external SDKs are installed for this platform.
+		/// Could be either a manual install or an AutoSDK.
+		/// </summary>
+		public SDKStatus HasRequiredSDKsInstalled()
+		{
+			// avoid redundant potentially expensive SDK checks.
+			if (SDKCheckStatus == -1)
+			{
+				bool bHasManualSDK = HasRequiredManualSDK() == SDKStatus.Valid;
+				bool bHasAutoSDK = HasRequiredAutoSDKInstalled() == SDKStatus.Valid;
+
+				// Per-Platform implementations can choose how to handle non-Auto SDK detection / handling.
+				SDKCheckStatus = (bHasManualSDK || bHasAutoSDK) ? 1 : 0;
+			}
+			return SDKCheckStatus == 1 ? SDKStatus.Valid : SDKStatus.Invalid;
+		}
+
+		// Arbitrates between manual SDKs and setting up AutoSDK based on program options and platform preferences.
+		public void ManageAndValidateSDK()
+		{
+			// do not modify installed manifests if parent process has already set everything up.
+			// this avoids problems with determining IsAutoSDKSafe and doing an incorrect invalidate.
+			if (bAllowAutoSDKSwitching && !HasParentProcessSetupAutoSDK(out _))
+			{
+				bool bSetSomeSDK = false;
+				bool bHasRequiredManualSDK = HasRequiredManualSDK() == SDKStatus.Valid;
+				if (IsAutoSDKSafe() && (PreferAutoSDK() || !bHasRequiredManualSDK))
+				{
+					SetupAutoSDK();
+					bSetSomeSDK = true;
+				}
+
+				//Setup manual SDK if autoSDK setup was skipped or failed for whatever reason.
+				if (bHasRequiredManualSDK && (HasRequiredAutoSDKInstalled() != SDKStatus.Valid))
+				{
+					SetupManualSDK();
+					bSetSomeSDK = true;
+				}
+
+				if (!bSetSomeSDK)
+				{
+					InvalidateCurrentlyInstalledAutoSDK();
+				}
+			}
+
+			// print all SDKs to log file (errors will print out later for builds and generateprojectfiles)
+			PrintSDKInfoAndReturnValidity(LogEventType.Log, LogFormatOptions.NoConsoleOutput, LogEventType.Verbose, LogFormatOptions.NoConsoleOutput);
+		}
+		#endregion
+
+	}
+}
\ No newline at end of file
diff --git a/Engine/Source/Programs/Shared/EpicGames.Build/Unreal.cs b/Engine/Source/Programs/Shared/EpicGames.Build/Unreal.cs
new file mode 100644
--- /dev/null
+++ b/Engine/Source/Programs/Shared/EpicGames.Build/Unreal.cs
@@ -0,0 +1,272 @@
+// Copyright Epic Games, Inc. All Rights Reserved.
+
+using EpicGames.Core;
+using System;
+using System.Collections.Generic;
+using System.IO;
+using System.Linq;
+using System.Reflection;
+using System.Text;
+
+namespace UnrealBuildBase
+{
+	public static class Unreal
+	{
+		private static DirectoryReference FindRootDirectory()
+		{
+			if (LocationOverride.RootDirectory != null)
+			{
+				return LocationOverride.RootDirectory;
+			}
+
+			// This base library may be used - and so be launched - from more than one location (at time of writing, UnrealBuildTool and AutomationTool)
+			// Programs that use this assembly must be located under "Engine/Binaries/DotNET" and so we look for that sequence of directories in that path of the executing assembly
+			
+			// Use the EntryAssembly (the application path), rather than the ExecutingAssembly (the library path)
+			string AssemblyLocation = Assembly.GetEntryAssembly().GetOriginalLocation();
+
+			DirectoryReference? FoundRootDirectory = DirectoryReference.FindCorrectCase(DirectoryReference.FromString(AssemblyLocation)!);
+
+			// Search up through the directory tree for the deepest instance of the sub-path "Engine/Binaries/DotNET"
+			while(FoundRootDirectory != null)
+			{
+				if (String.Equals("DotNET", FoundRootDirectory.GetDirectoryName()))
+				{
+					FoundRootDirectory = FoundRootDirectory.ParentDirectory;
+					if (FoundRootDirectory != null && String.Equals("Binaries", FoundRootDirectory.GetDirectoryName()))
+					{
+						FoundRootDirectory = FoundRootDirectory.ParentDirectory;
+						if (FoundRootDirectory != null && String.Equals("Engine", FoundRootDirectory.GetDirectoryName()))
+						{
+							FoundRootDirectory = FoundRootDirectory.ParentDirectory;
+							break;
+						}
+						continue;
+					}
+					continue;
+				}
+				FoundRootDirectory = FoundRootDirectory.ParentDirectory;
+			}
+
+			if (FoundRootDirectory == null)
+			{
+				throw new Exception($"The BuildUtilities assembly requires that applications are launched from a path containing \"Engine/Binaries/DotNET\". This application was launched from {Path.GetDirectoryName(AssemblyLocation)}");
+			}
+
+			// Confirm that we've found a valid root directory, by testing for the existence of a well-known file
+			FileReference ExpectedExistingFile = FileReference.Combine(FoundRootDirectory, "Engine", "Build", "Build.version");
+			if (!FileReference.Exists(ExpectedExistingFile))
+			{
+				throw new Exception($"Expected file \"Engine/Build/Build.version\" was not found at {ExpectedExistingFile.FullName}");
+			}
+
+			return FoundRootDirectory;
+		}
+
+		private static FileReference FindUnrealBuildTool()
+		{
+			// todo: use UnrealBuildTool.dll (same on all platforms). Will require changes wherever UnrealBuildTool is invoked.
+			string UBTName = "UnrealBuildTool" + RuntimePlatform.ExeExtension;
+
+			// the UnrealBuildTool executable is assumed to be located under {RootDirectory}/Engine/Binaries/DotNET/UnrealBuildTool/
+			FileReference UnrealBuildToolPath = FileReference.Combine(EngineDirectory, "Binaries", "DotNET", "UnrealBuildTool", UBTName);
+			
+			UnrealBuildToolPath = FileReference.FindCorrectCase(UnrealBuildToolPath);
+
+			if (!FileReference.Exists(UnrealBuildToolPath))
+			{
+				throw new Exception($"Unable to find {UBTName} in the expected location at {UnrealBuildToolPath.FullName}");
+			}
+
+			return UnrealBuildToolPath;
+		}
+
+		static private DirectoryReference FindDotnetDirectory()
+		{
+			string HostDotNetDirectoryName;
+			switch (RuntimePlatform.Current)
+			{
+				case RuntimePlatform.Type.Windows: HostDotNetDirectoryName = "Windows"; break;
+				case RuntimePlatform.Type.Mac:     HostDotNetDirectoryName = "Mac"; break;
+				case RuntimePlatform.Type.Linux:   HostDotNetDirectoryName = "Linux"; break;
+				default: throw new Exception("Unknown host platform");
+			}
+
+			return DirectoryReference.Combine(EngineDirectory, "Binaries", "ThirdParty", "DotNet", HostDotNetDirectoryName);
+		}
+
+		/// <summary>
+		/// The full name of the root UE directory
+		/// </summary>
+		public static readonly DirectoryReference RootDirectory = FindRootDirectory();
+
+		/// <summary>
+		/// The full name of the Engine directory
+		/// </summary>
+		public static readonly DirectoryReference EngineDirectory = DirectoryReference.Combine(RootDirectory, "Engine");
+
+		/// <summary>
+		/// The path to UBT
+		/// </summary>
+		public static readonly FileReference UnrealBuildToolPath = FindUnrealBuildTool();
+
+		/// <summary>
+		/// The directory containing the bundled .NET installation
+		/// </summary>
+		static public readonly DirectoryReference DotnetDirectory = FindDotnetDirectory();
+
+		/// <summary>
+		/// The path of the bundled dotnet executable
+		/// </summary>
+		static public readonly FileReference DotnetPath = FileReference.Combine(DotnetDirectory, "dotnet" + RuntimePlatform.ExeExtension);
+
+		/// <summary>
+		/// Whether we're running with engine installed
+		/// </summary>
+		static private bool? bIsEngineInstalled;
+
+		/// <summary>
+		/// Returns true if UnrealBuildTool is running using installed Engine components
+		/// </summary>
+		/// <returns>True if running using installed Engine components</returns>
+		static public bool IsEngineInstalled()
+		{
+			if (!bIsEngineInstalled.HasValue)
+			{
+				bIsEngineInstalled = FileReference.Exists(FileReference.Combine(EngineDirectory, "Build", "InstalledBuild.txt"));
+			}
+			return bIsEngineInstalled.Value;
+		}
+
+		public static class LocationOverride 
+		{
+			/// <summary>
+			/// If set, this value will be used to populate Unreal.RootDirectory
+			/// </summary>
+			public static DirectoryReference? RootDirectory = null;
+		}
+
+
+		// A subset of the functionality in DataDrivenPlatformInfo.GetAllPlatformInfos() - finds the DataDrivenPlatformInfo.ini files and records their existence, but does not parse them
+		// (perhaps DataDrivenPlatformInfo.GetAllPlatformInfos() could be modified to use this data to avoid an additional search through the filesystem)
+		public static HashSet<string>? IniPresentForPlatform = null;
+		private static bool DataDrivenPlatformInfoIniIsPresent(string PlatformName)
+		{
+			if (IniPresentForPlatform == null)
+			{
+				IniPresentForPlatform = new HashSet<string>(StringComparer.OrdinalIgnoreCase);
+
+				// find all platform directories (skipping NFL/NoRedist)
+				foreach (DirectoryReference EngineConfigDir in GetExtensionDirs(Unreal.EngineDirectory, "Config", bIncludeRestrictedDirectories: false))
+				{
+					// look through all config dirs looking for the data driven ini file
+					foreach (string FilePath in Directory.EnumerateFiles(EngineConfigDir.FullName, "DataDrivenPlatformInfo.ini", SearchOption.AllDirectories))
+					{
+						FileReference FileRef = new FileReference(FilePath);
+
+						// get the platform name from the path
+						string IniPlatformName;
+						if (FileRef.IsUnderDirectory(DirectoryReference.Combine(Unreal.EngineDirectory, "Config")))
+						{
+							// Foo/Engine/Config/<Platform>/DataDrivenPlatformInfo.ini
+							IniPlatformName = Path.GetFileName(Path.GetDirectoryName(FilePath))!;
+						}
+						else
+						{
+							// Foo/Engine/Platforms/<Platform>/Config/DataDrivenPlatformInfo.ini
+							IniPlatformName = Path.GetFileName(Path.GetDirectoryName(Path.GetDirectoryName(FilePath)))!;
+						}
+
+						// DataDrivenPlatformInfo.GetAllPlatformInfos() checks that [DataDrivenPlatformInfo] section exists as part of validating that the file exists
+						// This code should probably behave the same way.
+
+						IniPresentForPlatform.Add(IniPlatformName);
+					}
+				}
+			}
+			return IniPresentForPlatform.Contains(PlatformName);
+		}
+
+		// cached dictionary of BaseDir to extension directories
+		private static Dictionary<DirectoryReference, Tuple<List<DirectoryReference>, List<DirectoryReference>>> CachedExtensionDirectories = new Dictionary<DirectoryReference, Tuple<List<DirectoryReference>, List<DirectoryReference>>>();
+
+		/// <summary>
+		/// Finds all the extension directories for the given base directory. This includes platform extensions and restricted folders.
+		/// </summary>
+		/// <param name="BaseDir">Location of the base directory</param>
+		/// <param name="bIncludePlatformDirectories">If true, platform subdirectories are included (will return platform directories under Restricted dirs, even if bIncludeRestrictedDirectories is false)</param>
+		/// <param name="bIncludeRestrictedDirectories">If true, restricted (NotForLicensees, NoRedist) subdirectories are included</param>
+		/// <param name="bIncludeBaseDirectory">If true, BaseDir is included</param>
+		/// <returns>List of extension directories, including the given base directory</returns>
+		public static List<DirectoryReference> GetExtensionDirs(DirectoryReference BaseDir, bool bIncludePlatformDirectories=true, bool bIncludeRestrictedDirectories=true, bool bIncludeBaseDirectory=true)
+		{
+			Tuple<List<DirectoryReference>, List<DirectoryReference>>? CachedDirs;
+			if (!CachedExtensionDirectories.TryGetValue(BaseDir, out CachedDirs))
+			{
+				CachedDirs = Tuple.Create(new List<DirectoryReference>(), new List<DirectoryReference>());
+
+				CachedExtensionDirectories[BaseDir] = CachedDirs;
+
+				DirectoryReference PlatformExtensionBaseDir = DirectoryReference.Combine(BaseDir, "Platforms");
+				if (DirectoryReference.Exists(PlatformExtensionBaseDir))
+				{
+					CachedDirs.Item1.AddRange(DirectoryReference.EnumerateDirectories(PlatformExtensionBaseDir));
+				}
+
+				DirectoryReference RestrictedBaseDir = DirectoryReference.Combine(BaseDir, "Restricted");
+				if (DirectoryReference.Exists(RestrictedBaseDir))
+				{
+					IEnumerable<DirectoryReference> RestrictedDirs = DirectoryReference.EnumerateDirectories(RestrictedBaseDir);
+					CachedDirs.Item2.AddRange(RestrictedDirs);
+
+					// also look for nested platforms in the restricted
+					foreach (DirectoryReference RestrictedDir in RestrictedDirs)
+					{
+						DirectoryReference RestrictedPlatformExtensionBaseDir = DirectoryReference.Combine(RestrictedDir, "Platforms");
+						if (DirectoryReference.Exists(RestrictedPlatformExtensionBaseDir))
+						{
+							CachedDirs.Item1.AddRange(DirectoryReference.EnumerateDirectories(RestrictedPlatformExtensionBaseDir));
+						}
+					}
+				}
+
+				// remove any platform directories in non-engine locations if the engine doesn't have the platform 
+				if (BaseDir != Unreal.EngineDirectory && CachedDirs.Item1.Count > 0)
+				{
+					// if the DDPI.ini file doesn't exist, we haven't synced the platform, so just skip this directory
+					CachedDirs.Item1.RemoveAll(x => !DataDrivenPlatformInfoIniIsPresent(x.GetDirectoryName()));
+				}
+			}
+
+			// now return what the caller wanted (always include BaseDir)
+			List<DirectoryReference> ExtensionDirs = new List<DirectoryReference>();
+			if (bIncludeBaseDirectory)
+			{
+				ExtensionDirs.Add(BaseDir);
+			}
+			if (bIncludePlatformDirectories)
+			{
+				ExtensionDirs.AddRange(CachedDirs.Item1);
+			}
+			if (bIncludeRestrictedDirectories)
+			{
+				ExtensionDirs.AddRange(CachedDirs.Item2);
+			}
+			return ExtensionDirs;
+		}
+
+		/// <summary>
+		/// Finds all the extension directories for the given base directory. This includes platform extensions and restricted folders.
+		/// </summary>
+		/// <param name="BaseDir">Location of the base directory</param>
+		/// <param name="SubDir">The subdirectory to find</param>
+		/// <param name="bIncludePlatformDirectories">If true, platform subdirectories are included (will return platform directories under Restricted dirs, even if bIncludeRestrictedDirectories is false)</param>
+		/// <param name="bIncludeRestrictedDirectories">If true, restricted (NotForLicensees, NoRedist) subdirectories are included</param>
+		/// <param name="bIncludeBaseDirectory">If true, BaseDir is included</param>
+		/// <returns>List of extension directories, including the given base directory</returns>
+		public static List<DirectoryReference> GetExtensionDirs(DirectoryReference BaseDir, string SubDir, bool bIncludePlatformDirectories=true, bool bIncludeRestrictedDirectories=true, bool bIncludeBaseDirectory=true)
+		{
+			return GetExtensionDirs(BaseDir, bIncludePlatformDirectories, bIncludeRestrictedDirectories, bIncludeBaseDirectory).Select(x => DirectoryReference.Combine(x, SubDir)).Where(x => DirectoryReference.Exists(x)).ToList();
+		}
+	}
+}
diff --git a/Engine/Source/Programs/Shared/EpicGames.Core/EpicGames.Core.csproj b/Engine/Source/Programs/Shared/EpicGames.Core/EpicGames.Core.csproj
--- a/Engine/Source/Programs/Shared/EpicGames.Core/EpicGames.Core.csproj
+++ b/Engine/Source/Programs/Shared/EpicGames.Core/EpicGames.Core.csproj
@@ -22,7 +22,7 @@
     <Optimize>true</Optimize>
   </PropertyGroup>
   <ItemGroup>
-    <Compile Include="..\..\DotNETCommon\MetaData.cs">
+    <Compile Include="..\..\Shared\MetaData.cs">
       <Link>Properties\MetaData.cs</Link>
     </Compile>
   </ItemGroup>
diff --git a/Engine/Source/Programs/Shared/EpicGames.Horde.Client/EpicGames.Horde.Client.csproj b/Engine/Source/Programs/Shared/EpicGames.Horde.Client/EpicGames.Horde.Client.csproj
--- a/Engine/Source/Programs/Shared/EpicGames.Horde.Client/EpicGames.Horde.Client.csproj
+++ b/Engine/Source/Programs/Shared/EpicGames.Horde.Client/EpicGames.Horde.Client.csproj
@@ -22,7 +22,7 @@
     <Optimize>true</Optimize>
   </PropertyGroup>
   <ItemGroup>
-    <Compile Include="..\..\DotNETCommon\MetaData.cs">
+    <Compile Include="..\..\Shared\MetaData.cs">
       <Link>Properties\MetaData.cs</Link>
     </Compile>
   </ItemGroup>
diff --git a/Engine/Source/Programs/Shared/EpicGames.Horde.Common/EpicGames.Horde.Common.csproj b/Engine/Source/Programs/Shared/EpicGames.Horde.Common/EpicGames.Horde.Common.csproj
--- a/Engine/Source/Programs/Shared/EpicGames.Horde.Common/EpicGames.Horde.Common.csproj
+++ b/Engine/Source/Programs/Shared/EpicGames.Horde.Common/EpicGames.Horde.Common.csproj
@@ -22,7 +22,7 @@
     <Optimize>true</Optimize>
   </PropertyGroup>
   <ItemGroup>
-    <Compile Include="..\..\DotNETCommon\MetaData.cs">
+    <Compile Include="..\..\Shared\MetaData.cs">
       <Link>Properties\MetaData.cs</Link>
     </Compile>
   </ItemGroup>
diff --git a/Engine/Source/Programs/Shared/EpicGames.Horde/EpicGames.Horde.csproj b/Engine/Source/Programs/Shared/EpicGames.Horde/EpicGames.Horde.csproj
--- a/Engine/Source/Programs/Shared/EpicGames.Horde/EpicGames.Horde.csproj
+++ b/Engine/Source/Programs/Shared/EpicGames.Horde/EpicGames.Horde.csproj
@@ -1,4 +1,4 @@
-﻿<Project Sdk="Microsoft.NET.Sdk">
+<Project Sdk="Microsoft.NET.Sdk">
   <PropertyGroup>
     <TargetFramework>netcoreapp3.1</TargetFramework>
     <Configuration Condition=" '$(Configuration)' == '' ">Development</Configuration>
@@ -22,7 +22,7 @@
     <Optimize>true</Optimize>
   </PropertyGroup>
   <ItemGroup>
-    <Compile Include="..\..\DotNETCommon\MetaData.cs">
+    <Compile Include="..\..\Shared\MetaData.cs">
       <Link>Properties\MetaData.cs</Link>
     </Compile>
   </ItemGroup>
diff --git a/Engine/Source/Programs/Shared/EpicGames.IoHash/EpicGames.IoHash.csproj b/Engine/Source/Programs/Shared/EpicGames.IoHash/EpicGames.IoHash.csproj
--- a/Engine/Source/Programs/Shared/EpicGames.IoHash/EpicGames.IoHash.csproj
+++ b/Engine/Source/Programs/Shared/EpicGames.IoHash/EpicGames.IoHash.csproj
@@ -22,7 +22,7 @@
     <Optimize>true</Optimize>
   </PropertyGroup>
   <ItemGroup>
-    <Compile Include="..\..\DotNETCommon\MetaData.cs">
+    <Compile Include="..\..\Shared\MetaData.cs">
       <Link>Properties\MetaData.cs</Link>
     </Compile>
   </ItemGroup>
diff --git a/Engine/Source/Programs/Shared/EpicGames.Jupiter/EpicGames.Jupiter.csproj b/Engine/Source/Programs/Shared/EpicGames.Jupiter/EpicGames.Jupiter.csproj
--- a/Engine/Source/Programs/Shared/EpicGames.Jupiter/EpicGames.Jupiter.csproj
+++ b/Engine/Source/Programs/Shared/EpicGames.Jupiter/EpicGames.Jupiter.csproj
@@ -21,7 +21,7 @@
     <Optimize>true</Optimize>
   </PropertyGroup>
   <ItemGroup>
-    <Compile Include="..\..\DotNETCommon\MetaData.cs">
+    <Compile Include="..\..\Shared\MetaData.cs">
       <Link>Properties\MetaData.cs</Link>
     </Compile>
   </ItemGroup>
diff --git a/Engine/Source/Programs/Shared/EpicGames.MongoDB/EpicGames.MongoDB.csproj b/Engine/Source/Programs/Shared/EpicGames.MongoDB/EpicGames.MongoDB.csproj
--- a/Engine/Source/Programs/Shared/EpicGames.MongoDB/EpicGames.MongoDB.csproj
+++ b/Engine/Source/Programs/Shared/EpicGames.MongoDB/EpicGames.MongoDB.csproj
@@ -21,7 +21,7 @@
     <Optimize>true</Optimize>
   </PropertyGroup>
   <ItemGroup>
-    <Compile Include="..\..\DotNETCommon\MetaData.cs">
+    <Compile Include="..\..\Shared\MetaData.cs">
       <Link>Properties\MetaData.cs</Link>
     </Compile>
   </ItemGroup>
diff --git a/Engine/Source/Programs/Shared/EpicGames.Perforce.Managed/EpicGames.Perforce.Managed.csproj b/Engine/Source/Programs/Shared/EpicGames.Perforce.Managed/EpicGames.Perforce.Managed.csproj
--- a/Engine/Source/Programs/Shared/EpicGames.Perforce.Managed/EpicGames.Perforce.Managed.csproj
+++ b/Engine/Source/Programs/Shared/EpicGames.Perforce.Managed/EpicGames.Perforce.Managed.csproj
@@ -1,4 +1,4 @@
-﻿<Project Sdk="Microsoft.NET.Sdk">
+<Project Sdk="Microsoft.NET.Sdk">
 
   <PropertyGroup>
     <TargetFramework>netcoreapp3.1</TargetFramework>
@@ -25,7 +25,7 @@
     <ProjectReference Include="..\EpicGames.Serialization\EpicGames.Serialization.csproj" />
   </ItemGroup>
   <ItemGroup>
-    <Compile Include="..\..\DotNETCommon\MetaData.cs">
+    <Compile Include="..\..\Shared\MetaData.cs">
       <Link>Properties\MetaData.cs</Link>
     </Compile>
   </ItemGroup>
diff --git a/Engine/Source/Programs/Shared/EpicGames.Perforce/EpicGames.Perforce.csproj b/Engine/Source/Programs/Shared/EpicGames.Perforce/EpicGames.Perforce.csproj
--- a/Engine/Source/Programs/Shared/EpicGames.Perforce/EpicGames.Perforce.csproj
+++ b/Engine/Source/Programs/Shared/EpicGames.Perforce/EpicGames.Perforce.csproj
@@ -22,7 +22,7 @@
     <Optimize>true</Optimize>
   </PropertyGroup>
   <ItemGroup>
-    <Compile Include="..\..\DotNETCommon\MetaData.cs">
+    <Compile Include="..\..\Shared\MetaData.cs">
       <Link>Properties\MetaData.cs</Link>
     </Compile>
   </ItemGroup>
diff --git a/Engine/Source/Programs/Shared/EpicGames.Serialization.Tests/EpicGames.Serialization.Tests.csproj b/Engine/Source/Programs/Shared/EpicGames.Serialization.Tests/EpicGames.Serialization.Tests.csproj
--- a/Engine/Source/Programs/Shared/EpicGames.Serialization.Tests/EpicGames.Serialization.Tests.csproj
+++ b/Engine/Source/Programs/Shared/EpicGames.Serialization.Tests/EpicGames.Serialization.Tests.csproj
@@ -1,4 +1,4 @@
-﻿<Project Sdk="Microsoft.NET.Sdk">
+<Project Sdk="Microsoft.NET.Sdk">
 
   <PropertyGroup>
     <TargetFramework>netcoreapp3.1</TargetFramework>
@@ -15,7 +15,7 @@
     <Nullable>enable</Nullable>
   </PropertyGroup>
   <ItemGroup>
-    <Compile Include="..\..\DotNETCommon\MetaData.cs">
+    <Compile Include="..\..\Shared\MetaData.cs">
       <Link>Properties\MetaData.cs</Link>
     </Compile>
   </ItemGroup>
diff --git a/Engine/Source/Programs/Shared/EpicGames.Serialization/EpicGames.Serialization.csproj b/Engine/Source/Programs/Shared/EpicGames.Serialization/EpicGames.Serialization.csproj
--- a/Engine/Source/Programs/Shared/EpicGames.Serialization/EpicGames.Serialization.csproj
+++ b/Engine/Source/Programs/Shared/EpicGames.Serialization/EpicGames.Serialization.csproj
@@ -22,7 +22,7 @@
     <Optimize>true</Optimize>
   </PropertyGroup>
   <ItemGroup>
-    <Compile Include="..\..\DotNETCommon\MetaData.cs">
+    <Compile Include="..\..\Shared\MetaData.cs">
       <Link>Properties\MetaData.cs</Link>
     </Compile>
   </ItemGroup>
diff --git a/Engine/Source/Programs/Shared/MetaData.cs b/Engine/Source/Programs/Shared/MetaData.cs
new file mode 100644
--- /dev/null
+++ b/Engine/Source/Programs/Shared/MetaData.cs
@@ -0,0 +1,24 @@
+// Copyright Epic Games, Inc. All Rights Reserved.
+
+using System.Reflection;
+using System.Runtime.InteropServices;
+using System.Resources;
+
+// These are the assembly properties for all tools
+[assembly: AssemblyCompany( "Epic Games, Inc." )]
+[assembly: AssemblyProduct( "UnrealEngine" )]
+[assembly: AssemblyCopyright( "Copyright Epic Games, Inc. All Rights Reserved." )]
+[assembly: AssemblyTrademark( "" )]
+
+// Use a neutral culture to avoid some localisation issues
+[assembly: AssemblyCulture( "" )]
+
+[assembly: ComVisible( false )]
+[assembly: NeutralResourcesLanguageAttribute( "" )]
+
+#if !SPECIFIC_VERSION
+// Automatically generate a version number based on the time of compilation
+[assembly: AssemblyVersion( "4.0.0.0" )]
+[assembly: AssemblyInformationalVersion("4.0.0")]
+#endif
+
diff --git a/Engine/Source/Programs/UnrealBuildTool/ProjectFiles/ProjectFileGenerator.cs b/Engine/Source/Programs/UnrealBuildTool/ProjectFiles/ProjectFileGenerator.cs
--- a/Engine/Source/Programs/UnrealBuildTool/ProjectFiles/ProjectFileGenerator.cs
+++ b/Engine/Source/Programs/UnrealBuildTool/ProjectFiles/ProjectFileGenerator.cs
@@ -581,7 +581,7 @@ namespace UnrealBuildTool
 				
 				Writer.WriteLine("  <ItemGroup>");
 				Writer.WriteLine("    <ProjectReference Include=\"" + 
-					FileReference.Combine(Unreal.EngineDirectory, "Source/Programs/DotNETCommon/BuildUtilities/BuildUtilities.csproj").MakeRelativeTo(ProjectFilesDirectory) + 
+					FileReference.Combine(Unreal.EngineDirectory, "Source/Programs/Shared/EpicGames.Build/EpicGames.Build.csproj").MakeRelativeTo(ProjectFilesDirectory) + 
 					"\"><Private>false</Private></ProjectReference>"); // <Private>false</Private> to suppress copying of dependencies into output directory
 				Writer.WriteLine("    <ProjectReference Include=\"" + 
 					FileReference.Combine(Unreal.EngineDirectory, "Source", "Programs", "UnrealBuildTool", "UnrealBuildTool.csproj").MakeRelativeTo(ProjectFilesDirectory) + 
diff --git a/Engine/Source/Programs/UnrealBuildTool/UnrealBuildTool.csproj b/Engine/Source/Programs/UnrealBuildTool/UnrealBuildTool.csproj
--- a/Engine/Source/Programs/UnrealBuildTool/UnrealBuildTool.csproj
+++ b/Engine/Source/Programs/UnrealBuildTool/UnrealBuildTool.csproj
@@ -70,7 +70,7 @@
   </Choose>
 
   <ItemGroup>
-    <Compile Include="..\DotNETCommon\MetaData.cs">
+    <Compile Include="..\Shared\MetaData.cs">
       <Link>Properties\MetaData.cs</Link>
     </Compile>
   </ItemGroup>
@@ -97,7 +97,7 @@
   </ItemGroup>
 
   <ItemGroup>
-    <ProjectReference Include="..\DotNETCommon\BuildUtilities\BuildUtilities.csproj" />
+    <ProjectReference Include="..\Shared\EpicGames.Build\EpicGames.Build.csproj" />
     <ProjectReference Include="..\Shared\EpicGames.Core\EpicGames.Core.csproj" />
   </ItemGroup>
   <ItemGroup>
