update diagnostic text

#trivial

diff --git a/Engine/Source/Programs/Shared/EpicGames.Build/Unreal.cs b/Engine/Source/Programs/Shared/EpicGames.Build/Unreal.cs
--- a/Engine/Source/Programs/Shared/EpicGames.Build/Unreal.cs
+++ b/Engine/Source/Programs/Shared/EpicGames.Build/Unreal.cs
@@ -50,7 +50,7 @@ namespace UnrealBuildBase
 
 			if (FoundRootDirectory == null)
 			{
-				throw new Exception($"The BuildUtilities assembly requires that applications are launched from a path containing \"Engine/Binaries/DotNET\". This application was launched from {Path.GetDirectoryName(AssemblyLocation)}");
+				throw new Exception($"This code requires that applications using it are launched from a path containing \"Engine/Binaries/DotNET\". This application was launched from {Path.GetDirectoryName(AssemblyLocation)}");
 			}
 
 			// Confirm that we've found a valid root directory, by testing for the existence of a well-known file
