Horde: Add a library for accessing typed values in Redis.

diff --git a/Engine/Source/Programs/Shared/EpicGames.Redis/EpicGames.Redis.csproj b/Engine/Source/Programs/Shared/EpicGames.Redis/EpicGames.Redis.csproj
new file mode 100644
--- /dev/null
+++ b/Engine/Source/Programs/Shared/EpicGames.Redis/EpicGames.Redis.csproj
@@ -0,0 +1,37 @@
+<Project Sdk="Microsoft.NET.Sdk">
+  <PropertyGroup>
+    <TargetFramework>netcoreapp3.1</TargetFramework>
+    <Configuration Condition=" '$(Configuration)' == '' ">Development</Configuration>
+    <OutputType>Library</OutputType>
+    <Deterministic>true</Deterministic>
+    <RootNamespace>EpicGames.Redis</RootNamespace>
+    <GenerateAssemblyInfo>false</GenerateAssemblyInfo>
+    <GenerateTargetFrameworkAttribute>false</GenerateTargetFrameworkAttribute>
+    <Configurations>Debug;Release;Development</Configurations>
+    <Nullable>enable</Nullable>
+  </PropertyGroup>
+  <PropertyGroup Condition=" '$(Configuration)|$(Platform)' == 'Debug|AnyCPU' ">
+    <DocumentationFile></DocumentationFile>
+  </PropertyGroup>
+  <PropertyGroup Condition="'$(Configuration)|$(Platform)' == 'Development|AnyCPU'">
+    <DocumentationFile></DocumentationFile>
+    <Optimize>false</Optimize>
+  </PropertyGroup>
+  <PropertyGroup Condition="'$(Configuration)|$(Platform)' == 'Release|AnyCPU'">
+    <DocumentationFile></DocumentationFile>
+    <Optimize>true</Optimize>
+  </PropertyGroup>
+  <ItemGroup>
+    <Compile Include="..\..\Shared\MetaData.cs">
+      <Link>Properties\MetaData.cs</Link>
+    </Compile>
+  </ItemGroup>
+  <ItemGroup>
+    <PackageReference Include="Microsoft.Extensions.Logging" Version="5.0.0" />
+    <PackageReference Include="StackExchange.Redis" Version="2.2.4" />
+  </ItemGroup>
+  <ItemGroup>
+    <ProjectReference Include="..\EpicGames.Core\EpicGames.Core.csproj" />
+    <ProjectReference Include="..\EpicGames.Serialization\EpicGames.Serialization.csproj" />
+  </ItemGroup>
+</Project>
\ No newline at end of file
diff --git a/Engine/Source/Programs/Shared/EpicGames.Redis/RedisExtensions.cs b/Engine/Source/Programs/Shared/EpicGames.Redis/RedisExtensions.cs
new file mode 100644
--- /dev/null
+++ b/Engine/Source/Programs/Shared/EpicGames.Redis/RedisExtensions.cs
@@ -0,0 +1,36 @@
+// Copyright Epic Games, Inc. All Rights Reserved.
+
+using StackExchange.Redis;
+using System;
+using System.Collections.Generic;
+using System.Linq;
+using System.Threading.Tasks;
+
+namespace EpicGames.Redis
+{
+	/// <summary>
+	/// Extension methods for Redis
+	/// </summary>
+	public static class RedisExtensions
+	{
+		/// <summary>
+		/// Convert a RedisValue to a RedisKey
+		/// </summary>
+		/// <param name="Value"></param>
+		/// <returns></returns>
+		public static RedisKey AsKey(this RedisValue Value)
+		{
+			return (byte[])Value;
+		}
+
+		/// <summary>
+		/// Convert a RedisKey to a RedisValue
+		/// </summary>
+		/// <param name="Key"></param>
+		/// <returns></returns>
+		public static RedisValue AsValue(this RedisKey Key)
+		{
+			return (byte[])Key;
+		}
+	}
+}
diff --git a/Engine/Source/Programs/Shared/EpicGames.Redis/RedisList.cs b/Engine/Source/Programs/Shared/EpicGames.Redis/RedisList.cs
new file mode 100644
--- /dev/null
+++ b/Engine/Source/Programs/Shared/EpicGames.Redis/RedisList.cs
@@ -0,0 +1,82 @@
+// Copyright Epic Games, Inc. All Rights Reserved.
+
+using StackExchange.Redis;
+using System;
+using System.Collections.Generic;
+using System.Linq;
+using System.Text;
+using System.Threading.Tasks;
+
+namespace EpicGames.Redis
+{
+	/// <summary>
+	/// Represents a typed Redis list with a given key
+	/// </summary>
+	/// <typeparam name="TElement">The type of element stored in the list</typeparam>
+	public readonly struct RedisList<TElement> : IEquatable<RedisList<TElement>>
+	{
+		/// <summary>
+		/// The key for the list
+		/// </summary>
+		public readonly RedisKey Key { get; }
+
+		/// <summary>
+		/// Constructor
+		/// </summary>
+		/// <param name="Key"></param>
+		public RedisList(RedisKey Key) => this.Key = Key;
+
+		/// <inheritdoc/>
+		public override bool Equals(object? Obj) => Obj is RedisList<TElement> List && Key == List.Key;
+
+		/// <inheritdoc/>
+		public bool Equals(RedisList<TElement> Other) => Key == Other.Key;
+
+		/// <inheritdoc/>
+		public override int GetHashCode() => Key.GetHashCode();
+
+		/// <summary>Compares two instances for equality</summary>
+		public static bool operator ==(RedisList<TElement> Left, RedisList<TElement> Right) => Left.Key == Right.Key;
+
+		/// <summary>Compares two instances for equality</summary>
+		public static bool operator !=(RedisList<TElement> Left, RedisList<TElement> Right) => Left.Key != Right.Key;
+	}
+
+	/// <summary>
+	/// Extension methods for typed lists
+	/// </summary>
+	public static class RedisListExtensions
+	{
+		/// <inheritdoc cref="IDatabaseAsync.ListLengthAsync(RedisKey, CommandFlags)"/>
+		public static Task<long> ListLengthAsync<T>(this IDatabase Database, RedisList<T> List)
+		{
+			return Database.ListLengthAsync(List.Key);
+		}
+
+		/// <inheritdoc cref="IDatabaseAsync.ListLeftPushAsync(RedisKey, RedisValue, When, CommandFlags)"/>
+		public static Task<long> ListLeftPushAsync<T>(this IDatabase Database, RedisList<T> List, T Item, When When = When.Always, CommandFlags Flags = CommandFlags.None)
+		{
+			return Database.ListLeftPushAsync(List.Key, RedisSerializer.Serialize(Item), When, Flags);
+		}
+
+		/// <inheritdoc cref="IDatabaseAsync.ListLeftPushAsync(RedisKey, RedisValue[], When, CommandFlags)"/>
+		public static Task<long> ListLeftPushAsync<T>(this IDatabase Database, RedisList<T> List, IEnumerable<T> Items, When When = When.Always, CommandFlags Flags = CommandFlags.None)
+		{
+			RedisValue[] Values = Items.Select(x => RedisSerializer.Serialize(x)).ToArray();
+			return Database.ListLeftPushAsync(List.Key, Values, When, Flags);
+		}
+
+		/// <inheritdoc cref="IDatabaseAsync.ListRightPushAsync(RedisKey, RedisValue, When, CommandFlags)"/>
+		public static Task<long> ListRightPushAsync<T>(this IDatabase Database, RedisList<T> List, T Item, When When = When.Always, CommandFlags Flags = CommandFlags.None)
+		{
+			return Database.ListRightPushAsync(List.Key, RedisSerializer.Serialize(Item), When, Flags);
+		}
+
+		/// <inheritdoc cref="IDatabaseAsync.ListRightPushAsync(RedisKey, RedisValue[], When, CommandFlags)"/>
+		public static Task<long> ListRightPushAsync<T>(this IDatabase Database, RedisList<T> List, IEnumerable<T> Items, When When = When.Always, CommandFlags Flags = CommandFlags.None)
+		{
+			RedisValue[] Values = Items.Select(x => RedisSerializer.Serialize(x)).ToArray();
+			return Database.ListRightPushAsync(List.Key, Values, When, Flags);
+		}
+	}
+}
diff --git a/Engine/Source/Programs/Shared/EpicGames.Redis/RedisSerializer.cs b/Engine/Source/Programs/Shared/EpicGames.Redis/RedisSerializer.cs
new file mode 100644
--- /dev/null
+++ b/Engine/Source/Programs/Shared/EpicGames.Redis/RedisSerializer.cs
@@ -0,0 +1,159 @@
+// Copyright Epic Games, Inc. All Rights Reserved.
+
+using EpicGames.Serialization;
+using StackExchange.Redis;
+using System;
+using System.Collections.Generic;
+using System.ComponentModel;
+using System.Linq;
+using System.Reflection;
+using System.Threading.Tasks;
+
+namespace EpicGames.Redis
+{
+	/// <summary>
+	/// Attribute specifying the converter type to use for a class
+	/// </summary>
+	[AttributeUsage(AttributeTargets.Class)]
+	public class RedisConverterAttribute : Attribute
+	{
+		/// <summary>
+		/// Type of the converter to use
+		/// </summary>
+		public Type ConverterType { get; }
+
+		/// <summary>
+		/// Constructor
+		/// </summary>
+		/// <param name="ConverterType">The converter type</param>
+		public RedisConverterAttribute(Type ConverterType)
+		{
+			this.ConverterType = ConverterType;
+		}
+	}
+
+	/// <summary>
+	/// Converter to and from RedisValue types
+	/// </summary>
+	public interface IRedisConverter<T>
+	{
+		/// <summary>
+		/// Serailize an object to a RedisValue
+		/// </summary>
+		/// <param name="Value"></param>
+		/// <returns></returns>
+		RedisValue ToRedisValue(T Value);
+
+		/// <summary>
+		/// Deserialize an object from a RedisValue
+		/// </summary>
+		/// <param name="Value"></param>
+		/// <returns></returns>
+		T FromRedisValue(RedisValue Value);
+	}
+
+	/// <summary>
+	/// Redis serializer that uses compact binary to serialize objects
+	/// </summary>
+	/// <typeparam name="T"></typeparam>
+	public sealed class RedisCbConverter<T> : IRedisConverter<T>
+	{
+		/// <inheritdoc/>
+		public RedisValue ToRedisValue(T Value)
+		{
+			return CbSerializer.Serialize(Value).GetView();
+		}
+
+		/// <inheritdoc/>
+		public T FromRedisValue(RedisValue Value)
+		{
+			return CbSerializer.Deserialize<T>(new CbField((byte[])Value));
+		}
+	}
+
+	/// <summary>
+	/// Handles serialization of types to RedisValue instances
+	/// </summary>
+	public static class RedisSerializer
+	{
+		class RedisStringConverter<T> : IRedisConverter<T>
+		{
+			TypeConverter TypeConverter;
+
+			public RedisStringConverter(TypeConverter TypeConverter)
+			{
+				this.TypeConverter = TypeConverter;
+			}
+
+			public RedisValue ToRedisValue(T Value) => (string)TypeConverter.ConvertTo(Value, typeof(string));
+			public T FromRedisValue(RedisValue Value) => (T)TypeConverter.ConvertFrom((string)Value);
+		}
+
+		/// <summary>
+		/// Static class for caching converter lookups
+		/// </summary>
+		/// <typeparam name="T"></typeparam>
+		class CachedConverter<T>
+		{
+			public static IRedisConverter<T> Converter = CreateConverter();
+
+			static IRedisConverter<T> CreateConverter()
+			{
+				Type Type = typeof(T);
+
+				// Check for a custom converter
+				RedisConverterAttribute? Attribute = Type.GetCustomAttribute<RedisConverterAttribute>();
+				if (Attribute != null)
+				{
+					Type ConverterType = Attribute.ConverterType;
+					if (ConverterType.IsGenericTypeDefinition)
+					{
+						ConverterType = ConverterType.MakeGenericType(Type);
+					}
+					return (IRedisConverter<T>)Activator.CreateInstance(ConverterType)!;
+				}
+
+				// Check if there's a regular converter we can use to convert to/from a string
+				TypeConverter? Converter = TypeDescriptor.GetConverter(Type);
+				if (Converter != null && Converter.CanConvertFrom(typeof(string)) && Converter.CanConvertTo(typeof(string)))
+				{
+					return new RedisStringConverter<T>(Converter);
+				}
+
+				throw new Exception($"Unable to find Redis converter for {Type.Name}");
+			}
+		}
+
+		/// <summary>
+		/// Gets the converter for a particular type
+		/// </summary>
+		/// <typeparam name="T"></typeparam>
+		/// <returns></returns>
+		public static IRedisConverter<T> GetConverter<T>()
+		{
+			return CachedConverter<T>.Converter;
+		}
+
+		/// <summary>
+		/// Serialize an object to a <see cref="RedisValue"/>
+		/// </summary>
+		/// <typeparam name="T"></typeparam>
+		/// <param name="Value"></param>
+		/// <returns></returns>
+		public static RedisValue Serialize<T>(T Value)
+		{
+			return CachedConverter<T>.Converter.ToRedisValue(Value);
+		}
+
+		/// <summary>
+		/// Deserialize a <see cref="RedisValue"/>
+		/// </summary>
+		/// <typeparam name="T"></typeparam>
+		/// <param name="Value"></param>
+		/// <returns></returns>
+		public static T Deserialize<T>(RedisValue Value)
+		{
+			return CachedConverter<T>.Converter.FromRedisValue(Value);
+		}
+	}
+}
diff --git a/Engine/Source/Programs/Shared/EpicGames.Redis/RedisSet.cs b/Engine/Source/Programs/Shared/EpicGames.Redis/RedisSet.cs
new file mode 100644
--- /dev/null
+++ b/Engine/Source/Programs/Shared/EpicGames.Redis/RedisSet.cs
@@ -0,0 +1,61 @@
+// Copyright Epic Games, Inc. All Rights Reserved.
+
+using StackExchange.Redis;
+using System;
+using System.Collections.Generic;
+using System.Linq;
+using System.Threading.Tasks;
+
+namespace EpicGames.Redis
+{
+	/// <summary>
+	/// Represents a typed Redis list with a given key
+	/// </summary>
+	/// <typeparam name="TElement">The type of element stored in the list</typeparam>
+	public readonly struct RedisSet<TElement> : IEquatable<RedisSet<TElement>>
+	{
+		/// <summary>
+		/// The key for the list
+		/// </summary>
+		public readonly RedisKey Key { get; }
+
+		/// <summary>
+		/// Constructor
+		/// </summary>
+		/// <param name="Key"></param>
+		public RedisSet(RedisKey Key) => this.Key = Key;
+
+		/// <inheritdoc/>
+		public override bool Equals(object? Obj) => Obj is RedisSet<TElement> Set && Key == Set.Key;
+
+		/// <inheritdoc/>
+		public bool Equals(RedisSet<TElement> Other) => Key == Other.Key;
+
+		/// <inheritdoc/>
+		public override int GetHashCode() => Key.GetHashCode();
+
+		/// <summary>Compares two instances for equality</summary>
+		public static bool operator ==(RedisSet<TElement> Left, RedisSet<TElement> Right) => Left.Key == Right.Key;
+
+		/// <summary>Compares two instances for equality</summary>
+		public static bool operator !=(RedisSet<TElement> Left, RedisSet<TElement> Right) => Left.Key != Right.Key;
+	}
+
+	/// <summary>
+	/// Extension methods for <see cref="IDatabaseAsync"/>
+	/// </summary>
+	public static class RedisSetExtensions
+	{
+		/// <inheritdoc cref="IDatabaseAsync.SetAddAsync(RedisKey, RedisValue, CommandFlags)"/>
+		public static Task SetAddAsync<T>(this IDatabaseAsync Database, RedisSet<T> Set, T Value, CommandFlags Flags = CommandFlags.None)
+		{
+			return Database.SetAddAsync(Set.Key, RedisSerializer.Serialize(Value), Flags);
+		}
+
+		/// <inheritdoc cref="IDatabaseAsync.SetRemoveAsync(RedisKey, RedisValue, CommandFlags)"/>
+		public static Task SetRemoveAsync<T>(this IDatabaseAsync Database, RedisSet<T> Set, T Value, CommandFlags Flags = CommandFlags.None)
+		{
+			return Database.SetRemoveAsync(Set.Key, RedisSerializer.Serialize(Value), Flags);
+		}
+	}
+}
diff --git a/Engine/Source/Programs/Shared/EpicGames.Redis/RedisSortedSet.cs b/Engine/Source/Programs/Shared/EpicGames.Redis/RedisSortedSet.cs
new file mode 100644
--- /dev/null
+++ b/Engine/Source/Programs/Shared/EpicGames.Redis/RedisSortedSet.cs
@@ -0,0 +1,55 @@
+// Copyright Epic Games, Inc. All Rights Reserved.
+
+using StackExchange.Redis;
+using System;
+using System.Collections.Generic;
+using System.Linq;
+using System.Threading.Tasks;
+
+namespace EpicGames.Redis
+{
+	/// <summary>
+	/// Represents a typed Redis list with a given key
+	/// </summary>
+	/// <typeparam name="TElement">The type of element stored in the list</typeparam>
+	public readonly struct RedisSortedSet<TElement> : IEquatable<RedisSortedSet<TElement>>
+	{
+		/// <summary>
+		/// The key for the list
+		/// </summary>
+		public readonly RedisKey Key { get; }
+
+		/// <summary>
+		/// Constructor
+		/// </summary>
+		/// <param name="Key"></param>
+		public RedisSortedSet(RedisKey Key) => this.Key = Key;
+
+		/// <inheritdoc/>
+		public override bool Equals(object? Obj) => Obj is RedisSortedSet<TElement> Set && Key == Set.Key;
+
+		/// <inheritdoc/>
+		public bool Equals(RedisSortedSet<TElement> Other) => Key == Other.Key;
+
+		/// <inheritdoc/>
+		public override int GetHashCode() => Key.GetHashCode();
+
+		/// <summary>Compares two instances for equality</summary>
+		public static bool operator ==(RedisSortedSet<TElement> Left, RedisSortedSet<TElement> Right) => Left.Key == Right.Key;
+
+		/// <summary>Compares two instances for equality</summary>
+		public static bool operator !=(RedisSortedSet<TElement> Left, RedisSortedSet<TElement> Right) => Left.Key != Right.Key;
+	}
+
+	/// <summary>
+	/// Extension methods for <see cref="IDatabaseAsync"/>
+	/// </summary>
+	public static class RedisSortedSetExtensions
+	{
+		/// <inheritdoc cref="IDatabaseAsync.SortedSetAddAsync(RedisKey, RedisValue, double, CommandFlags)"/>
+		public static Task SortedSetAddAsync<T>(this IDatabaseAsync Database, RedisSortedSet<T> Set, T Value, double Score, CommandFlags Flags = CommandFlags.None)
+		{
+			return Database.SortedSetAddAsync(Set.Key, RedisSerializer.Serialize(Value), Score, Flags);
+		}
+	}
+}
diff --git a/Engine/Source/Programs/Shared/EpicGames.Redis/RedisString.cs b/Engine/Source/Programs/Shared/EpicGames.Redis/RedisString.cs
new file mode 100644
--- /dev/null
+++ b/Engine/Source/Programs/Shared/EpicGames.Redis/RedisString.cs
@@ -0,0 +1,62 @@
+// Copyright Epic Games, Inc. All Rights Reserved.
+
+using StackExchange.Redis;
+using System;
+using System.Collections.Generic;
+using System.Linq;
+using System.Threading.Tasks;
+
+namespace EpicGames.Redis
+{
+	/// <summary>
+	/// Represents a typed Redis list with a given key
+	/// </summary>
+	/// <typeparam name="TElement">The type of element stored in the list</typeparam>
+	public readonly struct RedisString<TElement> : IEquatable<RedisString<TElement>>
+	{
+		/// <summary>
+		/// The key for the list
+		/// </summary>
+		public readonly RedisKey Key { get; }
+
+		/// <summary>
+		/// Constructor
+		/// </summary>
+		/// <param name="Key"></param>
+		public RedisString(RedisKey Key) => this.Key = Key;
+
+		/// <inheritdoc/>
+		public override bool Equals(object? Obj) => Obj is RedisString<TElement> String && Key == String.Key;
+
+		/// <inheritdoc/>
+		public bool Equals(RedisString<TElement> Other) => Key == Other.Key;
+
+		/// <inheritdoc/>
+		public override int GetHashCode() => Key.GetHashCode();
+
+		/// <summary>Compares two instances for equality</summary>
+		public static bool operator ==(RedisString<TElement> Left, RedisString<TElement> Right) => Left.Key == Right.Key;
+
+		/// <summary>Compares two instances for equality</summary>
+		public static bool operator !=(RedisString<TElement> Left, RedisString<TElement> Right) => Left.Key != Right.Key;
+	}
+
+	/// <summary>
+	/// Extension methods for <see cref="IDatabaseAsync"/>
+	/// </summary>
+	public static class RedisStringExtensions
+	{
+		/// <inheritdoc cref="IDatabaseAsync.StringSetAsync(RedisKey, RedisValue, TimeSpan?, When, CommandFlags)"/>
+		public static Task StringSetAsync<T>(this IDatabaseAsync Database, RedisString<T> Str, T Value, TimeSpan? Expiry = null, When When = When.Always, CommandFlags Flags = CommandFlags.None)
+		{
+			return Database.StringSetAsync(Str.Key, RedisSerializer.Serialize(Value), Expiry, When, Flags);
+		}
+
+		/// <inheritdoc cref="IDatabaseAsync.StringGetAsync(RedisKey, CommandFlags)"/>
+		public static async Task<T> StringSetAsync<T>(this IDatabaseAsync Database, RedisString<T> Str, CommandFlags Flags = CommandFlags.None)
+		{
+			RedisValue Value = await Database.StringGetAsync(Str.Key, Flags);
+			return RedisSerializer.Deserialize<T>(Value);
+		}
+	}
+}
