Fix ListSegment excluding the first item in the collection.

diff --git a/Engine/Source/Programs/Shared/EpicGames.Core/ListSegment.cs b/Engine/Source/Programs/Shared/EpicGames.Core/ListSegment.cs
--- a/Engine/Source/Programs/Shared/EpicGames.Core/ListSegment.cs
+++ b/Engine/Source/Programs/Shared/EpicGames.Core/ListSegment.cs
@@ -133,7 +133,7 @@ namespace EpicGames.Core
 			public Enumerator(ListSegment<T> Segment)
 			{
 				this.List = Segment.List;
-				this.Index = Segment.Offset;
+				this.Index = Segment.Offset - 1;
 				this.MaxIndex = Segment.Offset + Segment.Count;
 			}
 
