Add support for serializing int, DateTime, Nullable, and enum types to compact binary through reflection.

Also fix deserialized DateTime values being treated as local times rather than UTC.

diff --git a/Engine/Source/Programs/Shared/EpicGames.Serialization/CompactBinary.cs b/Engine/Source/Programs/Shared/EpicGames.Serialization/CompactBinary.cs
--- a/Engine/Source/Programs/Shared/EpicGames.Serialization/CompactBinary.cs
+++ b/Engine/Source/Programs/Shared/EpicGames.Serialization/CompactBinary.cs
@@ -1110,14 +1110,23 @@ namespace EpicGames.Serialization
 			}
 		}
 
+		/// <summary>
+		/// Access the field as a DateTime.
+		/// </summary>
+		/// <returns></returns>
+		public DateTime AsDateTime()
+		{
+			return AsDateTime(default);
+		}
+
 		/// <summary>
 		/// Access the field as a DateTime.
 		/// </summary>
 		/// <param name="Default"></param>
 		/// <returns></returns>
-		public DateTime AsDateTime(DateTime Default = default)
+		public DateTime AsDateTime(DateTime Default)
 		{
-			return new DateTime(AsDateTimeTicks(Default.Ticks));
+			return new DateTime(AsDateTimeTicks(Default.ToUniversalTime().Ticks), DateTimeKind.Utc);
 		}
 
 		/// <summary>
diff --git a/Engine/Source/Programs/Shared/EpicGames.Serialization/CompactBinarySerializer.cs b/Engine/Source/Programs/Shared/EpicGames.Serialization/CompactBinarySerializer.cs
--- a/Engine/Source/Programs/Shared/EpicGames.Serialization/CompactBinarySerializer.cs
+++ b/Engine/Source/Programs/Shared/EpicGames.Serialization/CompactBinarySerializer.cs
@@ -79,7 +79,8 @@ namespace EpicGames.Serialization
 		{
 			new CbDefaultConverterFactory(),
 			new CbListConverterFactory(),
-			new CbDictionaryConverterFactory()
+			new CbDictionaryConverterFactory(),
+			new CbNullableConverterFactory()
 		};
 	
 		/// <summary>
diff --git a/Engine/Source/Programs/Shared/EpicGames.Serialization/Converters/CbDefaultConverter.cs b/Engine/Source/Programs/Shared/EpicGames.Serialization/Converters/CbDefaultConverter.cs
--- a/Engine/Source/Programs/Shared/EpicGames.Serialization/Converters/CbDefaultConverter.cs
+++ b/Engine/Source/Programs/Shared/EpicGames.Serialization/Converters/CbDefaultConverter.cs
@@ -98,11 +98,16 @@ namespace EpicGames.Serialization.Converters
 			}
 		}
 
+		static void WriteInt32Value(CbWriter Writer, int Value) => Writer.WriteIntegerValue(Value);
+		static void WriteInt32(CbWriter Writer, Utf8String Name, int Value) => Writer.WriteInteger(Name, Value);
+
 		static Dictionary<Type, CbConverterInfo> TypeToConverterInfo = new Dictionary<Type, CbConverterInfo>(new KeyValuePair<Type, CbConverterInfo>[]
 		{
 			GetPodConverterInfo(x => x.AsBool(), (w, v) => w.WriteBoolValue(v), (w, n, v) => w.WriteBool(n, v), null),
+			GetPodConverterInfo(x => x.AsInt32(), (w, v) => WriteInt32Value(w, v), (w, n, v) => WriteInt32(w, n, v), null),
 			GetPodConverterInfo(x => x.AsString(), (w, v) => w.WriteStringValue(v), (w, n, v) => w.WriteString(n, v), null),
-			GetPodConverterInfo(x => x.AsHash(), (w, v) => w.WriteHashValue(v), (w, n, v) => w.WriteHash(n, v), null)
+			GetPodConverterInfo(x => x.AsHash(), (w, v) => w.WriteHashValue(v), (w, n, v) => w.WriteHash(n, v), null),
+			GetPodConverterInfo(x => x.AsDateTime(), (w, v) => w.WriteDateTimeValue(v), (w, n, v) => w.WriteDateTime(n, v), null)
 			// TODO: other basic types
 		});
 
@@ -143,14 +148,32 @@ namespace EpicGames.Serialization.Converters
 				CbConverterInfo? ConverterInfo;
 				if (!TypeToConverterInfo.TryGetValue(Type, out ConverterInfo))
 				{
-					DynamicMethod ReadMethod = CreateObjectReader(Type);
-					DynamicMethod WriteMethod = CreateObjectWriter(Type);
-					DynamicMethod WriteNamedMethod = CreateNamedObjectWriter(Type);
+					if (Type.IsClass)
+					{
+						DynamicMethod ReadMethod = CreateObjectReader(Type);
+						DynamicMethod WriteMethod = CreateObjectWriter(Type);
+						DynamicMethod WriteNamedMethod = CreateNamedObjectWriter(Type);
+
+						Type GenericType = typeof(CbDefaultConverter<>).MakeGenericType(Type);
+						CbConverter Converter = (CbConverter)Activator.CreateInstance(GenericType, ReadMethod, WriteMethod, WriteNamedMethod)!;
+
+						ConverterInfo = new CbConverterInfo(Converter, ReadMethod, WriteMethod, WriteNamedMethod, SkipIfNullOrZero);
+					}
+					else if (Type.IsEnum)
+					{
+						DynamicMethod ReadMethod = CreateEnumReader(Type);
+						DynamicMethod WriteMethod = CreateEnumWriter(Type);
+						DynamicMethod WriteNamedMethod = CreateNamedEnumWriter(Type);
 
-					Type GenericType = typeof(CbDefaultConverter<>).MakeGenericType(Type);
-					CbConverter Converter = (CbConverter)Activator.CreateInstance(GenericType, ReadMethod, WriteMethod, WriteNamedMethod)!;
+						Type GenericType = typeof(CbDefaultConverter<>).MakeGenericType(Type);
+						CbConverter Converter = (CbConverter)Activator.CreateInstance(GenericType, ReadMethod, WriteMethod, WriteNamedMethod)!;
 
-					ConverterInfo = new CbConverterInfo(Converter, ReadMethod, WriteMethod, WriteNamedMethod, SkipIfNullOrZero);
+						ConverterInfo = new CbConverterInfo(Converter, ReadMethod, WriteMethod, WriteNamedMethod, SkipIfNullOrZero);
+					}
+					else
+					{
+						throw new InvalidOperationException($"Unable to serialize type {Type}");
+					}
 					TypeToConverterInfo.Add(Type, ConverterInfo);
 				}
 				return ConverterInfo.Converter;
@@ -174,6 +197,69 @@ namespace EpicGames.Serialization.Converters
 			}
 		}
 
+		static DynamicMethod CreateEnumWriter(Type Type)
+		{
+			// Create the new method
+			DynamicMethod DynamicMethod = new DynamicMethod("_", null, new Type[] { typeof(CbWriter), Type });
+
+			// Implement the body
+			ILGenerator Generator = DynamicMethod.GetILGenerator();
+
+			Generator.Emit(OpCodes.Ldarg_1);
+
+			Label SkipLabel = Generator.DefineLabel();
+			Generator.Emit(OpCodes.Brfalse, SkipLabel);
+
+			Generator.Emit(OpCodes.Ldarg_0);
+			Generator.Emit(OpCodes.Ldarg_1);
+			Generator.Emit(OpCodes.Conv_I8);
+			Generator.EmitCall(OpCodes.Call, GetMethodInfo<CbWriter>(x => x.WriteIntegerValue(0L)), null);
+
+			Generator.MarkLabel(SkipLabel);
+			Generator.Emit(OpCodes.Ret);
+
+			return DynamicMethod;
+		}
+
+		static DynamicMethod CreateNamedEnumWriter(Type Type)
+		{
+			// Create the new method
+			DynamicMethod DynamicMethod = new DynamicMethod("_", null, new Type[] { typeof(CbWriter), typeof(Utf8String), Type });
+
+			// Implement the body
+			ILGenerator Generator = DynamicMethod.GetILGenerator();
+			Generator.Emit(OpCodes.Ldarg_2);
+
+			Label SkipLabel = Generator.DefineLabel();
+			Generator.Emit(OpCodes.Brfalse, SkipLabel);
+
+			Generator.Emit(OpCodes.Ldarg_0);
+			Generator.Emit(OpCodes.Ldarg_1);
+			Generator.Emit(OpCodes.Ldarg_2);
+			Generator.Emit(OpCodes.Conv_I8);
+			Generator.EmitCall(OpCodes.Call, GetMethodInfo<CbWriter>(x => x.WriteInteger(default, 0L)), null);
+
+			Generator.MarkLabel(SkipLabel);
+			Generator.Emit(OpCodes.Ret);
+
+			return DynamicMethod;
+		}
+
+		static DynamicMethod CreateEnumReader(Type Type)
+		{
+			// Create the method
+			DynamicMethod DynamicMethod = new DynamicMethod("_", Type, new Type[] { typeof(CbField) });
+
+			// Generate the IL for it
+			ILGenerator Generator = DynamicMethod.GetILGenerator();
+
+			Generator.Emit(OpCodes.Ldarg_0);
+			Generator.EmitCall(OpCodes.Call, GetMethodInfo<CbField>(x => x.AsInt32()), null);
+			Generator.Emit(OpCodes.Ret);
+
+			return DynamicMethod;
+		}
+
 		static Dictionary<Type, DynamicMethod> WriteMethods = new Dictionary<Type, DynamicMethod>();
 
 		static DynamicMethod CreateObjectWriter(Type Type)
diff --git a/Engine/Source/Programs/Shared/EpicGames.Serialization/Converters/CbNullableConverter.cs b/Engine/Source/Programs/Shared/EpicGames.Serialization/Converters/CbNullableConverter.cs
new file mode 100644
--- /dev/null
+++ b/Engine/Source/Programs/Shared/EpicGames.Serialization/Converters/CbNullableConverter.cs
@@ -0,0 +1,57 @@
+// Copyright Epic Games, Inc. All Rights Reserved.
+
+using EpicGames.Core;
+using System;
+using System.Collections.Generic;
+using System.Text;
+
+namespace EpicGames.Serialization.Converters
+{
+	/// <summary>
+	/// Converter for list types
+	/// </summary>
+	/// <typeparam name="T"></typeparam>
+	class CbNullableConverter<T> : CbConverter<Nullable<T>> where T : struct
+	{
+		/// <inheritdoc/>
+		public override T? Read(CbField Field)
+		{
+			return CbSerializer.Deserialize<T>(Field);
+		}
+
+		/// <inheritdoc/>
+		public override void Write(CbWriter Writer, T? Nullable)
+		{
+			if (Nullable.HasValue)
+			{
+				CbSerializer.Serialize<T>(Writer, Nullable.Value);
+			}
+		}
+
+		/// <inheritdoc/>
+		public override void WriteNamed(CbWriter Writer, Utf8String Name, T? Nullable)
+		{
+			if (Nullable.HasValue)
+			{
+				CbSerializer.Serialize<T>(Writer, Name, Nullable.Value);
+			}
+		}
+	}
+
+	/// <summary>
+	/// Factory for CbNullableConverter
+	/// </summary>
+	class CbNullableConverterFactory : CbConverterFactory
+	{
+		/// <inheritdoc/>
+		public override CbConverter? CreateConverter(Type Type)
+		{
+			if (Type.IsGenericType && Type.GetGenericTypeDefinition() == typeof(Nullable<>))
+			{
+				Type ConverterType = typeof(CbNullableConverter<>).MakeGenericType(Type.GenericTypeArguments);
+				return (CbConverter)Activator.CreateInstance(ConverterType)!;
+			}
+			return null;
+		}
+	}
+}
