Horde: Add a typed abstraction for Redis pub/sub channels.

diff --git a/Engine/Source/Programs/Shared/EpicGames.Redis/RedisChannel.cs b/Engine/Source/Programs/Shared/EpicGames.Redis/RedisChannel.cs
new file mode 100644
--- /dev/null
+++ b/Engine/Source/Programs/Shared/EpicGames.Redis/RedisChannel.cs
@@ -0,0 +1,147 @@
+// Copyright Epic Games, Inc. All Rights Reserved.
+
+using StackExchange.Redis;
+using System;
+using System.Collections.Generic;
+using System.Text;
+using System.Threading.Tasks;
+
+namespace EpicGames.Redis
+{
+	using RedisChannel = StackExchange.Redis.RedisChannel;
+
+	/// <summary>
+	/// Represents a typed pub/sub channel with a particular value
+	/// </summary>
+	/// <typeparam name="T">The type of element stored in the channel</typeparam>
+	public readonly struct RedisChannel<T>
+	{
+		/// <summary>
+		/// The key for the list
+		/// </summary>
+		public readonly RedisChannel Channel { get; }
+
+		/// <summary>
+		/// Constructor
+		/// </summary>
+		/// <param name="Channel"></param>
+		public RedisChannel(RedisChannel Channel) => this.Channel = Channel;
+
+		/// <inheritdoc/>
+		public override bool Equals(object? Obj) => Obj is RedisChannel<T> List && Channel == List.Channel;
+
+		/// <inheritdoc/>
+		public bool Equals(RedisChannel<T> Other) => Channel == Other.Channel;
+
+		/// <inheritdoc/>
+		public override int GetHashCode() => Channel.GetHashCode();
+
+		/// <summary>Compares two instances for equality</summary>
+		public static bool operator ==(RedisChannel<T> Left, RedisChannel<T> Right) => Left.Channel == Right.Channel;
+
+		/// <summary>Compares two instances for equality</summary>
+		public static bool operator !=(RedisChannel<T> Left, RedisChannel<T> Right) => Left.Channel != Right.Channel;
+	}
+
+	/// <summary>
+	/// Subscription to a <see cref="RedisChannel{Task}"/>
+	/// </summary>
+	/// <typeparam name="T"></typeparam>
+	public class RedisChannelSubscription<T> : IDisposable, IAsyncDisposable
+	{
+		/// <summary>
+		/// The subscriber to register with
+		/// </summary>
+		ISubscriber? Subscriber;
+
+		/// <summary>
+		/// The channel to post on
+		/// </summary>
+		public RedisChannel<T> Channel { get; }
+
+		/// <summary>
+		/// The handler to call
+		/// </summary>
+		Action<RedisChannel<T>, T> Handler;
+
+		/// <summary>
+		/// Constructor
+		/// </summary>
+		/// <param name="Subscriber"></param>
+		/// <param name="Channel"></param>
+		/// <param name="Handler"></param>
+		RedisChannelSubscription(ISubscriber Subscriber, RedisChannel<T> Channel, Action<RedisChannel<T>, T> Handler)
+		{
+			this.Subscriber = Subscriber;
+			this.Channel = Channel;
+			this.Handler = Handler;
+		}
+		
+		/// <summary>
+		/// Start the subscription
+		/// </summary>
+		/// <param name="Flags"></param>
+		/// <returns></returns>
+		internal static async Task<RedisChannelSubscription<T>> CreateAsync(ISubscriber Subscriber, RedisChannel<T> Channel, Action<RedisChannel<T>, T> Handler, CommandFlags Flags = CommandFlags.None)
+		{
+			RedisChannelSubscription<T> Subscription = new RedisChannelSubscription<T>(Subscriber, Channel, Handler);
+			await Subscription.Subscriber!.SubscribeAsync(Channel.Channel, Subscription.UntypedHandler, Flags);
+			return Subscription;
+		}
+
+		/// <summary>
+		/// Unsubscribe from the channel
+		/// </summary>
+		/// <param name="Flags">Flags for the oper</param>
+		/// <returns></returns>
+		public async Task UnsubscribeAsync(CommandFlags Flags = CommandFlags.None)
+		{
+			if (Subscriber != null)
+			{
+				await Subscriber.UnsubscribeAsync(Channel.Channel, UntypedHandler, Flags);
+				Subscriber = null;
+			}
+		}
+
+		/// <summary>
+		/// The callback for messages being received
+		/// </summary>
+		/// <param name="_"></param>
+		/// <param name="Message"></param>
+		void UntypedHandler(RedisChannel _, RedisValue Message)
+		{
+			Handler(Channel, RedisSerializer.Deserialize<T>(Message));
+		}
+
+		/// <inheritdoc/>
+		public async ValueTask DisposeAsync()
+		{
+			await UnsubscribeAsync();
+		}
+
+		/// <inheritdoc/>
+		public void Dispose()
+		{
+			UnsubscribeAsync().Wait();
+		}
+	}
+
+	/// <summary>
+	/// Extension methods for typed lists
+	/// </summary>
+	public static class RedisChannelExtensions
+	{
+		/// <inheritdoc cref="ISubscriber.SubscribeAsync(StackExchange.Redis.RedisChannel, Action{StackExchange.Redis.RedisChannel, RedisValue}, CommandFlags)"/>
+		public static Task<RedisChannelSubscription<T>> SubscribeAsync<T>(this ISubscriber Subscriber, RedisChannel<T> Channel, Action<RedisChannel<T>, T> Handler, CommandFlags Flags = CommandFlags.None)
+		{
+			return RedisChannelSubscription<T>.CreateAsync(Subscriber, Channel, Handler, Flags);
+		}
+
+		/// <inheritdoc cref="IDatabaseAsync.PublishAsync(StackExchange.Redis.RedisChannel, RedisValue, CommandFlags)"/>
+		public static Task PublishAsync<T>(this IDatabaseAsync Database, RedisChannel<T> Channel, T Item, CommandFlags Flags = CommandFlags.None)
+		{
+			RedisValue Value = RedisSerializer.Serialize(Item);
+			return Database.PublishAsync(Channel.Channel, Value, Flags);
+		}
+	}
+}
