Horde: Add additional functionality to Redis library, including support for hashes.

diff --git a/Engine/Source/Programs/Shared/EpicGames.Redis/RedisHash.cs b/Engine/Source/Programs/Shared/EpicGames.Redis/RedisHash.cs
new file mode 100644
--- /dev/null
+++ b/Engine/Source/Programs/Shared/EpicGames.Redis/RedisHash.cs
@@ -0,0 +1,127 @@
+// Copyright Epic Games, Inc. All Rights Reserved.
+
+using StackExchange.Redis;
+using System;
+using System.Collections.Generic;
+using System.Linq;
+using System.Text;
+using System.Threading.Tasks;
+
+namespace EpicGames.Redis
+{
+	/// <summary>
+	/// Represents a typed Redis hash with a given key and value
+	/// </summary>
+	/// <typeparam name="TName">The key type for the hash</typeparam>
+	/// <typeparam name="TValue">The value type for the hash</typeparam>
+	public readonly struct RedisHash<TName, TValue> : IEquatable<RedisHash<TName, TValue>>
+	{
+		/// <summary>
+		/// The key for the hash
+		/// </summary>
+		public readonly RedisKey Key { get; }
+
+		/// <summary>
+		/// Constructor
+		/// </summary>
+		/// <param name="Key"></param>
+		public RedisHash(RedisKey Key) => this.Key = Key;
+
+		/// <inheritdoc/>
+		public override bool Equals(object? Obj) => Obj is RedisHash<TName, TValue> List && Key == List.Key;
+
+		/// <inheritdoc/>
+		public bool Equals(RedisHash<TName, TValue> Other) => Key == Other.Key;
+
+		/// <inheritdoc/>
+		public override int GetHashCode() => Key.GetHashCode();
+
+		/// <summary>Compares two instances for equality</summary>
+		public static bool operator ==(RedisHash<TName, TValue> Left, RedisHash<TName, TValue> Right) => Left.Key == Right.Key;
+
+		/// <summary>Compares two instances for equality</summary>
+		public static bool operator !=(RedisHash<TName, TValue> Left, RedisHash<TName, TValue> Right) => Left.Key != Right.Key;
+	}
+
+	/// <inheritdoc cref="HashEntry"/>
+	public readonly struct HashEntry<TName, TValue>
+	{
+		/// <inheritdoc cref="HashEntry.Name">
+		public readonly TName Name { get; }
+
+		/// <inheritdoc cref="HashEntry.Value">
+		public readonly TValue Value { get; }
+
+		/// <summary>
+		/// Constructor
+		/// </summary>
+		/// <param name="Name"></param>
+		/// <param name="Value"></param>
+		public HashEntry(TName Name, TValue Value)
+		{
+			this.Name = Name;
+			this.Value = Value;
+		}
+
+		/// <summary>
+		/// Implicit conversion to a <see cref="HashEntry"/>
+		/// </summary>
+		/// <param name="Entry"></param>
+		public static implicit operator HashEntry(HashEntry<TName, TValue> Entry)
+		{
+			return new HashEntry(RedisSerializer.Serialize(Entry.Name), RedisSerializer.Serialize(Entry.Value));
+		}
+
+		/// <summary>
+		/// Implicit conversion to a <see cref="KeyValuePair{TName, TValue}"/>
+		/// </summary>
+		/// <param name="Entry"></param>
+		public static implicit operator KeyValuePair<TName, TValue>(HashEntry<TName, TValue> Entry)
+		{
+			return new KeyValuePair<TName, TValue>(Entry.Name, Entry.Value);
+		}
+	}
+
+	/// <summary>
+	/// Extension methods for typed lists
+	/// </summary>
+	public static class RedisHashExtensions
+	{
+		/// <inheritdoc cref="IDatabaseAsync.HashDeleteAsync(RedisKey, RedisValue, CommandFlags)"/>
+		public static Task<bool> HashDeleteAsync<TName, TValue>(this IDatabase Redis, RedisHash<TName, TValue> Hash, TName Name, CommandFlags Flags = CommandFlags.None)
+		{
+			return Redis.HashDeleteAsync(Hash.Key, RedisSerializer.Serialize(Name), Flags);
+		}
+
+		/// <inheritdoc cref="IDatabaseAsync.HashDeleteAsync(RedisKey, RedisValue[], CommandFlags)"/>
+		public static Task<long> HashDeleteAsync<TName, TValue>(this IDatabase Redis, RedisHash<TName, TValue> Hash, TName[] Names, CommandFlags Flags = CommandFlags.None)
+		{
+			return Redis.HashDeleteAsync(Hash.Key, Array.ConvertAll(Names, x => RedisSerializer.Serialize(x)), Flags);
+		}
+
+		/// <inheritdoc cref="IDatabaseAsync.HashGetAllAsync(RedisKey, CommandFlags)"/>
+		public static async Task<HashEntry<TName, TValue>[]> HashGetAllAsync<TName, TValue>(this IDatabase Redis, RedisHash<TName, TValue> Hash, CommandFlags Flags = CommandFlags.None)
+		{
+			HashEntry[] Entries = await Redis.HashGetAllAsync(Hash.Key, Flags);
+			return Array.ConvertAll(Entries, x => new HashEntry<TName, TValue>(RedisSerializer.Deserialize<TName>(x.Name), RedisSerializer.Deserialize<TValue>(x.Value)));
+		}
+
+		/// <inheritdoc cref="IDatabaseAsync.HashLengthAsync(RedisKey, CommandFlags)"/>
+		public static Task<long> HashLengthAsync<TName, TValue>(this IDatabase Redis, RedisHash<TName, TValue> Hash, CommandFlags Flags = CommandFlags.None)
+		{
+			return Redis.HashLengthAsync(Hash.Key, Flags);
+		}
+
+		/// <inheritdoc cref="IDatabaseAsync.HashSetAsync(RedisKey, RedisValue, RedisValue, When, CommandFlags)"/>
+		public static Task HashSetAsync<TName, TValue>(this IDatabase Redis, RedisHash<TName, TValue> Hash, TName Name, TValue Value, When When = When.Always, CommandFlags Flags = CommandFlags.None)
+		{
+			return Redis.HashSetAsync(Hash.Key, RedisSerializer.Serialize(Name), RedisSerializer.Serialize(Value), When, Flags);
+		}
+
+		/// <inheritdoc cref="IDatabaseAsync.HashSetAsync(RedisKey, HashEntry[], CommandFlags)"/>
+		public static Task HashSetAsync<TName, TValue>(this IDatabase Redis, RedisHash<TName, TValue> Hash, HashEntry<TName, TValue>[] Entries, CommandFlags Flags = CommandFlags.None)
+		{
+			return Redis.HashSetAsync(Hash.Key, Array.ConvertAll(Entries, x => (HashEntry)x), Flags);
+		}
+	}
+}
diff --git a/Engine/Source/Programs/Shared/EpicGames.Redis/RedisList.cs b/Engine/Source/Programs/Shared/EpicGames.Redis/RedisList.cs
--- a/Engine/Source/Programs/Shared/EpicGames.Redis/RedisList.cs
+++ b/Engine/Source/Programs/Shared/EpicGames.Redis/RedisList.cs
@@ -66,6 +66,20 @@ namespace EpicGames.Redis
 			return Database.ListLeftPushAsync(List.Key, Values, When, Flags);
 		}
 
+		/// <inheritdoc cref="IDatabaseAsync.ListLeftPopAsync(RedisKey, CommandFlags)">
+		public static async Task<T?> ListLeftPopAsync<T>(this IDatabase Database, RedisList<T> List, CommandFlags Flags = CommandFlags.None) where T : class
+		{
+			RedisValue Value = await Database.ListLeftPopAsync(List.Key, Flags);
+			return Value.IsNull ? null : RedisSerializer.Deserialize<T>(Value);
+		}
+
+		/// <inheritdoc cref="IDatabaseAsync.ListRangeAsync(RedisKey, CommandFlags)">
+		public static async Task<T[]> ListRangeAsync<T>(this IDatabase Database, RedisList<T> List, long Start = 0, long Stop = -1, CommandFlags Flags = CommandFlags.None) where T : class
+		{
+			RedisValue[] Values = await Database.ListRangeAsync(List.Key, Start, Stop, Flags);
+			return Array.ConvertAll(Values, x => RedisSerializer.Deserialize<T>(x));
+		}
+
 		/// <inheritdoc cref="IDatabaseAsync.ListRightPushAsync(RedisKey, RedisValue, When, CommandFlags)"/>
 		public static Task<long> ListRightPushAsync<T>(this IDatabase Database, RedisList<T> List, T Item, When When = When.Always, CommandFlags Flags = CommandFlags.None)
 		{
@@ -78,5 +92,11 @@ namespace EpicGames.Redis
 			RedisValue[] Values = Items.Select(x => RedisSerializer.Serialize(x)).ToArray();
 			return Database.ListRightPushAsync(List.Key, Values, When, Flags);
 		}
+
+		/// <inheritdoc cref="IDatabaseAsync.ListTrimAsync(RedisKey, long, long, CommandFlags)">
+		public static async Task ListTrimAsync<T>(this IDatabase Database, RedisList<T> List, long Start, long Stop, CommandFlags Flags = CommandFlags.None) where T : class
+		{
+			await Database.ListTrimAsync(List.Key, Start, Stop, Flags);
+		}
 	}
 }
diff --git a/Engine/Source/Programs/Shared/EpicGames.Redis/RedisSet.cs b/Engine/Source/Programs/Shared/EpicGames.Redis/RedisSet.cs
--- a/Engine/Source/Programs/Shared/EpicGames.Redis/RedisSet.cs
+++ b/Engine/Source/Programs/Shared/EpicGames.Redis/RedisSet.cs
@@ -57,5 +57,12 @@ namespace EpicGames.Redis
 		{
 			return Database.SetRemoveAsync(Set.Key, RedisSerializer.Serialize(Value), Flags);
 		}
+
+		/// <inheritdoc cref="IDatabaseAsync.SetRemoveAsync(RedisKey, RedisValue, CommandFlags)"/>
+		public static async Task<T[]> SetMembersAsync<T>(this IDatabaseAsync Database, RedisSet<T> Set, CommandFlags Flags = CommandFlags.None)
+		{
+			RedisValue[] Values = await Database.SetMembersAsync(Set.Key, Flags);
+			return Array.ConvertAll(Values, x => RedisSerializer.Deserialize<T>(x));
+		}
 	}
 }
diff --git a/Engine/Source/Programs/Shared/EpicGames.Redis/RedisSortedSet.cs b/Engine/Source/Programs/Shared/EpicGames.Redis/RedisSortedSet.cs
--- a/Engine/Source/Programs/Shared/EpicGames.Redis/RedisSortedSet.cs
+++ b/Engine/Source/Programs/Shared/EpicGames.Redis/RedisSortedSet.cs
@@ -3,11 +3,50 @@
 using StackExchange.Redis;
 using System;
 using System.Collections.Generic;
+using System.Diagnostics.CodeAnalysis;
 using System.Linq;
 using System.Threading.Tasks;
 
 namespace EpicGames.Redis
 {
+	/// <summary>
+	/// Typed implementation of <see cref="SortedSetEntry"/>
+	/// </summary>
+	/// <typeparam name="T">The element type</typeparam>
+	public readonly struct SortedSetEntry<T>
+	{
+		/// <summary>
+		/// Accessor for the element type
+		/// </summary>
+		public readonly T Element { get; }
+
+		/// <summary>
+		/// Score for the entry
+		/// </summary>
+		public readonly double Score;
+
+		/// <summary>
+		/// Constructor
+		/// </summary>
+		/// <param name="Entry"></param>
+		public SortedSetEntry(SortedSetEntry Entry)
+		{
+			this.Element = RedisSerializer.Deserialize<T>(Entry.Element);
+			this.Score = Entry.Score;
+		}
+
+		/// <summary>
+		/// Constructor
+		/// </summary>
+		/// <param name="Element"></param>
+		/// <param name="Score"></param>
+		public SortedSetEntry(T Element, double Score)
+		{
+			this.Element = Element;
+			this.Score = Score;
+		}
+	}
+
 	/// <summary>
 	/// Represents a typed Redis list with a given key
 	/// </summary>
@@ -47,9 +86,46 @@ namespace EpicGames.Redis
 	public static class RedisSortedSetExtensions
 	{
 		/// <inheritdoc cref="IDatabaseAsync.SortedSetAddAsync(RedisKey, RedisValue, double, CommandFlags)"/>
-		public static Task SortedSetAddAsync<T>(this IDatabaseAsync Database, RedisSortedSet<T> Set, T Value, double Score, CommandFlags Flags = CommandFlags.None)
+		public static Task<bool> SortedSetAddAsync<T>(this IDatabaseAsync Database, RedisSortedSet<T> Set, T Value, double Score, CommandFlags Flags = CommandFlags.None)
 		{
 			return Database.SortedSetAddAsync(Set.Key, RedisSerializer.Serialize(Value), Score, Flags);
 		}
+
+		/// <inheritdoc cref="IDatabaseAsync.SortedSetAddAsync(RedisKey, RedisValue, double, CommandFlags)"/>
+		public static Task<long> SortedSetAddAsync<T>(this IDatabaseAsync Database, RedisSortedSet<T> Set, SortedSetEntry<T>[] Values, CommandFlags Flags = CommandFlags.None)
+		{
+			SortedSetEntry[] Untyped = Array.ConvertAll(Values, x => new SortedSetEntry(RedisSerializer.Serialize(x.Element), x.Score));
+			return Database.SortedSetAddAsync(Set.Key, Untyped, Flags);
+		}
+
+		/// <inheritdoc cref="IDatabaseAsync.SortedSetRemoveRangeByScoreAsync(RedisKey, double, double, Exclude, Order, long, long, CommandFlags)"/>
+		public static async Task<long> SortedSetRemoveRangeByScoreAsync<T>(this IDatabaseAsync Database, RedisSortedSet<T> Set, double Start = double.NegativeInfinity, double Stop = double.PositiveInfinity, Exclude Exclude = Exclude.None, CommandFlags Flags = CommandFlags.None)
+		{
+			return await Database.SortedSetRemoveRangeByScoreAsync(Set.Key, Start, Stop, Exclude, Flags);
+		}
+
+		/// <inheritdoc cref="IDatabaseAsync.SortedSetAddAsync(RedisKey, RedisValue, double, CommandFlags)"/>
+		public static async IAsyncEnumerable<SortedSetEntry<T>> SortedSetScanAsync<T>(this IDatabaseAsync Database, RedisSortedSet<T> Set, RedisValue Pattern = default, int PageSize = 250, long Cursor = 0, int PageOffset = 0, CommandFlags Flags = CommandFlags.None)
+		{
+			await foreach (SortedSetEntry Entry in Database.SortedSetScanAsync(Set.Key, Pattern, PageSize, Cursor, PageOffset, Flags))
+			{
+				yield return new SortedSetEntry<T>(Entry);
+			}
+		}
+
+		/// <inheritdoc cref="IDatabaseAsync.SortedSetRangeByScoreAsync(RedisKey, double, double, Exclude, Order, long, long, CommandFlags)"/>
+		public static async Task<T[]> SortedSetRangeByScoreAsync<T>(this IDatabaseAsync Database, RedisSortedSet<T> Set, double Start = double.NegativeInfinity, double Stop = double.PositiveInfinity, Exclude Exclude = Exclude.None, Order Order = Order.Ascending, long Skip = 0L, long Take = -1L, CommandFlags Flags = CommandFlags.None)
+		{
+			RedisValue[] Values = await Database.SortedSetRangeByScoreAsync(Set.Key, Start, Stop, Exclude, Order, Skip, Take, Flags);
+			return Array.ConvertAll(Values, x => RedisSerializer.Deserialize<T>(x));
+		}
+
+		/// <inheritdoc cref="IDatabaseAsync.SortedSetRangeByScoreWithScoresAsync(RedisKey, double, double, Exclude, Order, long, long, CommandFlags)"/>
+		public static async Task<SortedSetEntry<T>[]> SortedSetRangeByScoreWithScoresAsync<T>(this IDatabaseAsync Database, RedisSortedSet<T> Set, double Start = double.NegativeInfinity, double Stop = double.PositiveInfinity, Exclude Exclude = Exclude.None, Order Order = Order.Ascending, long Skip = 0L, long Take = -1L, CommandFlags Flags = CommandFlags.None)
+		{
+			SortedSetEntry[] Values = await Database.SortedSetRangeByScoreWithScoresAsync(Set.Key, Start, Stop, Exclude, Order, Skip, Take, Flags);
+			return Array.ConvertAll(Values, x => new SortedSetEntry<T>(x));
+		}
+
 	}
 }
