Allow implicit conversion from HashSet<T> to ReadOnlyHashSet<T>.

diff --git a/Engine/Source/Programs/Shared/EpicGames.Core/ReadOnlyHashSet.cs b/Engine/Source/Programs/Shared/EpicGames.Core/ReadOnlyHashSet.cs
--- a/Engine/Source/Programs/Shared/EpicGames.Core/ReadOnlyHashSet.cs
+++ b/Engine/Source/Programs/Shared/EpicGames.Core/ReadOnlyHashSet.cs
@@ -91,5 +91,14 @@ namespace EpicGames.Core
 		{
 			return ((IEnumerable)Inner).GetEnumerator();
 		}
+
+		/// <summary>
+		/// Implicit conversion operator from hashsets
+		/// </summary>
+		/// <param name="HashSet"></param>
+		public static implicit operator ReadOnlyHashSet<T>(HashSet<T> HashSet)
+		{
+			return new ReadOnlyHashSet<T>(HashSet);
+		}
 	}
 }
