Allow serializing structs to compact binary objects using reflection, and automatically instantiate converter types with the same generic arguments as the tagged class if necessary.

diff --git a/Engine/Source/Programs/Shared/EpicGames.Serialization/CompactBinarySerializer.cs b/Engine/Source/Programs/Shared/EpicGames.Serialization/CompactBinarySerializer.cs
--- a/Engine/Source/Programs/Shared/EpicGames.Serialization/CompactBinarySerializer.cs
+++ b/Engine/Source/Programs/Shared/EpicGames.Serialization/CompactBinarySerializer.cs
@@ -139,7 +139,12 @@ namespace EpicGames.Serialization
 					CbConverterAttribute? ConverterAttribute = Type.GetCustomAttribute<CbConverterAttribute>();
 					if (ConverterAttribute != null)
 					{
-						Converter = (CbConverter?)Activator.CreateInstance(ConverterAttribute.ConverterType)!;
+						Type ConverterType = ConverterAttribute.ConverterType;
+						if (Type.IsGenericType && ConverterType.IsGenericTypeDefinition)
+						{
+							ConverterType = ConverterType.MakeGenericType(Type.GetGenericArguments());
+						}
+						Converter = (CbConverter?)Activator.CreateInstance(ConverterType)!;
 					}
 					else
 					{
diff --git a/Engine/Source/Programs/Shared/EpicGames.Serialization/Converters/CbConverter.cs b/Engine/Source/Programs/Shared/EpicGames.Serialization/Converters/CbConverter.cs
--- a/Engine/Source/Programs/Shared/EpicGames.Serialization/Converters/CbConverter.cs
+++ b/Engine/Source/Programs/Shared/EpicGames.Serialization/Converters/CbConverter.cs
@@ -10,7 +10,7 @@ namespace EpicGames.Serialization.Converters
 	/// <summary>
 	/// Attribute declaring the converter to use for a type
 	/// </summary>
-	[AttributeUsage(AttributeTargets.Class)]
+	[AttributeUsage(AttributeTargets.Class | AttributeTargets.Struct)]
 	public class CbConverterAttribute : Attribute
 	{
 		/// <summary>
