Additional utility methods for Utf8String.

diff --git a/Engine/Source/Programs/Shared/EpicGames.Core/Utf8String.cs b/Engine/Source/Programs/Shared/EpicGames.Core/Utf8String.cs
--- a/Engine/Source/Programs/Shared/EpicGames.Core/Utf8String.cs
+++ b/Engine/Source/Programs/Shared/EpicGames.Core/Utf8String.cs
@@ -131,6 +131,12 @@ namespace EpicGames.Core
 			return Span.SequenceCompareTo(Other.Span);
 		}
 
+		/// <inheritdoc cref="String.Contains(string)"/>
+		public bool Contains(Utf8String String) => IndexOf(String) != -1;
+
+		/// <inheritdoc cref="String.Contains(string, StringComparison)"/>
+		public bool Contains(Utf8String String, Utf8StringComparer Comparer) => IndexOf(String, Comparer) != -1;
+
 		/// <inheritdoc cref="String.IndexOf(char)"/>
 		public int IndexOf(byte Char)
 		{
@@ -156,6 +162,19 @@ namespace EpicGames.Core
 			return Span.IndexOf(String.Span);
 		}
 
+		/// <inheritdoc cref="String.IndexOf(string, StringComparison)"/>
+		public int IndexOf(Utf8String String, Utf8StringComparer Comparer)
+		{
+			for (int Idx = 0; Idx < Length - String.Length; Idx++)
+			{
+				if (Comparer.Equals(String.Slice(Idx, String.Length), String))
+				{
+					return Idx;
+				}
+			}
+			return -1;
+		}
+
 		/// <summary>
 		/// Tests if this string starts with another string
 		/// </summary>
@@ -166,6 +185,17 @@ namespace EpicGames.Core
 			return Span.StartsWith(Other.Span);
 		}
 
+		/// <summary>
+		/// Tests if this string ends with another string
+		/// </summary>
+		/// <param name="Other">The string to check against</param>
+		/// <param name="Comparer">The string comparer</param>
+		/// <returns>True if this string ends with the other string</returns>
+		public bool StartsWith(Utf8String Other, Utf8StringComparer Comparer)
+		{
+			return Length >= Other.Length && Comparer.Equals(Slice(0, Other.Length), Other);
+		}
+
 		/// <summary>
 		/// Tests if this string ends with another string
 		/// </summary>
@@ -176,6 +206,17 @@ namespace EpicGames.Core
 			return Span.EndsWith(Other.Span);
 		}
 
+		/// <summary>
+		/// Tests if this string ends with another string
+		/// </summary>
+		/// <param name="Other">The string to check against</param>
+		/// <param name="Comparer">The string comparer</param>
+		/// <returns>True if this string ends with the other string</returns>
+		public bool EndsWith(Utf8String Other, Utf8StringComparer Comparer)
+		{
+			return Length >= Other.Length && Comparer.Equals(Slice(Length - Other.Length), Other);
+		}
+
 		/// <summary>
 		/// Slices this string at the given start position
 		/// </summary>
