Horde: Fix arguments when logging p4 protects commands.

diff --git a/Engine/Source/Programs/Shared/EpicGames.Perforce.Managed/ManagedWorkspace.cs b/Engine/Source/Programs/Shared/EpicGames.Perforce.Managed/ManagedWorkspace.cs
--- a/Engine/Source/Programs/Shared/EpicGames.Perforce.Managed/ManagedWorkspace.cs
+++ b/Engine/Source/Programs/Shared/EpicGames.Perforce.Managed/ManagedWorkspace.cs
@@ -1951,9 +1951,10 @@ namespace EpicGames.Perforce.Managed
 								Logger.LogInformation("SyncList: {Line}", Line);
 							}
 
-							await LogPerforceCommandAsync($"protects -M {LocalFile}", Logger);
-							await LogPerforceCommandAsync($"protects -M {LocalFile.FullName.ToLower()}", Logger);
-							await LogPerforceCommandAsync($"protects -M {LocalFile.FullName.Replace('\\', '/')}", Logger);
+							string Connection = $"-p {Client.ServerAndPort} -u {Client.UserName} -c {Client.ClientName}";
+							await LogPerforceCommandAsync($"{Connection} protects -M {LocalFile}", Logger);
+							await LogPerforceCommandAsync($"{Connection} protects -M {LocalFile.FullName.ToLower()}", Logger);
+							await LogPerforceCommandAsync($"{Connection} protects -M {LocalFile.FullName.Replace('\\', '/')}", Logger);
 
 							for (int Idx = 0; ; Idx++)
 							{
