Move functionality on PerforceConnection to extension methods, so we can have different implementations of the core connection handler.

diff --git a/Engine/Source/Programs/Shared/EpicGames.Perforce/PerforceConnection.cs b/Engine/Source/Programs/Shared/EpicGames.Perforce/PerforceConnection.cs
--- a/Engine/Source/Programs/Shared/EpicGames.Perforce/PerforceConnection.cs
+++ b/Engine/Source/Programs/Shared/EpicGames.Perforce/PerforceConnection.cs
@@ -117,6 +117,8 @@ namespace EpicGames.Perforce
 		/// <param name="ServerAndPort">The server address and port</param>
 		/// <param name="UserName">The user name</param>
 		/// <param name="ClientName">The client name</param>
+		/// <param name="AppName"></param>
+		/// <param name="AppVersion"></param>
 		/// <param name="Logger">Interface for logging</param>
 		public PerforceConnection(string? ServerAndPort, string? UserName, string? ClientName, string? AppName, string? AppVersion, ILogger Logger)
 		{
@@ -132,7 +134,6 @@ namespace EpicGames.Perforce
 		/// Constructor 
 		/// </summary>
 		/// <param name="Other">Connection to copy settings from</param>
-		/// <param name="ClientName">The client name</param>
 		public PerforceConnection(PerforceConnection Other)
 			: this(Other.ServerAndPort, Other.UserName, Other.ClientName, Other.AppName, Other.AppVersion, Other.Logger)
 		{
@@ -175,6 +176,94 @@ namespace EpicGames.Perforce
 			return Result.ToString();
 		}
 
+		/// <summary>
+		/// Execute a command and parse the response
+		/// </summary>
+		/// <param name="CommandLine">Arguments for the command</param>
+		/// <param name="InputData">Input data to pass to Perforce</param>
+		/// <param name="StatRecordType">The type of records to return for "stat" responses</param>
+		/// <param name="CancellationToken">Token used to cancel the operation</param>
+		/// <returns>List of objects returned by the server</returns>
+		public async Task<List<PerforceResponse>> CommandAsync(string CommandLine, byte[]? InputData, Type? StatRecordType, CancellationToken CancellationToken)
+		{
+			using (PerforceChildProcess Process = new PerforceChildProcess(InputData, GetFullCommandLine(CommandLine), Logger))
+			{
+				CachedRecordInfo? StatRecordInfo = (StatRecordType == null) ? null : PerforceReflection.GetCachedRecordInfo(StatRecordType);
+				return await Process.ReadResponsesAsync(StatRecordInfo, CancellationToken);
+			}
+		}
+
+		/// <summary>
+		/// Execute a command and parse the response
+		/// </summary>
+		/// <param name="CommandLine">Arguments for the command</param>
+		/// <param name="InputData">Input data to pass to Perforce</param>
+		/// <param name="HandleRecord">Delegate used to handle each record</param>
+		/// <param name="CancellationToken">Token used to cancel the operation</param>
+		/// <returns>List of objects returned by the server</returns>
+		public async Task RecordCommandAsync(string CommandLine, byte[]? InputData, Action<PerforceRecord> HandleRecord, CancellationToken CancellationToken)
+		{
+			using (PerforceChildProcess Process = new PerforceChildProcess(InputData, GetFullCommandLine(CommandLine), Logger))
+			{
+				await Process.ReadRecordsAsync(HandleRecord, CancellationToken);
+			}
+		}
+
+		#region p4 login
+
+		/// <summary>
+		/// Execute the 'login' command
+		/// </summary>
+		/// <param name="Password">Password to use to login</param>
+		/// <param name="CancellationToken">Token used to cancel the operation</param>
+		/// <returns>Response from the server</returns>
+		public async Task LoginAsync(string Password, CancellationToken CancellationToken)
+		{
+			// Some versions of P4.EXE do not support marshalled output for P4 login calls, so we only support this as a basic text query.
+			byte[] PasswordBytes = Encoding.UTF8.GetBytes(Password);
+			using (ManagedProcessGroup ChildProcessGroup = new ManagedProcessGroup())
+			{
+				string Executable = PerforceChildProcess.GetExecutable();
+				string FullArgumentList = GetFullCommandLine("login");
+
+				Logger.LogDebug("Running {0} {1}", Executable, FullArgumentList);
+				using (ManagedProcess ChildProcess = new ManagedProcess(ChildProcessGroup, Executable, FullArgumentList, null, null, PasswordBytes, ProcessPriorityClass.Normal))
+				{
+					StringBuilder Lines = new StringBuilder("Unable to log in: ");
+					for (; ; )
+					{
+						string? Line = await ChildProcess.ReadLineAsync();
+						if (Line == null)
+						{
+							break;
+						}
+
+						Lines.Append($"\n  {Line}");
+
+						if (Line.StartsWith("Enter password:", StringComparison.OrdinalIgnoreCase))
+						{
+							Lines.Append($" {new string('*', Password.Length)}");
+						}
+					}
+
+					ChildProcess.WaitForExit();
+
+					if (ChildProcess.ExitCode != 0)
+					{
+						throw new PerforceException($"{Lines}\n  (exit code {ChildProcess.ExitCode})");
+					}
+				}
+			}
+		}
+
+		#endregion
+	}
+
+	/// <summary>
+	/// Extension methods for <see cref="PerforceConnection"/>
+	/// </summary>
+	public static class PerforceConnectionExtensions
+	{
 		/// <summary>
 		/// Serializes a list of key/value pairs into binary format.
 		/// </summary>
@@ -213,30 +302,14 @@ namespace EpicGames.Perforce
 		/// <summary>
 		/// Execute a command and parse the response
 		/// </summary>
-		/// <param name="CommandLine">Arguments for the command</param>
-		/// <param name="InputData">Input data to pass to Perforce</param>
-		/// <param name="StatRecordType">The type of records to return for "stat" responses</param>
-		/// <param name="CancellationToken">Token used to cancel the operation</param>
-		/// <returns>List of objects returned by the server</returns>
-		public async Task<List<PerforceResponse>> CommandAsync(string CommandLine, byte[]? InputData, Type? StatRecordType, CancellationToken CancellationToken)
-		{
-			using (PerforceChildProcess Process = new PerforceChildProcess(InputData, GetFullCommandLine(CommandLine), Logger))
-			{
-				CachedRecordInfo? StatRecordInfo = (StatRecordType == null) ? null : PerforceReflection.GetCachedRecordInfo(StatRecordType);
-				return await Process.ReadResponsesAsync(StatRecordInfo, CancellationToken);
-			}
-		}
-
-		/// <summary>
-		/// Execute a command and parse the response
-		/// </summary>
+		/// <param name="Connection">Connection to the Perforce server</param>
 		/// <param name="Arguments">Arguments for the command</param>
 		/// <param name="InputData">Input data to pass to Perforce</param>
 		/// <param name="CancellationToken">Token used to cancel the operation</param>
 		/// <returns>List of objects returned by the server</returns>
-		public async Task<PerforceResponseList<T>> CommandAsync<T>(string Arguments, byte[]? InputData, CancellationToken CancellationToken) where T : class
+		public static async Task<PerforceResponseList<T>> CommandAsync<T>(this PerforceConnection Connection, string Arguments, byte[]? InputData, CancellationToken CancellationToken) where T : class
 		{
-			List<PerforceResponse> Responses = await CommandAsync(Arguments, InputData, typeof(T), CancellationToken);
+			List<PerforceResponse> Responses = await Connection.CommandAsync(Arguments, InputData, typeof(T), CancellationToken);
 
 			PerforceResponseList<T> TypedResponses = new PerforceResponseList<T>();
 			foreach (PerforceResponse Response in Responses)
@@ -249,12 +322,13 @@ namespace EpicGames.Perforce
 		/// <summary>
 		/// Execute a command and parse the response
 		/// </summary>
+		/// <param name="Connection">Connection to the Perforce server</param>
 		/// <param name="CommonArguments">Arguments for the command</param>
 		/// <param name="BatchedArguments">Arguments to pass to the command in batches</param>
 		/// <param name="InputData">Input data to pass to Perforce</param>
 		/// <param name="CancellationToken">Token used to cancel the operation</param>
 		/// <returns>List of objects returned by the server</returns>
-		public async Task<PerforceResponseList<T>> BatchedCommandAsync<T>(string CommonArguments, string[] BatchedArguments, byte[]? InputData, CancellationToken CancellationToken) where T : class
+		public static async Task<PerforceResponseList<T>> BatchedCommandAsync<T>(PerforceConnection Connection, string CommonArguments, string[] BatchedArguments, byte[]? InputData, CancellationToken CancellationToken) where T : class
 		{
 			PerforceResponseList<T> Responses = new PerforceResponseList<T>();
 			for (int FileSpecIdx = 0; FileSpecIdx < BatchedArguments.Length;)
@@ -264,43 +338,28 @@ namespace EpicGames.Perforce
 				{
 					Arguments.AppendFormat(" \"{0}\"", BatchedArguments[FileSpecIdx]);
 				}
-				Responses.AddRange(await CommandAsync<T>(Arguments.ToString(), null, CancellationToken));
+				Responses.AddRange(await CommandAsync<T>(Connection, Arguments.ToString(), null, CancellationToken));
 			}
 			return Responses;
 		}
 
-		/// <summary>
-		/// Execute a command and parse the response
-		/// </summary>
-		/// <param name="CommandLine">Arguments for the command</param>
-		/// <param name="InputData">Input data to pass to Perforce</param>
-		/// <param name="HandleRecord">Delegate used to handle each record</param>
-		/// <param name="CancellationToken">Token used to cancel the operation</param>
-		/// <returns>List of objects returned by the server</returns>
-		public async Task RecordCommandAsync(string CommandLine, byte[]? InputData, Action<PerforceRecord> HandleRecord, CancellationToken CancellationToken)
-		{
-			using (PerforceChildProcess Process = new PerforceChildProcess(InputData, GetFullCommandLine(CommandLine), Logger))
-			{
-				await Process.ReadRecordsAsync(HandleRecord, CancellationToken);
-			}
-		}
-
 		/// <summary>
 		/// Attempts to execute the given command, returning the results from the server or the first PerforceResponse object.
 		/// </summary>
+		/// <param name="Connection">Connection to the Perforce server</param>
 		/// <param name="Arguments">Arguments for the command.</param>
 		/// <param name="InputData">Input data for the command.</param>
 		/// <param name="StatRecordType">Type of element to return in the response</param>
 		/// <param name="CancellationToken">Token used to cancel the operation</param>
 		/// <returns>Response from the server; either an object of type T or error.</returns>
-		public async Task<PerforceResponse> SingleResponseCommandAsync(string Arguments, byte[]? InputData, Type? StatRecordType, CancellationToken CancellationToken)
+		static async Task<PerforceResponse> SingleResponseCommandAsync(PerforceConnection Connection, string Arguments, byte[]? InputData, Type? StatRecordType, CancellationToken CancellationToken)
 		{
-			List<PerforceResponse> Responses = await CommandAsync(Arguments, InputData, StatRecordType, CancellationToken);
+			List<PerforceResponse> Responses = await Connection.CommandAsync(Arguments, InputData, StatRecordType, CancellationToken);
 			if (Responses.Count != 1)
 			{
 				for (int Idx = 0; Idx < Responses.Count; Idx++)
 				{
-					Logger.LogDebug("Unexpected response {Idx}: {Text}", Idx, Responses[Idx].ToString());
+					Connection.Logger.LogDebug("Unexpected response {Idx}: {Text}", Idx, Responses[Idx].ToString());
 				}
 				throw new PerforceException("Expected one result from 'p4 {0}', got {1}", Arguments, Responses.Count);
 			}
@@ -311,13 +370,14 @@ namespace EpicGames.Perforce
 		/// Attempts to execute the given command, returning the results from the server or the first PerforceResponse object.
 		/// </summary>
 		/// <typeparam name="T">Type of record to parse</typeparam>
+		/// <param name="Connection">Connection to the Perforce server</param>
 		/// <param name="Arguments">Arguments for the command.</param>
 		/// <param name="InputData">Input data for the command.</param>
 		/// <param name="CancellationToken">Token used to cancel the operation</param>
 		/// <returns>Response from the server; either an object of type T or error.</returns>
-		public async Task<PerforceResponse<T>> SingleResponseCommandAsync<T>(string Arguments, byte[]? InputData, CancellationToken CancellationToken) where T : class
+		public static async Task<PerforceResponse<T>> SingleResponseCommandAsync<T>(PerforceConnection Connection, string Arguments, byte[]? InputData, CancellationToken CancellationToken) where T : class
 		{
-			return new PerforceResponse<T>(await SingleResponseCommandAsync(Arguments, InputData, typeof(T), CancellationToken));
+			return new PerforceResponse<T>(await SingleResponseCommandAsync(Connection, Arguments, InputData, typeof(T), CancellationToken));
 		}
 
 		#endregion
@@ -327,51 +387,55 @@ namespace EpicGames.Perforce
 		/// <summary>
 		/// Adds files to a pending changelist.
 		/// </summary>
+		/// <param name="Connection">Connection to the Perforce server</param>
 		/// <param name="ChangeNumber">Changelist to add files to</param>
 		/// <param name="FileNames">Files to be added</param>
 		/// <param name="CancellationToken">Token used to cancel the operation</param>
 		/// <returns>Response from the server</returns>
-		public async Task<List<AddRecord>> AddAsync(int ChangeNumber, string[] FileNames, CancellationToken CancellationToken)
+		public static async Task<List<AddRecord>> AddAsync(this PerforceConnection Connection, int ChangeNumber, string[] FileNames, CancellationToken CancellationToken)
 		{
-			return (await TryAddAsync(ChangeNumber, null, AddOptions.None, FileNames, CancellationToken)).Data;
+			return (await TryAddAsync(Connection, ChangeNumber, null, AddOptions.None, FileNames, CancellationToken)).Data;
 		}
 
 		/// <summary>
 		/// Adds files to a pending changelist.
 		/// </summary>
+		/// <param name="Connection">Connection to the Perforce server</param>
 		/// <param name="ChangeNumber">Changelist to add files to</param>
 		/// <param name="FileType">Type for new files</param>
 		/// <param name="Options">Options for the command</param>
 		/// <param name="FileNames">Files to be added</param>
 		/// <param name="CancellationToken">Token used to cancel the operation</param>
 		/// <returns>Response from the server</returns>
-		public async Task<List<AddRecord>> AddAsync(int ChangeNumber, string? FileType, AddOptions Options, string[] FileNames, CancellationToken CancellationToken)
+		public static async Task<List<AddRecord>> AddAsync(this PerforceConnection Connection, int ChangeNumber, string? FileType, AddOptions Options, string[] FileNames, CancellationToken CancellationToken)
 		{
-			return (await TryAddAsync(ChangeNumber, FileType, Options, FileNames, CancellationToken)).Data;
+			return (await TryAddAsync(Connection, ChangeNumber, FileType, Options, FileNames, CancellationToken)).Data;
 		}
 
 		/// <summary>
 		/// Adds files to a pending changelist.
 		/// </summary>
+		/// <param name="Connection">Connection to the Perforce server</param>
 		/// <param name="ChangeNumber">Changelist to add files to</param>
 		/// <param name="FileNames">Files to be added</param>
 		/// <param name="CancellationToken">Token used to cancel the operation</param>
 		/// <returns>Response from the server</returns>
-		public Task<PerforceResponseList<AddRecord>> TryAddAsync(int ChangeNumber, string[] FileNames, CancellationToken CancellationToken)
+		public static Task<PerforceResponseList<AddRecord>> TryAddAsync(this PerforceConnection Connection, int ChangeNumber, string[] FileNames, CancellationToken CancellationToken)
 		{
-			return TryAddAsync(ChangeNumber, null, AddOptions.None, FileNames, CancellationToken);
+			return TryAddAsync(Connection, ChangeNumber, null, AddOptions.None, FileNames, CancellationToken);
 		}
 
 		/// <summary>
 		/// Adds files to a pending changelist.
 		/// </summary>
+		/// <param name="Connection">Connection to the Perforce server</param>
 		/// <param name="ChangeNumber">Changelist to add files to</param>
 		/// <param name="FileType">Type for new files</param>
 		/// <param name="Options">Options for the command</param>
 		/// <param name="FileNames">Files to be added</param>
 		/// <param name="CancellationToken">Token used to cancel the operation</param>
 		/// <returns>Response from the server</returns>
-		public Task<PerforceResponseList<AddRecord>> TryAddAsync(int ChangeNumber, string? FileType, AddOptions Options, string[] FileNames, CancellationToken CancellationToken)
+		public static Task<PerforceResponseList<AddRecord>> TryAddAsync(this PerforceConnection Connection, int ChangeNumber, string? FileType, AddOptions Options, string[] FileNames, CancellationToken CancellationToken)
 		{
 			StringBuilder Arguments = new StringBuilder("add");
 			if (ChangeNumber != -1)
@@ -399,7 +463,7 @@ namespace EpicGames.Perforce
 				Arguments.AppendFormat(" -t \"{0}\"", FileType);
 			}
 
-			return BatchedCommandAsync<AddRecord>(Arguments.ToString(), FileNames, null, CancellationToken);
+			return BatchedCommandAsync<AddRecord>(Connection, Arguments.ToString(), FileNames, null, CancellationToken);
 		}
 
 		#endregion
@@ -409,28 +473,30 @@ namespace EpicGames.Perforce
 		/// <summary>
 		/// Creates a changelist with the p4 change command. 
 		/// </summary>
+		/// <param name="Connection">Connection to the Perforce server</param>
 		/// <param name="Record">Information for the change to create. The number field should be left set to -1.</param>
 		/// <param name="CancellationToken">Token used to cancel the operation</param>
 		/// <returns>The changelist number, or an error.</returns>
-		public async Task<ChangeRecord> CreateChangeAsync(ChangeRecord Record, CancellationToken CancellationToken)
+		public static async Task<ChangeRecord> CreateChangeAsync(this PerforceConnection Connection, ChangeRecord Record, CancellationToken CancellationToken)
 		{
-			return (await TryCreateChangeAsync(Record, CancellationToken)).Data;
+			return (await TryCreateChangeAsync(Connection, Record, CancellationToken)).Data;
 		}
 
 		/// <summary>
 		/// Creates a changelist with the p4 change command. 
 		/// </summary>
+		/// <param name="Connection">Connection to the Perforce server</param>
 		/// <param name="Record">Information for the change to create. The number field should be left set to -1.</param>
 		/// <param name="CancellationToken">Token used to cancel the operation</param>
 		/// <returns>The changelist number, or an error.</returns>
-		public async Task<PerforceResponse<ChangeRecord>> TryCreateChangeAsync(ChangeRecord Record, CancellationToken CancellationToken)
+		public static async Task<PerforceResponse<ChangeRecord>> TryCreateChangeAsync(this PerforceConnection Connection, ChangeRecord Record, CancellationToken CancellationToken)
 		{
 			if (Record.Number != -1)
 			{
 				throw new PerforceException("'Number' field should be set to -1 to create a new changelist.");
 			}
 
-			PerforceResponse Response = await SingleResponseCommandAsync("change -i", SerializeRecord(Record), null, CancellationToken);
+			PerforceResponse Response = await SingleResponseCommandAsync(Connection, "change -i", SerializeRecord(Record), null, CancellationToken);
 
 			PerforceError? Error = Response.Error;
 			if (Error != null)
@@ -457,23 +523,25 @@ namespace EpicGames.Perforce
 		/// <summary>
 		/// Updates an existing changelist.
 		/// </summary>
+		/// <param name="Connection">Connection to the Perforce server</param>
 		/// <param name="Options">Options for this command</param>
 		/// <param name="Record">Information for the change to create. The number field should be left set to zero.</param>
 		/// <param name="CancellationToken">Token used to cancel the operation</param>
 		/// <returns>The changelist number, or an error.</returns>
-		public async Task UpdateChangeAsync(UpdateChangeOptions Options, ChangeRecord Record, CancellationToken CancellationToken)
+		public static async Task UpdateChangeAsync(this PerforceConnection Connection, UpdateChangeOptions Options, ChangeRecord Record, CancellationToken CancellationToken)
 		{
-			(await TryUpdateChangeAsync(Options, Record, CancellationToken)).EnsureSuccess();
+			(await TryUpdateChangeAsync(Connection, Options, Record, CancellationToken)).EnsureSuccess();
 		}
 
 		/// <summary>
 		/// Updates an existing changelist.
 		/// </summary>
+		/// <param name="Connection">Connection to the Perforce server</param>
 		/// <param name="Options">Options for this command</param>
 		/// <param name="Record">Information for the change to create. The number field should be left set to zero.</param>
 		/// <param name="CancellationToken">Token used to cancel the operation</param>
 		/// <returns>The changelist number, or an error.</returns>
-		public Task<PerforceResponse> TryUpdateChangeAsync(UpdateChangeOptions Options, ChangeRecord Record, CancellationToken CancellationToken)
+		public static Task<PerforceResponse> TryUpdateChangeAsync(this PerforceConnection Connection, UpdateChangeOptions Options, ChangeRecord Record, CancellationToken CancellationToken)
 		{
 			if (Record.Number == -1)
 			{
@@ -490,29 +558,31 @@ namespace EpicGames.Perforce
 				Arguments.Append(" -u");
 			}
 
-			return SingleResponseCommandAsync(Arguments.ToString(), SerializeRecord(Record), null, CancellationToken);
+			return SingleResponseCommandAsync(Connection, Arguments.ToString(), SerializeRecord(Record), null, CancellationToken);
 		}
 
 		/// <summary>
 		/// Deletes a changelist (p4 change -d)
 		/// </summary>
+		/// <param name="Connection">Connection to the Perforce server</param>
 		/// <param name="Options">Options for the command</param>
 		/// <param name="ChangeNumber">Changelist number to delete</param>
 		/// <param name="CancellationToken">Token used to cancel the operation</param>
 		/// <returns>Response from the server</returns>
-		public async Task DeleteChangeAsync(DeleteChangeOptions Options, int ChangeNumber, CancellationToken CancellationToken)
+		public static async Task DeleteChangeAsync(this PerforceConnection Connection, DeleteChangeOptions Options, int ChangeNumber, CancellationToken CancellationToken)
 		{
-			(await TryDeleteChangeAsync(Options, ChangeNumber, CancellationToken)).EnsureSuccess();
+			(await TryDeleteChangeAsync(Connection, Options, ChangeNumber, CancellationToken)).EnsureSuccess();
 		}
 
 		/// <summary>
 		/// Deletes a changelist (p4 change -d)
 		/// </summary>
+		/// <param name="Connection">Connection to the Perforce server</param>
 		/// <param name="Options">Options for the command</param>
 		/// <param name="ChangeNumber">Changelist number to delete</param>
 		/// <param name="CancellationToken">Token used to cancel the operation</param>
 		/// <returns>Response from the server</returns>
-		public Task<PerforceResponse> TryDeleteChangeAsync(DeleteChangeOptions Options, int ChangeNumber, CancellationToken CancellationToken)
+		public static Task<PerforceResponse> TryDeleteChangeAsync(this PerforceConnection Connection, DeleteChangeOptions Options, int ChangeNumber, CancellationToken CancellationToken)
 		{
 			StringBuilder Arguments = new StringBuilder("change -d");
 			if ((Options & DeleteChangeOptions.Submitted) != 0)
@@ -525,29 +595,31 @@ namespace EpicGames.Perforce
 			}
 			Arguments.AppendFormat(" {0}", ChangeNumber);
 
-			return SingleResponseCommandAsync(Arguments.ToString(), null, null, CancellationToken);
+			return SingleResponseCommandAsync(Connection, Arguments.ToString(), null, null, CancellationToken);
 		}
 
 		/// <summary>
 		/// Gets a changelist
 		/// </summary>
+		/// <param name="Connection">Connection to the Perforce server</param>
 		/// <param name="Options">Options for the command</param>
 		/// <param name="ChangeNumber">Changelist number to retrieve. -1 is the default changelist for this workspace.</param>
 		/// <param name="CancellationToken">Token used to cancel the operation</param>
 		/// <returns>Response from the server</returns>
-		public async Task<ChangeRecord> GetChangeAsync(GetChangeOptions Options, int ChangeNumber, CancellationToken CancellationToken)
+		public static async Task<ChangeRecord> GetChangeAsync(this PerforceConnection Connection, GetChangeOptions Options, int ChangeNumber, CancellationToken CancellationToken)
 		{
-			return (await TryGetChange(Options, ChangeNumber, CancellationToken)).Data;
+			return (await TryGetChange(Connection, Options, ChangeNumber, CancellationToken)).Data;
 		}
 
 		/// <summary>
 		/// Gets a changelist
 		/// </summary>
+		/// <param name="Connection">Connection to the Perforce server</param>
 		/// <param name="Options">Options for the command</param>
 		/// <param name="ChangeNumber">Changelist number to retrieve. -1 is the default changelist for this workspace.</param>
 		/// <param name="CancellationToken">Token used to cancel the operation</param>
 		/// <returns>Response from the server</returns>
-		public Task<PerforceResponse<ChangeRecord>> TryGetChange(GetChangeOptions Options, int ChangeNumber, CancellationToken CancellationToken)
+		public static Task<PerforceResponse<ChangeRecord>> TryGetChange(this PerforceConnection Connection, GetChangeOptions Options, int ChangeNumber, CancellationToken CancellationToken)
 		{
 			StringBuilder Arguments = new StringBuilder("change -o");
 			if ((Options & GetChangeOptions.BeforeRenumber) != 0)
@@ -559,7 +631,7 @@ namespace EpicGames.Perforce
 				Arguments.AppendFormat(" {0}", ChangeNumber);
 			}
 
-			return SingleResponseCommandAsync<ChangeRecord>(Arguments.ToString(), null, CancellationToken);
+			return SingleResponseCommandAsync<ChangeRecord>(Connection, Arguments.ToString(), null, CancellationToken);
 		}
 
 		/// <summary>
@@ -567,7 +639,7 @@ namespace EpicGames.Perforce
 		/// </summary>
 		/// <param name="Input">The record to serialize</param>
 		/// <returns>Serialized record</returns>
-		byte[] SerializeRecord(ChangeRecord Input)
+		static byte[] SerializeRecord(ChangeRecord Input)
 		{
 			List<KeyValuePair<string, object>> NameToValue = new List<KeyValuePair<string, object>>();
 			if (Input.Number == -1)
@@ -604,20 +676,22 @@ namespace EpicGames.Perforce
 		/// <summary>
 		/// Enumerates changes on the server
 		/// </summary>
+		/// <param name="Connection">Connection to the Perforce server</param>
 		/// <param name="Options">Options for the command</param>
 		/// <param name="MaxChanges">List only the highest numbered changes</param>
 		/// <param name="Status">Limit the list to the changelists with the given status (pending, submitted or shelved)</param>
 		/// <param name="FileSpecs">Paths to query changes for</param>
 		/// <param name="CancellationToken">Token used to cancel the operation</param>
 		/// <returns>List of responses from the server.</returns>
-		public async Task<List<ChangesRecord>> GetChangesAsync(ChangesOptions Options, int MaxChanges, ChangeStatus Status, string[] FileSpecs, CancellationToken CancellationToken)
+		public static async Task<List<ChangesRecord>> GetChangesAsync(this PerforceConnection Connection, ChangesOptions Options, int MaxChanges, ChangeStatus Status, string[] FileSpecs, CancellationToken CancellationToken)
 		{
-			return (await TryGetChangesAsync(Options, MaxChanges, Status, FileSpecs, CancellationToken)).Data;
+			return (await TryGetChangesAsync(Connection, Options, MaxChanges, Status, FileSpecs, CancellationToken)).Data;
 		}
 
 		/// <summary>
 		/// Enumerates changes on the server
 		/// </summary>
+		/// <param name="Connection">Connection to the Perforce server</param>
 		/// <param name="Options">Options for the command</param>
 		/// <param name="ClientName">List only changes made from the named client workspace.</param>
 		/// <param name="MaxChanges">List only the highest numbered changes</param>
@@ -626,28 +700,30 @@ namespace EpicGames.Perforce
 		/// <param name="FileSpecs">Paths to query changes for</param>
 		/// <param name="CancellationToken">Token used to cancel the operation</param>
 		/// <returns>List of responses from the server.</returns>
-		public async Task<List<ChangesRecord>> GetChangesAsync(ChangesOptions Options, string? ClientName, int MaxChanges, ChangeStatus Status, string? UserName, string[]? FileSpecs, CancellationToken CancellationToken)
+		public static async Task<List<ChangesRecord>> GetChangesAsync(this PerforceConnection Connection, ChangesOptions Options, string? ClientName, int MaxChanges, ChangeStatus Status, string? UserName, string[]? FileSpecs, CancellationToken CancellationToken)
 		{
-			return (await TryGetChangesAsync(Options, ClientName, MaxChanges, Status, UserName, FileSpecs, CancellationToken)).Data;
+			return (await TryGetChangesAsync(Connection, Options, ClientName, MaxChanges, Status, UserName, FileSpecs, CancellationToken)).Data;
 		}
 
 		/// <summary>
 		/// Enumerates changes on the server
 		/// </summary>
+		/// <param name="Connection">Connection to the Perforce server</param>
 		/// <param name="Options">Options for the command</param>
 		/// <param name="MaxChanges">List only the highest numbered changes</param>
 		/// <param name="Status">Limit the list to the changelists with the given status (pending, submitted or shelved)</param>
 		/// <param name="FileSpecs">Paths to query changes for</param>
 		/// <param name="CancellationToken">Token used to cancel the operation</param>
 		/// <returns>List of responses from the server.</returns>
-		public Task<PerforceResponseList<ChangesRecord>> TryGetChangesAsync(ChangesOptions Options, int MaxChanges, ChangeStatus Status, string[] FileSpecs, CancellationToken CancellationToken)
+		public static Task<PerforceResponseList<ChangesRecord>> TryGetChangesAsync(this PerforceConnection Connection, ChangesOptions Options, int MaxChanges, ChangeStatus Status, string[] FileSpecs, CancellationToken CancellationToken)
 		{
-			return TryGetChangesAsync(Options, null, MaxChanges, Status, null, FileSpecs, CancellationToken);
+			return TryGetChangesAsync(Connection, Options, null, MaxChanges, Status, null, FileSpecs, CancellationToken);
 		}
 
 		/// <summary>
 		/// Enumerates changes on the server
 		/// </summary>
+		/// <param name="Connection">Connection to the Perforce server</param>
 		/// <param name="Options">Options for the command</param>
 		/// <param name="ClientName">List only changes made from the named client workspace.</param>
 		/// <param name="MaxChanges">List only the highest numbered changes</param>
@@ -656,7 +732,7 @@ namespace EpicGames.Perforce
 		/// <param name="FileSpecs">Paths to query changes for</param>
 		/// <param name="CancellationToken">Token used to cancel the operation</param>
 		/// <returns>List of responses from the server.</returns>
-		public Task<PerforceResponseList<ChangesRecord>> TryGetChangesAsync(ChangesOptions Options, string? ClientName, int MaxChanges, ChangeStatus Status, string? UserName, string[]? FileSpecs, CancellationToken CancellationToken)
+		public static Task<PerforceResponseList<ChangesRecord>> TryGetChangesAsync(this PerforceConnection Connection, ChangesOptions Options, string? ClientName, int MaxChanges, ChangeStatus Status, string? UserName, string[]? FileSpecs, CancellationToken CancellationToken)
 		{
 			StringBuilder Arguments = new StringBuilder("changes");
 			if ((Options & ChangesOptions.IncludeIntegrations) != 0)
@@ -702,7 +778,7 @@ namespace EpicGames.Perforce
 					Arguments.AppendFormat(" \"{0}\"", FileSpec);
 				}
 			}
-			return CommandAsync<ChangesRecord>(Arguments.ToString(), null, CancellationToken);
+			return CommandAsync<ChangesRecord>(Connection, Arguments.ToString(), null, CancellationToken);
 		}
 
 		#endregion
@@ -712,23 +788,25 @@ namespace EpicGames.Perforce
 		/// <summary>
 		/// Cleans the workspace
 		/// </summary>
+		/// <param name="Connection">Connection to the Perforce server</param>
 		/// <param name="Options">Options for the command</param>
 		/// <param name="FileSpecs">Files to sync</param>
 		/// <param name="CancellationToken">Token used to cancel the operation</param>
 		/// <returns>Response from the server</returns>
-		public async Task<List<CleanRecord>> CleanAsync(CleanOptions Options, string[] FileSpecs, CancellationToken CancellationToken)
+		public static async Task<List<CleanRecord>> CleanAsync(this PerforceConnection Connection, CleanOptions Options, string[] FileSpecs, CancellationToken CancellationToken)
 		{
-			return (await TryCleanAsync(Options, FileSpecs, CancellationToken)).Data;
+			return (await TryCleanAsync(Connection, Options, FileSpecs, CancellationToken)).Data;
 		}
 
 		/// <summary>
 		/// Cleans the workspace
 		/// </summary>
+		/// <param name="Connection">Connection to the Perforce server</param>
 		/// <param name="Options">Options for the command</param>
 		/// <param name="FileSpecs">Files to sync</param>
 		/// <param name="CancellationToken">Token used to cancel the operation</param>
 		/// <returns>Response from the server</returns>
-		public async Task<PerforceResponseList<CleanRecord>> TryCleanAsync(CleanOptions Options, string[] FileSpecs, CancellationToken CancellationToken)
+		public static async Task<PerforceResponseList<CleanRecord>> TryCleanAsync(this PerforceConnection Connection, CleanOptions Options, string[] FileSpecs, CancellationToken CancellationToken)
 		{
 			StringBuilder Arguments = new StringBuilder("clean");
 			if ((Options & CleanOptions.Edited) != 0)
@@ -764,7 +842,7 @@ namespace EpicGames.Perforce
 				Arguments.AppendFormat(" \"{0}\"", FileSpec);
 			}
 
-			PerforceResponseList<CleanRecord> Records = await CommandAsync<CleanRecord>(Arguments.ToString(), null, CancellationToken);
+			PerforceResponseList<CleanRecord> Records = await CommandAsync<CleanRecord>(Connection, Arguments.ToString(), null, CancellationToken);
 			Records.RemoveAll(x => x.Error != null && x.Error.Generic == PerforceGenericCode.Empty);
 			return Records;
 		}
@@ -776,67 +854,73 @@ namespace EpicGames.Perforce
 		/// <summary>
 		/// Creates a client
 		/// </summary>
+		/// <param name="Connection">Connection to the Perforce server</param>
 		/// <param name="Record">The client record</param>
 		/// <param name="CancellationToken">Token used to cancel the operation</param>
 		/// <returns>Response from the server</returns>
-		public async Task CreateClientAsync(ClientRecord Record, CancellationToken CancellationToken)
+		public static async Task CreateClientAsync(this PerforceConnection Connection, ClientRecord Record, CancellationToken CancellationToken)
 		{
-			(await TryCreateClientAsync(Record, CancellationToken)).EnsureSuccess();
+			(await TryCreateClientAsync(Connection, Record, CancellationToken)).EnsureSuccess();
 		}
 
 		/// <summary>
 		/// Creates a client
 		/// </summary>
+		/// <param name="Connection">Connection to the Perforce server</param>
 		/// <param name="Record">The client record</param>
 		/// <param name="CancellationToken">Token used to cancel the operation</param>
 		/// <returns>Response from the server</returns>
-		public Task<PerforceResponse> TryCreateClientAsync(ClientRecord Record, CancellationToken CancellationToken)
+		public static Task<PerforceResponse> TryCreateClientAsync(this PerforceConnection Connection, ClientRecord Record, CancellationToken CancellationToken)
 		{
-			return TryUpdateClientAsync(Record, CancellationToken);
+			return TryUpdateClientAsync(Connection, Record, CancellationToken);
 		}
 
 		/// <summary>
 		/// Creates a client
 		/// </summary>
+		/// <param name="Connection">Connection to the Perforce server</param>
 		/// <param name="Record">The client record</param>
 		/// <param name="CancellationToken">Token used to cancel the operation</param>
 		/// <returns>Response from the server</returns>
-		public async Task UpdateClientAsync(ClientRecord Record, CancellationToken CancellationToken)
+		public static async Task UpdateClientAsync(this PerforceConnection Connection, ClientRecord Record, CancellationToken CancellationToken)
 		{
-			(await TryUpdateClientAsync(Record, CancellationToken)).EnsureSuccess();
+			(await TryUpdateClientAsync(Connection, Record, CancellationToken)).EnsureSuccess();
 		}
 
 		/// <summary>
 		/// Update a client
 		/// </summary>
+		/// <param name="Connection">Connection to the Perforce server</param>
 		/// <param name="Record">The client record</param>
 		/// <param name="CancellationToken">Token used to cancel the operation</param>
 		/// <returns>Response from the server</returns>
-		public Task<PerforceResponse> TryUpdateClientAsync(ClientRecord Record, CancellationToken CancellationToken)
+		public static Task<PerforceResponse> TryUpdateClientAsync(this PerforceConnection Connection, ClientRecord Record, CancellationToken CancellationToken)
 		{
-			return SingleResponseCommandAsync("client -i", SerializeRecord(Record), null, CancellationToken);
+			return SingleResponseCommandAsync(Connection, "client -i", SerializeRecord(Record), null, CancellationToken);
 		}
 
 		/// <summary>
 		/// Deletes a client
 		/// </summary>
+		/// <param name="Connection">Connection to the Perforce server</param>
 		/// <param name="Options">Options for this command</param>
 		/// <param name="ClientName">Name of the client</param>
 		/// <param name="CancellationToken">Token used to cancel the operation</param>
 		/// <returns>Response from the server</returns>
-		public async Task DeleteClientAsync(DeleteClientOptions Options, string ClientName, CancellationToken CancellationToken)
+		public static async Task DeleteClientAsync(this PerforceConnection Connection, DeleteClientOptions Options, string ClientName, CancellationToken CancellationToken)
 		{
-			(await TryDeleteClientAsync(Options, ClientName, CancellationToken)).EnsureSuccess();
+			(await TryDeleteClientAsync(Connection, Options, ClientName, CancellationToken)).EnsureSuccess();
 		}
 
 		/// <summary>
 		/// Deletes a client
 		/// </summary>
+		/// <param name="Connection">Connection to the Perforce server</param>
 		/// <param name="Options">Options for this command</param>
 		/// <param name="ClientName">Name of the client</param>
 		/// <param name="CancellationToken">Token used to cancel the operation</param>
 		/// <returns>Response from the server</returns>
-		public Task<PerforceResponse> TryDeleteClientAsync(DeleteClientOptions Options, string ClientName, CancellationToken CancellationToken)
+		public static Task<PerforceResponse> TryDeleteClientAsync(this PerforceConnection Connection, DeleteClientOptions Options, string ClientName, CancellationToken CancellationToken)
 		{
 			StringBuilder Arguments = new StringBuilder("client -d");
 			if ((Options & DeleteClientOptions.Force) != 0)
@@ -848,31 +932,33 @@ namespace EpicGames.Perforce
 				Arguments.Append(" -Fs");
 			}
 			Arguments.AppendFormat(" \"{0}\"", ClientName);
-			return SingleResponseCommandAsync(Arguments.ToString(), null, null, CancellationToken);
+			return SingleResponseCommandAsync(Connection, Arguments.ToString(), null, null, CancellationToken);
 		}
 
 		/// <summary>
 		/// Changes the stream associated with a client
 		/// </summary>
+		/// <param name="Connection">Connection to the Perforce server</param>
 		/// <param name="ClientName">The client name</param>
 		/// <param name="StreamName">The new stream to be associated with the client</param>
 		/// <param name="Options">Options for this command</param>
 		/// <param name="CancellationToken">Token used to cancel the operation</param>
 		/// <returns>Response from the server</returns>
-		public async Task SwitchClientToStreamAsync(string ClientName, string StreamName, SwitchClientOptions Options, CancellationToken CancellationToken)
+		public static async Task SwitchClientToStreamAsync(this PerforceConnection Connection, string ClientName, string StreamName, SwitchClientOptions Options, CancellationToken CancellationToken)
 		{
-			(await TrySwitchClientToStreamAsync(ClientName, StreamName, Options, CancellationToken)).EnsureSuccess();
+			(await TrySwitchClientToStreamAsync(Connection, ClientName, StreamName, Options, CancellationToken)).EnsureSuccess();
 		}
 
 		/// <summary>
 		/// Changes the stream associated with a client
 		/// </summary>
+		/// <param name="Connection">Connection to the Perforce server</param>
 		/// <param name="ClientName">The client name</param>
 		/// <param name="StreamName">The new stream to be associated with the client</param>
 		/// <param name="Options">Options for this command</param>
 		/// <param name="CancellationToken">Token used to cancel the operation</param>
 		/// <returns>Response from the server</returns>
-		public Task<PerforceResponse> TrySwitchClientToStreamAsync(string ClientName, string StreamName, SwitchClientOptions Options, CancellationToken CancellationToken)
+		public static Task<PerforceResponse> TrySwitchClientToStreamAsync(this PerforceConnection Connection, string ClientName, string StreamName, SwitchClientOptions Options, CancellationToken CancellationToken)
 		{
 			StringBuilder Arguments = new StringBuilder("client -s");
 			if ((Options & SwitchClientOptions.IgnoreOpenFiles) != 0)
@@ -881,81 +967,87 @@ namespace EpicGames.Perforce
 			}
 			Arguments.AppendFormat(" -S \"{0}\"", StreamName);
 
-			return SingleResponseCommandAsync(Arguments.ToString(), null, null, CancellationToken);
+			return SingleResponseCommandAsync(Connection, Arguments.ToString(), null, null, CancellationToken);
 		}
 
 		/// <summary>
 		/// Changes a client to mirror a template
 		/// </summary>
+		/// <param name="Connection">Connection to the Perforce server</param>
 		/// <param name="ClientName">The client name</param>
 		/// <param name="TemplateName">The new stream to be associated with the client</param>
 		/// <param name="CancellationToken">Token used to cancel the operation</param>
 		/// <returns>Response from the server</returns>
-		public async Task SwitchClientToTemplateAsync(string ClientName, string TemplateName, CancellationToken CancellationToken)
+		public static async Task SwitchClientToTemplateAsync(this PerforceConnection Connection, string ClientName, string TemplateName, CancellationToken CancellationToken)
 		{
-			(await TrySwitchClientToTemplateAsync(ClientName, TemplateName, CancellationToken)).EnsureSuccess();
+			(await TrySwitchClientToTemplateAsync(Connection, ClientName, TemplateName, CancellationToken)).EnsureSuccess();
 		}
 
 		/// <summary>
 		/// Changes a client to mirror a template
 		/// </summary>
+		/// <param name="Connection">Connection to the Perforce server</param>
 		/// <param name="ClientName">The client name</param>
 		/// <param name="TemplateName">The new stream to be associated with the client</param>
 		/// <param name="CancellationToken">Token used to cancel the operation</param>
 		/// <returns>Response from the server</returns>
-		public Task<PerforceResponse> TrySwitchClientToTemplateAsync(string ClientName, string TemplateName, CancellationToken CancellationToken)
+		public static Task<PerforceResponse> TrySwitchClientToTemplateAsync(this PerforceConnection Connection, string ClientName, string TemplateName, CancellationToken CancellationToken)
 		{
 			string Arguments = String.Format("client -s -t \"{0}\"", TemplateName);
-			return SingleResponseCommandAsync(Arguments, null, null, CancellationToken);
+			return SingleResponseCommandAsync(Connection, Arguments, null, null, CancellationToken);
 		}
 
 		/// <summary>
 		/// Queries the current client definition
 		/// </summary>
+		/// <param name="Connection">Connection to the Perforce server</param>
 		/// <param name="ClientName">Name of the client. Specify null for the current client.</param>
 		/// <param name="CancellationToken">Token used to cancel the operation</param>
 		/// <returns>Response from the server; either a client record or error code</returns>
-		public async Task<ClientRecord> GetClientAsync(string? ClientName, CancellationToken CancellationToken)
+		public static async Task<ClientRecord> GetClientAsync(this PerforceConnection Connection, string? ClientName, CancellationToken CancellationToken)
 		{
-			return (await TryGetClientAsync(ClientName, CancellationToken)).Data;
+			return (await TryGetClientAsync(Connection, ClientName, CancellationToken)).Data;
 		}
 
 		/// <summary>
 		/// Queries the current client definition
 		/// </summary>
+		/// <param name="Connection">Connection to the Perforce server</param>
 		/// <param name="ClientName">Name of the client. Specify null for the current client.</param>
 		/// <param name="CancellationToken">Token used to cancel the operation</param>
 		/// <returns>Response from the server; either a client record or error code</returns>
-		public Task<PerforceResponse<ClientRecord>> TryGetClientAsync(string? ClientName, CancellationToken CancellationToken)
+		public static Task<PerforceResponse<ClientRecord>> TryGetClientAsync(this PerforceConnection Connection, string? ClientName, CancellationToken CancellationToken)
 		{
 			StringBuilder Arguments = new StringBuilder("client -o");
 			if (ClientName != null)
 			{
 				Arguments.AppendFormat(" \"{0}\"", ClientName);
 			}
-			return SingleResponseCommandAsync<ClientRecord>(Arguments.ToString(), null, CancellationToken);
+			return SingleResponseCommandAsync<ClientRecord>(Connection, Arguments.ToString(), null, CancellationToken);
 		}
 
 		/// <summary>
 		/// Queries information about a stream
 		/// </summary>
+		/// <param name="Connection">Connection to the Perforce server</param>
 		/// <param name="StreamName">Name of the stream to query</param>
 		/// <param name="bIncludeView">Whether to include the stream view in the output</param>
 		/// <param name="CancellationToken">Token used to cancel the operation</param>
 		/// <returns>Stream information record</returns>
-		public async Task<StreamRecord> GetStreamAsync(string StreamName, bool bIncludeView, CancellationToken CancellationToken)
+		public static async Task<StreamRecord> GetStreamAsync(this PerforceConnection Connection, string StreamName, bool bIncludeView, CancellationToken CancellationToken)
 		{
-			return (await TryGetStreamAsync(StreamName, bIncludeView, CancellationToken)).Data;
+			return (await TryGetStreamAsync(Connection, StreamName, bIncludeView, CancellationToken)).Data;
 		}
 
 		/// <summary>
 		/// Queries information about a stream
 		/// </summary>
+		/// <param name="Connection">Connection to the Perforce server</param>
 		/// <param name="StreamName">Name of the stream to query</param>
 		/// <param name="bIncludeView">Whether to include the stream view in the output</param>
 		/// <param name="CancellationToken">Token used to cancel the operation</param>
 		/// <returns>Stream information record</returns>
-		public Task<PerforceResponse<StreamRecord>> TryGetStreamAsync(string StreamName, bool bIncludeView, CancellationToken CancellationToken)
+		public static Task<PerforceResponse<StreamRecord>> TryGetStreamAsync(this PerforceConnection Connection, string StreamName, bool bIncludeView, CancellationToken CancellationToken)
 		{
 			StringBuilder Arguments = new StringBuilder("stream -o");
 			if (bIncludeView)
@@ -964,29 +1056,31 @@ namespace EpicGames.Perforce
 			}
 			Arguments.AppendFormat(" {0}", StreamName);
 
-			return SingleResponseCommandAsync<StreamRecord>(Arguments.ToString(), null, CancellationToken);
+			return SingleResponseCommandAsync<StreamRecord>(Connection, Arguments.ToString(), null, CancellationToken);
 		}
 
 		/// <summary>
 		/// Queries the view for a stream
 		/// </summary>
+		/// <param name="Connection">Connection to the Perforce server</param>
 		/// <param name="StreamName">Name of the stream.</param>
 		/// <param name="ChangeNumber">Changelist at which to query the stream view</param>
 		/// <param name="CancellationToken">Token used to cancel the operation</param>
 		/// <returns>Response from the server; either a client record or error code</returns>
-		public async Task<ClientRecord> GetStreamViewAsync(string StreamName, int ChangeNumber, CancellationToken CancellationToken)
+		public static async Task<ClientRecord> GetStreamViewAsync(this PerforceConnection Connection, string StreamName, int ChangeNumber, CancellationToken CancellationToken)
 		{
-			return (await TryGetStreamViewAsync(StreamName, ChangeNumber, CancellationToken)).Data;
+			return (await TryGetStreamViewAsync(Connection, StreamName, ChangeNumber, CancellationToken)).Data;
 		}
 
 		/// <summary>
 		/// Queries the view for a stream
 		/// </summary>
+		/// <param name="Connection">Connection to the Perforce server</param>
 		/// <param name="StreamName">Name of the stream.</param>
 		/// <param name="ChangeNumber">Changelist at which to query the stream view</param>
 		/// <param name="CancellationToken">Token used to cancel the operation</param>
 		/// <returns>Response from the server; either a client record or error code</returns>
-		public Task<PerforceResponse<ClientRecord>> TryGetStreamViewAsync(string StreamName, int ChangeNumber, CancellationToken CancellationToken)
+		public static Task<PerforceResponse<ClientRecord>> TryGetStreamViewAsync(this PerforceConnection Connection, string StreamName, int ChangeNumber, CancellationToken CancellationToken)
 		{
 			StringBuilder Arguments = new StringBuilder("client -o");
 			Arguments.AppendFormat(" -S \"{0}\"", StreamName);
@@ -994,7 +1088,7 @@ namespace EpicGames.Perforce
 			{
 				Arguments.AppendFormat(" -c {0}", ChangeNumber);
 			}
-			return SingleResponseCommandAsync<ClientRecord>(Arguments.ToString(), null, CancellationToken);
+			return SingleResponseCommandAsync<ClientRecord>(Connection, Arguments.ToString(), null, CancellationToken);
 		}
 
 		/// <summary>
@@ -1002,7 +1096,7 @@ namespace EpicGames.Perforce
 		/// </summary>
 		/// <param name="Input">The input record</param>
 		/// <returns>Serialized record data</returns>
-		byte[] SerializeRecord(ClientRecord Input)
+		static byte[] SerializeRecord(ClientRecord Input)
 		{
 			List<KeyValuePair<string, object>> NameToValue = new List<KeyValuePair<string, object>>();
 			if (Input.Name != null)
@@ -1059,30 +1153,33 @@ namespace EpicGames.Perforce
 		/// <summary>
 		/// Queries the current client definition
 		/// </summary>
+		/// <param name="Connection">Connection to the Perforce server</param>
 		/// <param name="Options">Options for this command</param>
 		/// <param name="UserName">List only client workspaces owned by this user.</param>
 		/// <param name="CancellationToken">Token used to cancel the operation</param>
 		/// <returns>Response from the server; either a client record or error code</returns>
-		public async Task<List<ClientsRecord>> GetClientsAsync(ClientsOptions Options, string? UserName, CancellationToken CancellationToken)
+		public static async Task<List<ClientsRecord>> GetClientsAsync(this PerforceConnection Connection, ClientsOptions Options, string? UserName, CancellationToken CancellationToken)
 		{
-			return (await TryGetClientsAsync(Options, UserName, CancellationToken)).Data;
+			return (await TryGetClientsAsync(Connection, Options, UserName, CancellationToken)).Data;
 		}
 
 		/// <summary>
 		/// Queries the current client definition
 		/// </summary>
+		/// <param name="Connection">Connection to the Perforce server</param>
 		/// <param name="Options">Options for this command</param>
 		/// <param name="UserName">List only client workspaces owned by this user.</param>
 		/// <param name="CancellationToken">Token used to cancel the operation</param>
 		/// <returns>Response from the server; either a client record or error code</returns>
-		public Task<PerforceResponseList<ClientsRecord>> TryGetClientsAsync(ClientsOptions Options, string? UserName, CancellationToken CancellationToken)
+		public static Task<PerforceResponseList<ClientsRecord>> TryGetClientsAsync(this PerforceConnection Connection, ClientsOptions Options, string? UserName, CancellationToken CancellationToken)
 		{
-			return TryGetClientsAsync(Options, null, -1, null, UserName, CancellationToken);
+			return TryGetClientsAsync(Connection, Options, null, -1, null, UserName, CancellationToken);
 		}
 
 		/// <summary>
 		/// Queries the current client definition
 		/// </summary>
+		/// <param name="Connection">Connection to the Perforce server</param>
 		/// <param name="Options">Options for this command</param>
 		/// <param name="Filter">List only client workspaces matching filter. Treated as case sensitive if <ref>ClientsOptions.CaseSensitiveFilter</ref> is set.</param>
 		/// <param name="MaxResults">Limit the number of results to return. -1 for all.</param>
@@ -1090,14 +1187,15 @@ namespace EpicGames.Perforce
 		/// <param name="UserName">List only client workspaces owned by this user.</param>
 		/// <param name="CancellationToken">Token used to cancel the operation</param>
 		/// <returns>Response from the server; either a client record or error code</returns>
-		public async Task<List<ClientsRecord>> GetClientsAsync(ClientsOptions Options, string? Filter, int MaxResults, string? Stream, string? UserName, CancellationToken CancellationToken)
+		public static async Task<List<ClientsRecord>> GetClientsAsync(this PerforceConnection Connection, ClientsOptions Options, string? Filter, int MaxResults, string? Stream, string? UserName, CancellationToken CancellationToken)
 		{
-			return (await TryGetClientsAsync(Options, Filter, MaxResults, Stream, UserName, CancellationToken)).Data;
+			return (await TryGetClientsAsync(Connection, Options, Filter, MaxResults, Stream, UserName, CancellationToken)).Data;
 		}
 
 		/// <summary>
 		/// Queries the current client definition
 		/// </summary>
+		/// <param name="Connection">Connection to the Perforce server</param>
 		/// <param name="Options">Options for this command</param>
 		/// <param name="Filter">List only client workspaces matching filter. Treated as case sensitive if <ref>ClientsOptions.CaseSensitiveFilter</ref> is set.</param>
 		/// <param name="MaxResults">Limit the number of results to return. -1 for all.</param>
@@ -1105,7 +1203,7 @@ namespace EpicGames.Perforce
 		/// <param name="UserName">List only client workspaces owned by this user.</param>
 		/// <param name="CancellationToken">Token used to cancel the operation</param>
 		/// <returns>Response from the server; either a client record or error code</returns>
-		public Task<PerforceResponseList<ClientsRecord>> TryGetClientsAsync(ClientsOptions Options, string? Filter, int MaxResults, string? Stream, string? UserName, CancellationToken CancellationToken)
+		public static Task<PerforceResponseList<ClientsRecord>> TryGetClientsAsync(this PerforceConnection Connection, ClientsOptions Options, string? Filter, int MaxResults, string? Stream, string? UserName, CancellationToken CancellationToken)
 		{
 			StringBuilder Arguments = new StringBuilder("clients");
 			if ((Options & ClientsOptions.All) != 0)
@@ -1143,7 +1241,7 @@ namespace EpicGames.Perforce
 			{
 				Arguments.Append(" -U");
 			}
-			return CommandAsync<ClientsRecord>(Arguments.ToString(), null, CancellationToken);
+			return CommandAsync<ClientsRecord>(Connection, Arguments.ToString(), null, CancellationToken);
 		}
 
 		#endregion
@@ -1153,23 +1251,27 @@ namespace EpicGames.Perforce
 		/// <summary>
 		/// Execute the 'delete' command
 		/// </summary>
+		/// <param name="Connection">Connection to the Perforce server</param>
+		/// <param name="ChangeNumber">The change to add deleted files to</param>
 		/// <param name="Options">Options for the command</param>
 		/// <param name="FileSpecs">List of file specifications to query</param>
 		/// <param name="CancellationToken">Token used to cancel the operation</param>
 		/// <returns>List of responses from the server</returns>
-		public async Task<List<DeleteRecord>> DeleteAsync(int ChangeNumber, DeleteOptions Options, string[] FileSpecs, CancellationToken CancellationToken)
+		public static async Task<List<DeleteRecord>> DeleteAsync(this PerforceConnection Connection, int ChangeNumber, DeleteOptions Options, string[] FileSpecs, CancellationToken CancellationToken)
 		{
-			return (await TryDeleteAsync(ChangeNumber, Options, FileSpecs, CancellationToken)).Data;
+			return (await TryDeleteAsync(Connection, ChangeNumber, Options, FileSpecs, CancellationToken)).Data;
 		}
 
 		/// <summary>
 		/// Execute the 'delete' command
 		/// </summary>
+		/// <param name="Connection">Connection to the Perforce server</param>
+		/// <param name="ChangeNumber">The change to add deleted files to</param>
 		/// <param name="Options">Options for the command</param>
 		/// <param name="FileSpecs">List of file specifications to query</param>
 		/// <param name="CancellationToken">Token used to cancel the operation</param>
 		/// <returns>List of responses from the server</returns>
-		public Task<PerforceResponseList<DeleteRecord>> TryDeleteAsync(int ChangeNumber, DeleteOptions Options, string[] FileSpecs, CancellationToken CancellationToken)
+		public static Task<PerforceResponseList<DeleteRecord>> TryDeleteAsync(this PerforceConnection Connection, int ChangeNumber, DeleteOptions Options, string[] FileSpecs, CancellationToken CancellationToken)
 		{
 			StringBuilder Arguments = new StringBuilder("delete");
 			if (ChangeNumber != -1)
@@ -1189,7 +1291,7 @@ namespace EpicGames.Perforce
 				Arguments.Append(" -v");
 			}
 
-			return BatchedCommandAsync<DeleteRecord>(Arguments.ToString(), FileSpecs, null, CancellationToken);
+			return BatchedCommandAsync<DeleteRecord>(Connection, Arguments.ToString(), FileSpecs, null, CancellationToken);
 		}
 
 		#endregion
@@ -1199,23 +1301,25 @@ namespace EpicGames.Perforce
 		/// <summary>
 		/// Describes a single changelist
 		/// </summary>
+		/// <param name="Connection">Connection to the Perforce server</param>
 		/// <param name="ChangeNumber">The changelist number to retrieve description for</param>
 		/// <param name="CancellationToken">Token used to cancel the operation</param>
 		/// <returns>Response from the server; either a describe record or error code</returns>
-		public async Task<DescribeRecord> DescribeAsync(int ChangeNumber, CancellationToken CancellationToken)
+		public static async Task<DescribeRecord> DescribeAsync(this PerforceConnection Connection, int ChangeNumber, CancellationToken CancellationToken)
 		{
-			return (await TryDescribeAsync(ChangeNumber, CancellationToken)).Data;
+			return (await TryDescribeAsync(Connection, ChangeNumber, CancellationToken)).Data;
 		}
 
 		/// <summary>
 		/// Describes a single changelist
 		/// </summary>
+		/// <param name="Connection">Connection to the Perforce server</param>
 		/// <param name="ChangeNumber">The changelist number to retrieve description for</param>
 		/// <param name="CancellationToken">Token used to cancel the operation</param>
 		/// <returns>Response from the server; either a describe record or error code</returns>
-		public async Task<PerforceResponse<DescribeRecord>> TryDescribeAsync(int ChangeNumber, CancellationToken CancellationToken)
+		public static async Task<PerforceResponse<DescribeRecord>> TryDescribeAsync(this PerforceConnection Connection, int ChangeNumber, CancellationToken CancellationToken)
 		{
-			PerforceResponseList<DescribeRecord> Records = await TryDescribeAsync(new int[] { ChangeNumber }, CancellationToken);
+			PerforceResponseList<DescribeRecord> Records = await TryDescribeAsync(Connection, new int[] { ChangeNumber }, CancellationToken);
 			if (Records.Count != 1)
 			{
 				throw new PerforceException("Expected only one record returned from p4 describe command, got {0}", Records.Count);
@@ -1226,48 +1330,56 @@ namespace EpicGames.Perforce
 		/// <summary>
 		/// Describes a set of changelists
 		/// </summary>
+		/// <param name="Connection">Connection to the Perforce server</param>
 		/// <param name="ChangeNumbers">The changelist numbers to retrieve descriptions for</param>
 		/// <param name="CancellationToken">Token used to cancel the operation</param>
 		/// <returns>List of responses from the server</returns>
-		public async Task<List<DescribeRecord>> DescribeAsync(int[] ChangeNumbers, CancellationToken CancellationToken)
+		public static async Task<List<DescribeRecord>> DescribeAsync(this PerforceConnection Connection, int[] ChangeNumbers, CancellationToken CancellationToken)
 		{
-			return (await TryDescribeAsync(ChangeNumbers, CancellationToken)).Data;
+			return (await TryDescribeAsync(Connection, ChangeNumbers, CancellationToken)).Data;
 		}
 
 		/// <summary>
 		/// Describes a set of changelists
 		/// </summary>
+		/// <param name="Connection">Connection to the Perforce server</param>
 		/// <param name="ChangeNumbers">The changelist numbers to retrieve descriptions for</param>
 		/// <param name="CancellationToken">Token used to cancel the operation</param>
 		/// <returns>List of responses from the server</returns>
-		public Task<PerforceResponseList<DescribeRecord>> TryDescribeAsync(int[] ChangeNumbers, CancellationToken CancellationToken)
+		public static Task<PerforceResponseList<DescribeRecord>> TryDescribeAsync(this PerforceConnection Connection, int[] ChangeNumbers, CancellationToken CancellationToken)
 		{
 			StringBuilder Arguments = new StringBuilder("describe -s");
 			foreach (int ChangeNumber in ChangeNumbers)
 			{
 				Arguments.AppendFormat(" {0}", ChangeNumber);
 			}
-			return CommandAsync<DescribeRecord>(Arguments.ToString(), null, CancellationToken);
+			return CommandAsync<DescribeRecord>(Connection, Arguments.ToString(), null, CancellationToken);
 		}
 
 		/// <summary>
 		/// Describes a set of changelists
 		/// </summary>
+		/// <param name="Connection">Connection to the Perforce server</param>
+		/// <param name="Options">Options for the command</param>
+		/// <param name="MaxNumFiles">Maximum number of files to return</param>
 		/// <param name="ChangeNumbers">The changelist numbers to retrieve descriptions for</param>
 		/// <param name="CancellationToken">Token used to cancel the operation</param>
 		/// <returns>List of responses from the server</returns>
-		public async Task<List<DescribeRecord>> DescribeAsync(DescribeOptions Options, int MaxNumFiles, int[] ChangeNumbers, CancellationToken CancellationToken)
+		public static async Task<List<DescribeRecord>> DescribeAsync(this PerforceConnection Connection, DescribeOptions Options, int MaxNumFiles, int[] ChangeNumbers, CancellationToken CancellationToken)
 		{
-			return (await TryDescribeAsync(Options, MaxNumFiles, ChangeNumbers, CancellationToken)).Data;
+			return (await TryDescribeAsync(Connection, Options, MaxNumFiles, ChangeNumbers, CancellationToken)).Data;
 		}
 
 		/// <summary>
 		/// Describes a set of changelists
 		/// </summary>
+		/// <param name="Connection">Connection to the Perforce server</param>
+		/// <param name="Options">Options for the command</param>
+		/// <param name="MaxNumFiles">Maximum number of files to return</param>
 		/// <param name="ChangeNumbers">The changelist numbers to retrieve descriptions for</param>
 		/// <param name="CancellationToken">Token used to cancel the operation</param>
 		/// <returns>List of responses from the server</returns>
-		public Task<PerforceResponseList<DescribeRecord>> TryDescribeAsync(DescribeOptions Options, int MaxNumFiles, int[] ChangeNumbers, CancellationToken CancellationToken)
+		public static Task<PerforceResponseList<DescribeRecord>> TryDescribeAsync(this PerforceConnection Connection, DescribeOptions Options, int MaxNumFiles, int[] ChangeNumbers, CancellationToken CancellationToken)
 		{
 			StringBuilder Arguments = new StringBuilder("describe -s");
 			if ((Options & DescribeOptions.ShowDescriptionForRestrictedChanges) != 0)
@@ -1294,7 +1406,7 @@ namespace EpicGames.Perforce
 			{
 				Arguments.AppendFormat(" {0}", ChangeNumber);
 			}
-			return CommandAsync<DescribeRecord>(Arguments.ToString(), null, CancellationToken);
+			return CommandAsync<DescribeRecord>(Connection, Arguments.ToString(), null, CancellationToken);
 		}
 
 		#endregion
@@ -1304,51 +1416,55 @@ namespace EpicGames.Perforce
 		/// <summary>
 		/// Opens files for edit
 		/// </summary>
+		/// <param name="Connection">Connection to the Perforce server</param>
 		/// <param name="ChangeNumber">Changelist to add files to</param>
 		/// <param name="FileSpecs">Files to be opened for edit</param>
 		/// <param name="CancellationToken">Token used to cancel the operation</param>
 		/// <returns>Response from the server</returns>
-		public async Task<List<EditRecord>> EditAsync(int ChangeNumber, string[] FileSpecs, CancellationToken CancellationToken)
+		public static async Task<List<EditRecord>> EditAsync(this PerforceConnection Connection, int ChangeNumber, string[] FileSpecs, CancellationToken CancellationToken)
 		{
-			return (await TryEditAsync(ChangeNumber, FileSpecs, CancellationToken)).Data;
+			return (await TryEditAsync(Connection, ChangeNumber, FileSpecs, CancellationToken)).Data;
 		}
 
 		/// <summary>
 		/// Opens files for edit
 		/// </summary>
+		/// <param name="Connection">Connection to the Perforce server</param>
 		/// <param name="ChangeNumber">Changelist to add files to</param>
 		/// <param name="FileSpecs">Files to be opened for edit</param>
 		/// <param name="CancellationToken">Token used to cancel the operation</param>
 		/// <returns>Response from the server</returns>
-		public Task<PerforceResponseList<EditRecord>> TryEditAsync(int ChangeNumber, string[] FileSpecs, CancellationToken CancellationToken)
+		public static Task<PerforceResponseList<EditRecord>> TryEditAsync(this PerforceConnection Connection, int ChangeNumber, string[] FileSpecs, CancellationToken CancellationToken)
 		{
-			return TryEditAsync(ChangeNumber, null, EditOptions.None, FileSpecs, CancellationToken);
+			return TryEditAsync(Connection, ChangeNumber, null, EditOptions.None, FileSpecs, CancellationToken);
 		}
 
 		/// <summary>
 		/// Opens files for edit
 		/// </summary>
+		/// <param name="Connection">Connection to the Perforce server</param>
 		/// <param name="ChangeNumber">Changelist to add files to</param>
 		/// <param name="FileType">Type for new files</param>
 		/// <param name="Options">Options for the command</param>
 		/// <param name="FileSpecs">Files to be opened for edit</param>
 		/// <param name="CancellationToken">Token used to cancel the operation</param>
 		/// <returns>Response from the server</returns>
-		public async Task<List<EditRecord>> EditAsync(int ChangeNumber, string? FileType, EditOptions Options, string[] FileSpecs, CancellationToken CancellationToken)
+		public static async Task<List<EditRecord>> EditAsync(this PerforceConnection Connection, int ChangeNumber, string? FileType, EditOptions Options, string[] FileSpecs, CancellationToken CancellationToken)
 		{
-			return (await TryEditAsync(ChangeNumber, FileType, Options, FileSpecs, CancellationToken)).Data;
+			return (await TryEditAsync(Connection, ChangeNumber, FileType, Options, FileSpecs, CancellationToken)).Data;
 		}
 
 		/// <summary>
 		/// Opens files for edit
 		/// </summary>
+		/// <param name="Connection">Connection to the Perforce server</param>
 		/// <param name="ChangeNumber">Changelist to add files to</param>
 		/// <param name="FileType">Type for new files</param>
 		/// <param name="Options">Options for the command</param>
 		/// <param name="FileSpecs">Files to be opened for edit</param>
 		/// <param name="CancellationToken">Token used to cancel the operation</param>
 		/// <returns>Response from the server</returns>
-		public Task<PerforceResponseList<EditRecord>> TryEditAsync(int ChangeNumber, string? FileType, EditOptions Options, string[] FileSpecs, CancellationToken CancellationToken)
+		public static Task<PerforceResponseList<EditRecord>> TryEditAsync(this PerforceConnection Connection, int ChangeNumber, string? FileType, EditOptions Options, string[] FileSpecs, CancellationToken CancellationToken)
 		{
 			StringBuilder Arguments = new StringBuilder("edit");
 			if (ChangeNumber != -1)
@@ -1368,7 +1484,7 @@ namespace EpicGames.Perforce
 				Arguments.AppendFormat(" -t \"{0}\"", FileType);
 			}
 
-			return BatchedCommandAsync<EditRecord>(Arguments.ToString(), FileSpecs, null, CancellationToken);
+			return BatchedCommandAsync<EditRecord>(Connection, Arguments.ToString(), FileSpecs, null, CancellationToken);
 		}
 
 		#endregion
@@ -1378,77 +1494,83 @@ namespace EpicGames.Perforce
 		/// <summary>
 		/// Execute the 'filelog' command
 		/// </summary>
+		/// <param name="Connection">Connection to the Perforce server</param>
 		/// <param name="Options">Options for the command</param>
 		/// <param name="FileSpecs">List of file specifications to query</param>
 		/// <param name="CancellationToken">Token used to cancel the operation</param>
 		/// <returns>List of responses from the server</returns>
-		public async Task<List<FileLogRecord>> FileLogAsync(FileLogOptions Options, string[] FileSpecs, CancellationToken CancellationToken)
+		public static async Task<List<FileLogRecord>> FileLogAsync(this PerforceConnection Connection, FileLogOptions Options, string[] FileSpecs, CancellationToken CancellationToken)
 		{
-			return (await TryFileLogAsync(Options, FileSpecs, CancellationToken)).Data;
+			return (await TryFileLogAsync(Connection, Options, FileSpecs, CancellationToken)).Data;
 		}
 
 		/// <summary>
 		/// Execute the 'filelog' command
 		/// </summary>
+		/// <param name="Connection">Connection to the Perforce server</param>
 		/// <param name="Options">Options for the command</param>
 		/// <param name="FileSpecs">List of file specifications to query</param>
 		/// <param name="CancellationToken">Token used to cancel the operation</param>
 		/// <returns>List of responses from the server</returns>
-		public Task<PerforceResponseList<FileLogRecord>> TryFileLogAsync(FileLogOptions Options, string[] FileSpecs, CancellationToken CancellationToken)
+		public static Task<PerforceResponseList<FileLogRecord>> TryFileLogAsync(this PerforceConnection Connection, FileLogOptions Options, string[] FileSpecs, CancellationToken CancellationToken)
 		{
-			return TryFileLogAsync(-1, -1, Options, FileSpecs, CancellationToken);
+			return TryFileLogAsync(Connection, -1, -1, Options, FileSpecs, CancellationToken);
 		}
 
 		/// <summary>
 		/// Execute the 'filelog' command
 		/// </summary>
+		/// <param name="Connection">Connection to the Perforce server</param>
 		/// <param name="MaxChanges">Number of changelists to show. Ignored if zero or negative.</param>
 		/// <param name="Options">Options for the command</param>
 		/// <param name="FileSpecs">List of file specifications to query</param>
 		/// <param name="CancellationToken">Token used to cancel the operation</param>
 		/// <returns>List of responses from the server</returns>
-		public async Task<List<FileLogRecord>> FileLogAsync(int MaxChanges, FileLogOptions Options, string[] FileSpecs, CancellationToken CancellationToken)
+		public static async Task<List<FileLogRecord>> FileLogAsync(this PerforceConnection Connection, int MaxChanges, FileLogOptions Options, string[] FileSpecs, CancellationToken CancellationToken)
 		{
-			return (await TryFileLogAsync(MaxChanges, Options, FileSpecs, CancellationToken)).Data;
+			return (await TryFileLogAsync(Connection, MaxChanges, Options, FileSpecs, CancellationToken)).Data;
 		}
 
 		/// <summary>
 		/// Execute the 'filelog' command
 		/// </summary>
+		/// <param name="Connection">Connection to the Perforce server</param>
 		/// <param name="MaxChanges">Number of changelists to show. Ignored if zero or negative.</param>
 		/// <param name="Options">Options for the command</param>
 		/// <param name="FileSpecs">List of file specifications to query</param>
 		/// <param name="CancellationToken">Token used to cancel the operation</param>
 		/// <returns>List of responses from the server</returns>
-		public Task<PerforceResponseList<FileLogRecord>> TryFileLogAsync(int MaxChanges, FileLogOptions Options, string[] FileSpecs, CancellationToken CancellationToken)
+		public static Task<PerforceResponseList<FileLogRecord>> TryFileLogAsync(this PerforceConnection Connection, int MaxChanges, FileLogOptions Options, string[] FileSpecs, CancellationToken CancellationToken)
 		{
-			return TryFileLogAsync(-1, MaxChanges, Options, FileSpecs, CancellationToken);
+			return TryFileLogAsync(Connection, -1, MaxChanges, Options, FileSpecs, CancellationToken);
 		}
 
 		/// <summary>
 		/// Execute the 'filelog' command
 		/// </summary>
+		/// <param name="Connection">Connection to the Perforce server</param>
 		/// <param name="ChangeNumber">Show only files modified by this changelist. Ignored if zero or negative.</param>
 		/// <param name="MaxChanges">Number of changelists to show. Ignored if zero or negative.</param>
 		/// <param name="Options">Options for the command</param>
 		/// <param name="FileSpecs">List of file specifications to query</param>
 		/// <param name="CancellationToken">Token used to cancel the operation</param>
 		/// <returns>List of responses from the server</returns>
-		public async Task<List<FileLogRecord>> FileLogAsync(int ChangeNumber, int MaxChanges, FileLogOptions Options, string[] FileSpecs, CancellationToken CancellationToken)
+		public static async Task<List<FileLogRecord>> FileLogAsync(this PerforceConnection Connection, int ChangeNumber, int MaxChanges, FileLogOptions Options, string[] FileSpecs, CancellationToken CancellationToken)
 		{
-			return (await TryFileLogAsync(ChangeNumber, MaxChanges, Options, FileSpecs, CancellationToken)).Data;
+			return (await TryFileLogAsync(Connection, ChangeNumber, MaxChanges, Options, FileSpecs, CancellationToken)).Data;
 		}
 
 		/// <summary>
 		/// Execute the 'filelog' command
 		/// </summary>
+		/// <param name="Connection">Connection to the Perforce server</param>
 		/// <param name="ChangeNumber">Show only files modified by this changelist. Ignored if zero or negative.</param>
 		/// <param name="MaxChanges">Number of changelists to show. Ignored if zero or negative.</param>
 		/// <param name="Options">Options for the command</param>
 		/// <param name="FileSpecs">List of file specifications to query</param>
 		/// <param name="CancellationToken">Token used to cancel the operation</param>
 		/// <returns>List of responses from the server</returns>
-		public Task<PerforceResponseList<FileLogRecord>> TryFileLogAsync(int ChangeNumber, int MaxChanges, FileLogOptions Options, string[] FileSpecs, CancellationToken CancellationToken)
+		public static Task<PerforceResponseList<FileLogRecord>> TryFileLogAsync(this PerforceConnection Connection, int ChangeNumber, int MaxChanges, FileLogOptions Options, string[] FileSpecs, CancellationToken CancellationToken)
 		{
 			// Build the argument list
 			StringBuilder Arguments = new StringBuilder("filelog");
@@ -1489,7 +1611,7 @@ namespace EpicGames.Perforce
 			Arguments.Append(" -t");
 
 			// Append all the arguments
-			return BatchedCommandAsync<FileLogRecord>(Arguments.ToString(), FileSpecs, null, CancellationToken);
+			return BatchedCommandAsync<FileLogRecord>(Connection, Arguments.ToString(), FileSpecs, null, CancellationToken);
 		}
 
 		#endregion
@@ -1499,56 +1621,61 @@ namespace EpicGames.Perforce
 		/// <summary>
 		/// Execute the 'fstat' command
 		/// </summary>
+		/// <param name="Connection">Connection to the Perforce server</param>
 		/// <param name="Options">Options for the command</param>
 		/// <param name="FileSpecs">List of file specifications to query</param>
 		/// <param name="CancellationToken">Token used to cancel the operation</param>
 		/// <returns>List of responses from the server</returns>
-		public async Task<List<FStatRecord>> FStatAsync(FStatOptions Options, string[] FileSpecs, CancellationToken CancellationToken)
+		public static async Task<List<FStatRecord>> FStatAsync(this PerforceConnection Connection, FStatOptions Options, string[] FileSpecs, CancellationToken CancellationToken)
 		{
-			return (await TryFStatAsync(Options, FileSpecs, CancellationToken)).Data;
+			return (await TryFStatAsync(Connection, Options, FileSpecs, CancellationToken)).Data;
 		}
 
 		/// <summary>
 		/// Execute the 'fstat' command
 		/// </summary>
+		/// <param name="Connection">Connection to the Perforce server</param>
 		/// <param name="Options">Options for the command</param>
 		/// <param name="FileSpecs">List of file specifications to query</param>
 		/// <param name="CancellationToken">Token used to cancel the operation</param>
 		/// <returns>List of responses from the server</returns>
-		public Task<PerforceResponseList<FStatRecord>> TryFStatAsync(FStatOptions Options, string[] FileSpecs, CancellationToken CancellationToken)
+		public static Task<PerforceResponseList<FStatRecord>> TryFStatAsync(this PerforceConnection Connection, FStatOptions Options, string[] FileSpecs, CancellationToken CancellationToken)
 		{
-			return TryFStatAsync(-1, Options, FileSpecs, CancellationToken);
+			return TryFStatAsync(Connection, -1, Options, FileSpecs, CancellationToken);
 		}
 
 		/// <summary>
 		/// Execute the 'fstat' command
 		/// </summary>
+		/// <param name="Connection">Connection to the Perforce server</param>
 		/// <param name="MaxFiles">Produce fstat output for only the first max files.</param>
 		/// <param name="Options">Options for the command</param>
 		/// <param name="FileSpecs">List of file specifications to query</param>
 		/// <param name="CancellationToken">Token used to cancel the operation</param>
 		/// <returns>List of responses from the server</returns>
-		public async Task<List<FStatRecord>> FStatAsync(int MaxFiles, FStatOptions Options, string[] FileSpecs, CancellationToken CancellationToken)
+		public static async Task<List<FStatRecord>> FStatAsync(this PerforceConnection Connection, int MaxFiles, FStatOptions Options, string[] FileSpecs, CancellationToken CancellationToken)
 		{
-			return (await TryFStatAsync(MaxFiles, Options, FileSpecs, CancellationToken)).Data;
+			return (await TryFStatAsync(Connection, MaxFiles, Options, FileSpecs, CancellationToken)).Data;
 		}
 
 		/// <summary>
 		/// Execute the 'fstat' command
 		/// </summary>
+		/// <param name="Connection">Connection to the Perforce server</param>
 		/// <param name="MaxFiles">Produce fstat output for only the first max files.</param>
 		/// <param name="Options">Options for the command</param>
 		/// <param name="FileSpecs">List of file specifications to query</param>
 		/// <param name="CancellationToken">Token used to cancel the operation</param>
 		/// <returns>List of responses from the server</returns>
-		public Task<PerforceResponseList<FStatRecord>> TryFStatAsync(int MaxFiles, FStatOptions Options, string[] FileSpecs, CancellationToken CancellationToken)
+		public static Task<PerforceResponseList<FStatRecord>> TryFStatAsync(this PerforceConnection Connection, int MaxFiles, FStatOptions Options, string[] FileSpecs, CancellationToken CancellationToken)
 		{
-			return TryFStatAsync(-1, -1, null, MaxFiles, Options, FileSpecs, CancellationToken);
+			return TryFStatAsync(Connection, -1, -1, null, MaxFiles, Options, FileSpecs, CancellationToken);
 		}
 
 		/// <summary>
 		/// Execute the 'fstat' command
 		/// </summary>
+		/// <param name="Connection">Connection to the Perforce server</param>
 		/// <param name="AfterChangeNumber">Return only files affected after the given changelist number.</param>
 		/// <param name="OnlyChangeNumber">Return only files affected by the given changelist number.</param>
 		/// <param name="Filter">List only those files that match the criteria specified.</param>
@@ -1557,14 +1684,15 @@ namespace EpicGames.Perforce
 		/// <param name="FileSpecs">List of file specifications to query</param>
 		/// <param name="CancellationToken">Token used to cancel the operation</param>
 		/// <returns>List of responses from the server</returns>
-		public async Task<List<FStatRecord>> FStatAsync(int AfterChangeNumber, int OnlyChangeNumber, string? Filter, int MaxFiles, FStatOptions Options, string[] FileSpecs, CancellationToken CancellationToken)
+		public static async Task<List<FStatRecord>> FStatAsync(this PerforceConnection Connection, int AfterChangeNumber, int OnlyChangeNumber, string? Filter, int MaxFiles, FStatOptions Options, string[] FileSpecs, CancellationToken CancellationToken)
 		{
-			return (await TryFStatAsync(AfterChangeNumber, OnlyChangeNumber, Filter, MaxFiles, Options, FileSpecs, CancellationToken)).Data;
+			return (await TryFStatAsync(Connection, AfterChangeNumber, OnlyChangeNumber, Filter, MaxFiles, Options, FileSpecs, CancellationToken)).Data;
 		}
 
 		/// <summary>
 		/// Execute the 'fstat' command
 		/// </summary>
+		/// <param name="Connection">Connection to the Perforce server</param>
 		/// <param name="AfterChangeNumber">Return only files affected after the given changelist number.</param>
 		/// <param name="OnlyChangeNumber">Return only files affected by the given changelist number.</param>
 		/// <param name="Filter">List only those files that match the criteria specified.</param>
@@ -1573,7 +1701,7 @@ namespace EpicGames.Perforce
 		/// <param name="FileSpecs">List of file specifications to query</param>
 		/// <param name="CancellationToken">Token used to cancel the operation</param>
 		/// <returns>List of responses from the server</returns>
-		public async Task<PerforceResponseList<FStatRecord>> TryFStatAsync(int AfterChangeNumber, int OnlyChangeNumber, string? Filter, int MaxFiles, FStatOptions Options, string[] FileSpecs, CancellationToken CancellationToken)
+		public static async Task<PerforceResponseList<FStatRecord>> TryFStatAsync(this PerforceConnection Connection, int AfterChangeNumber, int OnlyChangeNumber, string? Filter, int MaxFiles, FStatOptions Options, string[] FileSpecs, CancellationToken CancellationToken)
 		{
 			// Build the argument list
 			StringBuilder Arguments = new StringBuilder("fstat");
@@ -1671,7 +1799,7 @@ namespace EpicGames.Perforce
 			}
 
 			// Execute the command
-			PerforceResponseList<FStatRecord> Records = await BatchedCommandAsync<FStatRecord>(Arguments.ToString(), FileSpecs, null, CancellationToken);
+			PerforceResponseList<FStatRecord> Records = await BatchedCommandAsync<FStatRecord>(Connection, Arguments.ToString(), FileSpecs, null, CancellationToken);
 			if (OnlyChangeNumber != -1 && Records.Count > 0 && Records[Records.Count - 1].Succeeded && Records[Records.Count - 1].Data.Description != null)
 			{
 				Records.RemoveAt(Records.Count - 1);
@@ -1686,21 +1814,23 @@ namespace EpicGames.Perforce
 		/// <summary>
 		/// Execute the 'info' command
 		/// </summary>
+		/// <param name="Connection">Connection to the Perforce server</param>
 		/// <param name="Options">Options for the command</param>
 		/// <param name="CancellationToken">Token used to cancel the operation</param>
 		/// <returns>Response from the server; an InfoRecord or error code</returns>
-		public async Task<InfoRecord> GetInfoAsync(InfoOptions Options, CancellationToken CancellationToken)
+		public static async Task<InfoRecord> GetInfoAsync(this PerforceConnection Connection, InfoOptions Options, CancellationToken CancellationToken)
 		{
-			return (await TryGetInfoAsync(Options, CancellationToken)).Data;
+			return (await TryGetInfoAsync(Connection, Options, CancellationToken)).Data;
 		}
 
 		/// <summary>
 		/// Execute the 'info' command
 		/// </summary>
+		/// <param name="Connection">Connection to the Perforce server</param>
 		/// <param name="Options">Options for the command</param>
 		/// <param name="CancellationToken">Token used to cancel the operation</param>
 		/// <returns>Response from the server; an InfoRecord or error code</returns>
-		public Task<PerforceResponse<InfoRecord>> TryGetInfoAsync(InfoOptions Options, CancellationToken CancellationToken)
+		public static Task<PerforceResponse<InfoRecord>> TryGetInfoAsync(this PerforceConnection Connection, InfoOptions Options, CancellationToken CancellationToken)
 		{
 			// Build the argument list
 			StringBuilder Arguments = new StringBuilder("info");
@@ -1708,56 +1838,7 @@ namespace EpicGames.Perforce
 			{
 				Arguments.Append(" -s");
 			}
-			return SingleResponseCommandAsync<InfoRecord>(Arguments.ToString(), null, CancellationToken);
-		}
-
-		#endregion
-
-		#region p4 login
-
-		/// <summary>
-		/// Execute the 'login' command
-		/// </summary>
-		/// <param name="Password">Password to use to login</param>
-		/// <param name="CancellationToken">Token used to cancel the operation</param>
-		/// <returns>Response from the server</returns>
-		public async Task LoginAsync(string Password, CancellationToken CancellationToken)
-		{
-			// Some versions of P4.EXE do not support marshalled output for P4 login calls, so we only support this as a basic text query.
-			byte[] PasswordBytes = Encoding.UTF8.GetBytes(Password);
-			using (ManagedProcessGroup ChildProcessGroup = new ManagedProcessGroup())
-			{
-				string Executable = PerforceChildProcess.GetExecutable();
-				string FullArgumentList = GetFullCommandLine("login");
-
-				Logger.LogDebug("Running {0} {1}", Executable, FullArgumentList);
-				using (ManagedProcess ChildProcess = new ManagedProcess(ChildProcessGroup, Executable, FullArgumentList, null, null, PasswordBytes, ProcessPriorityClass.Normal))
-				{
-					StringBuilder Lines = new StringBuilder("Unable to log in: ");
-					for(; ;)
-					{
-						string? Line = await ChildProcess.ReadLineAsync();
-						if(Line == null)
-						{
-							break;
-						}
-
-						Lines.Append($"\n  {Line}");
-
-						if(Line.StartsWith("Enter password:", StringComparison.OrdinalIgnoreCase))
-						{
-							Lines.Append($" {new string('*', Password.Length)}");
-						}
-					}
-
-					ChildProcess.WaitForExit();
-
-					if (ChildProcess.ExitCode != 0)
-					{
-						throw new PerforceException($"{Lines}\n  (exit code {ChildProcess.ExitCode})");
-					}
-				}
-			}
+			return SingleResponseCommandAsync<InfoRecord>(Connection, Arguments.ToString(), null, CancellationToken);
 		}
 
 		#endregion
@@ -1767,6 +1848,7 @@ namespace EpicGames.Perforce
 		/// <summary>
 		/// Execute the 'merge' command
 		/// </summary>
+		/// <param name="Connection">Connection to the Perforce server</param>
 		/// <param name="Options">Options for the merge</param>
 		/// <param name="Change"></param>
 		/// <param name="MaxFiles">Maximum number of files to merge</param>
@@ -1774,7 +1856,7 @@ namespace EpicGames.Perforce
 		/// <param name="TargetFileSpec">The target filespec</param>
 		/// <param name="CancellationToken">Cancellation token</param>
 		/// <returns>List of records</returns>
-		public Task<PerforceResponseList<MergeRecord>> MergeAsync(MergeOptions Options, int Change, int MaxFiles, string SourceFileSpec, string TargetFileSpec, CancellationToken CancellationToken)
+		public static Task<PerforceResponseList<MergeRecord>> MergeAsync(this PerforceConnection Connection, MergeOptions Options, int Change, int MaxFiles, string SourceFileSpec, string TargetFileSpec, CancellationToken CancellationToken)
 		{
 			StringBuilder Arguments = new StringBuilder("merge");
 			if ((Options & MergeOptions.Preview) != 0)
@@ -1792,7 +1874,7 @@ namespace EpicGames.Perforce
 			Arguments.AppendFormat(" \"{0}\"", SourceFileSpec);
 			Arguments.AppendFormat(" \"{0}\"", TargetFileSpec);
 
-			return CommandAsync<MergeRecord>(Arguments.ToString(), null, CancellationToken);
+			return CommandAsync<MergeRecord>(Connection, Arguments.ToString(), null, CancellationToken);
 		}
 
 		#endregion
@@ -1802,6 +1884,7 @@ namespace EpicGames.Perforce
 		/// <summary>
 		/// Opens files for move
 		/// </summary>
+		/// <param name="Connection">Connection to the Perforce server</param>
 		/// <param name="ChangeNumber">Changelist to add files to</param>
 		/// <param name="FileType">Type for new files</param>
 		/// <param name="Options">Options for the command</param>
@@ -1809,14 +1892,15 @@ namespace EpicGames.Perforce
 		/// <param name="TargetFileSpec">The target file(s)</param>
 		/// <param name="CancellationToken">Token used to cancel the operation</param>
 		/// <returns>Response from the server</returns>
-		public async Task<List<MoveRecord>> MoveAsync(int ChangeNumber, string? FileType, MoveOptions Options, string SourceFileSpec, string TargetFileSpec, CancellationToken CancellationToken)
+		public static async Task<List<MoveRecord>> MoveAsync(this PerforceConnection Connection, int ChangeNumber, string? FileType, MoveOptions Options, string SourceFileSpec, string TargetFileSpec, CancellationToken CancellationToken)
 		{
-			return (await TryMoveAsync(ChangeNumber, FileType, Options, SourceFileSpec, TargetFileSpec, CancellationToken)).Data;
+			return (await TryMoveAsync(Connection, ChangeNumber, FileType, Options, SourceFileSpec, TargetFileSpec, CancellationToken)).Data;
 		}
 
 		/// <summary>
 		/// Opens files for move
 		/// </summary>
+		/// <param name="Connection">Connection to the Perforce server</param>
 		/// <param name="ChangeNumber">Changelist to add files to</param>
 		/// <param name="FileType">Type for new files</param>
 		/// <param name="Options">Options for the command</param>
@@ -1824,7 +1908,7 @@ namespace EpicGames.Perforce
 		/// <param name="TargetFileSpec">The target file(s)</param>
 		/// <param name="CancellationToken">Token used to cancel the operation</param>
 		/// <returns>Response from the server</returns>
-		public Task<PerforceResponseList<MoveRecord>> TryMoveAsync(int ChangeNumber, string? FileType, MoveOptions Options, string SourceFileSpec, string TargetFileSpec, CancellationToken CancellationToken)
+		public static Task<PerforceResponseList<MoveRecord>> TryMoveAsync(this PerforceConnection Connection, int ChangeNumber, string? FileType, MoveOptions Options, string SourceFileSpec, string TargetFileSpec, CancellationToken CancellationToken)
 		{
 			StringBuilder Arguments = new StringBuilder("move");
 			if (ChangeNumber != -1)
@@ -1849,7 +1933,7 @@ namespace EpicGames.Perforce
 			}
 			Arguments.AppendFormat(" \"{0}\" \"{1}\"", SourceFileSpec, TargetFileSpec);
 
-			return CommandAsync<MoveRecord>(Arguments.ToString(), null, CancellationToken);
+			return CommandAsync<MoveRecord>(Connection, Arguments.ToString(), null, CancellationToken);
 		}
 
 		#endregion
@@ -1859,6 +1943,7 @@ namespace EpicGames.Perforce
 		/// <summary>
 		/// Execute the 'opened' command
 		/// </summary>
+		/// <param name="Connection">Connection to the Perforce server</param>
 		/// <param name="Options">Options for the command</param>
 		/// <param name="ChangeNumber">List the files in pending changelist change. To list files in the default changelist, use DefaultChange.</param>
 		/// <param name="ClientName">List only files that are open in the given client</param>
@@ -1867,14 +1952,15 @@ namespace EpicGames.Perforce
 		/// <param name="FileSpecs">Specification for the files to list</param>
 		/// <param name="CancellationToken">Token used to cancel the operation</param>
 		/// <returns>Response from the server</returns>
-		public async Task<List<FStatRecord>> GetOpenFilesAsync(OpenedOptions Options, int ChangeNumber, string? ClientName, string? UserName, int MaxResults, string[]? FileSpecs, CancellationToken CancellationToken)
+		public static async Task<List<FStatRecord>> GetOpenFilesAsync(this PerforceConnection Connection, OpenedOptions Options, int ChangeNumber, string? ClientName, string? UserName, int MaxResults, string[]? FileSpecs, CancellationToken CancellationToken)
 		{
-			return (await TryGetOpenFilesAsync(Options, ChangeNumber, ClientName, UserName, MaxResults, FileSpecs, CancellationToken)).Data;
+			return (await TryGetOpenFilesAsync(Connection, Options, ChangeNumber, ClientName, UserName, MaxResults, FileSpecs, CancellationToken)).Data;
 		}
 
 		/// <summary>
 		/// Execute the 'opened' command
 		/// </summary>
+		/// <param name="Connection">Connection to the Perforce server</param>
 		/// <param name="Options">Options for the command</param>
 		/// <param name="ChangeNumber">List the files in pending changelist change. To list files in the default changelist, use DefaultChange.</param>
 		/// <param name="ClientName">List only files that are open in the given client</param>
@@ -1883,7 +1969,7 @@ namespace EpicGames.Perforce
 		/// <param name="FileSpecs">Specification for the files to list</param>
 		/// <param name="CancellationToken">Token used to cancel the operation</param>
 		/// <returns>Response from the server</returns>
-		public Task<PerforceResponseList<FStatRecord>> TryGetOpenFilesAsync(OpenedOptions Options, int ChangeNumber, string? ClientName, string? UserName, int MaxResults, string[]? FileSpecs, CancellationToken CancellationToken)
+		public static Task<PerforceResponseList<FStatRecord>> TryGetOpenFilesAsync(this PerforceConnection Connection, OpenedOptions Options, int ChangeNumber, string? ClientName, string? UserName, int MaxResults, string[]? FileSpecs, CancellationToken CancellationToken)
 		{
 			// Build the argument list
 			StringBuilder Arguments = new StringBuilder("opened");
@@ -1922,7 +2008,7 @@ namespace EpicGames.Perforce
 					Arguments.AppendFormat(" \"{0}\"", FileSpec);
 				}
 			}
-			return CommandAsync<FStatRecord>(Arguments.ToString(), null, CancellationToken);
+			return CommandAsync<FStatRecord>(Connection, Arguments.ToString(), null, CancellationToken);
 		}
 
 		#endregion
@@ -1932,29 +2018,31 @@ namespace EpicGames.Perforce
 		/// <summary>
 		/// Execute the 'print' command
 		/// </summary>
+		/// <param name="Connection">Connection to the Perforce server</param>
 		/// <param name="OutputFile">Output file to redirect output to</param>
 		/// <param name="FileSpec">Specification for the files to print</param>
 		/// <param name="CancellationToken">Token used to cancel the operation</param>
 		/// <returns>Response from the server</returns>
-		public async Task<PrintRecord> PrintAsync(string OutputFile, string FileSpec, CancellationToken CancellationToken)
+		public static async Task<PrintRecord> PrintAsync(this PerforceConnection Connection, string OutputFile, string FileSpec, CancellationToken CancellationToken)
 		{
-			return (await TryPrintAsync(OutputFile, FileSpec, CancellationToken)).Data;
+			return (await TryPrintAsync(Connection, OutputFile, FileSpec, CancellationToken)).Data;
 		}
 
 		/// <summary>
 		/// Execute the 'print' command
 		/// </summary>
+		/// <param name="Connection">Connection to the Perforce server</param>
 		/// <param name="OutputFile">Output file to redirect output to</param>
 		/// <param name="FileSpec">Specification for the files to print</param>
 		/// <param name="CancellationToken">Token used to cancel the operation</param>
 		/// <returns>Response from the server</returns>
-		public Task<PerforceResponse<PrintRecord>> TryPrintAsync(string OutputFile, string FileSpec, CancellationToken CancellationToken)
+		public static Task<PerforceResponse<PrintRecord>> TryPrintAsync(this PerforceConnection Connection, string OutputFile, string FileSpec, CancellationToken CancellationToken)
 		{
 			// Build the argument list
 			StringBuilder Arguments = new StringBuilder("print");
 			Arguments.AppendFormat(" -o \"{0}\"", OutputFile);
 			Arguments.AppendFormat(" \"{0}\"", FileSpec);
-			return SingleResponseCommandAsync<PrintRecord>(Arguments.ToString(), null, CancellationToken);
+			return SingleResponseCommandAsync<PrintRecord>(Connection, Arguments.ToString(), null, CancellationToken);
 		}
 
 		class PrintHandler : IDisposable
@@ -2027,14 +2115,15 @@ namespace EpicGames.Perforce
 		/// <summary>
 		/// Open files for add, delete, and/or edit in order to reconcile a workspace with changes made outside of Perforce.
 		/// </summary>
+		/// <param name="Connection">Connection to the Perforce server</param>
 		/// <param name="ChangeNumber">Changelist to open files to</param>
 		/// <param name="Options">Options for the command</param>
 		/// <param name="FileSpecs">Files to be reverted</param>
 		/// <param name="CancellationToken">Token used to cancel the operation</param>
 		/// <returns>Response from the server</returns>
-		public async Task<List<ReconcileRecord>> ReconcileAsync(int ChangeNumber, ReconcileOptions Options, string[] FileSpecs, CancellationToken CancellationToken)
+		public static async Task<List<ReconcileRecord>> ReconcileAsync(this PerforceConnection Connection, int ChangeNumber, ReconcileOptions Options, string[] FileSpecs, CancellationToken CancellationToken)
 		{
-			PerforceResponseList<ReconcileRecord> Records = await TryReconcileAsync(ChangeNumber, Options, FileSpecs, CancellationToken);
+			PerforceResponseList<ReconcileRecord> Records = await TryReconcileAsync(Connection, ChangeNumber, Options, FileSpecs, CancellationToken);
 			Records.RemoveAll(x => x.Info != null);
 			return Records.Data;
 		}
@@ -2042,12 +2131,13 @@ namespace EpicGames.Perforce
 		/// <summary>
 		/// Open files for add, delete, and/or edit in order to reconcile a workspace with changes made outside of Perforce.
 		/// </summary>
+		/// <param name="Connection">Connection to the Perforce server</param>
 		/// <param name="ChangeNumber">Changelist to open files to</param>
 		/// <param name="Options">Options for the command</param>
 		/// <param name="FileSpecs">Files to be reverted</param>
 		/// <param name="CancellationToken">Token used to cancel the operation</param>
 		/// <returns>Response from the server</returns>
-		public Task<PerforceResponseList<ReconcileRecord>> TryReconcileAsync(int ChangeNumber, ReconcileOptions Options, string[] FileSpecs, CancellationToken CancellationToken)
+		public static Task<PerforceResponseList<ReconcileRecord>> TryReconcileAsync(this PerforceConnection Connection, int ChangeNumber, ReconcileOptions Options, string[] FileSpecs, CancellationToken CancellationToken)
 		{
 			StringBuilder Arguments = new StringBuilder("reconcile");
 			if(ChangeNumber != -1)
@@ -2091,7 +2181,7 @@ namespace EpicGames.Perforce
 				Arguments.AppendFormat(" \"{0}\"", FileSpec);
 			}
 
-			return CommandAsync<ReconcileRecord>(Arguments.ToString(), null, CancellationToken);
+			return CommandAsync<ReconcileRecord>(Connection, Arguments.ToString(), null, CancellationToken);
 		}
 
 		#endregion
@@ -2101,25 +2191,27 @@ namespace EpicGames.Perforce
 		/// <summary>
 		/// Reopen a file
 		/// </summary>
+		/// <param name="Connection">Connection to the Perforce server</param>
 		/// <param name="ChangeNumber">Changelist to open files to</param>
 		/// <param name="FileType">New filetype</param>
 		/// <param name="FileSpec">Files to be reverted</param>
 		/// <param name="CancellationToken">Token used to cancel the operation</param>
 		/// <returns>Response from the server</returns>
-		public async Task<List<ReopenRecord>> ReopenAsync(int? ChangeNumber, string? FileType, string FileSpec, CancellationToken CancellationToken)
+		public static async Task<List<ReopenRecord>> ReopenAsync(this PerforceConnection Connection, int? ChangeNumber, string? FileType, string FileSpec, CancellationToken CancellationToken)
 		{
-			return (await TryReopenAsync(ChangeNumber, FileType, FileSpec, CancellationToken)).Data;
+			return (await TryReopenAsync(Connection, ChangeNumber, FileType, FileSpec, CancellationToken)).Data;
 		}
 
 		/// <summary>
 		/// Reopen a file
 		/// </summary>
+		/// <param name="Connection">Connection to the Perforce server</param>
 		/// <param name="ChangeNumber">Changelist to open files to</param>
 		/// <param name="FileType">New filetype</param>
 		/// <param name="FileSpec">Files to be reverted</param>
 		/// <param name="CancellationToken">Token used to cancel the operation</param>
 		/// <returns>Response from the server</returns>
-		public Task<PerforceResponseList<ReopenRecord>> TryReopenAsync(int? ChangeNumber, string? FileType, string FileSpec, CancellationToken CancellationToken)
+		public static Task<PerforceResponseList<ReopenRecord>> TryReopenAsync(this PerforceConnection Connection, int? ChangeNumber, string? FileType, string FileSpec, CancellationToken CancellationToken)
 		{
 			StringBuilder Arguments = new StringBuilder("reopen");
 			if (ChangeNumber != null)
@@ -2139,7 +2231,7 @@ namespace EpicGames.Perforce
 			}
 			Arguments.AppendFormat(" \"{0}\"", FileSpec);
 
-			return CommandAsync<ReopenRecord>(Arguments.ToString(), null, CancellationToken);
+			return CommandAsync<ReopenRecord>(Connection, Arguments.ToString(), null, CancellationToken);
 		}
 
 		#endregion
@@ -2149,13 +2241,14 @@ namespace EpicGames.Perforce
 		/// <summary>
 		/// Reloads a client workspace
 		/// </summary>
+		/// <param name="Connection">Connection to the Perforce server</param>
 		/// <param name="ClientName">Name of the client to reload</param>
 		/// <param name="SourceServerId">The source server id</param>
 		/// <param name="CancellationToken">Token used to cancel the operation</param>
 		/// <returns>Response from the server</returns>
-		public Task<List<PerforceResponse>> ReloadClient(string ClientName, string SourceServerId, CancellationToken CancellationToken)
+		public static Task<List<PerforceResponse>> ReloadClient(this PerforceConnection Connection, string ClientName, string SourceServerId, CancellationToken CancellationToken)
 		{
-			return CommandAsync($"reload -c \"{ClientName}\" -p \"{SourceServerId}\"", null, null, CancellationToken);
+			return Connection.CommandAsync($"reload -c \"{ClientName}\" -p \"{SourceServerId}\"", null, null, CancellationToken);
 		}
 
 		#endregion p4 reload
@@ -2165,25 +2258,27 @@ namespace EpicGames.Perforce
 		/// <summary>
 		/// Resolve conflicts between file revisions.
 		/// </summary>
+		/// <param name="Connection">Connection to the Perforce server</param>
 		/// <param name="ChangeNumber">Changelist to open files to</param>
 		/// <param name="Options">Options for the command</param>
 		/// <param name="FileSpecs">Files to be reverted</param>
 		/// <param name="CancellationToken">Token used to cancel the operation</param>
 		/// <returns>Response from the server</returns>
-		public async Task<List<ResolveRecord>> ResolveAsync(int ChangeNumber, ResolveOptions Options, string[]? FileSpecs, CancellationToken CancellationToken)
+		public static async Task<List<ResolveRecord>> ResolveAsync(this PerforceConnection Connection, int ChangeNumber, ResolveOptions Options, string[]? FileSpecs, CancellationToken CancellationToken)
 		{
-			return (await TryResolveAsync(ChangeNumber, Options, FileSpecs, CancellationToken)).Data;
+			return (await TryResolveAsync(Connection, ChangeNumber, Options, FileSpecs, CancellationToken)).Data;
 		}
 
 		/// <summary>
 		/// Resolve conflicts between file revisions.
 		/// </summary>
+		/// <param name="Connection">Connection to the Perforce server</param>
 		/// <param name="ChangeNumber">Changelist to open files to</param>
 		/// <param name="Options">Options for the command</param>
 		/// <param name="FileSpecs">Files to be reverted</param>
 		/// <param name="CancellationToken">Token used to cancel the operation</param>
 		/// <returns>Response from the server</returns>
-		public async Task<PerforceResponseList<ResolveRecord>> TryResolveAsync(int ChangeNumber, ResolveOptions Options, string[]? FileSpecs, CancellationToken CancellationToken)
+		public static async Task<PerforceResponseList<ResolveRecord>> TryResolveAsync(this PerforceConnection Connection, int ChangeNumber, ResolveOptions Options, string[]? FileSpecs, CancellationToken CancellationToken)
 		{
 			StringBuilder Arguments = new StringBuilder("resolve");
 			if ((Options & ResolveOptions.Automatic) != 0)
@@ -2234,7 +2329,7 @@ namespace EpicGames.Perforce
 				}
 			}
 
-			PerforceResponseList<ResolveRecord> Records = await CommandAsync<ResolveRecord>(Arguments.ToString(), null, CancellationToken);
+			PerforceResponseList<ResolveRecord> Records = await CommandAsync<ResolveRecord>(Connection, Arguments.ToString(), null, CancellationToken);
 			Records.RemoveAll(x => x.Error != null && x.Error.Generic == PerforceGenericCode.Empty);
 			return Records;
 		}
@@ -2246,27 +2341,29 @@ namespace EpicGames.Perforce
 		/// <summary>
 		/// Reverts files that have been added to a pending changelist.
 		/// </summary>
+		/// <param name="Connection">Connection to the Perforce server</param>
 		/// <param name="ChangeNumber">Changelist to add files to</param>
 		/// <param name="ClientName">Revert another user’s open files. </param>
 		/// <param name="Options">Options for the command</param>
 		/// <param name="FileSpecs">Files to be reverted</param>
 		/// <param name="CancellationToken">Token used to cancel the operation</param>
 		/// <returns>Response from the server</returns>
-		public async Task<List<RevertRecord>> RevertAsync(int ChangeNumber, string? ClientName, RevertOptions Options, string[] FileSpecs, CancellationToken CancellationToken)
+		public static async Task<List<RevertRecord>> RevertAsync(this PerforceConnection Connection, int ChangeNumber, string? ClientName, RevertOptions Options, string[] FileSpecs, CancellationToken CancellationToken)
 		{
-			return (await TryRevertAsync(ChangeNumber, ClientName, Options, FileSpecs, CancellationToken)).Data;
+			return (await TryRevertAsync(Connection, ChangeNumber, ClientName, Options, FileSpecs, CancellationToken)).Data;
 		}
 
 		/// <summary>
 		/// Reverts files that have been added to a pending changelist.
 		/// </summary>
+		/// <param name="Connection">Connection to the Perforce server</param>
 		/// <param name="ChangeNumber">Changelist to add files to</param>
 		/// <param name="ClientName">Revert another user’s open files. </param>
 		/// <param name="Options">Options for the command</param>
 		/// <param name="FileSpecs">Files to be reverted</param>
 		/// <param name="CancellationToken">Token used to cancel the operation</param>
 		/// <returns>Response from the server</returns>
-		public async Task<PerforceResponseList<RevertRecord>> TryRevertAsync(int ChangeNumber, string? ClientName, RevertOptions Options, string[] FileSpecs, CancellationToken CancellationToken)
+		public static async Task<PerforceResponseList<RevertRecord>> TryRevertAsync(this PerforceConnection Connection, int ChangeNumber, string? ClientName, RevertOptions Options, string[] FileSpecs, CancellationToken CancellationToken)
 		{
 			StringBuilder Arguments = new StringBuilder("revert");
 			if((Options & RevertOptions.Unchanged) != 0)
@@ -2294,7 +2391,7 @@ namespace EpicGames.Perforce
 				Arguments.AppendFormat(" -C \"{0}\"", ClientName);
 			}
 
-			PerforceResponseList<RevertRecord> Records = await BatchedCommandAsync<RevertRecord>(Arguments.ToString(), FileSpecs, null, CancellationToken);
+			PerforceResponseList<RevertRecord> Records = await BatchedCommandAsync<RevertRecord>(Connection, Arguments.ToString(), FileSpecs, null, CancellationToken);
 			Records.RemoveAll(x => x.Error != null && x.Error.Generic == PerforceGenericCode.Empty);
 			return Records;
 		}
@@ -2306,13 +2403,14 @@ namespace EpicGames.Perforce
 		/// <summary>
 		/// Sets an environment variable
 		/// </summary>
+		/// <param name="Connection">Connection to the Perforce server</param>
 		/// <param name="Name">Name of the variable to set</param>
 		/// <param name="Value">Value for the variable</param>
 		/// <param name="CancellationToken">Token used to cancel the operation</param>
 		/// <returns>Response from the server</returns>
-		public async Task SetAsync(string Name, string Value, CancellationToken CancellationToken)
+		public static async Task SetAsync(this PerforceConnection Connection, string Name, string Value, CancellationToken CancellationToken)
 		{
-			Tuple<bool, string> Result = await TrySetAsync(Name, Value, CancellationToken);
+			Tuple<bool, string> Result = await TrySetAsync(Connection, Name, Value, CancellationToken);
 			if (!Result.Item1)
 			{
 				throw new PerforceException(Result.Item2);
@@ -2322,17 +2420,18 @@ namespace EpicGames.Perforce
 		/// <summary>
 		/// Sets an environment variable
 		/// </summary>
+		/// <param name="Connection">Connection to the Perforce server</param>
 		/// <param name="Name">Name of the variable to set</param>
 		/// <param name="Value">Value for the variable</param>
 		/// <param name="CancellationToken">Token used to cancel the operation</param>
 		/// <returns>Response from the server</returns>
-		public async Task<Tuple<bool, string>> TrySetAsync(string Name, string Value, CancellationToken CancellationToken)
+		public static async Task<Tuple<bool, string>> TrySetAsync(this PerforceConnection Connection, string Name, string Value, CancellationToken CancellationToken)
 		{
 			StringBuilder Arguments = new StringBuilder("set");
 			Arguments.Append($" {Name}=");
 			Arguments.Append($"\"{Value}\"");
 
-			using (PerforceChildProcess ChildProcess = new PerforceChildProcess(null, Arguments.ToString(), Logger))
+			using (PerforceChildProcess ChildProcess = new PerforceChildProcess(null, Arguments.ToString(), Connection.Logger))
 			{
 				return await ChildProcess.TryReadToEndAsync(CancellationToken);
 			}
@@ -2341,12 +2440,13 @@ namespace EpicGames.Perforce
 		/// <summary>
 		/// Gets the setting of a Perforce variable
 		/// </summary>
+		/// <param name="Connection">Connection to the Perforce server</param>
 		/// <param name="Name">Name of the variable to get</param>
 		/// <param name="CancellationToken">Cancellation token for the request</param>
 		/// <returns>Value of the variable</returns>
-		public async Task<string?> TryGetSettingAsync(string Name, CancellationToken CancellationToken)
+		public static async Task<string?> TryGetSettingAsync(this PerforceConnection Connection, string Name, CancellationToken CancellationToken)
 		{
-			using (PerforceChildProcess ChildProcess = new PerforceChildProcess(null, $"set {Name}=", Logger))
+			using (PerforceChildProcess ChildProcess = new PerforceChildProcess(null, $"set {Name}=", Connection.Logger))
 			{
 				Tuple<bool, string> Response = await ChildProcess.TryReadToEndAsync(CancellationToken);
 				if (Response.Item1)
@@ -2364,7 +2464,7 @@ namespace EpicGames.Perforce
 					}
 				}
 
-				Logger.LogDebug($"Unable to get '{Name}' variable: {Response.Item2}");
+				Connection.Logger.LogDebug($"Unable to get '{Name}' variable: {Response.Item2}");
 				return null;
 			}
 		}
@@ -2376,25 +2476,27 @@ namespace EpicGames.Perforce
 		/// <summary>
 		/// Shelves a set of files
 		/// </summary>
+		/// <param name="Connection">Connection to the Perforce server</param>
 		/// <param name="ChangeNumber">The change number to receive the shelved files</param>
 		/// <param name="Options">Options for the command</param>
 		/// <param name="FileSpecs">Files to sync</param>
 		/// <param name="CancellationToken">Token used to cancel the operation</param>
 		/// <returns>Response from the server</returns>
-		public async Task<List<ShelveRecord>> ShelveAsync(int ChangeNumber, ShelveOptions Options, string[] FileSpecs, CancellationToken CancellationToken)
+		public static async Task<List<ShelveRecord>> ShelveAsync(this PerforceConnection Connection, int ChangeNumber, ShelveOptions Options, string[] FileSpecs, CancellationToken CancellationToken)
 		{
-			return (await TryShelveAsync(ChangeNumber, Options, FileSpecs, CancellationToken)).Data;
+			return (await TryShelveAsync(Connection, ChangeNumber, Options, FileSpecs, CancellationToken)).Data;
 		}
 
 		/// <summary>
 		/// Shelves a set of files
 		/// </summary>
+		/// <param name="Connection">Connection to the Perforce server</param>
 		/// <param name="ChangeNumber">The change number to receive the shelved files</param>
 		/// <param name="Options">Options for the command</param>
 		/// <param name="FileSpecs">Files to sync</param>
 		/// <param name="CancellationToken">Token used to cancel the operation</param>
 		/// <returns>Response from the server</returns>
-		public async Task<PerforceResponseList<ShelveRecord>> TryShelveAsync(int ChangeNumber, ShelveOptions Options, string[] FileSpecs, CancellationToken CancellationToken)
+		public static async Task<PerforceResponseList<ShelveRecord>> TryShelveAsync(this PerforceConnection Connection, int ChangeNumber, ShelveOptions Options, string[] FileSpecs, CancellationToken CancellationToken)
 		{
 			StringBuilder Arguments = new StringBuilder("shelve");
 			Arguments.AppendFormat(" -c {0}", ChangeNumber);
@@ -2411,7 +2513,7 @@ namespace EpicGames.Perforce
 				Arguments.AppendFormat(" \"{0}\"", FileSpec);
 			}
 
-			PerforceResponseList<ShelveRecord> Records = await CommandAsync<ShelveRecord>(Arguments.ToString(), null, CancellationToken);
+			PerforceResponseList<ShelveRecord> Records = await CommandAsync<ShelveRecord>(Connection, Arguments.ToString(), null, CancellationToken);
 			Records.RemoveAll(x => x.Error != null && x.Error.Generic == PerforceGenericCode.Empty);
 			return Records;
 		}
@@ -2419,13 +2521,14 @@ namespace EpicGames.Perforce
 		/// <summary>
 		/// Deletes files from a shelved changelist
 		/// </summary>
+		/// <param name="Connection">Connection to the Perforce server</param>
 		/// <param name="ChangeNumber">Changelist containing shelved files to be deleted</param>
 		/// <param name="FileSpecs">Files to delete</param>
 		/// <param name="CancellationToken">Token used to cancel the operation</param>
 		/// <returns>Response from the server</returns>
-		public async Task DeleteShelvedFilesAsync(int ChangeNumber, string[] FileSpecs, CancellationToken CancellationToken)
+		public static async Task DeleteShelvedFilesAsync(this PerforceConnection Connection, int ChangeNumber, string[] FileSpecs, CancellationToken CancellationToken)
 		{
-			List<PerforceResponse> Responses = await TryDeleteShelvedFilesAsync(ChangeNumber, FileSpecs, CancellationToken);
+			List<PerforceResponse> Responses = await TryDeleteShelvedFilesAsync(Connection, ChangeNumber, FileSpecs, CancellationToken);
 
 			PerforceResponse? ErrorResponse = Responses.FirstOrDefault(x => x.Error != null && x.Error.Generic != PerforceGenericCode.Empty);
 			if(ErrorResponse != null)
@@ -2437,11 +2540,12 @@ namespace EpicGames.Perforce
 		/// <summary>
 		/// Deletes files from a shelved changelist
 		/// </summary>
+		/// <param name="Connection">Connection to the Perforce server</param>
 		/// <param name="ChangeNumber">Changelist containing shelved files to be deleted</param>
 		/// <param name="FileSpecs">Files to delete</param>
 		/// <param name="CancellationToken">Token used to cancel the operation</param>
 		/// <returns>Response from the server</returns>
-		public Task<List<PerforceResponse>> TryDeleteShelvedFilesAsync(int ChangeNumber, string[] FileSpecs, CancellationToken CancellationToken)
+		public static Task<List<PerforceResponse>> TryDeleteShelvedFilesAsync(this PerforceConnection Connection, int ChangeNumber, string[] FileSpecs, CancellationToken CancellationToken)
 		{
 			StringBuilder Arguments = new StringBuilder("shelve -d");
 			if(ChangeNumber != -1)
@@ -2453,7 +2557,7 @@ namespace EpicGames.Perforce
 				Arguments.AppendFormat(" \"{0}\"", FileSpec);
 			}
 
-			return CommandAsync(Arguments.ToString(), null, null, CancellationToken);
+			return Connection.CommandAsync(Arguments.ToString(), null, null, CancellationToken);
 		}
 
 		#endregion
@@ -2463,49 +2567,53 @@ namespace EpicGames.Perforce
 		/// <summary>
 		/// Enumerates all streams in a depot
 		/// </summary>
+		/// <param name="Connection">Connection to the Perforce server</param>
 		/// <param name="StreamPath">The path for streams to enumerate (eg. "//UE4/...")</param>
 		/// <param name="CancellationToken">Token used to cancel the operation</param>
 		/// <returns>List of streams matching the given criteria</returns>
-		public async Task<List<StreamsRecord>> GetStreamsAsync(string StreamPath, CancellationToken CancellationToken)
+		public static async Task<List<StreamsRecord>> GetStreamsAsync(this PerforceConnection Connection, string StreamPath, CancellationToken CancellationToken)
 		{
-			return (await TryGetStreamsAsync(StreamPath, CancellationToken)).Data;
+			return (await TryGetStreamsAsync(Connection, StreamPath, CancellationToken)).Data;
 		}
 
 		/// <summary>
 		/// Enumerates all streams in a depot
 		/// </summary>
+		/// <param name="Connection">Connection to the Perforce server</param>
 		/// <param name="StreamPath">The path for streams to enumerate (eg. "//UE4/...")</param>
 		/// <param name="CancellationToken">Token used to cancel the operation</param>
 		/// <returns>List of streams matching the given criteria</returns>
-		public Task<PerforceResponseList<StreamsRecord>> TryGetStreamsAsync(string StreamPath, CancellationToken CancellationToken)
+		public static Task<PerforceResponseList<StreamsRecord>> TryGetStreamsAsync(this PerforceConnection Connection, string StreamPath, CancellationToken CancellationToken)
 		{
-			return TryGetStreamsAsync(StreamPath, -1, null, false, CancellationToken);
+			return TryGetStreamsAsync(Connection, StreamPath, -1, null, false, CancellationToken);
 		}
 
 		/// <summary>
 		/// Enumerates all streams in a depot
 		/// </summary>
+		/// <param name="Connection">Connection to the Perforce server</param>
 		/// <param name="StreamPath">The path for streams to enumerate (eg. "//UE4/...")</param>
 		/// <param name="MaxResults">Maximum number of results to return</param>
 		/// <param name="Filter">Additional filter to be applied to the results</param>
 		/// <param name="bUnloaded">Whether to enumerate unloaded workspaces</param>
 		/// <param name="CancellationToken">Token used to cancel the operation</param>
 		/// <returns>List of streams matching the given criteria</returns>
-		public async Task<List<StreamsRecord>> GetStreamsAsync(string StreamPath, int MaxResults, string? Filter, bool bUnloaded, CancellationToken CancellationToken)
+		public static async Task<List<StreamsRecord>> GetStreamsAsync(this PerforceConnection Connection, string StreamPath, int MaxResults, string? Filter, bool bUnloaded, CancellationToken CancellationToken)
 		{
-			return (await TryGetStreamsAsync(StreamPath, MaxResults, Filter, bUnloaded, CancellationToken)).Data;
+			return (await TryGetStreamsAsync(Connection, StreamPath, MaxResults, Filter, bUnloaded, CancellationToken)).Data;
 		}
 
 		/// <summary>
 		/// Enumerates all streams in a depot
 		/// </summary>
+		/// <param name="Connection">Connection to the Perforce server</param>
 		/// <param name="StreamPath">The path for streams to enumerate (eg. "//UE4/...")</param>
 		/// <param name="MaxResults">Maximum number of results to return</param>
 		/// <param name="Filter">Additional filter to be applied to the results</param>
 		/// <param name="bUnloaded">Whether to enumerate unloaded workspaces</param>
 		/// <param name="CancellationToken">Token used to cancel the operation</param>
 		/// <returns>List of streams matching the given criteria</returns>
-		public Task<PerforceResponseList<StreamsRecord>> TryGetStreamsAsync(string StreamPath, int MaxResults, string? Filter, bool bUnloaded, CancellationToken CancellationToken)
+		public static Task<PerforceResponseList<StreamsRecord>> TryGetStreamsAsync(this PerforceConnection Connection, string StreamPath, int MaxResults, string? Filter, bool bUnloaded, CancellationToken CancellationToken)
 		{
 			// Build the command line
 			StringBuilder Arguments = new StringBuilder("streams");
@@ -2524,7 +2632,7 @@ namespace EpicGames.Perforce
 			Arguments.AppendFormat(" \"{0}\"", StreamPath);
 
 			// Execute the command
-			return CommandAsync<StreamsRecord>(Arguments.ToString(), null, CancellationToken);
+			return CommandAsync<StreamsRecord>(Connection, Arguments.ToString(), null, CancellationToken);
 		}
 
 		#endregion
@@ -2534,23 +2642,25 @@ namespace EpicGames.Perforce
 		/// <summary>
 		/// Submits a pending changelist
 		/// </summary>
+		/// <param name="Connection">Connection to the Perforce server</param>
 		/// <param name="ChangeNumber">The changelist to submit</param>
 		/// <param name="Options">Options for the command</param>
 		/// <param name="CancellationToken">Token used to cancel the operation</param>
 		/// <returns>Response from the server</returns>
-		public async Task<SubmitRecord> SubmitAsync(int ChangeNumber, SubmitOptions Options, CancellationToken CancellationToken)
+		public static async Task<SubmitRecord> SubmitAsync(this PerforceConnection Connection, int ChangeNumber, SubmitOptions Options, CancellationToken CancellationToken)
 		{
-			return (await TrySubmitAsync(ChangeNumber, Options, CancellationToken)).Data;
+			return (await TrySubmitAsync(Connection, ChangeNumber, Options, CancellationToken)).Data;
 		}
 
 		/// <summary>
 		/// Submits a pending changelist
 		/// </summary>
+		/// <param name="Connection">Connection to the Perforce server</param>
 		/// <param name="ChangeNumber">The changelist to submit</param>
 		/// <param name="Options">Options for the command</param>
 		/// <param name="CancellationToken">Token used to cancel the operation</param>
 		/// <returns>Response from the server</returns>
-		public async Task<PerforceResponse<SubmitRecord>> TrySubmitAsync(int ChangeNumber, SubmitOptions Options, CancellationToken CancellationToken)
+		public static async Task<PerforceResponse<SubmitRecord>> TrySubmitAsync(this PerforceConnection Connection, int ChangeNumber, SubmitOptions Options, CancellationToken CancellationToken)
 		{
 			StringBuilder Arguments = new StringBuilder("submit");
 			if((Options & SubmitOptions.ReopenAsEdit) != 0)
@@ -2559,32 +2669,34 @@ namespace EpicGames.Perforce
 			}
 			Arguments.AppendFormat(" -c {0}", ChangeNumber);
 
-			return (await CommandAsync<SubmitRecord>(Arguments.ToString(), null, CancellationToken))[0];
+			return (await CommandAsync<SubmitRecord>(Connection, Arguments.ToString(), null, CancellationToken))[0];
 		}
 
 		/// <summary>
 		/// Submits a shelved changelist
 		/// </summary>
+		/// <param name="Connection">Connection to the Perforce server</param>
 		/// <param name="ChangeNumber">The changelist to submit</param>
 		/// <param name="CancellationToken">Token used to cancel the operation</param>
 		/// <returns>Response from the server</returns>
-		public async Task<SubmitRecord> SubmitShelvedAsync(int ChangeNumber, CancellationToken CancellationToken)
+		public static async Task<SubmitRecord> SubmitShelvedAsync(this PerforceConnection Connection, int ChangeNumber, CancellationToken CancellationToken)
 		{
-			return (await TrySubmitShelvedAsync(ChangeNumber, CancellationToken)).Data;
+			return (await TrySubmitShelvedAsync(Connection, ChangeNumber, CancellationToken)).Data;
 		}
 
 		/// <summary>
 		/// Submits a pending changelist
 		/// </summary>
+		/// <param name="Connection">Connection to the Perforce server</param>
 		/// <param name="ChangeNumber">The changelist to submit</param>
 		/// <param name="CancellationToken">Token used to cancel the operation</param>
 		/// <returns>Response from the server</returns>
-		public async Task<PerforceResponse<SubmitRecord>> TrySubmitShelvedAsync(int ChangeNumber, CancellationToken CancellationToken)
+		public static async Task<PerforceResponse<SubmitRecord>> TrySubmitShelvedAsync(this PerforceConnection Connection, int ChangeNumber, CancellationToken CancellationToken)
 		{
 			StringBuilder Arguments = new StringBuilder("submit");
 			Arguments.AppendFormat(" -e {0}", ChangeNumber);
 
-			return (await CommandAsync<SubmitRecord>(Arguments.ToString(), null, CancellationToken))[0];
+			return (await CommandAsync<SubmitRecord>(Connection, Arguments.ToString(), null, CancellationToken))[0];
 		}
 
 		#endregion
@@ -2594,54 +2706,59 @@ namespace EpicGames.Perforce
 		/// <summary>
 		/// Syncs files from the server
 		/// </summary>
+		/// <param name="Connection">Connection to the Perforce server</param>
 		/// <param name="FileSpecs">Files to sync</param>
 		/// <param name="CancellationToken">Token used to cancel the operation</param>
 		/// <returns>Response from the server</returns>
-		public async Task<List<SyncRecord>> SyncAsync(string[] FileSpecs, CancellationToken CancellationToken)
+		public static async Task<List<SyncRecord>> SyncAsync(this PerforceConnection Connection, string[] FileSpecs, CancellationToken CancellationToken)
 		{
-			return (await TrySyncAsync(FileSpecs, CancellationToken)).Data;
+			return (await TrySyncAsync(Connection, FileSpecs, CancellationToken)).Data;
 		}
 
 		/// <summary>
 		/// Syncs files from the server
 		/// </summary>
+		/// <param name="Connection">Connection to the Perforce server</param>
 		/// <param name="FileSpecs">Files to sync</param>
 		/// <param name="CancellationToken">Token used to cancel the operation</param>
 		/// <returns>Response from the server</returns>
-		public Task<PerforceResponseList<SyncRecord>> TrySyncAsync(string[] FileSpecs, CancellationToken CancellationToken)
+		public static Task<PerforceResponseList<SyncRecord>> TrySyncAsync(this PerforceConnection Connection, string[] FileSpecs, CancellationToken CancellationToken)
 		{
-			return TrySyncAsync(SyncOptions.None, -1, FileSpecs, CancellationToken);
+			return TrySyncAsync(Connection, SyncOptions.None, -1, FileSpecs, CancellationToken);
 		}
 
 		/// <summary>
 		/// Syncs files from the server
 		/// </summary>
+		/// <param name="Connection">Connection to the Perforce server</param>
 		/// <param name="Options">Options for the command</param>
 		/// <param name="MaxFiles">Syncs only the first number of files specified.</param>
 		/// <param name="FileSpecs">Files to sync</param>
 		/// <param name="CancellationToken">Token used to cancel the operation</param>
 		/// <returns>Response from the server</returns>
-		public async Task<List<SyncRecord>> SyncAsync(SyncOptions Options, int MaxFiles, string[] FileSpecs, CancellationToken CancellationToken)
+		public static async Task<List<SyncRecord>> SyncAsync(this PerforceConnection Connection, SyncOptions Options, int MaxFiles, string[] FileSpecs, CancellationToken CancellationToken)
 		{
-			return (await TrySyncAsync(Options, MaxFiles, FileSpecs, CancellationToken)).Data;
+			return (await TrySyncAsync(Connection, Options, MaxFiles, FileSpecs, CancellationToken)).Data;
 		}
 
 		/// <summary>
 		/// Syncs files from the server
 		/// </summary>
+		/// <param name="Connection">Connection to the Perforce server</param>
 		/// <param name="Options">Options for the command</param>
 		/// <param name="MaxFiles">Syncs only the first number of files specified.</param>
 		/// <param name="FileSpecs">Files to sync</param>
 		/// <param name="CancellationToken">Token used to cancel the operation</param>
 		/// <returns>Response from the server</returns>
-		public Task<PerforceResponseList<SyncRecord>> TrySyncAsync(SyncOptions Options, int MaxFiles, string[] FileSpecs, CancellationToken CancellationToken)
+		public static Task<PerforceResponseList<SyncRecord>> TrySyncAsync(this PerforceConnection Connection, SyncOptions Options, int MaxFiles, string[] FileSpecs, CancellationToken CancellationToken)
 		{
-			return TrySyncAsync(Options, MaxFiles, -1, -1, -1, -1, -1, FileSpecs, CancellationToken);
+			return TrySyncAsync(Connection, Options, MaxFiles, -1, -1, -1, -1, -1, FileSpecs, CancellationToken);
 		}
 
 		/// <summary>
 		/// Syncs files from the server
 		/// </summary>
+		/// <param name="Connection">Connection to the Perforce server</param>
 		/// <param name="Options">Options for the command</param>
 		/// <param name="MaxFiles">Syncs only the first number of files specified</param>
 		/// <param name="NumThreads">Sync in parallel using the given number of threads</param>
@@ -2652,14 +2769,15 @@ namespace EpicGames.Perforce
 		/// <param name="FileSpecs">Files to sync</param>
 		/// <param name="CancellationToken">Token used to cancel the operation</param>
 		/// <returns>Response from the server</returns>
-		public async Task<List<SyncRecord>> SyncAsync(SyncOptions Options, int MaxFiles, int NumThreads, int Batch, int BatchSize, int Min, int MinSize, string[] FileSpecs, CancellationToken CancellationToken)
+		public static async Task<List<SyncRecord>> SyncAsync(this PerforceConnection Connection, SyncOptions Options, int MaxFiles, int NumThreads, int Batch, int BatchSize, int Min, int MinSize, string[] FileSpecs, CancellationToken CancellationToken)
 		{
-			return (await TrySyncAsync(Options, MaxFiles, NumThreads, Batch, BatchSize, Min, MinSize, FileSpecs, CancellationToken)).Data;
+			return (await TrySyncAsync(Connection, Options, MaxFiles, NumThreads, Batch, BatchSize, Min, MinSize, FileSpecs, CancellationToken)).Data;
 		}
 
 		/// <summary>
 		/// Syncs files from the server
 		/// </summary>
+		/// <param name="Connection">Connection to the Perforce server</param>
 		/// <param name="Options">Options for the command</param>
 		/// <param name="MaxFiles">Syncs only the first number of files specified</param>
 		/// <param name="NumThreads">Sync in parallel using the given number of threads</param>
@@ -2670,10 +2788,10 @@ namespace EpicGames.Perforce
 		/// <param name="FileSpecs">Files to sync</param>
 		/// <param name="CancellationToken">Token used to cancel the operation</param>
 		/// <returns>Response from the server</returns>
-		public async Task<PerforceResponseList<SyncRecord>> TrySyncAsync(SyncOptions Options, int MaxFiles, int NumThreads, int Batch, int BatchSize, int Min, int MinSize, string[] FileSpecs, CancellationToken CancellationToken)
+		public static async Task<PerforceResponseList<SyncRecord>> TrySyncAsync(this PerforceConnection Connection, SyncOptions Options, int MaxFiles, int NumThreads, int Batch, int BatchSize, int Min, int MinSize, string[] FileSpecs, CancellationToken CancellationToken)
 		{
 			// Perforce annoyingly returns 'up-to-date' as an error. Ignore it.
-			PerforceResponseList<SyncRecord> Records = await SyncInternalAsync<SyncRecord>(Options, MaxFiles, NumThreads, Batch, BatchSize, Min, MinSize, FileSpecs, false, CancellationToken);
+			PerforceResponseList<SyncRecord> Records = await SyncInternalAsync<SyncRecord>(Connection, Options, MaxFiles, NumThreads, Batch, BatchSize, Min, MinSize, FileSpecs, false, CancellationToken);
 			Records.RemoveAll(x => x.Error != null && x.Error.Generic == PerforceGenericCode.Empty);
 			return Records;
 		}
@@ -2681,14 +2799,15 @@ namespace EpicGames.Perforce
 		/// <summary>
 		/// Syncs files from the server
 		/// </summary>
+		/// <param name="Connection">Connection to the Perforce server</param>
 		/// <param name="Options">Options for the command</param>
 		/// <param name="MaxFiles">Syncs only the first number of files specified.</param>
 		/// <param name="FileSpecs">Files to sync</param>
 		/// <param name="CancellationToken">Token used to cancel the operation</param>
 		/// <returns>Response from the server</returns>
-		public async Task<List<SyncSummaryRecord>> SyncQuietAsync(SyncOptions Options, int MaxFiles, string[] FileSpecs, CancellationToken CancellationToken)
+		public static async Task<List<SyncSummaryRecord>> SyncQuietAsync(this PerforceConnection Connection, SyncOptions Options, int MaxFiles, string[] FileSpecs, CancellationToken CancellationToken)
 		{
-			PerforceResponseList<SyncSummaryRecord> Records = await TrySyncQuietAsync(Options, MaxFiles, FileSpecs, CancellationToken);
+			PerforceResponseList<SyncSummaryRecord> Records = await TrySyncQuietAsync(Connection, Options, MaxFiles, FileSpecs, CancellationToken);
 			Records.RemoveAll(x => (x.Error != null && x.Error.Generic == PerforceGenericCode.Empty) || x.Info != null);
 			return Records.Data;
 		}
@@ -2696,19 +2815,21 @@ namespace EpicGames.Perforce
 		/// <summary>
 		/// Syncs files from the server
 		/// </summary>
+		/// <param name="Connection">Connection to the Perforce server</param>
 		/// <param name="Options">Options for the command</param>
 		/// <param name="MaxFiles">Syncs only the first number of files specified.</param>
 		/// <param name="FileSpecs">Files to sync</param>
 		/// <param name="CancellationToken">Token used to cancel the operation</param>
 		/// <returns>Response from the server</returns>
-		public Task<PerforceResponseList<SyncSummaryRecord>> TrySyncQuietAsync(SyncOptions Options, int MaxFiles, string[] FileSpecs, CancellationToken CancellationToken)
+		public static Task<PerforceResponseList<SyncSummaryRecord>> TrySyncQuietAsync(this PerforceConnection Connection, SyncOptions Options, int MaxFiles, string[] FileSpecs, CancellationToken CancellationToken)
 		{
-			return TrySyncQuietAsync(Options, MaxFiles, -1, -1, -1, -1, -1, FileSpecs, null, CancellationToken);
+			return TrySyncQuietAsync(Connection, Options, MaxFiles, -1, -1, -1, -1, -1, FileSpecs, null, CancellationToken);
 		}
 
 		/// <summary>
 		/// Syncs files from the server without returning detailed file info
 		/// </summary>
+		/// <param name="Connection">Connection to the Perforce server</param>
 		/// <param name="Options">Options for the command</param>
 		/// <param name="MaxFiles">Syncs only the first number of files specified</param>
 		/// <param name="NumThreads">Sync in parallel using the given number of threads</param>
@@ -2720,14 +2841,15 @@ namespace EpicGames.Perforce
 		/// <param name="ResponseFile">Response file for list of files to sync</param>
 		/// <param name="CancellationToken">Token used to cancel the operation</param>
 		/// <returns>Response from the server</returns>
-		public async Task<List<SyncSummaryRecord>> SyncQuietAsync(SyncOptions Options, int MaxFiles, int NumThreads, int Batch, int BatchSize, int Min, int MinSize, string[] FileSpecs, FileReference? ResponseFile, CancellationToken CancellationToken)
+		public static async Task<List<SyncSummaryRecord>> SyncQuietAsync(this PerforceConnection Connection, SyncOptions Options, int MaxFiles, int NumThreads, int Batch, int BatchSize, int Min, int MinSize, string[] FileSpecs, FileReference? ResponseFile, CancellationToken CancellationToken)
 		{
-			return (await TrySyncQuietAsync(Options, MaxFiles, NumThreads, Batch, BatchSize, Min, MinSize, FileSpecs, ResponseFile, CancellationToken)).Data;
+			return (await TrySyncQuietAsync(Connection, Options, MaxFiles, NumThreads, Batch, BatchSize, Min, MinSize, FileSpecs, ResponseFile, CancellationToken)).Data;
 		}
 
 		/// <summary>
 		/// Syncs files from the server without returning detailed file info
 		/// </summary>
+		/// <param name="Connection">Connection to the Perforce server</param>
 		/// <param name="Options">Options for the command</param>
 		/// <param name="MaxFiles">Syncs only the first number of files specified</param>
 		/// <param name="NumThreads">Sync in parallel using the given number of threads</param>
@@ -2739,14 +2861,15 @@ namespace EpicGames.Perforce
 		/// <param name="ResponseFile">Response file for list of files to sync</param>
 		/// <param name="CancellationToken">Token used to cancel the operation</param>
 		/// <returns>Response from the server</returns>
-		public Task<PerforceResponseList<SyncSummaryRecord>> TrySyncQuietAsync(SyncOptions Options, int MaxFiles, int NumThreads, int Batch, int BatchSize, int Min, int MinSize, string[] FileSpecs, FileReference? ResponseFile, CancellationToken CancellationToken)
+		public static Task<PerforceResponseList<SyncSummaryRecord>> TrySyncQuietAsync(this PerforceConnection Connection, SyncOptions Options, int MaxFiles, int NumThreads, int Batch, int BatchSize, int Min, int MinSize, string[] FileSpecs, FileReference? ResponseFile, CancellationToken CancellationToken)
 		{
-			return SyncInternalAsync<SyncSummaryRecord>(Options, MaxFiles, NumThreads, Batch, BatchSize, Min, MinSize, FileSpecs, true, CancellationToken);
+			return SyncInternalAsync<SyncSummaryRecord>(Connection, Options, MaxFiles, NumThreads, Batch, BatchSize, Min, MinSize, FileSpecs, true, CancellationToken);
 		}
 
 		/// <summary>
 		/// Gets arguments for a sync command
 		/// </summary>
+		/// <param name="Connection">Connection to the Perforce server</param>
 		/// <param name="Options">Options for the command</param>
 		/// <param name="MaxFiles">Syncs only the first number of files specified</param>
 		/// <param name="NumThreads">Sync in parallel using the given number of threads</param>
@@ -2758,7 +2881,7 @@ namespace EpicGames.Perforce
 		/// <param name="bQuiet">Whether to use quiet output</param>
 		/// <param name="CancellationToken">Token used to cancel the operation</param>
 		/// <returns>Arguments for the command</returns>
-		private async Task<PerforceResponseList<T>> SyncInternalAsync<T>(SyncOptions Options, int MaxFiles, int NumThreads, int Batch, int BatchSize, int Min, int MinSize, string[] FileSpecs, bool bQuiet, CancellationToken CancellationToken) where T : class
+		private static async Task<PerforceResponseList<T>> SyncInternalAsync<T>(this PerforceConnection Connection, SyncOptions Options, int MaxFiles, int NumThreads, int Batch, int BatchSize, int Min, int MinSize, string[] FileSpecs, bool bQuiet, CancellationToken CancellationToken) where T : class
 		{
 			StringBuilder Arguments = new StringBuilder("sync");
 			if((Options & SyncOptions.Force) != 0)
@@ -2827,7 +2950,7 @@ namespace EpicGames.Perforce
 				{
 					Arguments.AppendFormat(" \"{0}\"", FileSpec);
 				}
-				return await CommandAsync<T>(Arguments.ToString(), null, CancellationToken);
+				return await CommandAsync<T>(Connection, Arguments.ToString(), null, CancellationToken);
 			}
 			else
 			{
@@ -2836,7 +2959,7 @@ namespace EpicGames.Perforce
 				{
 					await File.WriteAllLinesAsync(TempFileName, FileSpecs);
 					Arguments.Insert(0, $"-x\"{TempFileName}\" ");
-					return await CommandAsync<T>(Arguments.ToString(), null, CancellationToken);
+					return await CommandAsync<T>(Connection, Arguments.ToString(), null, CancellationToken);
 				}
 				finally
 				{
@@ -2852,6 +2975,7 @@ namespace EpicGames.Perforce
 		/// <summary>
 		/// Restore shelved files from a pending change into a workspace
 		/// </summary>
+		/// <param name="Connection">Connection to the Perforce server</param>
 		/// <param name="ChangeNumber">The changelist containing shelved files</param>
 		/// <param name="IntoChangeNumber">The changelist to receive the unshelved files</param>
 		/// <param name="UsingBranchSpec">The branchspec to use when unshelving files</param>
@@ -2861,14 +2985,15 @@ namespace EpicGames.Perforce
 		/// <param name="FileSpecs">Files to unshelve</param>
 		/// <param name="CancellationToken">Token used to cancel the operation</param>
 		/// <returns>Response from the server</returns>
-		public async Task<List<UnshelveRecord>> UnshelveAsync(int ChangeNumber, int IntoChangeNumber, string UsingBranchSpec, string UsingStream, string ForceParentStream, UnshelveOptions Options, string[] FileSpecs, CancellationToken CancellationToken)
+		public static async Task<List<UnshelveRecord>> UnshelveAsync(this PerforceConnection Connection, int ChangeNumber, int IntoChangeNumber, string UsingBranchSpec, string UsingStream, string ForceParentStream, UnshelveOptions Options, string[] FileSpecs, CancellationToken CancellationToken)
 		{
-			return (await TryUnshelveAsync(ChangeNumber, IntoChangeNumber, UsingBranchSpec, UsingStream, ForceParentStream, Options, FileSpecs, CancellationToken)).Data;
+			return (await TryUnshelveAsync(Connection, ChangeNumber, IntoChangeNumber, UsingBranchSpec, UsingStream, ForceParentStream, Options, FileSpecs, CancellationToken)).Data;
 		}
 
 		/// <summary>
 		/// Restore shelved files from a pending change into a workspace
 		/// </summary>
+		/// <param name="Connection">Connection to the Perforce server</param>
 		/// <param name="ChangeNumber">The changelist containing shelved files</param>
 		/// <param name="IntoChangeNumber">The changelist to receive the unshelved files</param>
 		/// <param name="UsingBranchSpec">The branchspec to use when unshelving files</param>
@@ -2878,7 +3003,7 @@ namespace EpicGames.Perforce
 		/// <param name="FileSpecs">Files to unshelve</param>
 		/// <param name="CancellationToken">Token used to cancel the operation</param>
 		/// <returns>Response from the server</returns>
-		public Task<PerforceResponseList<UnshelveRecord>> TryUnshelveAsync(int ChangeNumber, int IntoChangeNumber, string UsingBranchSpec, string UsingStream, string ForceParentStream, UnshelveOptions Options, string[] FileSpecs, CancellationToken CancellationToken)
+		public static Task<PerforceResponseList<UnshelveRecord>> TryUnshelveAsync(this PerforceConnection Connection, int ChangeNumber, int IntoChangeNumber, string UsingBranchSpec, string UsingStream, string ForceParentStream, UnshelveOptions Options, string[] FileSpecs, CancellationToken CancellationToken)
 		{
 			StringBuilder Arguments = new StringBuilder("unshelve");
 			Arguments.AppendFormat(" -s {0}", ChangeNumber);
@@ -2911,7 +3036,7 @@ namespace EpicGames.Perforce
 				Arguments.AppendFormat(" \"{0}\"", FileSpec);
 			}
 
-			return CommandAsync<UnshelveRecord>(Arguments.ToString(), null, CancellationToken);
+			return CommandAsync<UnshelveRecord>(Connection, Arguments.ToString(), null, CancellationToken);
 		}
 
 		#endregion
@@ -2921,25 +3046,27 @@ namespace EpicGames.Perforce
 		/// <summary>
 		/// Enumerates all streams in a depot
 		/// </summary>
+		/// <param name="Connection">Connection to the Perforce server</param>
 		/// <param name="UserName">Name of the user to fetch information for</param>
 		/// <param name="CancellationToken">Token used to cancel the operation</param>
 		/// <returns>Response from the server</returns>
-		public async Task<UserRecord> GetUserAsync(string UserName, CancellationToken CancellationToken)
+		public static async Task<UserRecord> GetUserAsync(this PerforceConnection Connection, string UserName, CancellationToken CancellationToken)
 		{
-			return (await TryGetUserAsync(UserName, CancellationToken))[0].Data;
+			return (await TryGetUserAsync(Connection, UserName, CancellationToken))[0].Data;
 		}
 
 		/// <summary>
 		/// Enumerates all streams in a depot
 		/// </summary>
+		/// <param name="Connection">Connection to the Perforce server</param>
 		/// <param name="UserName">Name of the user to fetch information for</param>
 		/// <param name="CancellationToken">Token used to cancel the operation</param>
 		/// <returns>Response from the server</returns>
-		public Task<PerforceResponseList<UserRecord>> TryGetUserAsync(string UserName, CancellationToken CancellationToken)
+		public static Task<PerforceResponseList<UserRecord>> TryGetUserAsync(this PerforceConnection Connection, string UserName, CancellationToken CancellationToken)
 		{
 			StringBuilder Arguments = new StringBuilder("user");
 			Arguments.AppendFormat(" -o \"{0}\"", UserName);
-			return CommandAsync<UserRecord>(Arguments.ToString(), null, CancellationToken);
+			return CommandAsync<UserRecord>(Connection, Arguments.ToString(), null, CancellationToken);
 		}
 
 		#endregion
@@ -2949,12 +3076,13 @@ namespace EpicGames.Perforce
 		/// <summary>
 		/// Retrieves the location of a file of set of files in the workspace
 		/// </summary>
+		/// <param name="Connection">Connection to the Perforce server</param>
 		/// <param name="FileSpecs">Patterns for the files to query</param>
 		/// <param name="CancellationToken">Token used to cancel the operation</param>
 		/// <returns>List of responses from the server</returns>
-		public async Task<List<WhereRecord>> WhereAsync(string[] FileSpecs, CancellationToken CancellationToken)
+		public static async Task<List<WhereRecord>> WhereAsync(this PerforceConnection Connection, string[] FileSpecs, CancellationToken CancellationToken)
 		{
-			List<WhereRecord> Records = (await TryWhereAsync(FileSpecs, CancellationToken)).Data;
+			List<WhereRecord> Records = (await TryWhereAsync(Connection, FileSpecs, CancellationToken)).Data;
 			if (Records.Count != FileSpecs.Length)
 			{
 				throw new PerforceException($"Unexpected response count; expected {FileSpecs.Length}, got {Records.Count}.");
@@ -2965,10 +3093,11 @@ namespace EpicGames.Perforce
 		/// <summary>
 		/// Retrieves the location of a file of set of files in the workspace
 		/// </summary>
+		/// <param name="Connection">Connection to the Perforce server</param>
 		/// <param name="FileSpecs">Patterns for the files to query</param>
 		/// <param name="CancellationToken">Token used to cancel the operation</param>
 		/// <returns>List of responses from the server</returns>
-		public async Task<PerforceResponseList<WhereRecord>> TryWhereAsync(string[] FileSpecs, CancellationToken CancellationToken)
+		public static async Task<PerforceResponseList<WhereRecord>> TryWhereAsync(this PerforceConnection Connection, string[] FileSpecs, CancellationToken CancellationToken)
 		{
 			PerforceResponseList<WhereRecord> Responses = new PerforceResponseList<WhereRecord>();
 			for (int Idx = 0; Idx < FileSpecs.Length; )
@@ -2978,7 +3107,7 @@ namespace EpicGames.Perforce
 				{
 					Arguments.AppendFormat(" \"{0}\"", FileSpecs[Idx]);
 				}
-				Responses.AddRange(await CommandAsync<WhereRecord>(Arguments.ToString(), null, CancellationToken));
+				Responses.AddRange(await CommandAsync<WhereRecord>(Connection, Arguments.ToString(), null, CancellationToken));
 			}
 			return Responses;
 		}
@@ -2986,32 +3115,35 @@ namespace EpicGames.Perforce
 		#endregion
 
 		#region p4 undo
+
 		/// <summary>
-		/// perform undo on a changelist (p4 undo -c <targetCL> //...@undoCL)
+		/// perform undo on a changelist (p4 undo -c [targetCL] //...@undoCL)
 		/// </summary>
+		/// <param name="Connection">Connection to the Perforce server</param>
 		/// <param name="ChangeNumberToUndo">Changelist number to undo</param>
 		/// <param name="ChangeNumber">Changelist number to receive the changes</param>
 		/// <param name="CancellationToken">Token used to cancel the operation</param>
 		/// <returns>Response from the server</returns>
-		public async Task UndoChangeAsync(int ChangeNumberToUndo, int ChangeNumber, CancellationToken CancellationToken)
+		public static async Task UndoChangeAsync(this PerforceConnection Connection, int ChangeNumberToUndo, int ChangeNumber, CancellationToken CancellationToken)
 		{
-			(await TryUndoChangeAsync(ChangeNumberToUndo, ChangeNumber, CancellationToken))[0].EnsureSuccess();
+			(await TryUndoChangeAsync(Connection, ChangeNumberToUndo, ChangeNumber, CancellationToken))[0].EnsureSuccess();
 		}
 
 		/// <summary>
-		/// perform undo on a changelist (p4 undo -c <targetCL> //...@undoCL)
+		/// perform undo on a changelist (p4 undo -c [targetCL] //...@undoCL)
 		/// </summary>
+		/// <param name="Connection">Connection to the Perforce server</param>
 		/// <param name="ChangeNumberToUndo">Changelist number to undo</param>
 		/// <param name="ChangeNumber">Changelist number to receive the changes</param>
 		/// <param name="CancellationToken">Token used to cancel the operation</param>
 		/// <returns>Response from the server</returns>
-		public Task<PerforceResponseList<UndoRecord>> TryUndoChangeAsync(int ChangeNumberToUndo, int ChangeNumber, CancellationToken CancellationToken)
+		public static Task<PerforceResponseList<UndoRecord>> TryUndoChangeAsync(this PerforceConnection Connection, int ChangeNumberToUndo, int ChangeNumber, CancellationToken CancellationToken)
 		{
 			StringBuilder Arguments = new StringBuilder("undo -c");
 			
 			Arguments.AppendFormat(" {0} //...@{1}", ChangeNumber, ChangeNumberToUndo);
 
-			return CommandAsync<UndoRecord>(Arguments.ToString(), null, CancellationToken);
+			return CommandAsync<UndoRecord>(Connection, Arguments.ToString(), null, CancellationToken);
 		}
 		#endregion
 
@@ -3020,23 +3152,25 @@ namespace EpicGames.Perforce
 		/// <summary>
 		/// Runs the annotate p4 command
 		/// </summary>
+		/// <param name="Connection">Connection to the Perforce server</param>
 		/// <param name="FileSpec">Depot path to the file</param>
 		/// <param name="Options">Options for the anotate command</param>
 		/// <param name="CancellationToken">Token used to cancel the operation</param>
 		/// <returns>List of annotate records</returns>
-		public async Task<List<AnnotateRecord>> AnnotateAsync(string FileSpec, AnnotateOptions Options, CancellationToken CancellationToken)
+		public static async Task<List<AnnotateRecord>> AnnotateAsync(this PerforceConnection Connection, string FileSpec, AnnotateOptions Options, CancellationToken CancellationToken)
 		{
-			return (await TryAnnotateAsync(FileSpec, Options, CancellationToken)).Data;
+			return (await TryAnnotateAsync(Connection, FileSpec, Options, CancellationToken)).Data;
 		}
 
 		/// <summary>
 		/// Runs the annotate p4 command 
 		/// </summary>
+		/// <param name="Connection">Connection to the Perforce server</param>
 		/// <param name="FileSpec">Depot path to the file</param>
 		/// <param name="Options">Options for the anotate command</param>
 		/// <param name="CancellationToken">Token used to cancel the operation</param>
 		/// <returns>List of annotate records</returns>
-		public Task<PerforceResponseList<AnnotateRecord>> TryAnnotateAsync(string FileSpec, AnnotateOptions Options, CancellationToken CancellationToken)
+		public static Task<PerforceResponseList<AnnotateRecord>> TryAnnotateAsync(this PerforceConnection Connection, string FileSpec, AnnotateOptions Options, CancellationToken CancellationToken)
 		{
 			PerforceResponseList<AnnotateRecord> Responses = new PerforceResponseList<AnnotateRecord>();
 			StringBuilder Arguments = new StringBuilder("annotate ");
@@ -3059,7 +3193,7 @@ namespace EpicGames.Perforce
 
 			Arguments.AppendFormat(" {0}", FileSpec);
 
-			return CommandAsync<AnnotateRecord>(Arguments.ToString(), null, CancellationToken);
+			return CommandAsync<AnnotateRecord>(Connection, Arguments.ToString(), null, CancellationToken);
 		}
 
 		#endregion
