Fix xmldoc markup for EpicGames.Perforce.

diff --git a/Engine/Source/Programs/Shared/EpicGames.Perforce/EpicGames.Perforce.csproj b/Engine/Source/Programs/Shared/EpicGames.Perforce/EpicGames.Perforce.csproj
--- a/Engine/Source/Programs/Shared/EpicGames.Perforce/EpicGames.Perforce.csproj
+++ b/Engine/Source/Programs/Shared/EpicGames.Perforce/EpicGames.Perforce.csproj
@@ -11,7 +11,7 @@
     <nullable>enable</nullable>
   </PropertyGroup>
   <PropertyGroup Condition=" '$(Configuration)|$(Platform)' == 'Debug|AnyCPU' ">
-    <DocumentationFile></DocumentationFile>
+    <DocumentationFile>EpicGames.Perforce.xml</DocumentationFile>
   </PropertyGroup>
   <PropertyGroup Condition="'$(Configuration)|$(Platform)' == 'Development|AnyCPU'">
     <DocumentationFile></DocumentationFile>
diff --git a/Engine/Source/Programs/Shared/EpicGames.Perforce/PerforceChildProcess.cs b/Engine/Source/Programs/Shared/EpicGames.Perforce/PerforceChildProcess.cs
--- a/Engine/Source/Programs/Shared/EpicGames.Perforce/PerforceChildProcess.cs
+++ b/Engine/Source/Programs/Shared/EpicGames.Perforce/PerforceChildProcess.cs
@@ -178,9 +178,9 @@ namespace EpicGames.Perforce
 		/// <summary>
 		/// Formats the current contents of the buffer to a string
 		/// </summary>
-		/// <param name="Dat">The next byte that was read</param>
+		/// <param name="Data">The next byte that was read</param>
 		/// <returns>String representation of the buffer</returns>
-		private static string FormatDataAsString(Span<byte> Data)
+		private static string FormatDataAsString(ReadOnlySpan<byte> Data)
 		{
 			StringBuilder Result = new StringBuilder();
 			if (Data.Length > 0)
@@ -227,7 +227,7 @@ namespace EpicGames.Perforce
 		/// <summary>
 		/// Formats the current contents of the buffer to a string
 		/// </summary>
-		/// <param name="NextByte">The next byte that was read</param>
+		/// <param name="BufferIdx">The next byte that was read</param>
 		/// <returns>String representation of the buffer</returns>
 		private string FormatDataAsHexDump(int BufferIdx)
 		{
@@ -660,6 +660,7 @@ namespace EpicGames.Perforce
 		/// </summary>
 		/// <param name="RequiredSuffix">The required suffix for any subobject arrays.</param>
 		/// <param name="RecordInfo">Reflection information for the type being serialized into.</param>
+		/// <param name="Record">Receives the deserialized record on success</param>
 		/// <returns>The parsed object.</returns>
 		bool TryReadTypedRecord(Utf8String RequiredSuffix, CachedRecordInfo RecordInfo, [NotNullWhen(true)] out object? Record)
 		{
diff --git a/Engine/Source/Programs/Shared/EpicGames.Perforce/PerforceReflection.cs b/Engine/Source/Programs/Shared/EpicGames.Perforce/PerforceReflection.cs
--- a/Engine/Source/Programs/Shared/EpicGames.Perforce/PerforceReflection.cs
+++ b/Engine/Source/Programs/Shared/EpicGames.Perforce/PerforceReflection.cs
@@ -228,7 +228,7 @@ namespace EpicGames.Perforce
 		/// <summary>
 		/// Parses the given text as an enum
 		/// </summary>
-		/// <param name="Text"></param>
+		/// <param name="Name"></param>
 		/// <returns></returns>
 		public int ParseToInteger(Utf8String Name)
 		{
diff --git a/Engine/Source/Programs/Shared/EpicGames.Perforce/PerforceUtils.cs b/Engine/Source/Programs/Shared/EpicGames.Perforce/PerforceUtils.cs
--- a/Engine/Source/Programs/Shared/EpicGames.Perforce/PerforceUtils.cs
+++ b/Engine/Source/Programs/Shared/EpicGames.Perforce/PerforceUtils.cs
@@ -10,8 +10,14 @@ using System.Threading.Tasks.Dataflow;
 
 namespace EpicGames.Perforce
 {
+	/// <summary>
+	/// Utility methods for dealing with Perforce paths
+	/// </summary>
 	public static class PerforceUtils
 	{
+		/// <summary>
+		/// Escape a path to Perforce syntax
+		/// </summary>
 		static public string EscapePath(string Path)
 		{
 			string NewPath = Path;
@@ -22,6 +28,9 @@ namespace EpicGames.Perforce
 			return NewPath;
 		}
 
+		/// <summary>
+		/// Remove escape characters from a path
+		/// </summary>
 		static public string UnescapePath(string Path)
 		{
 			string NewPath = Path;
@@ -33,6 +42,9 @@ namespace EpicGames.Perforce
 			return NewPath;
 		}
 
+		/// <summary>
+		/// Remove escape characters from a UTF8 path
+		/// </summary>
 		static public Utf8String UnescapePath(Utf8String Path)
 		{
 			ReadOnlySpan<byte> PathSpan = Path.Span;
diff --git a/Engine/Source/Programs/Shared/EpicGames.Perforce/Records/AnnotateLineRecord.cs b/Engine/Source/Programs/Shared/EpicGames.Perforce/Records/AnnotateLineRecord.cs
--- a/Engine/Source/Programs/Shared/EpicGames.Perforce/Records/AnnotateLineRecord.cs
+++ b/Engine/Source/Programs/Shared/EpicGames.Perforce/Records/AnnotateLineRecord.cs
@@ -4,8 +4,14 @@ using System;
 
 namespace EpicGames.Perforce
 {
+	/// <summary>
+	/// 
+	/// </summary>
 	public class AnnotateLineRecord
 	{
+		/// <summary>
+		/// 
+		/// </summary>
 		[PerforceTag("depotFile")]
 		public readonly string DepotFile;
 
diff --git a/Engine/Source/Programs/Shared/EpicGames.Perforce/Records/AnnotateOptions.cs b/Engine/Source/Programs/Shared/EpicGames.Perforce/Records/AnnotateOptions.cs
--- a/Engine/Source/Programs/Shared/EpicGames.Perforce/Records/AnnotateOptions.cs
+++ b/Engine/Source/Programs/Shared/EpicGames.Perforce/Records/AnnotateOptions.cs
@@ -4,6 +4,9 @@ using System;
 
 namespace EpicGames.Perforce
 {
+	/// <summary>
+	/// 
+	/// </summary>
 	[Flags]
 	public enum AnnotateOptions
 	{
diff --git a/Engine/Source/Programs/Shared/EpicGames.Perforce/Records/AnnotateRecord.cs b/Engine/Source/Programs/Shared/EpicGames.Perforce/Records/AnnotateRecord.cs
--- a/Engine/Source/Programs/Shared/EpicGames.Perforce/Records/AnnotateRecord.cs
+++ b/Engine/Source/Programs/Shared/EpicGames.Perforce/Records/AnnotateRecord.cs
@@ -5,6 +5,9 @@ using System.Collections.Generic;
 
 namespace EpicGames.Perforce
 {
+	/// <summary>
+	/// 
+	/// </summary>
 	public class AnnotateRecord
 	{
 		/// <summary>
@@ -60,7 +63,10 @@ namespace EpicGames.Perforce
 		/// </summary>
 		[PerforceTag("data", Optional = true)]
 		public readonly string Data;
-
+		
+		/// <summary>
+		/// 
+		/// </summary>
 		[PerforceRecordList]
 		public List<AnnotateLineRecord> LineRecord = new List<AnnotateLineRecord>();
 
diff --git a/Engine/Source/Programs/Shared/EpicGames.Perforce/Records/FileType.cs b/Engine/Source/Programs/Shared/EpicGames.Perforce/Records/FileType.cs
--- a/Engine/Source/Programs/Shared/EpicGames.Perforce/Records/FileType.cs
+++ b/Engine/Source/Programs/Shared/EpicGames.Perforce/Records/FileType.cs
@@ -257,7 +257,7 @@ namespace EpicGames.Perforce.Managed
 		/// Try to parse modifiers from a utf8 string
 		/// </summary>
 		/// <param name="Text"></param>
-		/// <param name="Modifiers"></param>
+		/// <param name="Result"></param>
 		/// <returns></returns>
 		public static bool TryParseModifiers(ReadOnlySpan<byte> Text, out (FileTypeModifiers, int) Result)
 		{
