Introduce an interface for Perforce connections in C#.

diff --git a/Engine/Source/Programs/Shared/EpicGames.Perforce/PerforceConnection.cs b/Engine/Source/Programs/Shared/EpicGames.Perforce/PerforceConnection.cs
--- a/Engine/Source/Programs/Shared/EpicGames.Perforce/PerforceConnection.cs
+++ b/Engine/Source/Programs/Shared/EpicGames.Perforce/PerforceConnection.cs
@@ -16,10 +16,66 @@ using System.Threading.Tasks;
 
 namespace EpicGames.Perforce
 {
+	/// <summary>
+	/// Base interface for Perforce clients
+	/// </summary>
+	public interface IPerforceConnection
+	{
+		/// <summary>
+		/// Logger for this connection
+		/// </summary>
+		ILogger Logger { get; }
+
+		/// <summary>
+		/// Execute a command and parse the response
+		/// </summary>
+		/// <param name="CommandLine">Arguments for the command</param>
+		/// <param name="InputData">Input data to pass to Perforce</param>
+		/// <param name="StatRecordType">The type of records to return for "stat" responses</param>
+		/// <param name="CancellationToken">Token used to cancel the operation</param>
+		/// <returns>List of objects returned by the server</returns>
+		Task<List<PerforceResponse>> CommandAsync(string CommandLine, byte[]? InputData, Type? StatRecordType, CancellationToken CancellationToken);
+
+		/// <summary>
+		/// Execute a command and parse the response
+		/// </summary>
+		/// <param name="CommandLine">Arguments for the command</param>
+		/// <param name="InputData">Input data to pass to Perforce</param>
+		/// <param name="HandleRecord">Delegate used to handle each record</param>
+		/// <param name="CancellationToken">Token used to cancel the operation</param>
+		/// <returns>List of objects returned by the server</returns>
+		Task RecordCommandAsync(string CommandLine, byte[]? InputData, Action<PerforceRecord> HandleRecord, CancellationToken CancellationToken);
+
+		/// <summary>
+		/// Execute the 'login' command
+		/// </summary>
+		/// <param name="Password">Password to use to login</param>
+		/// <param name="CancellationToken">Token used to cancel the operation</param>
+		/// <returns>Response from the server</returns>
+		Task LoginAsync(string Password, CancellationToken CancellationToken);
+
+		/// <summary>
+		/// Sets an environment variable
+		/// </summary>
+		/// <param name="Name">Name of the variable to set</param>
+		/// <param name="Value">Value for the variable</param>
+		/// <param name="CancellationToken">Token used to cancel the operation</param>
+		/// <returns>Response from the server</returns>
+		Task SetAsync(string Name, string Value, CancellationToken CancellationToken);
+
+		/// <summary>
+		/// Gets the setting of a Perforce variable
+		/// </summary>
+		/// <param name="Name">Name of the variable to get</param>
+		/// <param name="CancellationToken">Cancellation token for the request</param>
+		/// <returns>Value of the variable</returns>
+		Task<string?> TryGetSettingAsync(string Name, CancellationToken CancellationToken);
+	}
+
 	/// <summary>
 	/// Stores settings for communicating with a Perforce server.
 	/// </summary>
-	public class PerforceConnection
+	public class PerforceConnection : IPerforceConnection
 	{
 		#region Plumbing
 
@@ -257,6 +313,76 @@ namespace EpicGames.Perforce
 		}
 
 		#endregion
+
+		#region p4 set
+
+		/// <summary>
+		/// Sets an environment variable
+		/// </summary>
+		/// <param name="Name">Name of the variable to set</param>
+		/// <param name="Value">Value for the variable</param>
+		/// <param name="CancellationToken">Token used to cancel the operation</param>
+		/// <returns>Response from the server</returns>
+		public async Task SetAsync(string Name, string Value, CancellationToken CancellationToken)
+		{
+			Tuple<bool, string> Result = await TrySetAsync(Name, Value, CancellationToken);
+			if (!Result.Item1)
+			{
+				throw new PerforceException(Result.Item2);
+			}
+		}
+
+		/// <summary>
+		/// Sets an environment variable
+		/// </summary>
+		/// <param name="Name">Name of the variable to set</param>
+		/// <param name="Value">Value for the variable</param>
+		/// <param name="CancellationToken">Token used to cancel the operation</param>
+		/// <returns>Response from the server</returns>
+		public async Task<Tuple<bool, string>> TrySetAsync(string Name, string Value, CancellationToken CancellationToken)
+		{
+			StringBuilder Arguments = new StringBuilder("set");
+			Arguments.Append($" {Name}=");
+			Arguments.Append($"\"{Value}\"");
+
+			using (PerforceChildProcess ChildProcess = new PerforceChildProcess(null, Arguments.ToString(), Logger))
+			{
+				return await ChildProcess.TryReadToEndAsync(CancellationToken);
+			}
+		}
+
+		/// <summary>
+		/// Gets the setting of a Perforce variable
+		/// </summary>
+		/// <param name="Name">Name of the variable to get</param>
+		/// <param name="CancellationToken">Cancellation token for the request</param>
+		/// <returns>Value of the variable</returns>
+		public async Task<string?> TryGetSettingAsync(string Name, CancellationToken CancellationToken)
+		{
+			using (PerforceChildProcess ChildProcess = new PerforceChildProcess(null, $"set {Name}=", Logger))
+			{
+				Tuple<bool, string> Response = await ChildProcess.TryReadToEndAsync(CancellationToken);
+				if (Response.Item1)
+				{
+					string TrimResponse = Response.Item2.Trim();
+					if (TrimResponse.Length == 0)
+					{
+						return null;
+					}
+
+					string Prefix = Name + "=";
+					if (TrimResponse.StartsWith(Prefix, StringComparison.OrdinalIgnoreCase))
+					{
+						return TrimResponse.Substring(Prefix.Length);
+					}
+				}
+
+				Logger.LogDebug($"Unable to get '{Name}' variable: {Response.Item2}");
+				return null;
+			}
+		}
+
+		#endregion
 	}
 
 	/// <summary>
@@ -307,7 +433,7 @@ namespace EpicGames.Perforce
 		/// <param name="InputData">Input data to pass to Perforce</param>
 		/// <param name="CancellationToken">Token used to cancel the operation</param>
 		/// <returns>List of objects returned by the server</returns>
-		public static async Task<PerforceResponseList<T>> CommandAsync<T>(this PerforceConnection Connection, string Arguments, byte[]? InputData, CancellationToken CancellationToken) where T : class
+		public static async Task<PerforceResponseList<T>> CommandAsync<T>(this IPerforceConnection Connection, string Arguments, byte[]? InputData, CancellationToken CancellationToken) where T : class
 		{
 			List<PerforceResponse> Responses = await Connection.CommandAsync(Arguments, InputData, typeof(T), CancellationToken);
 
@@ -328,7 +454,7 @@ namespace EpicGames.Perforce
 		/// <param name="InputData">Input data to pass to Perforce</param>
 		/// <param name="CancellationToken">Token used to cancel the operation</param>
 		/// <returns>List of objects returned by the server</returns>
-		public static async Task<PerforceResponseList<T>> BatchedCommandAsync<T>(PerforceConnection Connection, string CommonArguments, string[] BatchedArguments, byte[]? InputData, CancellationToken CancellationToken) where T : class
+		public static async Task<PerforceResponseList<T>> BatchedCommandAsync<T>(IPerforceConnection Connection, string CommonArguments, string[] BatchedArguments, byte[]? InputData, CancellationToken CancellationToken) where T : class
 		{
 			PerforceResponseList<T> Responses = new PerforceResponseList<T>();
 			for (int FileSpecIdx = 0; FileSpecIdx < BatchedArguments.Length;)
@@ -352,7 +478,7 @@ namespace EpicGames.Perforce
 		/// <param name="StatRecordType">Type of element to return in the response</param>
 		/// <param name="CancellationToken">Token used to cancel the operation</param>
 		/// <returns>Response from the server; either an object of type T or error.</returns>
-		static async Task<PerforceResponse> SingleResponseCommandAsync(PerforceConnection Connection, string Arguments, byte[]? InputData, Type? StatRecordType, CancellationToken CancellationToken)
+		static async Task<PerforceResponse> SingleResponseCommandAsync(IPerforceConnection Connection, string Arguments, byte[]? InputData, Type? StatRecordType, CancellationToken CancellationToken)
 		{
 			List<PerforceResponse> Responses = await Connection.CommandAsync(Arguments, InputData, StatRecordType, CancellationToken);
 			if (Responses.Count != 1)
@@ -375,7 +501,7 @@ namespace EpicGames.Perforce
 		/// <param name="InputData">Input data for the command.</param>
 		/// <param name="CancellationToken">Token used to cancel the operation</param>
 		/// <returns>Response from the server; either an object of type T or error.</returns>
-		public static async Task<PerforceResponse<T>> SingleResponseCommandAsync<T>(PerforceConnection Connection, string Arguments, byte[]? InputData, CancellationToken CancellationToken) where T : class
+		public static async Task<PerforceResponse<T>> SingleResponseCommandAsync<T>(IPerforceConnection Connection, string Arguments, byte[]? InputData, CancellationToken CancellationToken) where T : class
 		{
 			return new PerforceResponse<T>(await SingleResponseCommandAsync(Connection, Arguments, InputData, typeof(T), CancellationToken));
 		}
@@ -392,7 +518,7 @@ namespace EpicGames.Perforce
 		/// <param name="FileNames">Files to be added</param>
 		/// <param name="CancellationToken">Token used to cancel the operation</param>
 		/// <returns>Response from the server</returns>
-		public static async Task<List<AddRecord>> AddAsync(this PerforceConnection Connection, int ChangeNumber, string[] FileNames, CancellationToken CancellationToken)
+		public static async Task<List<AddRecord>> AddAsync(this IPerforceConnection Connection, int ChangeNumber, string[] FileNames, CancellationToken CancellationToken)
 		{
 			return (await TryAddAsync(Connection, ChangeNumber, null, AddOptions.None, FileNames, CancellationToken)).Data;
 		}
@@ -407,7 +533,7 @@ namespace EpicGames.Perforce
 		/// <param name="FileNames">Files to be added</param>
 		/// <param name="CancellationToken">Token used to cancel the operation</param>
 		/// <returns>Response from the server</returns>
-		public static async Task<List<AddRecord>> AddAsync(this PerforceConnection Connection, int ChangeNumber, string? FileType, AddOptions Options, string[] FileNames, CancellationToken CancellationToken)
+		public static async Task<List<AddRecord>> AddAsync(this IPerforceConnection Connection, int ChangeNumber, string? FileType, AddOptions Options, string[] FileNames, CancellationToken CancellationToken)
 		{
 			return (await TryAddAsync(Connection, ChangeNumber, FileType, Options, FileNames, CancellationToken)).Data;
 		}
@@ -420,7 +546,7 @@ namespace EpicGames.Perforce
 		/// <param name="FileNames">Files to be added</param>
 		/// <param name="CancellationToken">Token used to cancel the operation</param>
 		/// <returns>Response from the server</returns>
-		public static Task<PerforceResponseList<AddRecord>> TryAddAsync(this PerforceConnection Connection, int ChangeNumber, string[] FileNames, CancellationToken CancellationToken)
+		public static Task<PerforceResponseList<AddRecord>> TryAddAsync(this IPerforceConnection Connection, int ChangeNumber, string[] FileNames, CancellationToken CancellationToken)
 		{
 			return TryAddAsync(Connection, ChangeNumber, null, AddOptions.None, FileNames, CancellationToken);
 		}
@@ -435,7 +561,7 @@ namespace EpicGames.Perforce
 		/// <param name="FileNames">Files to be added</param>
 		/// <param name="CancellationToken">Token used to cancel the operation</param>
 		/// <returns>Response from the server</returns>
-		public static Task<PerforceResponseList<AddRecord>> TryAddAsync(this PerforceConnection Connection, int ChangeNumber, string? FileType, AddOptions Options, string[] FileNames, CancellationToken CancellationToken)
+		public static Task<PerforceResponseList<AddRecord>> TryAddAsync(this IPerforceConnection Connection, int ChangeNumber, string? FileType, AddOptions Options, string[] FileNames, CancellationToken CancellationToken)
 		{
 			StringBuilder Arguments = new StringBuilder("add");
 			if (ChangeNumber != -1)
@@ -477,7 +603,7 @@ namespace EpicGames.Perforce
 		/// <param name="Record">Information for the change to create. The number field should be left set to -1.</param>
 		/// <param name="CancellationToken">Token used to cancel the operation</param>
 		/// <returns>The changelist number, or an error.</returns>
-		public static async Task<ChangeRecord> CreateChangeAsync(this PerforceConnection Connection, ChangeRecord Record, CancellationToken CancellationToken)
+		public static async Task<ChangeRecord> CreateChangeAsync(this IPerforceConnection Connection, ChangeRecord Record, CancellationToken CancellationToken)
 		{
 			return (await TryCreateChangeAsync(Connection, Record, CancellationToken)).Data;
 		}
@@ -489,7 +615,7 @@ namespace EpicGames.Perforce
 		/// <param name="Record">Information for the change to create. The number field should be left set to -1.</param>
 		/// <param name="CancellationToken">Token used to cancel the operation</param>
 		/// <returns>The changelist number, or an error.</returns>
-		public static async Task<PerforceResponse<ChangeRecord>> TryCreateChangeAsync(this PerforceConnection Connection, ChangeRecord Record, CancellationToken CancellationToken)
+		public static async Task<PerforceResponse<ChangeRecord>> TryCreateChangeAsync(this IPerforceConnection Connection, ChangeRecord Record, CancellationToken CancellationToken)
 		{
 			if (Record.Number != -1)
 			{
@@ -528,7 +654,7 @@ namespace EpicGames.Perforce
 		/// <param name="Record">Information for the change to create. The number field should be left set to zero.</param>
 		/// <param name="CancellationToken">Token used to cancel the operation</param>
 		/// <returns>The changelist number, or an error.</returns>
-		public static async Task UpdateChangeAsync(this PerforceConnection Connection, UpdateChangeOptions Options, ChangeRecord Record, CancellationToken CancellationToken)
+		public static async Task UpdateChangeAsync(this IPerforceConnection Connection, UpdateChangeOptions Options, ChangeRecord Record, CancellationToken CancellationToken)
 		{
 			(await TryUpdateChangeAsync(Connection, Options, Record, CancellationToken)).EnsureSuccess();
 		}
@@ -541,7 +667,7 @@ namespace EpicGames.Perforce
 		/// <param name="Record">Information for the change to create. The number field should be left set to zero.</param>
 		/// <param name="CancellationToken">Token used to cancel the operation</param>
 		/// <returns>The changelist number, or an error.</returns>
-		public static Task<PerforceResponse> TryUpdateChangeAsync(this PerforceConnection Connection, UpdateChangeOptions Options, ChangeRecord Record, CancellationToken CancellationToken)
+		public static Task<PerforceResponse> TryUpdateChangeAsync(this IPerforceConnection Connection, UpdateChangeOptions Options, ChangeRecord Record, CancellationToken CancellationToken)
 		{
 			if (Record.Number == -1)
 			{
@@ -569,7 +695,7 @@ namespace EpicGames.Perforce
 		/// <param name="ChangeNumber">Changelist number to delete</param>
 		/// <param name="CancellationToken">Token used to cancel the operation</param>
 		/// <returns>Response from the server</returns>
-		public static async Task DeleteChangeAsync(this PerforceConnection Connection, DeleteChangeOptions Options, int ChangeNumber, CancellationToken CancellationToken)
+		public static async Task DeleteChangeAsync(this IPerforceConnection Connection, DeleteChangeOptions Options, int ChangeNumber, CancellationToken CancellationToken)
 		{
 			(await TryDeleteChangeAsync(Connection, Options, ChangeNumber, CancellationToken)).EnsureSuccess();
 		}
@@ -582,7 +708,7 @@ namespace EpicGames.Perforce
 		/// <param name="ChangeNumber">Changelist number to delete</param>
 		/// <param name="CancellationToken">Token used to cancel the operation</param>
 		/// <returns>Response from the server</returns>
-		public static Task<PerforceResponse> TryDeleteChangeAsync(this PerforceConnection Connection, DeleteChangeOptions Options, int ChangeNumber, CancellationToken CancellationToken)
+		public static Task<PerforceResponse> TryDeleteChangeAsync(this IPerforceConnection Connection, DeleteChangeOptions Options, int ChangeNumber, CancellationToken CancellationToken)
 		{
 			StringBuilder Arguments = new StringBuilder("change -d");
 			if ((Options & DeleteChangeOptions.Submitted) != 0)
@@ -606,7 +732,7 @@ namespace EpicGames.Perforce
 		/// <param name="ChangeNumber">Changelist number to retrieve. -1 is the default changelist for this workspace.</param>
 		/// <param name="CancellationToken">Token used to cancel the operation</param>
 		/// <returns>Response from the server</returns>
-		public static async Task<ChangeRecord> GetChangeAsync(this PerforceConnection Connection, GetChangeOptions Options, int ChangeNumber, CancellationToken CancellationToken)
+		public static async Task<ChangeRecord> GetChangeAsync(this IPerforceConnection Connection, GetChangeOptions Options, int ChangeNumber, CancellationToken CancellationToken)
 		{
 			return (await TryGetChange(Connection, Options, ChangeNumber, CancellationToken)).Data;
 		}
@@ -619,7 +745,7 @@ namespace EpicGames.Perforce
 		/// <param name="ChangeNumber">Changelist number to retrieve. -1 is the default changelist for this workspace.</param>
 		/// <param name="CancellationToken">Token used to cancel the operation</param>
 		/// <returns>Response from the server</returns>
-		public static Task<PerforceResponse<ChangeRecord>> TryGetChange(this PerforceConnection Connection, GetChangeOptions Options, int ChangeNumber, CancellationToken CancellationToken)
+		public static Task<PerforceResponse<ChangeRecord>> TryGetChange(this IPerforceConnection Connection, GetChangeOptions Options, int ChangeNumber, CancellationToken CancellationToken)
 		{
 			StringBuilder Arguments = new StringBuilder("change -o");
 			if ((Options & GetChangeOptions.BeforeRenumber) != 0)
@@ -683,7 +809,7 @@ namespace EpicGames.Perforce
 		/// <param name="FileSpecs">Paths to query changes for</param>
 		/// <param name="CancellationToken">Token used to cancel the operation</param>
 		/// <returns>List of responses from the server.</returns>
-		public static async Task<List<ChangesRecord>> GetChangesAsync(this PerforceConnection Connection, ChangesOptions Options, int MaxChanges, ChangeStatus Status, string[] FileSpecs, CancellationToken CancellationToken)
+		public static async Task<List<ChangesRecord>> GetChangesAsync(this IPerforceConnection Connection, ChangesOptions Options, int MaxChanges, ChangeStatus Status, string[] FileSpecs, CancellationToken CancellationToken)
 		{
 			return (await TryGetChangesAsync(Connection, Options, MaxChanges, Status, FileSpecs, CancellationToken)).Data;
 		}
@@ -700,7 +826,7 @@ namespace EpicGames.Perforce
 		/// <param name="FileSpecs">Paths to query changes for</param>
 		/// <param name="CancellationToken">Token used to cancel the operation</param>
 		/// <returns>List of responses from the server.</returns>
-		public static async Task<List<ChangesRecord>> GetChangesAsync(this PerforceConnection Connection, ChangesOptions Options, string? ClientName, int MaxChanges, ChangeStatus Status, string? UserName, string[]? FileSpecs, CancellationToken CancellationToken)
+		public static async Task<List<ChangesRecord>> GetChangesAsync(this IPerforceConnection Connection, ChangesOptions Options, string? ClientName, int MaxChanges, ChangeStatus Status, string? UserName, string[]? FileSpecs, CancellationToken CancellationToken)
 		{
 			return (await TryGetChangesAsync(Connection, Options, ClientName, MaxChanges, Status, UserName, FileSpecs, CancellationToken)).Data;
 		}
@@ -715,7 +841,7 @@ namespace EpicGames.Perforce
 		/// <param name="FileSpecs">Paths to query changes for</param>
 		/// <param name="CancellationToken">Token used to cancel the operation</param>
 		/// <returns>List of responses from the server.</returns>
-		public static Task<PerforceResponseList<ChangesRecord>> TryGetChangesAsync(this PerforceConnection Connection, ChangesOptions Options, int MaxChanges, ChangeStatus Status, string[] FileSpecs, CancellationToken CancellationToken)
+		public static Task<PerforceResponseList<ChangesRecord>> TryGetChangesAsync(this IPerforceConnection Connection, ChangesOptions Options, int MaxChanges, ChangeStatus Status, string[] FileSpecs, CancellationToken CancellationToken)
 		{
 			return TryGetChangesAsync(Connection, Options, null, MaxChanges, Status, null, FileSpecs, CancellationToken);
 		}
@@ -732,7 +858,7 @@ namespace EpicGames.Perforce
 		/// <param name="FileSpecs">Paths to query changes for</param>
 		/// <param name="CancellationToken">Token used to cancel the operation</param>
 		/// <returns>List of responses from the server.</returns>
-		public static Task<PerforceResponseList<ChangesRecord>> TryGetChangesAsync(this PerforceConnection Connection, ChangesOptions Options, string? ClientName, int MaxChanges, ChangeStatus Status, string? UserName, string[]? FileSpecs, CancellationToken CancellationToken)
+		public static Task<PerforceResponseList<ChangesRecord>> TryGetChangesAsync(this IPerforceConnection Connection, ChangesOptions Options, string? ClientName, int MaxChanges, ChangeStatus Status, string? UserName, string[]? FileSpecs, CancellationToken CancellationToken)
 		{
 			StringBuilder Arguments = new StringBuilder("changes");
 			if ((Options & ChangesOptions.IncludeIntegrations) != 0)
@@ -793,7 +919,7 @@ namespace EpicGames.Perforce
 		/// <param name="FileSpecs">Files to sync</param>
 		/// <param name="CancellationToken">Token used to cancel the operation</param>
 		/// <returns>Response from the server</returns>
-		public static async Task<List<CleanRecord>> CleanAsync(this PerforceConnection Connection, CleanOptions Options, string[] FileSpecs, CancellationToken CancellationToken)
+		public static async Task<List<CleanRecord>> CleanAsync(this IPerforceConnection Connection, CleanOptions Options, string[] FileSpecs, CancellationToken CancellationToken)
 		{
 			return (await TryCleanAsync(Connection, Options, FileSpecs, CancellationToken)).Data;
 		}
@@ -806,7 +932,7 @@ namespace EpicGames.Perforce
 		/// <param name="FileSpecs">Files to sync</param>
 		/// <param name="CancellationToken">Token used to cancel the operation</param>
 		/// <returns>Response from the server</returns>
-		public static async Task<PerforceResponseList<CleanRecord>> TryCleanAsync(this PerforceConnection Connection, CleanOptions Options, string[] FileSpecs, CancellationToken CancellationToken)
+		public static async Task<PerforceResponseList<CleanRecord>> TryCleanAsync(this IPerforceConnection Connection, CleanOptions Options, string[] FileSpecs, CancellationToken CancellationToken)
 		{
 			StringBuilder Arguments = new StringBuilder("clean");
 			if ((Options & CleanOptions.Edited) != 0)
@@ -858,7 +984,7 @@ namespace EpicGames.Perforce
 		/// <param name="Record">The client record</param>
 		/// <param name="CancellationToken">Token used to cancel the operation</param>
 		/// <returns>Response from the server</returns>
-		public static async Task CreateClientAsync(this PerforceConnection Connection, ClientRecord Record, CancellationToken CancellationToken)
+		public static async Task CreateClientAsync(this IPerforceConnection Connection, ClientRecord Record, CancellationToken CancellationToken)
 		{
 			(await TryCreateClientAsync(Connection, Record, CancellationToken)).EnsureSuccess();
 		}
@@ -870,7 +996,7 @@ namespace EpicGames.Perforce
 		/// <param name="Record">The client record</param>
 		/// <param name="CancellationToken">Token used to cancel the operation</param>
 		/// <returns>Response from the server</returns>
-		public static Task<PerforceResponse> TryCreateClientAsync(this PerforceConnection Connection, ClientRecord Record, CancellationToken CancellationToken)
+		public static Task<PerforceResponse> TryCreateClientAsync(this IPerforceConnection Connection, ClientRecord Record, CancellationToken CancellationToken)
 		{
 			return TryUpdateClientAsync(Connection, Record, CancellationToken);
 		}
@@ -882,7 +1008,7 @@ namespace EpicGames.Perforce
 		/// <param name="Record">The client record</param>
 		/// <param name="CancellationToken">Token used to cancel the operation</param>
 		/// <returns>Response from the server</returns>
-		public static async Task UpdateClientAsync(this PerforceConnection Connection, ClientRecord Record, CancellationToken CancellationToken)
+		public static async Task UpdateClientAsync(this IPerforceConnection Connection, ClientRecord Record, CancellationToken CancellationToken)
 		{
 			(await TryUpdateClientAsync(Connection, Record, CancellationToken)).EnsureSuccess();
 		}
@@ -894,7 +1020,7 @@ namespace EpicGames.Perforce
 		/// <param name="Record">The client record</param>
 		/// <param name="CancellationToken">Token used to cancel the operation</param>
 		/// <returns>Response from the server</returns>
-		public static Task<PerforceResponse> TryUpdateClientAsync(this PerforceConnection Connection, ClientRecord Record, CancellationToken CancellationToken)
+		public static Task<PerforceResponse> TryUpdateClientAsync(this IPerforceConnection Connection, ClientRecord Record, CancellationToken CancellationToken)
 		{
 			return SingleResponseCommandAsync(Connection, "client -i", SerializeRecord(Record), null, CancellationToken);
 		}
@@ -907,7 +1033,7 @@ namespace EpicGames.Perforce
 		/// <param name="ClientName">Name of the client</param>
 		/// <param name="CancellationToken">Token used to cancel the operation</param>
 		/// <returns>Response from the server</returns>
-		public static async Task DeleteClientAsync(this PerforceConnection Connection, DeleteClientOptions Options, string ClientName, CancellationToken CancellationToken)
+		public static async Task DeleteClientAsync(this IPerforceConnection Connection, DeleteClientOptions Options, string ClientName, CancellationToken CancellationToken)
 		{
 			(await TryDeleteClientAsync(Connection, Options, ClientName, CancellationToken)).EnsureSuccess();
 		}
@@ -920,7 +1046,7 @@ namespace EpicGames.Perforce
 		/// <param name="ClientName">Name of the client</param>
 		/// <param name="CancellationToken">Token used to cancel the operation</param>
 		/// <returns>Response from the server</returns>
-		public static Task<PerforceResponse> TryDeleteClientAsync(this PerforceConnection Connection, DeleteClientOptions Options, string ClientName, CancellationToken CancellationToken)
+		public static Task<PerforceResponse> TryDeleteClientAsync(this IPerforceConnection Connection, DeleteClientOptions Options, string ClientName, CancellationToken CancellationToken)
 		{
 			StringBuilder Arguments = new StringBuilder("client -d");
 			if ((Options & DeleteClientOptions.Force) != 0)
@@ -944,7 +1070,7 @@ namespace EpicGames.Perforce
 		/// <param name="Options">Options for this command</param>
 		/// <param name="CancellationToken">Token used to cancel the operation</param>
 		/// <returns>Response from the server</returns>
-		public static async Task SwitchClientToStreamAsync(this PerforceConnection Connection, string ClientName, string StreamName, SwitchClientOptions Options, CancellationToken CancellationToken)
+		public static async Task SwitchClientToStreamAsync(this IPerforceConnection Connection, string ClientName, string StreamName, SwitchClientOptions Options, CancellationToken CancellationToken)
 		{
 			(await TrySwitchClientToStreamAsync(Connection, ClientName, StreamName, Options, CancellationToken)).EnsureSuccess();
 		}
@@ -958,7 +1084,7 @@ namespace EpicGames.Perforce
 		/// <param name="Options">Options for this command</param>
 		/// <param name="CancellationToken">Token used to cancel the operation</param>
 		/// <returns>Response from the server</returns>
-		public static Task<PerforceResponse> TrySwitchClientToStreamAsync(this PerforceConnection Connection, string ClientName, string StreamName, SwitchClientOptions Options, CancellationToken CancellationToken)
+		public static Task<PerforceResponse> TrySwitchClientToStreamAsync(this IPerforceConnection Connection, string ClientName, string StreamName, SwitchClientOptions Options, CancellationToken CancellationToken)
 		{
 			StringBuilder Arguments = new StringBuilder("client -s");
 			if ((Options & SwitchClientOptions.IgnoreOpenFiles) != 0)
@@ -978,7 +1104,7 @@ namespace EpicGames.Perforce
 		/// <param name="TemplateName">The new stream to be associated with the client</param>
 		/// <param name="CancellationToken">Token used to cancel the operation</param>
 		/// <returns>Response from the server</returns>
-		public static async Task SwitchClientToTemplateAsync(this PerforceConnection Connection, string ClientName, string TemplateName, CancellationToken CancellationToken)
+		public static async Task SwitchClientToTemplateAsync(this IPerforceConnection Connection, string ClientName, string TemplateName, CancellationToken CancellationToken)
 		{
 			(await TrySwitchClientToTemplateAsync(Connection, ClientName, TemplateName, CancellationToken)).EnsureSuccess();
 		}
@@ -991,7 +1117,7 @@ namespace EpicGames.Perforce
 		/// <param name="TemplateName">The new stream to be associated with the client</param>
 		/// <param name="CancellationToken">Token used to cancel the operation</param>
 		/// <returns>Response from the server</returns>
-		public static Task<PerforceResponse> TrySwitchClientToTemplateAsync(this PerforceConnection Connection, string ClientName, string TemplateName, CancellationToken CancellationToken)
+		public static Task<PerforceResponse> TrySwitchClientToTemplateAsync(this IPerforceConnection Connection, string ClientName, string TemplateName, CancellationToken CancellationToken)
 		{
 			string Arguments = String.Format("client -s -t \"{0}\"", TemplateName);
 			return SingleResponseCommandAsync(Connection, Arguments, null, null, CancellationToken);
@@ -1004,7 +1130,7 @@ namespace EpicGames.Perforce
 		/// <param name="ClientName">Name of the client. Specify null for the current client.</param>
 		/// <param name="CancellationToken">Token used to cancel the operation</param>
 		/// <returns>Response from the server; either a client record or error code</returns>
-		public static async Task<ClientRecord> GetClientAsync(this PerforceConnection Connection, string? ClientName, CancellationToken CancellationToken)
+		public static async Task<ClientRecord> GetClientAsync(this IPerforceConnection Connection, string? ClientName, CancellationToken CancellationToken)
 		{
 			return (await TryGetClientAsync(Connection, ClientName, CancellationToken)).Data;
 		}
@@ -1016,7 +1142,7 @@ namespace EpicGames.Perforce
 		/// <param name="ClientName">Name of the client. Specify null for the current client.</param>
 		/// <param name="CancellationToken">Token used to cancel the operation</param>
 		/// <returns>Response from the server; either a client record or error code</returns>
-		public static Task<PerforceResponse<ClientRecord>> TryGetClientAsync(this PerforceConnection Connection, string? ClientName, CancellationToken CancellationToken)
+		public static Task<PerforceResponse<ClientRecord>> TryGetClientAsync(this IPerforceConnection Connection, string? ClientName, CancellationToken CancellationToken)
 		{
 			StringBuilder Arguments = new StringBuilder("client -o");
 			if (ClientName != null)
@@ -1034,7 +1160,7 @@ namespace EpicGames.Perforce
 		/// <param name="bIncludeView">Whether to include the stream view in the output</param>
 		/// <param name="CancellationToken">Token used to cancel the operation</param>
 		/// <returns>Stream information record</returns>
-		public static async Task<StreamRecord> GetStreamAsync(this PerforceConnection Connection, string StreamName, bool bIncludeView, CancellationToken CancellationToken)
+		public static async Task<StreamRecord> GetStreamAsync(this IPerforceConnection Connection, string StreamName, bool bIncludeView, CancellationToken CancellationToken)
 		{
 			return (await TryGetStreamAsync(Connection, StreamName, bIncludeView, CancellationToken)).Data;
 		}
@@ -1047,7 +1173,7 @@ namespace EpicGames.Perforce
 		/// <param name="bIncludeView">Whether to include the stream view in the output</param>
 		/// <param name="CancellationToken">Token used to cancel the operation</param>
 		/// <returns>Stream information record</returns>
-		public static Task<PerforceResponse<StreamRecord>> TryGetStreamAsync(this PerforceConnection Connection, string StreamName, bool bIncludeView, CancellationToken CancellationToken)
+		public static Task<PerforceResponse<StreamRecord>> TryGetStreamAsync(this IPerforceConnection Connection, string StreamName, bool bIncludeView, CancellationToken CancellationToken)
 		{
 			StringBuilder Arguments = new StringBuilder("stream -o");
 			if (bIncludeView)
@@ -1067,7 +1193,7 @@ namespace EpicGames.Perforce
 		/// <param name="ChangeNumber">Changelist at which to query the stream view</param>
 		/// <param name="CancellationToken">Token used to cancel the operation</param>
 		/// <returns>Response from the server; either a client record or error code</returns>
-		public static async Task<ClientRecord> GetStreamViewAsync(this PerforceConnection Connection, string StreamName, int ChangeNumber, CancellationToken CancellationToken)
+		public static async Task<ClientRecord> GetStreamViewAsync(this IPerforceConnection Connection, string StreamName, int ChangeNumber, CancellationToken CancellationToken)
 		{
 			return (await TryGetStreamViewAsync(Connection, StreamName, ChangeNumber, CancellationToken)).Data;
 		}
@@ -1080,7 +1206,7 @@ namespace EpicGames.Perforce
 		/// <param name="ChangeNumber">Changelist at which to query the stream view</param>
 		/// <param name="CancellationToken">Token used to cancel the operation</param>
 		/// <returns>Response from the server; either a client record or error code</returns>
-		public static Task<PerforceResponse<ClientRecord>> TryGetStreamViewAsync(this PerforceConnection Connection, string StreamName, int ChangeNumber, CancellationToken CancellationToken)
+		public static Task<PerforceResponse<ClientRecord>> TryGetStreamViewAsync(this IPerforceConnection Connection, string StreamName, int ChangeNumber, CancellationToken CancellationToken)
 		{
 			StringBuilder Arguments = new StringBuilder("client -o");
 			Arguments.AppendFormat(" -S \"{0}\"", StreamName);
@@ -1158,7 +1284,7 @@ namespace EpicGames.Perforce
 		/// <param name="UserName">List only client workspaces owned by this user.</param>
 		/// <param name="CancellationToken">Token used to cancel the operation</param>
 		/// <returns>Response from the server; either a client record or error code</returns>
-		public static async Task<List<ClientsRecord>> GetClientsAsync(this PerforceConnection Connection, ClientsOptions Options, string? UserName, CancellationToken CancellationToken)
+		public static async Task<List<ClientsRecord>> GetClientsAsync(this IPerforceConnection Connection, ClientsOptions Options, string? UserName, CancellationToken CancellationToken)
 		{
 			return (await TryGetClientsAsync(Connection, Options, UserName, CancellationToken)).Data;
 		}
@@ -1171,7 +1297,7 @@ namespace EpicGames.Perforce
 		/// <param name="UserName">List only client workspaces owned by this user.</param>
 		/// <param name="CancellationToken">Token used to cancel the operation</param>
 		/// <returns>Response from the server; either a client record or error code</returns>
-		public static Task<PerforceResponseList<ClientsRecord>> TryGetClientsAsync(this PerforceConnection Connection, ClientsOptions Options, string? UserName, CancellationToken CancellationToken)
+		public static Task<PerforceResponseList<ClientsRecord>> TryGetClientsAsync(this IPerforceConnection Connection, ClientsOptions Options, string? UserName, CancellationToken CancellationToken)
 		{
 			return TryGetClientsAsync(Connection, Options, null, -1, null, UserName, CancellationToken);
 		}
@@ -1187,7 +1313,7 @@ namespace EpicGames.Perforce
 		/// <param name="UserName">List only client workspaces owned by this user.</param>
 		/// <param name="CancellationToken">Token used to cancel the operation</param>
 		/// <returns>Response from the server; either a client record or error code</returns>
-		public static async Task<List<ClientsRecord>> GetClientsAsync(this PerforceConnection Connection, ClientsOptions Options, string? Filter, int MaxResults, string? Stream, string? UserName, CancellationToken CancellationToken)
+		public static async Task<List<ClientsRecord>> GetClientsAsync(this IPerforceConnection Connection, ClientsOptions Options, string? Filter, int MaxResults, string? Stream, string? UserName, CancellationToken CancellationToken)
 		{
 			return (await TryGetClientsAsync(Connection, Options, Filter, MaxResults, Stream, UserName, CancellationToken)).Data;
 		}
@@ -1203,7 +1329,7 @@ namespace EpicGames.Perforce
 		/// <param name="UserName">List only client workspaces owned by this user.</param>
 		/// <param name="CancellationToken">Token used to cancel the operation</param>
 		/// <returns>Response from the server; either a client record or error code</returns>
-		public static Task<PerforceResponseList<ClientsRecord>> TryGetClientsAsync(this PerforceConnection Connection, ClientsOptions Options, string? Filter, int MaxResults, string? Stream, string? UserName, CancellationToken CancellationToken)
+		public static Task<PerforceResponseList<ClientsRecord>> TryGetClientsAsync(this IPerforceConnection Connection, ClientsOptions Options, string? Filter, int MaxResults, string? Stream, string? UserName, CancellationToken CancellationToken)
 		{
 			StringBuilder Arguments = new StringBuilder("clients");
 			if ((Options & ClientsOptions.All) != 0)
@@ -1257,7 +1383,7 @@ namespace EpicGames.Perforce
 		/// <param name="FileSpecs">List of file specifications to query</param>
 		/// <param name="CancellationToken">Token used to cancel the operation</param>
 		/// <returns>List of responses from the server</returns>
-		public static async Task<List<DeleteRecord>> DeleteAsync(this PerforceConnection Connection, int ChangeNumber, DeleteOptions Options, string[] FileSpecs, CancellationToken CancellationToken)
+		public static async Task<List<DeleteRecord>> DeleteAsync(this IPerforceConnection Connection, int ChangeNumber, DeleteOptions Options, string[] FileSpecs, CancellationToken CancellationToken)
 		{
 			return (await TryDeleteAsync(Connection, ChangeNumber, Options, FileSpecs, CancellationToken)).Data;
 		}
@@ -1271,7 +1397,7 @@ namespace EpicGames.Perforce
 		/// <param name="FileSpecs">List of file specifications to query</param>
 		/// <param name="CancellationToken">Token used to cancel the operation</param>
 		/// <returns>List of responses from the server</returns>
-		public static Task<PerforceResponseList<DeleteRecord>> TryDeleteAsync(this PerforceConnection Connection, int ChangeNumber, DeleteOptions Options, string[] FileSpecs, CancellationToken CancellationToken)
+		public static Task<PerforceResponseList<DeleteRecord>> TryDeleteAsync(this IPerforceConnection Connection, int ChangeNumber, DeleteOptions Options, string[] FileSpecs, CancellationToken CancellationToken)
 		{
 			StringBuilder Arguments = new StringBuilder("delete");
 			if (ChangeNumber != -1)
@@ -1305,7 +1431,7 @@ namespace EpicGames.Perforce
 		/// <param name="ChangeNumber">The changelist number to retrieve description for</param>
 		/// <param name="CancellationToken">Token used to cancel the operation</param>
 		/// <returns>Response from the server; either a describe record or error code</returns>
-		public static async Task<DescribeRecord> DescribeAsync(this PerforceConnection Connection, int ChangeNumber, CancellationToken CancellationToken)
+		public static async Task<DescribeRecord> DescribeAsync(this IPerforceConnection Connection, int ChangeNumber, CancellationToken CancellationToken)
 		{
 			return (await TryDescribeAsync(Connection, ChangeNumber, CancellationToken)).Data;
 		}
@@ -1317,7 +1443,7 @@ namespace EpicGames.Perforce
 		/// <param name="ChangeNumber">The changelist number to retrieve description for</param>
 		/// <param name="CancellationToken">Token used to cancel the operation</param>
 		/// <returns>Response from the server; either a describe record or error code</returns>
-		public static async Task<PerforceResponse<DescribeRecord>> TryDescribeAsync(this PerforceConnection Connection, int ChangeNumber, CancellationToken CancellationToken)
+		public static async Task<PerforceResponse<DescribeRecord>> TryDescribeAsync(this IPerforceConnection Connection, int ChangeNumber, CancellationToken CancellationToken)
 		{
 			PerforceResponseList<DescribeRecord> Records = await TryDescribeAsync(Connection, new int[] { ChangeNumber }, CancellationToken);
 			if (Records.Count != 1)
@@ -1334,7 +1460,7 @@ namespace EpicGames.Perforce
 		/// <param name="ChangeNumbers">The changelist numbers to retrieve descriptions for</param>
 		/// <param name="CancellationToken">Token used to cancel the operation</param>
 		/// <returns>List of responses from the server</returns>
-		public static async Task<List<DescribeRecord>> DescribeAsync(this PerforceConnection Connection, int[] ChangeNumbers, CancellationToken CancellationToken)
+		public static async Task<List<DescribeRecord>> DescribeAsync(this IPerforceConnection Connection, int[] ChangeNumbers, CancellationToken CancellationToken)
 		{
 			return (await TryDescribeAsync(Connection, ChangeNumbers, CancellationToken)).Data;
 		}
@@ -1346,7 +1472,7 @@ namespace EpicGames.Perforce
 		/// <param name="ChangeNumbers">The changelist numbers to retrieve descriptions for</param>
 		/// <param name="CancellationToken">Token used to cancel the operation</param>
 		/// <returns>List of responses from the server</returns>
-		public static Task<PerforceResponseList<DescribeRecord>> TryDescribeAsync(this PerforceConnection Connection, int[] ChangeNumbers, CancellationToken CancellationToken)
+		public static Task<PerforceResponseList<DescribeRecord>> TryDescribeAsync(this IPerforceConnection Connection, int[] ChangeNumbers, CancellationToken CancellationToken)
 		{
 			StringBuilder Arguments = new StringBuilder("describe -s");
 			foreach (int ChangeNumber in ChangeNumbers)
@@ -1365,7 +1491,7 @@ namespace EpicGames.Perforce
 		/// <param name="ChangeNumbers">The changelist numbers to retrieve descriptions for</param>
 		/// <param name="CancellationToken">Token used to cancel the operation</param>
 		/// <returns>List of responses from the server</returns>
-		public static async Task<List<DescribeRecord>> DescribeAsync(this PerforceConnection Connection, DescribeOptions Options, int MaxNumFiles, int[] ChangeNumbers, CancellationToken CancellationToken)
+		public static async Task<List<DescribeRecord>> DescribeAsync(this IPerforceConnection Connection, DescribeOptions Options, int MaxNumFiles, int[] ChangeNumbers, CancellationToken CancellationToken)
 		{
 			return (await TryDescribeAsync(Connection, Options, MaxNumFiles, ChangeNumbers, CancellationToken)).Data;
 		}
@@ -1379,7 +1505,7 @@ namespace EpicGames.Perforce
 		/// <param name="ChangeNumbers">The changelist numbers to retrieve descriptions for</param>
 		/// <param name="CancellationToken">Token used to cancel the operation</param>
 		/// <returns>List of responses from the server</returns>
-		public static Task<PerforceResponseList<DescribeRecord>> TryDescribeAsync(this PerforceConnection Connection, DescribeOptions Options, int MaxNumFiles, int[] ChangeNumbers, CancellationToken CancellationToken)
+		public static Task<PerforceResponseList<DescribeRecord>> TryDescribeAsync(this IPerforceConnection Connection, DescribeOptions Options, int MaxNumFiles, int[] ChangeNumbers, CancellationToken CancellationToken)
 		{
 			StringBuilder Arguments = new StringBuilder("describe -s");
 			if ((Options & DescribeOptions.ShowDescriptionForRestrictedChanges) != 0)
@@ -1421,7 +1547,7 @@ namespace EpicGames.Perforce
 		/// <param name="FileSpecs">Files to be opened for edit</param>
 		/// <param name="CancellationToken">Token used to cancel the operation</param>
 		/// <returns>Response from the server</returns>
-		public static async Task<List<EditRecord>> EditAsync(this PerforceConnection Connection, int ChangeNumber, string[] FileSpecs, CancellationToken CancellationToken)
+		public static async Task<List<EditRecord>> EditAsync(this IPerforceConnection Connection, int ChangeNumber, string[] FileSpecs, CancellationToken CancellationToken)
 		{
 			return (await TryEditAsync(Connection, ChangeNumber, FileSpecs, CancellationToken)).Data;
 		}
@@ -1434,7 +1560,7 @@ namespace EpicGames.Perforce
 		/// <param name="FileSpecs">Files to be opened for edit</param>
 		/// <param name="CancellationToken">Token used to cancel the operation</param>
 		/// <returns>Response from the server</returns>
-		public static Task<PerforceResponseList<EditRecord>> TryEditAsync(this PerforceConnection Connection, int ChangeNumber, string[] FileSpecs, CancellationToken CancellationToken)
+		public static Task<PerforceResponseList<EditRecord>> TryEditAsync(this IPerforceConnection Connection, int ChangeNumber, string[] FileSpecs, CancellationToken CancellationToken)
 		{
 			return TryEditAsync(Connection, ChangeNumber, null, EditOptions.None, FileSpecs, CancellationToken);
 		}
@@ -1449,7 +1575,7 @@ namespace EpicGames.Perforce
 		/// <param name="FileSpecs">Files to be opened for edit</param>
 		/// <param name="CancellationToken">Token used to cancel the operation</param>
 		/// <returns>Response from the server</returns>
-		public static async Task<List<EditRecord>> EditAsync(this PerforceConnection Connection, int ChangeNumber, string? FileType, EditOptions Options, string[] FileSpecs, CancellationToken CancellationToken)
+		public static async Task<List<EditRecord>> EditAsync(this IPerforceConnection Connection, int ChangeNumber, string? FileType, EditOptions Options, string[] FileSpecs, CancellationToken CancellationToken)
 		{
 			return (await TryEditAsync(Connection, ChangeNumber, FileType, Options, FileSpecs, CancellationToken)).Data;
 		}
@@ -1464,7 +1590,7 @@ namespace EpicGames.Perforce
 		/// <param name="FileSpecs">Files to be opened for edit</param>
 		/// <param name="CancellationToken">Token used to cancel the operation</param>
 		/// <returns>Response from the server</returns>
-		public static Task<PerforceResponseList<EditRecord>> TryEditAsync(this PerforceConnection Connection, int ChangeNumber, string? FileType, EditOptions Options, string[] FileSpecs, CancellationToken CancellationToken)
+		public static Task<PerforceResponseList<EditRecord>> TryEditAsync(this IPerforceConnection Connection, int ChangeNumber, string? FileType, EditOptions Options, string[] FileSpecs, CancellationToken CancellationToken)
 		{
 			StringBuilder Arguments = new StringBuilder("edit");
 			if (ChangeNumber != -1)
@@ -1499,7 +1625,7 @@ namespace EpicGames.Perforce
 		/// <param name="FileSpecs">List of file specifications to query</param>
 		/// <param name="CancellationToken">Token used to cancel the operation</param>
 		/// <returns>List of responses from the server</returns>
-		public static async Task<List<FileLogRecord>> FileLogAsync(this PerforceConnection Connection, FileLogOptions Options, string[] FileSpecs, CancellationToken CancellationToken)
+		public static async Task<List<FileLogRecord>> FileLogAsync(this IPerforceConnection Connection, FileLogOptions Options, string[] FileSpecs, CancellationToken CancellationToken)
 		{
 			return (await TryFileLogAsync(Connection, Options, FileSpecs, CancellationToken)).Data;
 		}
@@ -1512,7 +1638,7 @@ namespace EpicGames.Perforce
 		/// <param name="FileSpecs">List of file specifications to query</param>
 		/// <param name="CancellationToken">Token used to cancel the operation</param>
 		/// <returns>List of responses from the server</returns>
-		public static Task<PerforceResponseList<FileLogRecord>> TryFileLogAsync(this PerforceConnection Connection, FileLogOptions Options, string[] FileSpecs, CancellationToken CancellationToken)
+		public static Task<PerforceResponseList<FileLogRecord>> TryFileLogAsync(this IPerforceConnection Connection, FileLogOptions Options, string[] FileSpecs, CancellationToken CancellationToken)
 		{
 			return TryFileLogAsync(Connection, -1, -1, Options, FileSpecs, CancellationToken);
 		}
@@ -1526,7 +1652,7 @@ namespace EpicGames.Perforce
 		/// <param name="FileSpecs">List of file specifications to query</param>
 		/// <param name="CancellationToken">Token used to cancel the operation</param>
 		/// <returns>List of responses from the server</returns>
-		public static async Task<List<FileLogRecord>> FileLogAsync(this PerforceConnection Connection, int MaxChanges, FileLogOptions Options, string[] FileSpecs, CancellationToken CancellationToken)
+		public static async Task<List<FileLogRecord>> FileLogAsync(this IPerforceConnection Connection, int MaxChanges, FileLogOptions Options, string[] FileSpecs, CancellationToken CancellationToken)
 		{
 			return (await TryFileLogAsync(Connection, MaxChanges, Options, FileSpecs, CancellationToken)).Data;
 		}
@@ -1540,7 +1666,7 @@ namespace EpicGames.Perforce
 		/// <param name="FileSpecs">List of file specifications to query</param>
 		/// <param name="CancellationToken">Token used to cancel the operation</param>
 		/// <returns>List of responses from the server</returns>
-		public static Task<PerforceResponseList<FileLogRecord>> TryFileLogAsync(this PerforceConnection Connection, int MaxChanges, FileLogOptions Options, string[] FileSpecs, CancellationToken CancellationToken)
+		public static Task<PerforceResponseList<FileLogRecord>> TryFileLogAsync(this IPerforceConnection Connection, int MaxChanges, FileLogOptions Options, string[] FileSpecs, CancellationToken CancellationToken)
 		{
 			return TryFileLogAsync(Connection, -1, MaxChanges, Options, FileSpecs, CancellationToken);
 		}
@@ -1555,7 +1681,7 @@ namespace EpicGames.Perforce
 		/// <param name="FileSpecs">List of file specifications to query</param>
 		/// <param name="CancellationToken">Token used to cancel the operation</param>
 		/// <returns>List of responses from the server</returns>
-		public static async Task<List<FileLogRecord>> FileLogAsync(this PerforceConnection Connection, int ChangeNumber, int MaxChanges, FileLogOptions Options, string[] FileSpecs, CancellationToken CancellationToken)
+		public static async Task<List<FileLogRecord>> FileLogAsync(this IPerforceConnection Connection, int ChangeNumber, int MaxChanges, FileLogOptions Options, string[] FileSpecs, CancellationToken CancellationToken)
 		{
 			return (await TryFileLogAsync(Connection, ChangeNumber, MaxChanges, Options, FileSpecs, CancellationToken)).Data;
 		}
@@ -1570,7 +1696,7 @@ namespace EpicGames.Perforce
 		/// <param name="FileSpecs">List of file specifications to query</param>
 		/// <param name="CancellationToken">Token used to cancel the operation</param>
 		/// <returns>List of responses from the server</returns>
-		public static Task<PerforceResponseList<FileLogRecord>> TryFileLogAsync(this PerforceConnection Connection, int ChangeNumber, int MaxChanges, FileLogOptions Options, string[] FileSpecs, CancellationToken CancellationToken)
+		public static Task<PerforceResponseList<FileLogRecord>> TryFileLogAsync(this IPerforceConnection Connection, int ChangeNumber, int MaxChanges, FileLogOptions Options, string[] FileSpecs, CancellationToken CancellationToken)
 		{
 			// Build the argument list
 			StringBuilder Arguments = new StringBuilder("filelog");
@@ -1626,7 +1752,7 @@ namespace EpicGames.Perforce
 		/// <param name="FileSpecs">List of file specifications to query</param>
 		/// <param name="CancellationToken">Token used to cancel the operation</param>
 		/// <returns>List of responses from the server</returns>
-		public static async Task<List<FStatRecord>> FStatAsync(this PerforceConnection Connection, FStatOptions Options, string[] FileSpecs, CancellationToken CancellationToken)
+		public static async Task<List<FStatRecord>> FStatAsync(this IPerforceConnection Connection, FStatOptions Options, string[] FileSpecs, CancellationToken CancellationToken)
 		{
 			return (await TryFStatAsync(Connection, Options, FileSpecs, CancellationToken)).Data;
 		}
@@ -1639,7 +1765,7 @@ namespace EpicGames.Perforce
 		/// <param name="FileSpecs">List of file specifications to query</param>
 		/// <param name="CancellationToken">Token used to cancel the operation</param>
 		/// <returns>List of responses from the server</returns>
-		public static Task<PerforceResponseList<FStatRecord>> TryFStatAsync(this PerforceConnection Connection, FStatOptions Options, string[] FileSpecs, CancellationToken CancellationToken)
+		public static Task<PerforceResponseList<FStatRecord>> TryFStatAsync(this IPerforceConnection Connection, FStatOptions Options, string[] FileSpecs, CancellationToken CancellationToken)
 		{
 			return TryFStatAsync(Connection, -1, Options, FileSpecs, CancellationToken);
 		}
@@ -1653,7 +1779,7 @@ namespace EpicGames.Perforce
 		/// <param name="FileSpecs">List of file specifications to query</param>
 		/// <param name="CancellationToken">Token used to cancel the operation</param>
 		/// <returns>List of responses from the server</returns>
-		public static async Task<List<FStatRecord>> FStatAsync(this PerforceConnection Connection, int MaxFiles, FStatOptions Options, string[] FileSpecs, CancellationToken CancellationToken)
+		public static async Task<List<FStatRecord>> FStatAsync(this IPerforceConnection Connection, int MaxFiles, FStatOptions Options, string[] FileSpecs, CancellationToken CancellationToken)
 		{
 			return (await TryFStatAsync(Connection, MaxFiles, Options, FileSpecs, CancellationToken)).Data;
 		}
@@ -1667,7 +1793,7 @@ namespace EpicGames.Perforce
 		/// <param name="FileSpecs">List of file specifications to query</param>
 		/// <param name="CancellationToken">Token used to cancel the operation</param>
 		/// <returns>List of responses from the server</returns>
-		public static Task<PerforceResponseList<FStatRecord>> TryFStatAsync(this PerforceConnection Connection, int MaxFiles, FStatOptions Options, string[] FileSpecs, CancellationToken CancellationToken)
+		public static Task<PerforceResponseList<FStatRecord>> TryFStatAsync(this IPerforceConnection Connection, int MaxFiles, FStatOptions Options, string[] FileSpecs, CancellationToken CancellationToken)
 		{
 			return TryFStatAsync(Connection, -1, -1, null, MaxFiles, Options, FileSpecs, CancellationToken);
 		}
@@ -1684,7 +1810,7 @@ namespace EpicGames.Perforce
 		/// <param name="FileSpecs">List of file specifications to query</param>
 		/// <param name="CancellationToken">Token used to cancel the operation</param>
 		/// <returns>List of responses from the server</returns>
-		public static async Task<List<FStatRecord>> FStatAsync(this PerforceConnection Connection, int AfterChangeNumber, int OnlyChangeNumber, string? Filter, int MaxFiles, FStatOptions Options, string[] FileSpecs, CancellationToken CancellationToken)
+		public static async Task<List<FStatRecord>> FStatAsync(this IPerforceConnection Connection, int AfterChangeNumber, int OnlyChangeNumber, string? Filter, int MaxFiles, FStatOptions Options, string[] FileSpecs, CancellationToken CancellationToken)
 		{
 			return (await TryFStatAsync(Connection, AfterChangeNumber, OnlyChangeNumber, Filter, MaxFiles, Options, FileSpecs, CancellationToken)).Data;
 		}
@@ -1701,7 +1827,7 @@ namespace EpicGames.Perforce
 		/// <param name="FileSpecs">List of file specifications to query</param>
 		/// <param name="CancellationToken">Token used to cancel the operation</param>
 		/// <returns>List of responses from the server</returns>
-		public static async Task<PerforceResponseList<FStatRecord>> TryFStatAsync(this PerforceConnection Connection, int AfterChangeNumber, int OnlyChangeNumber, string? Filter, int MaxFiles, FStatOptions Options, string[] FileSpecs, CancellationToken CancellationToken)
+		public static async Task<PerforceResponseList<FStatRecord>> TryFStatAsync(this IPerforceConnection Connection, int AfterChangeNumber, int OnlyChangeNumber, string? Filter, int MaxFiles, FStatOptions Options, string[] FileSpecs, CancellationToken CancellationToken)
 		{
 			// Build the argument list
 			StringBuilder Arguments = new StringBuilder("fstat");
@@ -1818,7 +1944,7 @@ namespace EpicGames.Perforce
 		/// <param name="Options">Options for the command</param>
 		/// <param name="CancellationToken">Token used to cancel the operation</param>
 		/// <returns>Response from the server; an InfoRecord or error code</returns>
-		public static async Task<InfoRecord> GetInfoAsync(this PerforceConnection Connection, InfoOptions Options, CancellationToken CancellationToken)
+		public static async Task<InfoRecord> GetInfoAsync(this IPerforceConnection Connection, InfoOptions Options, CancellationToken CancellationToken)
 		{
 			return (await TryGetInfoAsync(Connection, Options, CancellationToken)).Data;
 		}
@@ -1830,7 +1956,7 @@ namespace EpicGames.Perforce
 		/// <param name="Options">Options for the command</param>
 		/// <param name="CancellationToken">Token used to cancel the operation</param>
 		/// <returns>Response from the server; an InfoRecord or error code</returns>
-		public static Task<PerforceResponse<InfoRecord>> TryGetInfoAsync(this PerforceConnection Connection, InfoOptions Options, CancellationToken CancellationToken)
+		public static Task<PerforceResponse<InfoRecord>> TryGetInfoAsync(this IPerforceConnection Connection, InfoOptions Options, CancellationToken CancellationToken)
 		{
 			// Build the argument list
 			StringBuilder Arguments = new StringBuilder("info");
@@ -1856,7 +1982,7 @@ namespace EpicGames.Perforce
 		/// <param name="TargetFileSpec">The target filespec</param>
 		/// <param name="CancellationToken">Cancellation token</param>
 		/// <returns>List of records</returns>
-		public static Task<PerforceResponseList<MergeRecord>> MergeAsync(this PerforceConnection Connection, MergeOptions Options, int Change, int MaxFiles, string SourceFileSpec, string TargetFileSpec, CancellationToken CancellationToken)
+		public static Task<PerforceResponseList<MergeRecord>> MergeAsync(this IPerforceConnection Connection, MergeOptions Options, int Change, int MaxFiles, string SourceFileSpec, string TargetFileSpec, CancellationToken CancellationToken)
 		{
 			StringBuilder Arguments = new StringBuilder("merge");
 			if ((Options & MergeOptions.Preview) != 0)
@@ -1892,7 +2018,7 @@ namespace EpicGames.Perforce
 		/// <param name="TargetFileSpec">The target file(s)</param>
 		/// <param name="CancellationToken">Token used to cancel the operation</param>
 		/// <returns>Response from the server</returns>
-		public static async Task<List<MoveRecord>> MoveAsync(this PerforceConnection Connection, int ChangeNumber, string? FileType, MoveOptions Options, string SourceFileSpec, string TargetFileSpec, CancellationToken CancellationToken)
+		public static async Task<List<MoveRecord>> MoveAsync(this IPerforceConnection Connection, int ChangeNumber, string? FileType, MoveOptions Options, string SourceFileSpec, string TargetFileSpec, CancellationToken CancellationToken)
 		{
 			return (await TryMoveAsync(Connection, ChangeNumber, FileType, Options, SourceFileSpec, TargetFileSpec, CancellationToken)).Data;
 		}
@@ -1908,7 +2034,7 @@ namespace EpicGames.Perforce
 		/// <param name="TargetFileSpec">The target file(s)</param>
 		/// <param name="CancellationToken">Token used to cancel the operation</param>
 		/// <returns>Response from the server</returns>
-		public static Task<PerforceResponseList<MoveRecord>> TryMoveAsync(this PerforceConnection Connection, int ChangeNumber, string? FileType, MoveOptions Options, string SourceFileSpec, string TargetFileSpec, CancellationToken CancellationToken)
+		public static Task<PerforceResponseList<MoveRecord>> TryMoveAsync(this IPerforceConnection Connection, int ChangeNumber, string? FileType, MoveOptions Options, string SourceFileSpec, string TargetFileSpec, CancellationToken CancellationToken)
 		{
 			StringBuilder Arguments = new StringBuilder("move");
 			if (ChangeNumber != -1)
@@ -1952,7 +2078,7 @@ namespace EpicGames.Perforce
 		/// <param name="FileSpecs">Specification for the files to list</param>
 		/// <param name="CancellationToken">Token used to cancel the operation</param>
 		/// <returns>Response from the server</returns>
-		public static async Task<List<FStatRecord>> GetOpenFilesAsync(this PerforceConnection Connection, OpenedOptions Options, int ChangeNumber, string? ClientName, string? UserName, int MaxResults, string[]? FileSpecs, CancellationToken CancellationToken)
+		public static async Task<List<FStatRecord>> GetOpenFilesAsync(this IPerforceConnection Connection, OpenedOptions Options, int ChangeNumber, string? ClientName, string? UserName, int MaxResults, string[]? FileSpecs, CancellationToken CancellationToken)
 		{
 			return (await TryGetOpenFilesAsync(Connection, Options, ChangeNumber, ClientName, UserName, MaxResults, FileSpecs, CancellationToken)).Data;
 		}
@@ -1969,7 +2095,7 @@ namespace EpicGames.Perforce
 		/// <param name="FileSpecs">Specification for the files to list</param>
 		/// <param name="CancellationToken">Token used to cancel the operation</param>
 		/// <returns>Response from the server</returns>
-		public static Task<PerforceResponseList<FStatRecord>> TryGetOpenFilesAsync(this PerforceConnection Connection, OpenedOptions Options, int ChangeNumber, string? ClientName, string? UserName, int MaxResults, string[]? FileSpecs, CancellationToken CancellationToken)
+		public static Task<PerforceResponseList<FStatRecord>> TryGetOpenFilesAsync(this IPerforceConnection Connection, OpenedOptions Options, int ChangeNumber, string? ClientName, string? UserName, int MaxResults, string[]? FileSpecs, CancellationToken CancellationToken)
 		{
 			// Build the argument list
 			StringBuilder Arguments = new StringBuilder("opened");
@@ -2023,7 +2149,7 @@ namespace EpicGames.Perforce
 		/// <param name="FileSpec">Specification for the files to print</param>
 		/// <param name="CancellationToken">Token used to cancel the operation</param>
 		/// <returns>Response from the server</returns>
-		public static async Task<PrintRecord> PrintAsync(this PerforceConnection Connection, string OutputFile, string FileSpec, CancellationToken CancellationToken)
+		public static async Task<PrintRecord> PrintAsync(this IPerforceConnection Connection, string OutputFile, string FileSpec, CancellationToken CancellationToken)
 		{
 			return (await TryPrintAsync(Connection, OutputFile, FileSpec, CancellationToken)).Data;
 		}
@@ -2036,7 +2162,7 @@ namespace EpicGames.Perforce
 		/// <param name="FileSpec">Specification for the files to print</param>
 		/// <param name="CancellationToken">Token used to cancel the operation</param>
 		/// <returns>Response from the server</returns>
-		public static Task<PerforceResponse<PrintRecord>> TryPrintAsync(this PerforceConnection Connection, string OutputFile, string FileSpec, CancellationToken CancellationToken)
+		public static Task<PerforceResponse<PrintRecord>> TryPrintAsync(this IPerforceConnection Connection, string OutputFile, string FileSpec, CancellationToken CancellationToken)
 		{
 			// Build the argument list
 			StringBuilder Arguments = new StringBuilder("print");
@@ -2121,7 +2247,7 @@ namespace EpicGames.Perforce
 		/// <param name="FileSpecs">Files to be reverted</param>
 		/// <param name="CancellationToken">Token used to cancel the operation</param>
 		/// <returns>Response from the server</returns>
-		public static async Task<List<ReconcileRecord>> ReconcileAsync(this PerforceConnection Connection, int ChangeNumber, ReconcileOptions Options, string[] FileSpecs, CancellationToken CancellationToken)
+		public static async Task<List<ReconcileRecord>> ReconcileAsync(this IPerforceConnection Connection, int ChangeNumber, ReconcileOptions Options, string[] FileSpecs, CancellationToken CancellationToken)
 		{
 			PerforceResponseList<ReconcileRecord> Records = await TryReconcileAsync(Connection, ChangeNumber, Options, FileSpecs, CancellationToken);
 			Records.RemoveAll(x => x.Info != null);
@@ -2137,7 +2263,7 @@ namespace EpicGames.Perforce
 		/// <param name="FileSpecs">Files to be reverted</param>
 		/// <param name="CancellationToken">Token used to cancel the operation</param>
 		/// <returns>Response from the server</returns>
-		public static Task<PerforceResponseList<ReconcileRecord>> TryReconcileAsync(this PerforceConnection Connection, int ChangeNumber, ReconcileOptions Options, string[] FileSpecs, CancellationToken CancellationToken)
+		public static Task<PerforceResponseList<ReconcileRecord>> TryReconcileAsync(this IPerforceConnection Connection, int ChangeNumber, ReconcileOptions Options, string[] FileSpecs, CancellationToken CancellationToken)
 		{
 			StringBuilder Arguments = new StringBuilder("reconcile");
 			if(ChangeNumber != -1)
@@ -2197,7 +2323,7 @@ namespace EpicGames.Perforce
 		/// <param name="FileSpec">Files to be reverted</param>
 		/// <param name="CancellationToken">Token used to cancel the operation</param>
 		/// <returns>Response from the server</returns>
-		public static async Task<List<ReopenRecord>> ReopenAsync(this PerforceConnection Connection, int? ChangeNumber, string? FileType, string FileSpec, CancellationToken CancellationToken)
+		public static async Task<List<ReopenRecord>> ReopenAsync(this IPerforceConnection Connection, int? ChangeNumber, string? FileType, string FileSpec, CancellationToken CancellationToken)
 		{
 			return (await TryReopenAsync(Connection, ChangeNumber, FileType, FileSpec, CancellationToken)).Data;
 		}
@@ -2211,7 +2337,7 @@ namespace EpicGames.Perforce
 		/// <param name="FileSpec">Files to be reverted</param>
 		/// <param name="CancellationToken">Token used to cancel the operation</param>
 		/// <returns>Response from the server</returns>
-		public static Task<PerforceResponseList<ReopenRecord>> TryReopenAsync(this PerforceConnection Connection, int? ChangeNumber, string? FileType, string FileSpec, CancellationToken CancellationToken)
+		public static Task<PerforceResponseList<ReopenRecord>> TryReopenAsync(this IPerforceConnection Connection, int? ChangeNumber, string? FileType, string FileSpec, CancellationToken CancellationToken)
 		{
 			StringBuilder Arguments = new StringBuilder("reopen");
 			if (ChangeNumber != null)
@@ -2246,7 +2372,7 @@ namespace EpicGames.Perforce
 		/// <param name="SourceServerId">The source server id</param>
 		/// <param name="CancellationToken">Token used to cancel the operation</param>
 		/// <returns>Response from the server</returns>
-		public static Task<List<PerforceResponse>> ReloadClient(this PerforceConnection Connection, string ClientName, string SourceServerId, CancellationToken CancellationToken)
+		public static Task<List<PerforceResponse>> ReloadClient(this IPerforceConnection Connection, string ClientName, string SourceServerId, CancellationToken CancellationToken)
 		{
 			return Connection.CommandAsync($"reload -c \"{ClientName}\" -p \"{SourceServerId}\"", null, null, CancellationToken);
 		}
@@ -2264,7 +2390,7 @@ namespace EpicGames.Perforce
 		/// <param name="FileSpecs">Files to be reverted</param>
 		/// <param name="CancellationToken">Token used to cancel the operation</param>
 		/// <returns>Response from the server</returns>
-		public static async Task<List<ResolveRecord>> ResolveAsync(this PerforceConnection Connection, int ChangeNumber, ResolveOptions Options, string[]? FileSpecs, CancellationToken CancellationToken)
+		public static async Task<List<ResolveRecord>> ResolveAsync(this IPerforceConnection Connection, int ChangeNumber, ResolveOptions Options, string[]? FileSpecs, CancellationToken CancellationToken)
 		{
 			return (await TryResolveAsync(Connection, ChangeNumber, Options, FileSpecs, CancellationToken)).Data;
 		}
@@ -2278,7 +2404,7 @@ namespace EpicGames.Perforce
 		/// <param name="FileSpecs">Files to be reverted</param>
 		/// <param name="CancellationToken">Token used to cancel the operation</param>
 		/// <returns>Response from the server</returns>
-		public static async Task<PerforceResponseList<ResolveRecord>> TryResolveAsync(this PerforceConnection Connection, int ChangeNumber, ResolveOptions Options, string[]? FileSpecs, CancellationToken CancellationToken)
+		public static async Task<PerforceResponseList<ResolveRecord>> TryResolveAsync(this IPerforceConnection Connection, int ChangeNumber, ResolveOptions Options, string[]? FileSpecs, CancellationToken CancellationToken)
 		{
 			StringBuilder Arguments = new StringBuilder("resolve");
 			if ((Options & ResolveOptions.Automatic) != 0)
@@ -2348,7 +2474,7 @@ namespace EpicGames.Perforce
 		/// <param name="FileSpecs">Files to be reverted</param>
 		/// <param name="CancellationToken">Token used to cancel the operation</param>
 		/// <returns>Response from the server</returns>
-		public static async Task<List<RevertRecord>> RevertAsync(this PerforceConnection Connection, int ChangeNumber, string? ClientName, RevertOptions Options, string[] FileSpecs, CancellationToken CancellationToken)
+		public static async Task<List<RevertRecord>> RevertAsync(this IPerforceConnection Connection, int ChangeNumber, string? ClientName, RevertOptions Options, string[] FileSpecs, CancellationToken CancellationToken)
 		{
 			return (await TryRevertAsync(Connection, ChangeNumber, ClientName, Options, FileSpecs, CancellationToken)).Data;
 		}
@@ -2363,7 +2489,7 @@ namespace EpicGames.Perforce
 		/// <param name="FileSpecs">Files to be reverted</param>
 		/// <param name="CancellationToken">Token used to cancel the operation</param>
 		/// <returns>Response from the server</returns>
-		public static async Task<PerforceResponseList<RevertRecord>> TryRevertAsync(this PerforceConnection Connection, int ChangeNumber, string? ClientName, RevertOptions Options, string[] FileSpecs, CancellationToken CancellationToken)
+		public static async Task<PerforceResponseList<RevertRecord>> TryRevertAsync(this IPerforceConnection Connection, int ChangeNumber, string? ClientName, RevertOptions Options, string[] FileSpecs, CancellationToken CancellationToken)
 		{
 			StringBuilder Arguments = new StringBuilder("revert");
 			if((Options & RevertOptions.Unchanged) != 0)
@@ -2398,79 +2524,6 @@ namespace EpicGames.Perforce
 
 		#endregion
 
-		#region p4 set
-
-		/// <summary>
-		/// Sets an environment variable
-		/// </summary>
-		/// <param name="Connection">Connection to the Perforce server</param>
-		/// <param name="Name">Name of the variable to set</param>
-		/// <param name="Value">Value for the variable</param>
-		/// <param name="CancellationToken">Token used to cancel the operation</param>
-		/// <returns>Response from the server</returns>
-		public static async Task SetAsync(this PerforceConnection Connection, string Name, string Value, CancellationToken CancellationToken)
-		{
-			Tuple<bool, string> Result = await TrySetAsync(Connection, Name, Value, CancellationToken);
-			if (!Result.Item1)
-			{
-				throw new PerforceException(Result.Item2);
-			}
-		}
-
-		/// <summary>
-		/// Sets an environment variable
-		/// </summary>
-		/// <param name="Connection">Connection to the Perforce server</param>
-		/// <param name="Name">Name of the variable to set</param>
-		/// <param name="Value">Value for the variable</param>
-		/// <param name="CancellationToken">Token used to cancel the operation</param>
-		/// <returns>Response from the server</returns>
-		public static async Task<Tuple<bool, string>> TrySetAsync(this PerforceConnection Connection, string Name, string Value, CancellationToken CancellationToken)
-		{
-			StringBuilder Arguments = new StringBuilder("set");
-			Arguments.Append($" {Name}=");
-			Arguments.Append($"\"{Value}\"");
-
-			using (PerforceChildProcess ChildProcess = new PerforceChildProcess(null, Arguments.ToString(), Connection.Logger))
-			{
-				return await ChildProcess.TryReadToEndAsync(CancellationToken);
-			}
-		}
-
-		/// <summary>
-		/// Gets the setting of a Perforce variable
-		/// </summary>
-		/// <param name="Connection">Connection to the Perforce server</param>
-		/// <param name="Name">Name of the variable to get</param>
-		/// <param name="CancellationToken">Cancellation token for the request</param>
-		/// <returns>Value of the variable</returns>
-		public static async Task<string?> TryGetSettingAsync(this PerforceConnection Connection, string Name, CancellationToken CancellationToken)
-		{
-			using (PerforceChildProcess ChildProcess = new PerforceChildProcess(null, $"set {Name}=", Connection.Logger))
-			{
-				Tuple<bool, string> Response = await ChildProcess.TryReadToEndAsync(CancellationToken);
-				if (Response.Item1)
-				{
-					string TrimResponse = Response.Item2.Trim();
-					if (TrimResponse.Length == 0)
-					{
-						return null;
-					}
-
-					string Prefix = Name + "=";
-					if (TrimResponse.StartsWith(Prefix, StringComparison.OrdinalIgnoreCase))
-					{
-						return TrimResponse.Substring(Prefix.Length);
-					}
-				}
-
-				Connection.Logger.LogDebug($"Unable to get '{Name}' variable: {Response.Item2}");
-				return null;
-			}
-		}
-
-		#endregion
-
 		#region p4 shelve
 
 		/// <summary>
@@ -2482,7 +2535,7 @@ namespace EpicGames.Perforce
 		/// <param name="FileSpecs">Files to sync</param>
 		/// <param name="CancellationToken">Token used to cancel the operation</param>
 		/// <returns>Response from the server</returns>
-		public static async Task<List<ShelveRecord>> ShelveAsync(this PerforceConnection Connection, int ChangeNumber, ShelveOptions Options, string[] FileSpecs, CancellationToken CancellationToken)
+		public static async Task<List<ShelveRecord>> ShelveAsync(this IPerforceConnection Connection, int ChangeNumber, ShelveOptions Options, string[] FileSpecs, CancellationToken CancellationToken)
 		{
 			return (await TryShelveAsync(Connection, ChangeNumber, Options, FileSpecs, CancellationToken)).Data;
 		}
@@ -2496,7 +2549,7 @@ namespace EpicGames.Perforce
 		/// <param name="FileSpecs">Files to sync</param>
 		/// <param name="CancellationToken">Token used to cancel the operation</param>
 		/// <returns>Response from the server</returns>
-		public static async Task<PerforceResponseList<ShelveRecord>> TryShelveAsync(this PerforceConnection Connection, int ChangeNumber, ShelveOptions Options, string[] FileSpecs, CancellationToken CancellationToken)
+		public static async Task<PerforceResponseList<ShelveRecord>> TryShelveAsync(this IPerforceConnection Connection, int ChangeNumber, ShelveOptions Options, string[] FileSpecs, CancellationToken CancellationToken)
 		{
 			StringBuilder Arguments = new StringBuilder("shelve");
 			Arguments.AppendFormat(" -c {0}", ChangeNumber);
@@ -2526,7 +2579,7 @@ namespace EpicGames.Perforce
 		/// <param name="FileSpecs">Files to delete</param>
 		/// <param name="CancellationToken">Token used to cancel the operation</param>
 		/// <returns>Response from the server</returns>
-		public static async Task DeleteShelvedFilesAsync(this PerforceConnection Connection, int ChangeNumber, string[] FileSpecs, CancellationToken CancellationToken)
+		public static async Task DeleteShelvedFilesAsync(this IPerforceConnection Connection, int ChangeNumber, string[] FileSpecs, CancellationToken CancellationToken)
 		{
 			List<PerforceResponse> Responses = await TryDeleteShelvedFilesAsync(Connection, ChangeNumber, FileSpecs, CancellationToken);
 
@@ -2545,7 +2598,7 @@ namespace EpicGames.Perforce
 		/// <param name="FileSpecs">Files to delete</param>
 		/// <param name="CancellationToken">Token used to cancel the operation</param>
 		/// <returns>Response from the server</returns>
-		public static Task<List<PerforceResponse>> TryDeleteShelvedFilesAsync(this PerforceConnection Connection, int ChangeNumber, string[] FileSpecs, CancellationToken CancellationToken)
+		public static Task<List<PerforceResponse>> TryDeleteShelvedFilesAsync(this IPerforceConnection Connection, int ChangeNumber, string[] FileSpecs, CancellationToken CancellationToken)
 		{
 			StringBuilder Arguments = new StringBuilder("shelve -d");
 			if(ChangeNumber != -1)
@@ -2571,7 +2624,7 @@ namespace EpicGames.Perforce
 		/// <param name="StreamPath">The path for streams to enumerate (eg. "//UE4/...")</param>
 		/// <param name="CancellationToken">Token used to cancel the operation</param>
 		/// <returns>List of streams matching the given criteria</returns>
-		public static async Task<List<StreamsRecord>> GetStreamsAsync(this PerforceConnection Connection, string StreamPath, CancellationToken CancellationToken)
+		public static async Task<List<StreamsRecord>> GetStreamsAsync(this IPerforceConnection Connection, string StreamPath, CancellationToken CancellationToken)
 		{
 			return (await TryGetStreamsAsync(Connection, StreamPath, CancellationToken)).Data;
 		}
@@ -2583,7 +2636,7 @@ namespace EpicGames.Perforce
 		/// <param name="StreamPath">The path for streams to enumerate (eg. "//UE4/...")</param>
 		/// <param name="CancellationToken">Token used to cancel the operation</param>
 		/// <returns>List of streams matching the given criteria</returns>
-		public static Task<PerforceResponseList<StreamsRecord>> TryGetStreamsAsync(this PerforceConnection Connection, string StreamPath, CancellationToken CancellationToken)
+		public static Task<PerforceResponseList<StreamsRecord>> TryGetStreamsAsync(this IPerforceConnection Connection, string StreamPath, CancellationToken CancellationToken)
 		{
 			return TryGetStreamsAsync(Connection, StreamPath, -1, null, false, CancellationToken);
 		}
@@ -2598,7 +2651,7 @@ namespace EpicGames.Perforce
 		/// <param name="bUnloaded">Whether to enumerate unloaded workspaces</param>
 		/// <param name="CancellationToken">Token used to cancel the operation</param>
 		/// <returns>List of streams matching the given criteria</returns>
-		public static async Task<List<StreamsRecord>> GetStreamsAsync(this PerforceConnection Connection, string StreamPath, int MaxResults, string? Filter, bool bUnloaded, CancellationToken CancellationToken)
+		public static async Task<List<StreamsRecord>> GetStreamsAsync(this IPerforceConnection Connection, string StreamPath, int MaxResults, string? Filter, bool bUnloaded, CancellationToken CancellationToken)
 		{
 			return (await TryGetStreamsAsync(Connection, StreamPath, MaxResults, Filter, bUnloaded, CancellationToken)).Data;
 		}
@@ -2613,7 +2666,7 @@ namespace EpicGames.Perforce
 		/// <param name="bUnloaded">Whether to enumerate unloaded workspaces</param>
 		/// <param name="CancellationToken">Token used to cancel the operation</param>
 		/// <returns>List of streams matching the given criteria</returns>
-		public static Task<PerforceResponseList<StreamsRecord>> TryGetStreamsAsync(this PerforceConnection Connection, string StreamPath, int MaxResults, string? Filter, bool bUnloaded, CancellationToken CancellationToken)
+		public static Task<PerforceResponseList<StreamsRecord>> TryGetStreamsAsync(this IPerforceConnection Connection, string StreamPath, int MaxResults, string? Filter, bool bUnloaded, CancellationToken CancellationToken)
 		{
 			// Build the command line
 			StringBuilder Arguments = new StringBuilder("streams");
@@ -2647,7 +2700,7 @@ namespace EpicGames.Perforce
 		/// <param name="Options">Options for the command</param>
 		/// <param name="CancellationToken">Token used to cancel the operation</param>
 		/// <returns>Response from the server</returns>
-		public static async Task<SubmitRecord> SubmitAsync(this PerforceConnection Connection, int ChangeNumber, SubmitOptions Options, CancellationToken CancellationToken)
+		public static async Task<SubmitRecord> SubmitAsync(this IPerforceConnection Connection, int ChangeNumber, SubmitOptions Options, CancellationToken CancellationToken)
 		{
 			return (await TrySubmitAsync(Connection, ChangeNumber, Options, CancellationToken)).Data;
 		}
@@ -2660,7 +2713,7 @@ namespace EpicGames.Perforce
 		/// <param name="Options">Options for the command</param>
 		/// <param name="CancellationToken">Token used to cancel the operation</param>
 		/// <returns>Response from the server</returns>
-		public static async Task<PerforceResponse<SubmitRecord>> TrySubmitAsync(this PerforceConnection Connection, int ChangeNumber, SubmitOptions Options, CancellationToken CancellationToken)
+		public static async Task<PerforceResponse<SubmitRecord>> TrySubmitAsync(this IPerforceConnection Connection, int ChangeNumber, SubmitOptions Options, CancellationToken CancellationToken)
 		{
 			StringBuilder Arguments = new StringBuilder("submit");
 			if((Options & SubmitOptions.ReopenAsEdit) != 0)
@@ -2679,7 +2732,7 @@ namespace EpicGames.Perforce
 		/// <param name="ChangeNumber">The changelist to submit</param>
 		/// <param name="CancellationToken">Token used to cancel the operation</param>
 		/// <returns>Response from the server</returns>
-		public static async Task<SubmitRecord> SubmitShelvedAsync(this PerforceConnection Connection, int ChangeNumber, CancellationToken CancellationToken)
+		public static async Task<SubmitRecord> SubmitShelvedAsync(this IPerforceConnection Connection, int ChangeNumber, CancellationToken CancellationToken)
 		{
 			return (await TrySubmitShelvedAsync(Connection, ChangeNumber, CancellationToken)).Data;
 		}
@@ -2691,7 +2744,7 @@ namespace EpicGames.Perforce
 		/// <param name="ChangeNumber">The changelist to submit</param>
 		/// <param name="CancellationToken">Token used to cancel the operation</param>
 		/// <returns>Response from the server</returns>
-		public static async Task<PerforceResponse<SubmitRecord>> TrySubmitShelvedAsync(this PerforceConnection Connection, int ChangeNumber, CancellationToken CancellationToken)
+		public static async Task<PerforceResponse<SubmitRecord>> TrySubmitShelvedAsync(this IPerforceConnection Connection, int ChangeNumber, CancellationToken CancellationToken)
 		{
 			StringBuilder Arguments = new StringBuilder("submit");
 			Arguments.AppendFormat(" -e {0}", ChangeNumber);
@@ -2710,7 +2763,7 @@ namespace EpicGames.Perforce
 		/// <param name="FileSpecs">Files to sync</param>
 		/// <param name="CancellationToken">Token used to cancel the operation</param>
 		/// <returns>Response from the server</returns>
-		public static async Task<List<SyncRecord>> SyncAsync(this PerforceConnection Connection, string[] FileSpecs, CancellationToken CancellationToken)
+		public static async Task<List<SyncRecord>> SyncAsync(this IPerforceConnection Connection, string[] FileSpecs, CancellationToken CancellationToken)
 		{
 			return (await TrySyncAsync(Connection, FileSpecs, CancellationToken)).Data;
 		}
@@ -2722,7 +2775,7 @@ namespace EpicGames.Perforce
 		/// <param name="FileSpecs">Files to sync</param>
 		/// <param name="CancellationToken">Token used to cancel the operation</param>
 		/// <returns>Response from the server</returns>
-		public static Task<PerforceResponseList<SyncRecord>> TrySyncAsync(this PerforceConnection Connection, string[] FileSpecs, CancellationToken CancellationToken)
+		public static Task<PerforceResponseList<SyncRecord>> TrySyncAsync(this IPerforceConnection Connection, string[] FileSpecs, CancellationToken CancellationToken)
 		{
 			return TrySyncAsync(Connection, SyncOptions.None, -1, FileSpecs, CancellationToken);
 		}
@@ -2736,7 +2789,7 @@ namespace EpicGames.Perforce
 		/// <param name="FileSpecs">Files to sync</param>
 		/// <param name="CancellationToken">Token used to cancel the operation</param>
 		/// <returns>Response from the server</returns>
-		public static async Task<List<SyncRecord>> SyncAsync(this PerforceConnection Connection, SyncOptions Options, int MaxFiles, string[] FileSpecs, CancellationToken CancellationToken)
+		public static async Task<List<SyncRecord>> SyncAsync(this IPerforceConnection Connection, SyncOptions Options, int MaxFiles, string[] FileSpecs, CancellationToken CancellationToken)
 		{
 			return (await TrySyncAsync(Connection, Options, MaxFiles, FileSpecs, CancellationToken)).Data;
 		}
@@ -2750,7 +2803,7 @@ namespace EpicGames.Perforce
 		/// <param name="FileSpecs">Files to sync</param>
 		/// <param name="CancellationToken">Token used to cancel the operation</param>
 		/// <returns>Response from the server</returns>
-		public static Task<PerforceResponseList<SyncRecord>> TrySyncAsync(this PerforceConnection Connection, SyncOptions Options, int MaxFiles, string[] FileSpecs, CancellationToken CancellationToken)
+		public static Task<PerforceResponseList<SyncRecord>> TrySyncAsync(this IPerforceConnection Connection, SyncOptions Options, int MaxFiles, string[] FileSpecs, CancellationToken CancellationToken)
 		{
 			return TrySyncAsync(Connection, Options, MaxFiles, -1, -1, -1, -1, -1, FileSpecs, CancellationToken);
 		}
@@ -2769,7 +2822,7 @@ namespace EpicGames.Perforce
 		/// <param name="FileSpecs">Files to sync</param>
 		/// <param name="CancellationToken">Token used to cancel the operation</param>
 		/// <returns>Response from the server</returns>
-		public static async Task<List<SyncRecord>> SyncAsync(this PerforceConnection Connection, SyncOptions Options, int MaxFiles, int NumThreads, int Batch, int BatchSize, int Min, int MinSize, string[] FileSpecs, CancellationToken CancellationToken)
+		public static async Task<List<SyncRecord>> SyncAsync(this IPerforceConnection Connection, SyncOptions Options, int MaxFiles, int NumThreads, int Batch, int BatchSize, int Min, int MinSize, string[] FileSpecs, CancellationToken CancellationToken)
 		{
 			return (await TrySyncAsync(Connection, Options, MaxFiles, NumThreads, Batch, BatchSize, Min, MinSize, FileSpecs, CancellationToken)).Data;
 		}
@@ -2788,7 +2841,7 @@ namespace EpicGames.Perforce
 		/// <param name="FileSpecs">Files to sync</param>
 		/// <param name="CancellationToken">Token used to cancel the operation</param>
 		/// <returns>Response from the server</returns>
-		public static async Task<PerforceResponseList<SyncRecord>> TrySyncAsync(this PerforceConnection Connection, SyncOptions Options, int MaxFiles, int NumThreads, int Batch, int BatchSize, int Min, int MinSize, string[] FileSpecs, CancellationToken CancellationToken)
+		public static async Task<PerforceResponseList<SyncRecord>> TrySyncAsync(this IPerforceConnection Connection, SyncOptions Options, int MaxFiles, int NumThreads, int Batch, int BatchSize, int Min, int MinSize, string[] FileSpecs, CancellationToken CancellationToken)
 		{
 			// Perforce annoyingly returns 'up-to-date' as an error. Ignore it.
 			PerforceResponseList<SyncRecord> Records = await SyncInternalAsync<SyncRecord>(Connection, Options, MaxFiles, NumThreads, Batch, BatchSize, Min, MinSize, FileSpecs, false, CancellationToken);
@@ -2805,7 +2858,7 @@ namespace EpicGames.Perforce
 		/// <param name="FileSpecs">Files to sync</param>
 		/// <param name="CancellationToken">Token used to cancel the operation</param>
 		/// <returns>Response from the server</returns>
-		public static async Task<List<SyncSummaryRecord>> SyncQuietAsync(this PerforceConnection Connection, SyncOptions Options, int MaxFiles, string[] FileSpecs, CancellationToken CancellationToken)
+		public static async Task<List<SyncSummaryRecord>> SyncQuietAsync(this IPerforceConnection Connection, SyncOptions Options, int MaxFiles, string[] FileSpecs, CancellationToken CancellationToken)
 		{
 			PerforceResponseList<SyncSummaryRecord> Records = await TrySyncQuietAsync(Connection, Options, MaxFiles, FileSpecs, CancellationToken);
 			Records.RemoveAll(x => (x.Error != null && x.Error.Generic == PerforceGenericCode.Empty) || x.Info != null);
@@ -2821,7 +2874,7 @@ namespace EpicGames.Perforce
 		/// <param name="FileSpecs">Files to sync</param>
 		/// <param name="CancellationToken">Token used to cancel the operation</param>
 		/// <returns>Response from the server</returns>
-		public static Task<PerforceResponseList<SyncSummaryRecord>> TrySyncQuietAsync(this PerforceConnection Connection, SyncOptions Options, int MaxFiles, string[] FileSpecs, CancellationToken CancellationToken)
+		public static Task<PerforceResponseList<SyncSummaryRecord>> TrySyncQuietAsync(this IPerforceConnection Connection, SyncOptions Options, int MaxFiles, string[] FileSpecs, CancellationToken CancellationToken)
 		{
 			return TrySyncQuietAsync(Connection, Options, MaxFiles, -1, -1, -1, -1, -1, FileSpecs, null, CancellationToken);
 		}
@@ -2841,7 +2894,7 @@ namespace EpicGames.Perforce
 		/// <param name="ResponseFile">Response file for list of files to sync</param>
 		/// <param name="CancellationToken">Token used to cancel the operation</param>
 		/// <returns>Response from the server</returns>
-		public static async Task<List<SyncSummaryRecord>> SyncQuietAsync(this PerforceConnection Connection, SyncOptions Options, int MaxFiles, int NumThreads, int Batch, int BatchSize, int Min, int MinSize, string[] FileSpecs, FileReference? ResponseFile, CancellationToken CancellationToken)
+		public static async Task<List<SyncSummaryRecord>> SyncQuietAsync(this IPerforceConnection Connection, SyncOptions Options, int MaxFiles, int NumThreads, int Batch, int BatchSize, int Min, int MinSize, string[] FileSpecs, FileReference? ResponseFile, CancellationToken CancellationToken)
 		{
 			return (await TrySyncQuietAsync(Connection, Options, MaxFiles, NumThreads, Batch, BatchSize, Min, MinSize, FileSpecs, ResponseFile, CancellationToken)).Data;
 		}
@@ -2861,7 +2914,7 @@ namespace EpicGames.Perforce
 		/// <param name="ResponseFile">Response file for list of files to sync</param>
 		/// <param name="CancellationToken">Token used to cancel the operation</param>
 		/// <returns>Response from the server</returns>
-		public static Task<PerforceResponseList<SyncSummaryRecord>> TrySyncQuietAsync(this PerforceConnection Connection, SyncOptions Options, int MaxFiles, int NumThreads, int Batch, int BatchSize, int Min, int MinSize, string[] FileSpecs, FileReference? ResponseFile, CancellationToken CancellationToken)
+		public static Task<PerforceResponseList<SyncSummaryRecord>> TrySyncQuietAsync(this IPerforceConnection Connection, SyncOptions Options, int MaxFiles, int NumThreads, int Batch, int BatchSize, int Min, int MinSize, string[] FileSpecs, FileReference? ResponseFile, CancellationToken CancellationToken)
 		{
 			return SyncInternalAsync<SyncSummaryRecord>(Connection, Options, MaxFiles, NumThreads, Batch, BatchSize, Min, MinSize, FileSpecs, true, CancellationToken);
 		}
@@ -2881,7 +2934,7 @@ namespace EpicGames.Perforce
 		/// <param name="bQuiet">Whether to use quiet output</param>
 		/// <param name="CancellationToken">Token used to cancel the operation</param>
 		/// <returns>Arguments for the command</returns>
-		private static async Task<PerforceResponseList<T>> SyncInternalAsync<T>(this PerforceConnection Connection, SyncOptions Options, int MaxFiles, int NumThreads, int Batch, int BatchSize, int Min, int MinSize, string[] FileSpecs, bool bQuiet, CancellationToken CancellationToken) where T : class
+		private static async Task<PerforceResponseList<T>> SyncInternalAsync<T>(this IPerforceConnection Connection, SyncOptions Options, int MaxFiles, int NumThreads, int Batch, int BatchSize, int Min, int MinSize, string[] FileSpecs, bool bQuiet, CancellationToken CancellationToken) where T : class
 		{
 			StringBuilder Arguments = new StringBuilder("sync");
 			if((Options & SyncOptions.Force) != 0)
@@ -2985,7 +3038,7 @@ namespace EpicGames.Perforce
 		/// <param name="FileSpecs">Files to unshelve</param>
 		/// <param name="CancellationToken">Token used to cancel the operation</param>
 		/// <returns>Response from the server</returns>
-		public static async Task<List<UnshelveRecord>> UnshelveAsync(this PerforceConnection Connection, int ChangeNumber, int IntoChangeNumber, string UsingBranchSpec, string UsingStream, string ForceParentStream, UnshelveOptions Options, string[] FileSpecs, CancellationToken CancellationToken)
+		public static async Task<List<UnshelveRecord>> UnshelveAsync(this IPerforceConnection Connection, int ChangeNumber, int IntoChangeNumber, string UsingBranchSpec, string UsingStream, string ForceParentStream, UnshelveOptions Options, string[] FileSpecs, CancellationToken CancellationToken)
 		{
 			return (await TryUnshelveAsync(Connection, ChangeNumber, IntoChangeNumber, UsingBranchSpec, UsingStream, ForceParentStream, Options, FileSpecs, CancellationToken)).Data;
 		}
@@ -3003,7 +3056,7 @@ namespace EpicGames.Perforce
 		/// <param name="FileSpecs">Files to unshelve</param>
 		/// <param name="CancellationToken">Token used to cancel the operation</param>
 		/// <returns>Response from the server</returns>
-		public static Task<PerforceResponseList<UnshelveRecord>> TryUnshelveAsync(this PerforceConnection Connection, int ChangeNumber, int IntoChangeNumber, string UsingBranchSpec, string UsingStream, string ForceParentStream, UnshelveOptions Options, string[] FileSpecs, CancellationToken CancellationToken)
+		public static Task<PerforceResponseList<UnshelveRecord>> TryUnshelveAsync(this IPerforceConnection Connection, int ChangeNumber, int IntoChangeNumber, string UsingBranchSpec, string UsingStream, string ForceParentStream, UnshelveOptions Options, string[] FileSpecs, CancellationToken CancellationToken)
 		{
 			StringBuilder Arguments = new StringBuilder("unshelve");
 			Arguments.AppendFormat(" -s {0}", ChangeNumber);
@@ -3050,7 +3103,7 @@ namespace EpicGames.Perforce
 		/// <param name="UserName">Name of the user to fetch information for</param>
 		/// <param name="CancellationToken">Token used to cancel the operation</param>
 		/// <returns>Response from the server</returns>
-		public static async Task<UserRecord> GetUserAsync(this PerforceConnection Connection, string UserName, CancellationToken CancellationToken)
+		public static async Task<UserRecord> GetUserAsync(this IPerforceConnection Connection, string UserName, CancellationToken CancellationToken)
 		{
 			return (await TryGetUserAsync(Connection, UserName, CancellationToken))[0].Data;
 		}
@@ -3062,7 +3115,7 @@ namespace EpicGames.Perforce
 		/// <param name="UserName">Name of the user to fetch information for</param>
 		/// <param name="CancellationToken">Token used to cancel the operation</param>
 		/// <returns>Response from the server</returns>
-		public static Task<PerforceResponseList<UserRecord>> TryGetUserAsync(this PerforceConnection Connection, string UserName, CancellationToken CancellationToken)
+		public static Task<PerforceResponseList<UserRecord>> TryGetUserAsync(this IPerforceConnection Connection, string UserName, CancellationToken CancellationToken)
 		{
 			StringBuilder Arguments = new StringBuilder("user");
 			Arguments.AppendFormat(" -o \"{0}\"", UserName);
@@ -3080,7 +3133,7 @@ namespace EpicGames.Perforce
 		/// <param name="FileSpecs">Patterns for the files to query</param>
 		/// <param name="CancellationToken">Token used to cancel the operation</param>
 		/// <returns>List of responses from the server</returns>
-		public static async Task<List<WhereRecord>> WhereAsync(this PerforceConnection Connection, string[] FileSpecs, CancellationToken CancellationToken)
+		public static async Task<List<WhereRecord>> WhereAsync(this IPerforceConnection Connection, string[] FileSpecs, CancellationToken CancellationToken)
 		{
 			List<WhereRecord> Records = (await TryWhereAsync(Connection, FileSpecs, CancellationToken)).Data;
 			if (Records.Count != FileSpecs.Length)
@@ -3097,7 +3150,7 @@ namespace EpicGames.Perforce
 		/// <param name="FileSpecs">Patterns for the files to query</param>
 		/// <param name="CancellationToken">Token used to cancel the operation</param>
 		/// <returns>List of responses from the server</returns>
-		public static async Task<PerforceResponseList<WhereRecord>> TryWhereAsync(this PerforceConnection Connection, string[] FileSpecs, CancellationToken CancellationToken)
+		public static async Task<PerforceResponseList<WhereRecord>> TryWhereAsync(this IPerforceConnection Connection, string[] FileSpecs, CancellationToken CancellationToken)
 		{
 			PerforceResponseList<WhereRecord> Responses = new PerforceResponseList<WhereRecord>();
 			for (int Idx = 0; Idx < FileSpecs.Length; )
@@ -3124,7 +3177,7 @@ namespace EpicGames.Perforce
 		/// <param name="ChangeNumber">Changelist number to receive the changes</param>
 		/// <param name="CancellationToken">Token used to cancel the operation</param>
 		/// <returns>Response from the server</returns>
-		public static async Task UndoChangeAsync(this PerforceConnection Connection, int ChangeNumberToUndo, int ChangeNumber, CancellationToken CancellationToken)
+		public static async Task UndoChangeAsync(this IPerforceConnection Connection, int ChangeNumberToUndo, int ChangeNumber, CancellationToken CancellationToken)
 		{
 			(await TryUndoChangeAsync(Connection, ChangeNumberToUndo, ChangeNumber, CancellationToken))[0].EnsureSuccess();
 		}
@@ -3137,7 +3190,7 @@ namespace EpicGames.Perforce
 		/// <param name="ChangeNumber">Changelist number to receive the changes</param>
 		/// <param name="CancellationToken">Token used to cancel the operation</param>
 		/// <returns>Response from the server</returns>
-		public static Task<PerforceResponseList<UndoRecord>> TryUndoChangeAsync(this PerforceConnection Connection, int ChangeNumberToUndo, int ChangeNumber, CancellationToken CancellationToken)
+		public static Task<PerforceResponseList<UndoRecord>> TryUndoChangeAsync(this IPerforceConnection Connection, int ChangeNumberToUndo, int ChangeNumber, CancellationToken CancellationToken)
 		{
 			StringBuilder Arguments = new StringBuilder("undo -c");
 			
@@ -3157,7 +3210,7 @@ namespace EpicGames.Perforce
 		/// <param name="Options">Options for the anotate command</param>
 		/// <param name="CancellationToken">Token used to cancel the operation</param>
 		/// <returns>List of annotate records</returns>
-		public static async Task<List<AnnotateRecord>> AnnotateAsync(this PerforceConnection Connection, string FileSpec, AnnotateOptions Options, CancellationToken CancellationToken)
+		public static async Task<List<AnnotateRecord>> AnnotateAsync(this IPerforceConnection Connection, string FileSpec, AnnotateOptions Options, CancellationToken CancellationToken)
 		{
 			return (await TryAnnotateAsync(Connection, FileSpec, Options, CancellationToken)).Data;
 		}
@@ -3170,7 +3223,7 @@ namespace EpicGames.Perforce
 		/// <param name="Options">Options for the anotate command</param>
 		/// <param name="CancellationToken">Token used to cancel the operation</param>
 		/// <returns>List of annotate records</returns>
-		public static Task<PerforceResponseList<AnnotateRecord>> TryAnnotateAsync(this PerforceConnection Connection, string FileSpec, AnnotateOptions Options, CancellationToken CancellationToken)
+		public static Task<PerforceResponseList<AnnotateRecord>> TryAnnotateAsync(this IPerforceConnection Connection, string FileSpec, AnnotateOptions Options, CancellationToken CancellationToken)
 		{
 			PerforceResponseList<AnnotateRecord> Responses = new PerforceResponseList<AnnotateRecord>();
 			StringBuilder Arguments = new StringBuilder("annotate ");
