P4: Specify command arguments as a list rather than a concatenated string, for future compatibility with C++ library.

diff --git a/Engine/Source/Programs/Shared/EpicGames.Perforce.Managed/ManagedWorkspace.cs b/Engine/Source/Programs/Shared/EpicGames.Perforce.Managed/ManagedWorkspace.cs
--- a/Engine/Source/Programs/Shared/EpicGames.Perforce.Managed/ManagedWorkspace.cs
+++ b/Engine/Source/Programs/Shared/EpicGames.Perforce.Managed/ManagedWorkspace.cs
@@ -1503,7 +1503,15 @@ namespace EpicGames.Perforce.Managed
 				};
 
 				// Create the workspace, and add records for all the files. Exclude deleted files with digest = null.
-				await PerforceClient.RecordCommandAsync($"fstat -Ol -Op -Os -Rh -T \"{String.Join(",", FStatIndexedRecord.FieldNames)}\" \"//{PerforceClient.ClientName}/...@{ChangeNumber}\"", null, HandleRecord, CancellationToken);
+				List<string> Arguments = new List<string>();
+				Arguments.Add("-Ol");
+				Arguments.Add("-Op");
+				Arguments.Add("-Os");
+				Arguments.Add("-Rh");
+				Arguments.Add("-T");
+				Arguments.Add(String.Join(",", FStatIndexedRecord.FieldNames));
+				Arguments.Add($"//{PerforceClient.ClientName}/...@{ChangeNumber}");
+				await PerforceClient.RecordCommandAsync("fstat", Arguments, null, HandleRecord, CancellationToken);
 
 				// Output the elapsed time
 				Scope.Progress = $"({Timer.Elapsed.TotalSeconds:0.0}s)";
diff --git a/Engine/Source/Programs/Shared/EpicGames.Perforce/PerforceChildProcess.cs b/Engine/Source/Programs/Shared/EpicGames.Perforce/PerforceChildProcess.cs
--- a/Engine/Source/Programs/Shared/EpicGames.Perforce/PerforceChildProcess.cs
+++ b/Engine/Source/Programs/Shared/EpicGames.Perforce/PerforceChildProcess.cs
@@ -79,17 +79,25 @@ namespace EpicGames.Perforce
 		/// <summary>
 		/// Constructor
 		/// </summary>
-		/// <param name="InputData">Input data to pass to the child process</param>
+		/// <param name="Command"></param>
 		/// <param name="Arguments">Command line arguments</param>
+		/// <param name="InputData">Input data to pass to the child process</param>
+		/// <param name="GlobalOptions"></param>
 		/// <param name="Logger">Logging device</param>
-		public PerforceChildProcess(byte[]? InputData, string Arguments, ILogger Logger)
+		public PerforceChildProcess(string Command, IReadOnlyList<string> Arguments, byte[]? InputData, IReadOnlyList<string> GlobalOptions, ILogger Logger)
 		{
 			string PerforceFileName = GetExecutable();
 
-			string FullArgumentList = "-G " + Arguments;
+			List<string> FullArguments = new List<string>();
+			FullArguments.Add("-G");
+			FullArguments.AddRange(GlobalOptions);
+			FullArguments.Add(Command);
+			FullArguments.AddRange(Arguments);
+
+			string FullArgumentList = CommandLineArguments.Join(FullArguments);
 			Logger.LogDebug("Running {0} {1}", PerforceFileName, FullArgumentList);
 
-			Scope = TraceSpan.Create($"{GetCommandName(Arguments)}", Service: "perforce");
+			Scope = TraceSpan.Create($"{Command}", Service: "perforce");
 			Scope.AddMetadata("arguments", FullArgumentList);
 
 			ChildProcessGroup = new ManagedProcessGroup();
@@ -98,29 +106,6 @@ namespace EpicGames.Perforce
 			Buffer = new byte[64 * 1024];
 		}
 
-		/// <summary>
-		/// Gets the command name from the argument list
-		/// </summary>
-		/// <param name="CommandLine">The command line arguments</param>
-		/// <returns>The command name</returns>
-		public static string GetCommandName(string CommandLine)
-		{
-			string[] Arguments = CommandLineArguments.Split(CommandLine);
-			for(int Idx = 0; Idx < Arguments.Length; Idx++)
-			{
-				string Argument = Arguments[Idx];
-				if(Argument == "-p" || Argument == "-u" || Argument == "-c")
-				{
-					Idx++;
-				}
-				else if(Argument.Length > 0 && Char.IsLetter(Argument[0]))
-				{
-					return Argument;
-				}
-			}
-			return "(unknown)";
-		}
-
 		/// <summary>
 		/// Gets the path to the P4.EXE executable
 		/// </summary>
diff --git a/Engine/Source/Programs/Shared/EpicGames.Perforce/PerforceConnection.cs b/Engine/Source/Programs/Shared/EpicGames.Perforce/PerforceConnection.cs
--- a/Engine/Source/Programs/Shared/EpicGames.Perforce/PerforceConnection.cs
+++ b/Engine/Source/Programs/Shared/EpicGames.Perforce/PerforceConnection.cs
@@ -29,22 +29,24 @@ namespace EpicGames.Perforce
 		/// <summary>
 		/// Execute a command and parse the response
 		/// </summary>
-		/// <param name="CommandLine">Arguments for the command</param>
+		/// <param name="Command">The command to execute</param>
+		/// <param name="Arguments">Command line arguments</param>
 		/// <param name="InputData">Input data to pass to Perforce</param>
 		/// <param name="StatRecordType">The type of records to return for "stat" responses</param>
 		/// <param name="CancellationToken">Token used to cancel the operation</param>
 		/// <returns>List of objects returned by the server</returns>
-		Task<List<PerforceResponse>> CommandAsync(string CommandLine, byte[]? InputData, Type? StatRecordType, CancellationToken CancellationToken);
+		Task<List<PerforceResponse>> CommandAsync(string Command, IReadOnlyList<string> Arguments, byte[]? InputData, Type? StatRecordType, CancellationToken CancellationToken);
 
 		/// <summary>
 		/// Execute a command and parse the response
 		/// </summary>
-		/// <param name="CommandLine">Arguments for the command</param>
+		/// <param name="Command">The command to execute</param>
+		/// <param name="Arguments">Command line arguments</param>
 		/// <param name="InputData">Input data to pass to Perforce</param>
 		/// <param name="HandleRecord">Delegate used to handle each record</param>
 		/// <param name="CancellationToken">Token used to cancel the operation</param>
 		/// <returns>List of objects returned by the server</returns>
-		Task RecordCommandAsync(string CommandLine, byte[]? InputData, Action<PerforceRecord> HandleRecord, CancellationToken CancellationToken);
+		Task RecordCommandAsync(string Command, IReadOnlyList<string> Arguments, byte[]? InputData, Action<PerforceRecord> HandleRecord, CancellationToken CancellationToken);
 
 		/// <summary>
 		/// Execute the 'login' command
@@ -196,53 +198,45 @@ namespace EpicGames.Perforce
 			GlobalOptions.AddRange(Other.GlobalOptions);
 		}
 
-		/// <summary>
-		/// Gets all the arguments
-		/// </summary>
-		/// <returns></returns>
-		string GetFullCommandLine(string CommandLine)
+		List<string> GetGlobalArguments()
 		{
-			StringBuilder Result = new StringBuilder();
+			List<string> Arguments = new List<string>();
 			if (ServerAndPort != null)
 			{
-				Result.AppendFormat("-p {0} ", ServerAndPort);
+				Arguments.Add($"-p{ServerAndPort}");
 			}
 			if (UserName != null)
 			{
-				Result.AppendFormat("-u {0} ", UserName);
+				Arguments.Add($"-u{UserName}");
 			}
 			if (ClientName != null)
 			{
-				Result.AppendFormat("-c {0} ", ClientName);
+				Arguments.Add($"-c{ClientName}");
 			}
 			if (AppName != null)
 			{
-				Result.AppendFormat("-zprog={0} ", AppName);
+				Arguments.Add($"-zprog={AppName}");
 			}
 			if (AppVersion != null)
 			{
-				Result.AppendFormat("-zversion={0} ", AppVersion);
-			}
-
-			foreach (string GlobalOption in GlobalOptions)
-			{
-				Result.AppendFormat("{0} ", GlobalOption);
+				Arguments.Add($"-zversion={AppVersion}");
 			}
-			Result.Append(CommandLine);
-			return Result.ToString();
+			Arguments.AddRange(GlobalOptions);
+			return Arguments;
 		}
 
 		/// <summary>
 		/// Execute a command and parse the response
 		/// </summary>
-		/// <param name="CommandLine">Arguments for the command</param>
+		/// <param name="Command">Command to execute</param>
+		/// <param name="Arguments">Arguments for the command</param>
 		/// <param name="InputData">Input data to pass to Perforce</param>
 		/// <param name="StatRecordType">The type of records to return for "stat" responses</param>
 		/// <param name="CancellationToken">Token used to cancel the operation</param>
 		/// <returns>List of objects returned by the server</returns>
-		public async Task<List<PerforceResponse>> CommandAsync(string CommandLine, byte[]? InputData, Type? StatRecordType, CancellationToken CancellationToken)
+		public async Task<List<PerforceResponse>> CommandAsync(string Command, IReadOnlyList<string> Arguments, byte[]? InputData, Type? StatRecordType, CancellationToken CancellationToken)
 		{
-			using (PerforceChildProcess Process = new PerforceChildProcess(InputData, GetFullCommandLine(CommandLine), Logger))
+			using (PerforceChildProcess Process = new PerforceChildProcess(Command, Arguments, InputData, GetGlobalArguments(), Logger))
 			{
 				CachedRecordInfo? StatRecordInfo = (StatRecordType == null) ? null : PerforceReflection.GetCachedRecordInfo(StatRecordType);
 				return await Process.ReadResponsesAsync(StatRecordInfo, CancellationToken);
@@ -252,14 +246,15 @@ namespace EpicGames.Perforce
 		/// <summary>
 		/// Execute a command and parse the response
 		/// </summary>
-		/// <param name="CommandLine">Arguments for the command</param>
+		/// <param name="Command">Command to execute</param>
+		/// <param name="Arguments">Arguments for the command</param>
 		/// <param name="InputData">Input data to pass to Perforce</param>
 		/// <param name="HandleRecord">Delegate used to handle each record</param>
 		/// <param name="CancellationToken">Token used to cancel the operation</param>
 		/// <returns>List of objects returned by the server</returns>
-		public async Task RecordCommandAsync(string CommandLine, byte[]? InputData, Action<PerforceRecord> HandleRecord, CancellationToken CancellationToken)
+		public async Task RecordCommandAsync(string Command, IReadOnlyList<string> Arguments, byte[]? InputData, Action<PerforceRecord> HandleRecord, CancellationToken CancellationToken)
 		{
-			using (PerforceChildProcess Process = new PerforceChildProcess(InputData, GetFullCommandLine(CommandLine), Logger))
+			using (PerforceChildProcess Process = new PerforceChildProcess(Command, Arguments, InputData, GetGlobalArguments(), Logger))
 			{
 				await Process.ReadRecordsAsync(HandleRecord, CancellationToken);
 			}
@@ -280,7 +275,11 @@ namespace EpicGames.Perforce
 			using (ManagedProcessGroup ChildProcessGroup = new ManagedProcessGroup())
 			{
 				string Executable = PerforceChildProcess.GetExecutable();
-				string FullArgumentList = GetFullCommandLine("login");
+
+				List<string> Arguments = GetGlobalArguments();
+				Arguments.Add("login");
+
+				string FullArgumentList = CommandLineArguments.Join(Arguments);
 
 				Logger.LogDebug("Running {0} {1}", Executable, FullArgumentList);
 				using (ManagedProcess ChildProcess = new ManagedProcess(ChildProcessGroup, Executable, FullArgumentList, null, null, PasswordBytes, ProcessPriorityClass.Normal))
@@ -341,11 +340,10 @@ namespace EpicGames.Perforce
 		/// <returns>Response from the server</returns>
 		public async Task<Tuple<bool, string>> TrySetAsync(string Name, string Value, CancellationToken CancellationToken)
 		{
-			StringBuilder Arguments = new StringBuilder("set");
-			Arguments.Append($" {Name}=");
-			Arguments.Append($"\"{Value}\"");
+			List<string> Arguments = new List<string>();
+			Arguments.Append($"{Name}={Value}");
 
-			using (PerforceChildProcess ChildProcess = new PerforceChildProcess(null, Arguments.ToString(), Logger))
+			using (PerforceChildProcess ChildProcess = new PerforceChildProcess("set", Arguments, null, GetGlobalArguments(), Logger))
 			{
 				return await ChildProcess.TryReadToEndAsync(CancellationToken);
 			}
@@ -359,7 +357,7 @@ namespace EpicGames.Perforce
 		/// <returns>Value of the variable</returns>
 		public async Task<string?> TryGetSettingAsync(string Name, CancellationToken CancellationToken)
 		{
-			using (PerforceChildProcess ChildProcess = new PerforceChildProcess(null, $"set {Name}=", Logger))
+			using (PerforceChildProcess ChildProcess = new PerforceChildProcess("set", new List<string> { $"{Name}=" }, null, GetGlobalArguments(), Logger))
 			{
 				Tuple<bool, string> Response = await ChildProcess.TryReadToEndAsync(CancellationToken);
 				if (Response.Item1)
@@ -429,13 +427,14 @@ namespace EpicGames.Perforce
 		/// Execute a command and parse the response
 		/// </summary>
 		/// <param name="Connection">Connection to the Perforce server</param>
+		/// <param name="Command">The command to execute</param>
 		/// <param name="Arguments">Arguments for the command</param>
 		/// <param name="InputData">Input data to pass to Perforce</param>
 		/// <param name="CancellationToken">Token used to cancel the operation</param>
 		/// <returns>List of objects returned by the server</returns>
-		public static async Task<PerforceResponseList<T>> CommandAsync<T>(this IPerforceConnection Connection, string Arguments, byte[]? InputData, CancellationToken CancellationToken) where T : class
+		public static async Task<PerforceResponseList<T>> CommandAsync<T>(this IPerforceConnection Connection, string Command, List<string> Arguments, byte[]? InputData, CancellationToken CancellationToken) where T : class
 		{
-			List<PerforceResponse> Responses = await Connection.CommandAsync(Arguments, InputData, typeof(T), CancellationToken);
+			List<PerforceResponse> Responses = await Connection.CommandAsync(Command, Arguments, InputData, typeof(T), CancellationToken);
 
 			PerforceResponseList<T> TypedResponses = new PerforceResponseList<T>();
 			foreach (PerforceResponse Response in Responses)
@@ -449,22 +448,29 @@ namespace EpicGames.Perforce
 		/// Execute a command and parse the response
 		/// </summary>
 		/// <param name="Connection">Connection to the Perforce server</param>
+		/// <param name="Command">The command to execute</param>
 		/// <param name="CommonArguments">Arguments for the command</param>
 		/// <param name="BatchedArguments">Arguments to pass to the command in batches</param>
 		/// <param name="InputData">Input data to pass to Perforce</param>
 		/// <param name="CancellationToken">Token used to cancel the operation</param>
 		/// <returns>List of objects returned by the server</returns>
-		public static async Task<PerforceResponseList<T>> BatchedCommandAsync<T>(IPerforceConnection Connection, string CommonArguments, string[] BatchedArguments, byte[]? InputData, CancellationToken CancellationToken) where T : class
+		public static async Task<PerforceResponseList<T>> BatchedCommandAsync<T>(IPerforceConnection Connection, string Command, IReadOnlyList<string> CommonArguments, IReadOnlyList<string> BatchedArguments, byte[]? InputData, CancellationToken CancellationToken) where T : class
 		{
 			PerforceResponseList<T> Responses = new PerforceResponseList<T>();
-			for (int FileSpecIdx = 0; FileSpecIdx < BatchedArguments.Length;)
+			for (int FileSpecIdx = 0; FileSpecIdx < BatchedArguments.Count;)
 			{
-				StringBuilder Arguments = new StringBuilder(CommonArguments);
-				for (; FileSpecIdx < BatchedArguments.Length && Arguments.Length < 4096; FileSpecIdx++)
+				List<string> Arguments = new List<string>();
+				Arguments.AddRange(CommonArguments);
+
+				const int PerArgumentExtra = 5;
+				int Length = (Command.Length + PerArgumentExtra) + Arguments.Sum(x => x.Length + PerArgumentExtra);
+
+				for (; FileSpecIdx < BatchedArguments.Count && Length < 4096; FileSpecIdx++)
 				{
-					Arguments.AppendFormat(" \"{0}\"", BatchedArguments[FileSpecIdx]);
+					Arguments.Add(BatchedArguments[FileSpecIdx]);
+					Length += BatchedArguments[FileSpecIdx].Length + PerArgumentExtra;
 				}
-				Responses.AddRange(await CommandAsync<T>(Connection, Arguments.ToString(), null, CancellationToken));
+				Responses.AddRange(await CommandAsync<T>(Connection, Command, Arguments, null, CancellationToken));
 			}
 			return Responses;
 		}
@@ -473,14 +479,15 @@ namespace EpicGames.Perforce
 		/// Attempts to execute the given command, returning the results from the server or the first PerforceResponse object.
 		/// </summary>
 		/// <param name="Connection">Connection to the Perforce server</param>
+		/// <param name="Command">The command to execute</param>
 		/// <param name="Arguments">Arguments for the command.</param>
 		/// <param name="InputData">Input data for the command.</param>
 		/// <param name="StatRecordType">Type of element to return in the response</param>
 		/// <param name="CancellationToken">Token used to cancel the operation</param>
 		/// <returns>Response from the server; either an object of type T or error.</returns>
-		static async Task<PerforceResponse> SingleResponseCommandAsync(IPerforceConnection Connection, string Arguments, byte[]? InputData, Type? StatRecordType, CancellationToken CancellationToken)
+		static async Task<PerforceResponse> SingleResponseCommandAsync(IPerforceConnection Connection, string Command, IReadOnlyList<string> Arguments, byte[]? InputData, Type? StatRecordType, CancellationToken CancellationToken)
 		{
-			List<PerforceResponse> Responses = await Connection.CommandAsync(Arguments, InputData, StatRecordType, CancellationToken);
+			List<PerforceResponse> Responses = await Connection.CommandAsync(Command, Arguments, InputData, StatRecordType, CancellationToken);
 			if (Responses.Count != 1)
 			{
 				for (int Idx = 0; Idx < Responses.Count; Idx++)
@@ -497,13 +504,14 @@ namespace EpicGames.Perforce
 		/// </summary>
 		/// <typeparam name="T">Type of record to parse</typeparam>
 		/// <param name="Connection">Connection to the Perforce server</param>
+		/// <param name="Command">The command to execute</param>
 		/// <param name="Arguments">Arguments for the command.</param>
 		/// <param name="InputData">Input data for the command.</param>
 		/// <param name="CancellationToken">Token used to cancel the operation</param>
 		/// <returns>Response from the server; either an object of type T or error.</returns>
-		public static async Task<PerforceResponse<T>> SingleResponseCommandAsync<T>(IPerforceConnection Connection, string Arguments, byte[]? InputData, CancellationToken CancellationToken) where T : class
+		public static async Task<PerforceResponse<T>> SingleResponseCommandAsync<T>(IPerforceConnection Connection, string Command, List<string> Arguments, byte[]? InputData, CancellationToken CancellationToken) where T : class
 		{
-			return new PerforceResponse<T>(await SingleResponseCommandAsync(Connection, Arguments, InputData, typeof(T), CancellationToken));
+			return new PerforceResponse<T>(await SingleResponseCommandAsync(Connection, Command, Arguments, InputData, typeof(T), CancellationToken));
 		}
 
 		#endregion
@@ -561,35 +569,35 @@ namespace EpicGames.Perforce
 		/// <param name="FileNames">Files to be added</param>
 		/// <param name="CancellationToken">Token used to cancel the operation</param>
 		/// <returns>Response from the server</returns>
-		public static Task<PerforceResponseList<AddRecord>> TryAddAsync(this IPerforceConnection Connection, int ChangeNumber, string? FileType, AddOptions Options, string[] FileNames, CancellationToken CancellationToken)
+		public static Task<PerforceResponseList<AddRecord>> TryAddAsync(this IPerforceConnection Connection, int ChangeNumber, string? FileType, AddOptions Options, IReadOnlyList<string> FileNames, CancellationToken CancellationToken)
 		{
-			StringBuilder Arguments = new StringBuilder("add");
+			List<string> Arguments = new List<string>();
 			if (ChangeNumber != -1)
 			{
-				Arguments.AppendFormat(" -c {0}", ChangeNumber);
+				Arguments.Add($"-c{ChangeNumber}");
 			}
 			if ((Options & AddOptions.DowngradeToAdd) != 0)
 			{
-				Arguments.Append(" -d");
+				Arguments.Add("-d");
 			}
 			if ((Options & AddOptions.IncludeWildcards) != 0)
 			{
-				Arguments.Append(" -f");
+				Arguments.Add("-f");
 			}
 			if ((Options & AddOptions.NoIgnore) != 0)
 			{
-				Arguments.Append(" -I");
+				Arguments.Add("-I");
 			}
 			if ((Options & AddOptions.PreviewOnly) != 0)
 			{
-				Arguments.Append(" -n");
+				Arguments.Add("-n");
 			}
 			if (FileType != null)
 			{
-				Arguments.AppendFormat(" -t \"{0}\"", FileType);
+				Arguments.Add($"-t{FileType}");
 			}
 
-			return BatchedCommandAsync<AddRecord>(Connection, Arguments.ToString(), FileNames, null, CancellationToken);
+			return BatchedCommandAsync<AddRecord>(Connection, "add", Arguments, FileNames, null, CancellationToken);
 		}
 
 		#endregion
@@ -622,7 +630,7 @@ namespace EpicGames.Perforce
 				throw new PerforceException("'Number' field should be set to -1 to create a new changelist.");
 			}
 
-			PerforceResponse Response = await SingleResponseCommandAsync(Connection, "change -i", SerializeRecord(Record), null, CancellationToken);
+			PerforceResponse Response = await SingleResponseCommandAsync(Connection, "change", new List<string> { "-i" }, SerializeRecord(Record), null, CancellationToken);
 
 			PerforceError? Error = Response.Error;
 			if (Error != null)
@@ -674,17 +682,18 @@ namespace EpicGames.Perforce
 				throw new PerforceException("'Number' field must be set to update a changelist.");
 			}
 
-			StringBuilder Arguments = new StringBuilder("change -i");
+			List<string> Arguments = new List<string>();
+			Arguments.Add("-i");
 			if ((Options & UpdateChangeOptions.Force) != 0)
 			{
-				Arguments.Append(" -f");
+				Arguments.Add("-f");
 			}
 			if ((Options & UpdateChangeOptions.Submitted) != 0)
 			{
-				Arguments.Append(" -u");
+				Arguments.Add("-u");
 			}
 
-			return SingleResponseCommandAsync(Connection, Arguments.ToString(), SerializeRecord(Record), null, CancellationToken);
+			return SingleResponseCommandAsync(Connection, "change", Arguments, SerializeRecord(Record), null, CancellationToken);
 		}
 
 		/// <summary>
@@ -710,18 +719,18 @@ namespace EpicGames.Perforce
 		/// <returns>Response from the server</returns>
 		public static Task<PerforceResponse> TryDeleteChangeAsync(this IPerforceConnection Connection, DeleteChangeOptions Options, int ChangeNumber, CancellationToken CancellationToken)
 		{
-			StringBuilder Arguments = new StringBuilder("change -d");
+			List<string> Arguments = new List<string> { "-d" };
 			if ((Options & DeleteChangeOptions.Submitted) != 0)
 			{
-				Arguments.Append(" -f");
+				Arguments.Add("-f");
 			}
 			if ((Options & DeleteChangeOptions.BeforeRenumber) != 0)
 			{
-				Arguments.Append(" -O");
+				Arguments.Add("-O");
 			}
-			Arguments.AppendFormat(" {0}", ChangeNumber);
+			Arguments.Add($"{ChangeNumber}");
 
-			return SingleResponseCommandAsync(Connection, Arguments.ToString(), null, null, CancellationToken);
+			return SingleResponseCommandAsync(Connection, "change", Arguments, null, null, CancellationToken);
 		}
 
 		/// <summary>
@@ -747,17 +756,17 @@ namespace EpicGames.Perforce
 		/// <returns>Response from the server</returns>
 		public static Task<PerforceResponse<ChangeRecord>> TryGetChange(this IPerforceConnection Connection, GetChangeOptions Options, int ChangeNumber, CancellationToken CancellationToken)
 		{
-			StringBuilder Arguments = new StringBuilder("change -o");
+			List<string> Arguments = new List<string> { "-o" };
 			if ((Options & GetChangeOptions.BeforeRenumber) != 0)
 			{
-				Arguments.Append(" -O");
+				Arguments.Add("-O");
 			}
 			if (ChangeNumber != -1)
 			{
-				Arguments.AppendFormat(" {0}", ChangeNumber);
+				Arguments.Add($"{ChangeNumber}");
 			}
 
-			return SingleResponseCommandAsync<ChangeRecord>(Connection, Arguments.ToString(), null, CancellationToken);
+			return SingleResponseCommandAsync<ChangeRecord>(Connection, "change", Arguments, null, CancellationToken);
 		}
 
 		/// <summary>
@@ -860,51 +869,48 @@ namespace EpicGames.Perforce
 		/// <returns>List of responses from the server.</returns>
 		public static Task<PerforceResponseList<ChangesRecord>> TryGetChangesAsync(this IPerforceConnection Connection, ChangesOptions Options, string? ClientName, int MaxChanges, ChangeStatus Status, string? UserName, string[]? FileSpecs, CancellationToken CancellationToken)
 		{
-			StringBuilder Arguments = new StringBuilder("changes");
+			List<string> Arguments = new List<string>();
 			if ((Options & ChangesOptions.IncludeIntegrations) != 0)
 			{
-				Arguments.Append(" -i");
+				Arguments.Add("-i");
 			}
 			if ((Options & ChangesOptions.IncludeTimes) != 0)
 			{
-				Arguments.Append(" -t");
+				Arguments.Add("-t");
 			}
 			if ((Options & ChangesOptions.LongOutput) != 0)
 			{
-				Arguments.Append(" -l");
+				Arguments.Add("-l");
 			}
 			if ((Options & ChangesOptions.TruncatedLongOutput) != 0)
 			{
-				Arguments.Append(" -L");
+				Arguments.Add("-L");
 			}
 			if ((Options & ChangesOptions.IncludeRestricted) != 0)
 			{
-				Arguments.Append(" -f");
+				Arguments.Add("-f");
 			}
 			if (ClientName != null)
 			{
-				Arguments.AppendFormat(" -c \"{0}\"", ClientName);
+				Arguments.Add($"-c{ClientName}");
 			}
 			if (MaxChanges != -1)
 			{
-				Arguments.AppendFormat(" -m {0}", MaxChanges);
+				Arguments.Add($"-m{MaxChanges}");
 			}
 			if (Status != ChangeStatus.All)
 			{
-				Arguments.AppendFormat(" -s {0}", PerforceReflection.GetEnumText(typeof(ChangeStatus), Status));
+				Arguments.Add($"-s{PerforceReflection.GetEnumText(typeof(ChangeStatus), Status)}");
 			}
 			if (UserName != null)
 			{
-				Arguments.AppendFormat(" -u {0}", UserName);
+				Arguments.Add($"-u{UserName}");
 			}
 			if (FileSpecs != null)
 			{
-				foreach (string FileSpec in FileSpecs)
-				{
-					Arguments.AppendFormat(" \"{0}\"", FileSpec);
-				}
+				Arguments.AddRange(FileSpecs);
 			}
-			return CommandAsync<ChangesRecord>(Connection, Arguments.ToString(), null, CancellationToken);
+			return CommandAsync<ChangesRecord>(Connection, "changes", Arguments, null, CancellationToken);
 		}
 
 		#endregion
@@ -934,41 +940,38 @@ namespace EpicGames.Perforce
 		/// <returns>Response from the server</returns>
 		public static async Task<PerforceResponseList<CleanRecord>> TryCleanAsync(this IPerforceConnection Connection, CleanOptions Options, string[] FileSpecs, CancellationToken CancellationToken)
 		{
-			StringBuilder Arguments = new StringBuilder("clean");
+			List<string> Arguments = new List<string>();
 			if ((Options & CleanOptions.Edited) != 0)
 			{
-				Arguments.Append(" -e");
+				Arguments.Add("-e");
 			}
 			if ((Options & CleanOptions.Added) != 0)
 			{
-				Arguments.Append(" -a");
+				Arguments.Add("-a");
 			}
 			if ((Options & CleanOptions.Deleted) != 0)
 			{
-				Arguments.Append(" -d");
+				Arguments.Add("-d");
 			}
 			if ((Options & CleanOptions.Preview) != 0)
 			{
-				Arguments.Append(" -n");
+				Arguments.Add("-n");
 			}
 			if ((Options & CleanOptions.NoIgnoreChecking) != 0)
 			{
-				Arguments.Append(" -I");
+				Arguments.Add("-I");
 			}
 			if ((Options & CleanOptions.LocalSyntax) != 0)
 			{
-				Arguments.Append(" -l");
+				Arguments.Add("-l");
 			}
 			if ((Options & CleanOptions.ModifiedTimes) != 0)
 			{
-				Arguments.Append(" -m");
-			}
-			foreach (string FileSpec in FileSpecs)
-			{
-				Arguments.AppendFormat(" \"{0}\"", FileSpec);
+				Arguments.Add("-m");
 			}
+			Arguments.AddRange(FileSpecs);
 
-			PerforceResponseList<CleanRecord> Records = await CommandAsync<CleanRecord>(Connection, Arguments.ToString(), null, CancellationToken);
+			PerforceResponseList<CleanRecord> Records = await CommandAsync<CleanRecord>(Connection, "clean", Arguments, null, CancellationToken);
 			Records.RemoveAll(x => x.Error != null && x.Error.Generic == PerforceGenericCode.Empty);
 			return Records;
 		}
@@ -1022,7 +1025,7 @@ namespace EpicGames.Perforce
 		/// <returns>Response from the server</returns>
 		public static Task<PerforceResponse> TryUpdateClientAsync(this IPerforceConnection Connection, ClientRecord Record, CancellationToken CancellationToken)
 		{
-			return SingleResponseCommandAsync(Connection, "client -i", SerializeRecord(Record), null, CancellationToken);
+			return SingleResponseCommandAsync(Connection, "client", new List<string> { "-i" }, SerializeRecord(Record), null, CancellationToken);
 		}
 
 		/// <summary>
@@ -1048,17 +1051,18 @@ namespace EpicGames.Perforce
 		/// <returns>Response from the server</returns>
 		public static Task<PerforceResponse> TryDeleteClientAsync(this IPerforceConnection Connection, DeleteClientOptions Options, string ClientName, CancellationToken CancellationToken)
 		{
-			StringBuilder Arguments = new StringBuilder("client -d");
+			List<string> Arguments = new List<string> { "-d" };
 			if ((Options & DeleteClientOptions.Force) != 0)
 			{
-				Arguments.Append(" -f");
+				Arguments.Add("-f");
 			}
 			if ((Options & DeleteClientOptions.DeleteShelved) != 0)
 			{
-				Arguments.Append(" -Fs");
+				Arguments.Add("-Fs");
 			}
-			Arguments.AppendFormat(" \"{0}\"", ClientName);
-			return SingleResponseCommandAsync(Connection, Arguments.ToString(), null, null, CancellationToken);
+			Arguments.Add(ClientName);
+
+			return SingleResponseCommandAsync(Connection, "client", Arguments, null, null, CancellationToken);
 		}
 
 		/// <summary>
@@ -1086,14 +1090,14 @@ namespace EpicGames.Perforce
 		/// <returns>Response from the server</returns>
 		public static Task<PerforceResponse> TrySwitchClientToStreamAsync(this IPerforceConnection Connection, string ClientName, string StreamName, SwitchClientOptions Options, CancellationToken CancellationToken)
 		{
-			StringBuilder Arguments = new StringBuilder("client -s");
+			List<string> Arguments = new List<string> { "-s" };
 			if ((Options & SwitchClientOptions.IgnoreOpenFiles) != 0)
 			{
-				Arguments.Append(" -f");
+				Arguments.Add("-f");
 			}
-			Arguments.AppendFormat(" -S \"{0}\"", StreamName);
+			Arguments.Add($"-S{StreamName}");
 
-			return SingleResponseCommandAsync(Connection, Arguments.ToString(), null, null, CancellationToken);
+			return SingleResponseCommandAsync(Connection, "client", Arguments, null, null, CancellationToken);
 		}
 
 		/// <summary>
@@ -1119,8 +1123,11 @@ namespace EpicGames.Perforce
 		/// <returns>Response from the server</returns>
 		public static Task<PerforceResponse> TrySwitchClientToTemplateAsync(this IPerforceConnection Connection, string ClientName, string TemplateName, CancellationToken CancellationToken)
 		{
-			string Arguments = String.Format("client -s -t \"{0}\"", TemplateName);
-			return SingleResponseCommandAsync(Connection, Arguments, null, null, CancellationToken);
+			List<string> Arguments = new List<string>();
+			Arguments.Add("-s");
+			Arguments.Add($"-t{TemplateName}");
+
+			return SingleResponseCommandAsync(Connection, "client", Arguments, null, null, CancellationToken);
 		}
 
 		/// <summary>
@@ -1144,12 +1151,12 @@ namespace EpicGames.Perforce
 		/// <returns>Response from the server; either a client record or error code</returns>
 		public static Task<PerforceResponse<ClientRecord>> TryGetClientAsync(this IPerforceConnection Connection, string? ClientName, CancellationToken CancellationToken)
 		{
-			StringBuilder Arguments = new StringBuilder("client -o");
+			List<string> Arguments = new List<string> { "-o" };
 			if (ClientName != null)
 			{
-				Arguments.AppendFormat(" \"{0}\"", ClientName);
+				Arguments.Add(ClientName);
 			}
-			return SingleResponseCommandAsync<ClientRecord>(Connection, Arguments.ToString(), null, CancellationToken);
+			return SingleResponseCommandAsync<ClientRecord>(Connection, "client", Arguments, null, CancellationToken);
 		}
 
 		/// <summary>
@@ -1175,14 +1182,14 @@ namespace EpicGames.Perforce
 		/// <returns>Stream information record</returns>
 		public static Task<PerforceResponse<StreamRecord>> TryGetStreamAsync(this IPerforceConnection Connection, string StreamName, bool bIncludeView, CancellationToken CancellationToken)
 		{
-			StringBuilder Arguments = new StringBuilder("stream -o");
+			List<string> Arguments = new List<string> { "-o" };
 			if (bIncludeView)
 			{
-				Arguments.Append(" -v");
+				Arguments.Add("-v");
 			}
-			Arguments.AppendFormat(" {0}", StreamName);
+			Arguments.Add(StreamName);
 
-			return SingleResponseCommandAsync<StreamRecord>(Connection, Arguments.ToString(), null, CancellationToken);
+			return SingleResponseCommandAsync<StreamRecord>(Connection, "stream", Arguments, null, CancellationToken);
 		}
 
 		/// <summary>
@@ -1208,13 +1215,14 @@ namespace EpicGames.Perforce
 		/// <returns>Response from the server; either a client record or error code</returns>
 		public static Task<PerforceResponse<ClientRecord>> TryGetStreamViewAsync(this IPerforceConnection Connection, string StreamName, int ChangeNumber, CancellationToken CancellationToken)
 		{
-			StringBuilder Arguments = new StringBuilder("client -o");
-			Arguments.AppendFormat(" -S \"{0}\"", StreamName);
+			List<string> Arguments = new List<string> { "-o" };
+			Arguments.Add($"-S{StreamName}");
 			if (ChangeNumber != -1)
 			{
-				Arguments.AppendFormat(" -c {0}", ChangeNumber);
+				Arguments.Add($"-c{ChangeNumber}");
 			}
-			return SingleResponseCommandAsync<ClientRecord>(Connection, Arguments.ToString(), null, CancellationToken);
+
+			return SingleResponseCommandAsync<ClientRecord>(Connection, "client", Arguments, null, CancellationToken);
 		}
 
 		/// <summary>
@@ -1331,43 +1339,45 @@ namespace EpicGames.Perforce
 		/// <returns>Response from the server; either a client record or error code</returns>
 		public static Task<PerforceResponseList<ClientsRecord>> TryGetClientsAsync(this IPerforceConnection Connection, ClientsOptions Options, string? Filter, int MaxResults, string? Stream, string? UserName, CancellationToken CancellationToken)
 		{
-			StringBuilder Arguments = new StringBuilder("clients");
+			List<string> Arguments = new List<string>();
 			if ((Options & ClientsOptions.All) != 0)
 			{
-				Arguments.Append(" -a");
+				Arguments.Add("-a");
 			}
 			if (Filter != null)
 			{
 				if ((Options & ClientsOptions.CaseSensitiveFilter) != 0)
 				{
-					Arguments.AppendFormat(" -e \"{0}\"", Filter);
+					Arguments.Add("-e");
+					Arguments.Add(Filter);
 				}
 				else
 				{
-					Arguments.AppendFormat(" -E \"{0}\"", Filter);
+					Arguments.Add("-E");
+					Arguments.Add(Filter);
 				}
 			}
 			if (MaxResults != -1)
 			{
-				Arguments.AppendFormat(" -m {0}", MaxResults);
+				Arguments.Add($"-m{MaxResults}");
 			}
 			if (Stream != null)
 			{
-				Arguments.AppendFormat(" -S \"{0}\"", Stream);
+				Arguments.Add($"-S{Stream}");
 			}
 			if ((Options & ClientsOptions.WithTimes) != 0)
 			{
-				Arguments.Append(" -t");
+				Arguments.Add("-t");
 			}
 			if (UserName != null)
 			{
-				Arguments.AppendFormat(" -u \"{0}\"", UserName);
+				Arguments.Add($"-u{UserName}");
 			}
 			if ((Options & ClientsOptions.Unloaded) != 0)
 			{
-				Arguments.Append(" -U");
+				Arguments.Add("-U");
 			}
-			return CommandAsync<ClientsRecord>(Connection, Arguments.ToString(), null, CancellationToken);
+			return CommandAsync<ClientsRecord>(Connection, "clients", Arguments, null, CancellationToken);
 		}
 
 		#endregion
@@ -1399,25 +1409,25 @@ namespace EpicGames.Perforce
 		/// <returns>List of responses from the server</returns>
 		public static Task<PerforceResponseList<DeleteRecord>> TryDeleteAsync(this IPerforceConnection Connection, int ChangeNumber, DeleteOptions Options, string[] FileSpecs, CancellationToken CancellationToken)
 		{
-			StringBuilder Arguments = new StringBuilder("delete");
+			List<string> Arguments = new List<string>();
 			if (ChangeNumber != -1)
 			{
-				Arguments.AppendFormat(" -c {0}", ChangeNumber);
+				Arguments.Add($"-c{ChangeNumber}");
 			}
 			if ((Options & DeleteOptions.PreviewOnly) != 0)
 			{
-				Arguments.Append(" -n");
+				Arguments.Add("-n");
 			}
 			if ((Options & DeleteOptions.KeepWorkspaceFiles) != 0)
 			{
-				Arguments.Append(" -k");
+				Arguments.Add("-k");
 			}
 			if ((Options & DeleteOptions.WithoutSyncing) != 0)
 			{
-				Arguments.Append(" -v");
+				Arguments.Add("-v");
 			}
 
-			return BatchedCommandAsync<DeleteRecord>(Connection, Arguments.ToString(), FileSpecs, null, CancellationToken);
+			return BatchedCommandAsync<DeleteRecord>(Connection, "delete", Arguments, FileSpecs, null, CancellationToken);
 		}
 
 		#endregion
@@ -1474,12 +1484,12 @@ namespace EpicGames.Perforce
 		/// <returns>List of responses from the server</returns>
 		public static Task<PerforceResponseList<DescribeRecord>> TryDescribeAsync(this IPerforceConnection Connection, int[] ChangeNumbers, CancellationToken CancellationToken)
 		{
-			StringBuilder Arguments = new StringBuilder("describe -s");
+			List<string> Arguments = new List<string> { "-s" };
 			foreach (int ChangeNumber in ChangeNumbers)
 			{
-				Arguments.AppendFormat(" {0}", ChangeNumber);
+				Arguments.Add($"{ChangeNumber}");
 			}
-			return CommandAsync<DescribeRecord>(Connection, Arguments.ToString(), null, CancellationToken);
+			return CommandAsync<DescribeRecord>(Connection, "describe", Arguments, null, CancellationToken);
 		}
 
 		/// <summary>
@@ -1507,32 +1517,32 @@ namespace EpicGames.Perforce
 		/// <returns>List of responses from the server</returns>
 		public static Task<PerforceResponseList<DescribeRecord>> TryDescribeAsync(this IPerforceConnection Connection, DescribeOptions Options, int MaxNumFiles, int[] ChangeNumbers, CancellationToken CancellationToken)
 		{
-			StringBuilder Arguments = new StringBuilder("describe -s");
+			List<string> Arguments = new List<string> { "-s" };
 			if ((Options & DescribeOptions.ShowDescriptionForRestrictedChanges) != 0)
 			{
-				Arguments.Append(" -f");
+				Arguments.Add("-f");
 			}
 			if ((Options & DescribeOptions.Identity) != 0)
 			{
-				Arguments.Append(" -I");
+				Arguments.Add("-I");
 			}
 			if (MaxNumFiles != -1)
 			{
-				Arguments.AppendFormat(" -m{0}", MaxNumFiles);
+				Arguments.Add($"-m{MaxNumFiles}");
 			}
 			if ((Options & DescribeOptions.OriginalChangeNumber) != 0)
 			{
-				Arguments.Append(" -O");
+				Arguments.Add("-O");
 			}
 			if ((Options & DescribeOptions.Shelved) != 0)
 			{
-				Arguments.Append(" -S");
+				Arguments.Add("-S");
 			}
 			foreach (int ChangeNumber in ChangeNumbers)
 			{
-				Arguments.AppendFormat(" {0}", ChangeNumber);
+				Arguments.Add($"{ChangeNumber}");
 			}
-			return CommandAsync<DescribeRecord>(Connection, Arguments.ToString(), null, CancellationToken);
+			return CommandAsync<DescribeRecord>(Connection, "describe", Arguments, null, CancellationToken);
 		}
 
 		#endregion
@@ -1592,25 +1602,25 @@ namespace EpicGames.Perforce
 		/// <returns>Response from the server</returns>
 		public static Task<PerforceResponseList<EditRecord>> TryEditAsync(this IPerforceConnection Connection, int ChangeNumber, string? FileType, EditOptions Options, string[] FileSpecs, CancellationToken CancellationToken)
 		{
-			StringBuilder Arguments = new StringBuilder("edit");
+			List<string> Arguments = new List<string>();
 			if (ChangeNumber != -1)
 			{
-				Arguments.AppendFormat(" -c {0}", ChangeNumber);
+				Arguments.Add($"-c{ChangeNumber}");
 			}
 			if ((Options & EditOptions.KeepWorkspaceFiles) != 0)
 			{
-				Arguments.Append(" -k");
+				Arguments.Add("-k");
 			}
 			if ((Options & EditOptions.PreviewOnly) != 0)
 			{
-				Arguments.Append(" -n");
+				Arguments.Add("-n");
 			}
 			if (FileType != null)
 			{
-				Arguments.AppendFormat(" -t \"{0}\"", FileType);
+				Arguments.Add($"-t{FileType}");
 			}
 
-			return BatchedCommandAsync<EditRecord>(Connection, Arguments.ToString(), FileSpecs, null, CancellationToken);
+			return BatchedCommandAsync<EditRecord>(Connection, "edit", Arguments, FileSpecs, null, CancellationToken);
 		}
 
 		#endregion
@@ -1699,45 +1709,45 @@ namespace EpicGames.Perforce
 		public static Task<PerforceResponseList<FileLogRecord>> TryFileLogAsync(this IPerforceConnection Connection, int ChangeNumber, int MaxChanges, FileLogOptions Options, string[] FileSpecs, CancellationToken CancellationToken)
 		{
 			// Build the argument list
-			StringBuilder Arguments = new StringBuilder("filelog");
+			List<string> Arguments = new List<string>();
 			if (ChangeNumber > 0)
 			{
-				Arguments.AppendFormat(" -c {0}", ChangeNumber);
+				Arguments.Add($"-c{ChangeNumber}");
 			}
 			if ((Options & FileLogOptions.ContentHistory) != 0)
 			{
-				Arguments.Append(" -h");
+				Arguments.Add("-h");
 			}
 			if ((Options & FileLogOptions.FollowAcrossBranches) != 0)
 			{
-				Arguments.Append(" -i");
+				Arguments.Add("-i");
 			}
 			if ((Options & FileLogOptions.FullDescriptions) != 0)
 			{
-				Arguments.Append(" -l");
+				Arguments.Add("-l");
 			}
 			if ((Options & FileLogOptions.LongDescriptions) != 0)
 			{
-				Arguments.Append(" -L");
+				Arguments.Add("-L");
 			}
 			if (MaxChanges > 0)
 			{
-				Arguments.AppendFormat(" -m {0}", MaxChanges);
+				Arguments.Add($"-m{MaxChanges}");
 			}
 			if ((Options & FileLogOptions.DoNotFollowPromotedTaskStreams) != 0)
 			{
-				Arguments.Append(" -p");
+				Arguments.Add("-p");
 			}
 			if ((Options & FileLogOptions.IgnoreNonContributoryIntegrations) != 0)
 			{
-				Arguments.Append(" -s");
+				Arguments.Add("-s");
 			}
 
 			// Always include times to simplify parsing
-			Arguments.Append(" -t");
+			Arguments.Add("-t");
 
 			// Append all the arguments
-			return BatchedCommandAsync<FileLogRecord>(Connection, Arguments.ToString(), FileSpecs, null, CancellationToken);
+			return BatchedCommandAsync<FileLogRecord>(Connection, "filelog", Arguments, FileSpecs, null, CancellationToken);
 		}
 
 		#endregion
@@ -1830,102 +1840,103 @@ namespace EpicGames.Perforce
 		public static async Task<PerforceResponseList<FStatRecord>> TryFStatAsync(this IPerforceConnection Connection, int AfterChangeNumber, int OnlyChangeNumber, string? Filter, int MaxFiles, FStatOptions Options, string[] FileSpecs, CancellationToken CancellationToken)
 		{
 			// Build the argument list
-			StringBuilder Arguments = new StringBuilder("fstat");
+			List<string> Arguments = new List<string>();
 			if (AfterChangeNumber != -1)
 			{
-				Arguments.AppendFormat(" -c {0}", AfterChangeNumber);
+				Arguments.Add($"-c{AfterChangeNumber}");
 			}
 			if(OnlyChangeNumber != -1)
 			{
-				Arguments.AppendFormat(" -e {0}", OnlyChangeNumber);
+				Arguments.Add($"-e{OnlyChangeNumber}");
 			}
 			if(Filter != null)
 			{
-				Arguments.AppendFormat(" -F \"{0}\"", Filter);
+				Arguments.Add("-F");
+				Arguments.Add(Filter);
 			}
 			if((Options & FStatOptions.ReportDepotSyntax) != 0)
 			{
-				Arguments.Append(" -L");
+				Arguments.Add("-L");
 			}
 			if((Options & FStatOptions.AllRevisions) != 0)
 			{
-				Arguments.Append(" -Of");
+				Arguments.Add("-Of");
 			}
 			if((Options & FStatOptions.IncludeFileSizes) != 0)
 			{
-				Arguments.Append(" -Ol");
+				Arguments.Add("-Ol");
 			}
 			if((Options & FStatOptions.ClientFileInPerforceSyntax) != 0)
 			{
-				Arguments.Append(" -Op");
+				Arguments.Add("-Op");
 			}
 			if((Options & FStatOptions.ShowPendingIntegrations) != 0)
 			{
-				Arguments.Append(" -Or");
+				Arguments.Add("-Or");
 			}
 			if((Options & FStatOptions.ShortenOutput) != 0)
 			{
-				Arguments.Append(" -Os");
+				Arguments.Add("-Os");
 			}
 			if((Options & FStatOptions.ReverseOrder) != 0)
 			{
-				Arguments.Append(" -r");
+				Arguments.Add("-r");
 			}
 			if((Options & FStatOptions.OnlyMapped) != 0)
 			{
-				Arguments.Append(" -Rc");
+				Arguments.Add("-Rc");
 			}
 			if((Options & FStatOptions.OnlyHave) != 0)
 			{
-				Arguments.Append(" -Rh");
+				Arguments.Add("-Rh");
 			}
 			if((Options & FStatOptions.OnlyOpenedBeforeHead) != 0)
 			{
-				Arguments.Append(" -Rn");
+				Arguments.Add("-Rn");
 			}
 			if((Options & FStatOptions.OnlyOpenInWorkspace) != 0)
 			{
-				Arguments.Append(" -Ro");
+				Arguments.Add("-Ro");
 			}
 			if((Options & FStatOptions.OnlyOpenAndResolved) != 0)
 			{
-				Arguments.Append(" -Rr");
+				Arguments.Add("-Rr");
 			}
 			if((Options & FStatOptions.OnlyShelved) != 0)
 			{
-				Arguments.Append(" -Rs");
+				Arguments.Add("-Rs");
 			}
 			if((Options & FStatOptions.OnlyUnresolved) != 0)
 			{
-				Arguments.Append(" -Ru");
+				Arguments.Add("-Ru");
 			}
 			if((Options & FStatOptions.SortByDate) != 0)
 			{
-				Arguments.Append(" -Sd");
+				Arguments.Add("-Sd");
 			}
 			if((Options & FStatOptions.SortByHaveRevision) != 0)
 			{
-				Arguments.Append(" -Sh");
+				Arguments.Add("-Sh");
 			}
 			if((Options & FStatOptions.SortByHeadRevision) != 0)
 			{
-				Arguments.Append(" -Sr");
+				Arguments.Add("-Sr");
 			}
 			if((Options & FStatOptions.SortByFileSize) != 0)
 			{
-				Arguments.Append(" -Ss");
+				Arguments.Add("-Ss");
 			}
 			if((Options & FStatOptions.SortByFileType) != 0)
 			{
-				Arguments.Append(" -St");
+				Arguments.Add("-St");
 			}
 			if((Options & FStatOptions.IncludeFilesInUnloadDepot) != 0)
 			{
-				Arguments.Append(" -U");
+				Arguments.Add("-U");
 			}
 
 			// Execute the command
-			PerforceResponseList<FStatRecord> Records = await BatchedCommandAsync<FStatRecord>(Connection, Arguments.ToString(), FileSpecs, null, CancellationToken);
+			PerforceResponseList<FStatRecord> Records = await BatchedCommandAsync<FStatRecord>(Connection, "fstat", Arguments, FileSpecs, null, CancellationToken);
 			if (OnlyChangeNumber != -1 && Records.Count > 0 && Records[Records.Count - 1].Succeeded && Records[Records.Count - 1].Data.Description != null)
 			{
 				Records.RemoveAt(Records.Count - 1);
@@ -1959,12 +1970,12 @@ namespace EpicGames.Perforce
 		public static Task<PerforceResponse<InfoRecord>> TryGetInfoAsync(this IPerforceConnection Connection, InfoOptions Options, CancellationToken CancellationToken)
 		{
 			// Build the argument list
-			StringBuilder Arguments = new StringBuilder("info");
+			List<string> Arguments = new List<string>();
 			if((Options & InfoOptions.ShortOutput) != 0)
 			{
-				Arguments.Append(" -s");
+				Arguments.Add("-s");
 			}
-			return SingleResponseCommandAsync<InfoRecord>(Connection, Arguments.ToString(), null, CancellationToken);
+			return SingleResponseCommandAsync<InfoRecord>(Connection, "info", Arguments, null, CancellationToken);
 		}
 
 		#endregion
@@ -1984,23 +1995,23 @@ namespace EpicGames.Perforce
 		/// <returns>List of records</returns>
 		public static Task<PerforceResponseList<MergeRecord>> MergeAsync(this IPerforceConnection Connection, MergeOptions Options, int Change, int MaxFiles, string SourceFileSpec, string TargetFileSpec, CancellationToken CancellationToken)
 		{
-			StringBuilder Arguments = new StringBuilder("merge");
+			List<string> Arguments = new List<string>();
 			if ((Options & MergeOptions.Preview) != 0)
 			{
-				Arguments.AppendFormat(" -n");
+				Arguments.Add($"-n");
 			}
 			if (Change != -1)
 			{
-				Arguments.AppendFormat(" -c {0}", Change);
+				Arguments.Add($"-c{Change}");
 			}
 			if (MaxFiles != -1)
 			{
-				Arguments.AppendFormat(" -m {0}", MaxFiles);
+				Arguments.Add($"-m{MaxFiles}");
 			}
-			Arguments.AppendFormat(" \"{0}\"", SourceFileSpec);
-			Arguments.AppendFormat(" \"{0}\"", TargetFileSpec);
+			Arguments.Add(SourceFileSpec);
+			Arguments.Add(TargetFileSpec);
 
-			return CommandAsync<MergeRecord>(Connection, Arguments.ToString(), null, CancellationToken);
+			return CommandAsync<MergeRecord>(Connection, "merge", Arguments, null, CancellationToken);
 		}
 
 		#endregion
@@ -2036,30 +2047,31 @@ namespace EpicGames.Perforce
 		/// <returns>Response from the server</returns>
 		public static Task<PerforceResponseList<MoveRecord>> TryMoveAsync(this IPerforceConnection Connection, int ChangeNumber, string? FileType, MoveOptions Options, string SourceFileSpec, string TargetFileSpec, CancellationToken CancellationToken)
 		{
-			StringBuilder Arguments = new StringBuilder("move");
+			List<string> Arguments = new List<string>();
 			if (ChangeNumber != -1)
 			{
-				Arguments.AppendFormat(" -c {0}", ChangeNumber);
+				Arguments.Add($"-c{ChangeNumber}");
 			}
 			if ((Options & MoveOptions.KeepWorkspaceFiles) != 0)
 			{
-				Arguments.Append(" -k");
+				Arguments.Add("-k");
 			}
 			if ((Options & MoveOptions.RenameOnly) != 0)
 			{
-				Arguments.Append(" -r");
+				Arguments.Add("-r");
 			}
 			if ((Options & MoveOptions.PreviewOnly) != 0)
 			{
-				Arguments.Append(" -n");
+				Arguments.Add("-n");
 			}
 			if (FileType != null)
 			{
-				Arguments.AppendFormat(" -t \"{0}\"", FileType);
+				Arguments.Add($"-t{FileType}");
 			}
-			Arguments.AppendFormat(" \"{0}\" \"{1}\"", SourceFileSpec, TargetFileSpec);
+			Arguments.Add(SourceFileSpec);
+			Arguments.Add(TargetFileSpec);
 
-			return CommandAsync<MoveRecord>(Connection, Arguments.ToString(), null, CancellationToken);
+			return CommandAsync<MoveRecord>(Connection, "move", Arguments, null, CancellationToken);
 		}
 
 		#endregion
@@ -2098,43 +2110,40 @@ namespace EpicGames.Perforce
 		public static Task<PerforceResponseList<FStatRecord>> TryGetOpenFilesAsync(this IPerforceConnection Connection, OpenedOptions Options, int ChangeNumber, string? ClientName, string? UserName, int MaxResults, string[]? FileSpecs, CancellationToken CancellationToken)
 		{
 			// Build the argument list
-			StringBuilder Arguments = new StringBuilder("opened");
+			List<string> Arguments = new List<string>();
 			if((Options & OpenedOptions.AllWorkspaces) != 0)
 			{
-				Arguments.AppendFormat(" -a");
+				Arguments.Add($"-a");
 			}
 			if (ChangeNumber == PerforceReflection.DefaultChange)
 			{
-				Arguments.AppendFormat(" -c default");
+				Arguments.Add($"-cdefault");
 			}
 			else if (ChangeNumber != -1)
 			{
-				Arguments.AppendFormat(" -c {0}", ChangeNumber);
+				Arguments.Add($"-c{ChangeNumber}");
 			}
 			if(ClientName != null)
 			{
-				Arguments.AppendFormat(" -C \"{0}\"", ClientName);
+				Arguments.Add($"-C{ClientName}");
 			}
 			if(UserName != null)
 			{
-				Arguments.AppendFormat(" -u \"{0}\"", UserName);
+				Arguments.Add($"-u{UserName}");
 			}
 			if(MaxResults != -1)
 			{
-				Arguments.AppendFormat(" -m {0}", MaxResults);
+				Arguments.Add($"-m{MaxResults}");
 			}
 			if((Options & OpenedOptions.ShortOutput) != 0)
 			{
-				Arguments.AppendFormat(" -s");
+				Arguments.Add($"-s");
 			}
 			if (FileSpecs != null)
 			{
-				foreach (string FileSpec in FileSpecs)
-				{
-					Arguments.AppendFormat(" \"{0}\"", FileSpec);
-				}
+				Arguments.AddRange(FileSpecs);
 			}
-			return CommandAsync<FStatRecord>(Connection, Arguments.ToString(), null, CancellationToken);
+			return CommandAsync<FStatRecord>(Connection, "opened", Arguments, null, CancellationToken);
 		}
 
 		#endregion
@@ -2165,10 +2174,11 @@ namespace EpicGames.Perforce
 		public static Task<PerforceResponse<PrintRecord>> TryPrintAsync(this IPerforceConnection Connection, string OutputFile, string FileSpec, CancellationToken CancellationToken)
 		{
 			// Build the argument list
-			StringBuilder Arguments = new StringBuilder("print");
-			Arguments.AppendFormat(" -o \"{0}\"", OutputFile);
-			Arguments.AppendFormat(" \"{0}\"", FileSpec);
-			return SingleResponseCommandAsync<PrintRecord>(Connection, Arguments.ToString(), null, CancellationToken);
+			List<string> Arguments = new List<string>();
+			Arguments.Add("-o");
+			Arguments.Add(OutputFile);
+			Arguments.Add(FileSpec);
+			return SingleResponseCommandAsync<PrintRecord>(Connection, "print", Arguments, null, CancellationToken);
 		}
 
 		class PrintHandler : IDisposable
@@ -2265,49 +2275,46 @@ namespace EpicGames.Perforce
 		/// <returns>Response from the server</returns>
 		public static Task<PerforceResponseList<ReconcileRecord>> TryReconcileAsync(this IPerforceConnection Connection, int ChangeNumber, ReconcileOptions Options, string[] FileSpecs, CancellationToken CancellationToken)
 		{
-			StringBuilder Arguments = new StringBuilder("reconcile");
+			List<string> Arguments = new List<string>();
 			if(ChangeNumber != -1)
 			{
-				Arguments.AppendFormat(" -c {0}", ChangeNumber);
+				Arguments.Add($"-c{ChangeNumber}");
 			}
 			if((Options & ReconcileOptions.Edit) != 0)
 			{
-				Arguments.Append(" -e");
+				Arguments.Add("-e");
 			}
 			if((Options & ReconcileOptions.Add) != 0)
 			{
-				Arguments.Append(" -a");
+				Arguments.Add("-a");
 			}
 			if((Options & ReconcileOptions.Delete) != 0)
 			{
-				Arguments.Append(" -d");
+				Arguments.Add("-d");
 			}
 			if((Options & ReconcileOptions.PreviewOnly) != 0)
 			{
-				Arguments.Append(" -n");
+				Arguments.Add("-n");
 			}
 			if((Options & ReconcileOptions.AllowWildcards) != 0)
 			{
-				Arguments.Append(" -f");
+				Arguments.Add("-f");
 			}
 			if((Options & ReconcileOptions.NoIgnore) != 0)
 			{
-				Arguments.Append(" -I");
+				Arguments.Add("-I");
 			}
 			if((Options & ReconcileOptions.LocalFileSyntax) != 0)
 			{
-				Arguments.Append(" -l");
+				Arguments.Add("-l");
 			}
 			if ((Options & ReconcileOptions.UseFileModification) != 0)
 			{
-				Arguments.Append(" -m");
-			}
-			foreach (string FileSpec in FileSpecs)
-			{
-				Arguments.AppendFormat(" \"{0}\"", FileSpec);
+				Arguments.Add("-m");
 			}
+			Arguments.AddRange(FileSpecs);
 
-			return CommandAsync<ReconcileRecord>(Connection, Arguments.ToString(), null, CancellationToken);
+			return CommandAsync<ReconcileRecord>(Connection, "reconcile", Arguments, null, CancellationToken);
 		}
 
 		#endregion
@@ -2339,25 +2346,25 @@ namespace EpicGames.Perforce
 		/// <returns>Response from the server</returns>
 		public static Task<PerforceResponseList<ReopenRecord>> TryReopenAsync(this IPerforceConnection Connection, int? ChangeNumber, string? FileType, string FileSpec, CancellationToken CancellationToken)
 		{
-			StringBuilder Arguments = new StringBuilder("reopen");
+			List<string> Arguments = new List<string>();
 			if (ChangeNumber != null)
 			{
 				if (ChangeNumber == PerforceReflection.DefaultChange)
 				{
-					Arguments.Append(" -c default");
+					Arguments.Add("-cdefault");
 				}
 				else
 				{
-					Arguments.AppendFormat(" -c {0}", ChangeNumber);
+					Arguments.Add($"-c{ChangeNumber}");
 				}
 			}
 			if (FileType != null)
 			{
-				Arguments.AppendFormat(" -t \"{0}\"", FileType);
+				Arguments.Add($"-t{FileType}");
 			}
-			Arguments.AppendFormat(" \"{0}\"", FileSpec);
+			Arguments.Add(FileSpec);
 
-			return CommandAsync<ReopenRecord>(Connection, Arguments.ToString(), null, CancellationToken);
+			return CommandAsync<ReopenRecord>(Connection, "reopen", Arguments, null, CancellationToken);
 		}
 
 		#endregion
@@ -2374,7 +2381,10 @@ namespace EpicGames.Perforce
 		/// <returns>Response from the server</returns>
 		public static Task<List<PerforceResponse>> ReloadClient(this IPerforceConnection Connection, string ClientName, string SourceServerId, CancellationToken CancellationToken)
 		{
-			return Connection.CommandAsync($"reload -c \"{ClientName}\" -p \"{SourceServerId}\"", null, null, CancellationToken);
+			List<string> Arguments = new List<string>();
+			Arguments.Add($"-c{ClientName}");
+			Arguments.Add($"-p{SourceServerId}");
+			return Connection.CommandAsync("reload", Arguments, null, null, CancellationToken);
 		}
 
 		#endregion p4 reload
@@ -2406,56 +2416,53 @@ namespace EpicGames.Perforce
 		/// <returns>Response from the server</returns>
 		public static async Task<PerforceResponseList<ResolveRecord>> TryResolveAsync(this IPerforceConnection Connection, int ChangeNumber, ResolveOptions Options, string[]? FileSpecs, CancellationToken CancellationToken)
 		{
-			StringBuilder Arguments = new StringBuilder("resolve");
+			List<string> Arguments = new List<string>();
 			if ((Options & ResolveOptions.Automatic) != 0)
 			{
-				Arguments.Append(" -am");
+				Arguments.Add("-am");
 			}
 			if ((Options & ResolveOptions.AcceptYours) != 0)
 			{
-				Arguments.Append(" -ay");
+				Arguments.Add("-ay");
 			}
 			if ((Options & ResolveOptions.AcceptTheirs) != 0)
 			{
-				Arguments.Append(" -at");
+				Arguments.Add("-at");
 			}
 			if ((Options & ResolveOptions.SafeAccept) != 0)
 			{
-				Arguments.Append(" -as");
+				Arguments.Add("-as");
 			}
 			if ((Options & ResolveOptions.ForceAccept) != 0)
 			{
-				Arguments.Append(" -af");
+				Arguments.Add("-af");
 			}
 			if((Options & ResolveOptions.IgnoreWhitespaceOnly) != 0)
 			{
-				Arguments.Append(" -db");
+				Arguments.Add("-db");
 			}
 			if((Options & ResolveOptions.IgnoreWhitespace) != 0)
 			{
-				Arguments.Append(" -dw");
+				Arguments.Add("-dw");
 			}
 			if((Options & ResolveOptions.IgnoreLineEndings) != 0)
 			{
-				Arguments.Append(" -dl");
+				Arguments.Add("-dl");
 			}
 			if ((Options & ResolveOptions.ResolveAgain) != 0)
 			{
-				Arguments.Append(" -f");
+				Arguments.Add("-f");
 			}
 			if ((Options & ResolveOptions.PreviewOnly) != 0)
 			{
-				Arguments.Append(" -n");
+				Arguments.Add("-n");
 			}
 			if (FileSpecs != null)
 			{
-				foreach (string FileSpec in FileSpecs)
-				{
-					Arguments.AppendFormat(" \"{0}\"", FileSpec);
-				}
+				Arguments.AddRange(FileSpecs);
 			}
 
-			PerforceResponseList<ResolveRecord> Records = await CommandAsync<ResolveRecord>(Connection, Arguments.ToString(), null, CancellationToken);
+			PerforceResponseList<ResolveRecord> Records = await CommandAsync<ResolveRecord>(Connection, "resolve", Arguments, null, CancellationToken);
 			Records.RemoveAll(x => x.Error != null && x.Error.Generic == PerforceGenericCode.Empty);
 			return Records;
 		}
@@ -2491,33 +2498,33 @@ namespace EpicGames.Perforce
 		/// <returns>Response from the server</returns>
 		public static async Task<PerforceResponseList<RevertRecord>> TryRevertAsync(this IPerforceConnection Connection, int ChangeNumber, string? ClientName, RevertOptions Options, string[] FileSpecs, CancellationToken CancellationToken)
 		{
-			StringBuilder Arguments = new StringBuilder("revert");
+			List<string> Arguments = new List<string>();
 			if((Options & RevertOptions.Unchanged) != 0)
 			{
-				Arguments.Append(" -a");
+				Arguments.Add("-a");
 			}
 			if((Options & RevertOptions.PreviewOnly) != 0)
 			{
-				Arguments.Append(" -n");
+				Arguments.Add("-n");
 			}
 			if((Options & RevertOptions.KeepWorkspaceFiles) != 0)
 			{
-				Arguments.Append(" -k");
+				Arguments.Add("-k");
 			}
 			if((Options & RevertOptions.DeleteAddedFiles) != 0)
 			{
-				Arguments.Append(" -w");
+				Arguments.Add("-w");
 			}
 			if(ChangeNumber != -1)
 			{
-				Arguments.AppendFormat(" -c {0}", ChangeNumber);
+				Arguments.Add($"-c{ChangeNumber}");
 			}
 			if(ClientName != null)
 			{
-				Arguments.AppendFormat(" -C \"{0}\"", ClientName);
+				Arguments.Add($"-C{ClientName}");
 			}
 
-			PerforceResponseList<RevertRecord> Records = await BatchedCommandAsync<RevertRecord>(Connection, Arguments.ToString(), FileSpecs, null, CancellationToken);
+			PerforceResponseList<RevertRecord> Records = await BatchedCommandAsync<RevertRecord>(Connection, "revert", Arguments, FileSpecs, null, CancellationToken);
 			Records.RemoveAll(x => x.Error != null && x.Error.Generic == PerforceGenericCode.Empty);
 			return Records;
 		}
@@ -2551,22 +2558,19 @@ namespace EpicGames.Perforce
 		/// <returns>Response from the server</returns>
 		public static async Task<PerforceResponseList<ShelveRecord>> TryShelveAsync(this IPerforceConnection Connection, int ChangeNumber, ShelveOptions Options, string[] FileSpecs, CancellationToken CancellationToken)
 		{
-			StringBuilder Arguments = new StringBuilder("shelve");
-			Arguments.AppendFormat(" -c {0}", ChangeNumber);
+			List<string> Arguments = new List<string>();
+			Arguments.Add($"-c{ChangeNumber}");
 			if((Options & ShelveOptions.OnlyChanged) != 0)
 			{
-				Arguments.Append(" -a leaveunchanged");
+				Arguments.Add("-aleaveunchanged");
 			}
 			if((Options & ShelveOptions.Overwrite) != 0)
 			{
-				Arguments.Append(" -f");
-			}
-			foreach(string FileSpec in FileSpecs)
-			{
-				Arguments.AppendFormat(" \"{0}\"", FileSpec);
+				Arguments.Add("-f");
 			}
+			Arguments.AddRange(FileSpecs);
 
-			PerforceResponseList<ShelveRecord> Records = await CommandAsync<ShelveRecord>(Connection, Arguments.ToString(), null, CancellationToken);
+			PerforceResponseList<ShelveRecord> Records = await CommandAsync<ShelveRecord>(Connection, "shelve", Arguments, null, CancellationToken);
 			Records.RemoveAll(x => x.Error != null && x.Error.Generic == PerforceGenericCode.Empty);
 			return Records;
 		}
@@ -2600,17 +2604,15 @@ namespace EpicGames.Perforce
 		/// <returns>Response from the server</returns>
 		public static Task<List<PerforceResponse>> TryDeleteShelvedFilesAsync(this IPerforceConnection Connection, int ChangeNumber, string[] FileSpecs, CancellationToken CancellationToken)
 		{
-			StringBuilder Arguments = new StringBuilder("shelve -d");
+			List<string> Arguments = new List<string>();
+			Arguments.Add("-d");
 			if(ChangeNumber != -1)
 			{
-				Arguments.AppendFormat(" -c {0}", ChangeNumber);
-			}
-			foreach (string FileSpec in FileSpecs)
-			{
-				Arguments.AppendFormat(" \"{0}\"", FileSpec);
+				Arguments.Add($"-c{ChangeNumber}");
 			}
+			Arguments.AddRange(FileSpecs);
 
-			return Connection.CommandAsync(Arguments.ToString(), null, null, CancellationToken);
+			return Connection.CommandAsync("shelve", Arguments, null, null, CancellationToken);
 		}
 
 		#endregion
@@ -2669,23 +2671,24 @@ namespace EpicGames.Perforce
 		public static Task<PerforceResponseList<StreamsRecord>> TryGetStreamsAsync(this IPerforceConnection Connection, string StreamPath, int MaxResults, string? Filter, bool bUnloaded, CancellationToken CancellationToken)
 		{
 			// Build the command line
-			StringBuilder Arguments = new StringBuilder("streams");
+			List<string> Arguments = new List<string>();
 			if (bUnloaded)
 			{
-				Arguments.Append(" -U");
+				Arguments.Add("-U");
 			}
 			if (Filter != null)
 			{
-				Arguments.AppendFormat("-F \"{0}\"", Filter);
+				Arguments.Add("-F");
+				Arguments.Add(Filter);
 			}
 			if (MaxResults > 0)
 			{
-				Arguments.AppendFormat("-m {0}", MaxResults);
+				Arguments.Add($"-m{MaxResults}");
 			}
-			Arguments.AppendFormat(" \"{0}\"", StreamPath);
+			Arguments.Add(StreamPath);
 
 			// Execute the command
-			return CommandAsync<StreamsRecord>(Connection, Arguments.ToString(), null, CancellationToken);
+			return CommandAsync<StreamsRecord>(Connection, "streams", Arguments, null, CancellationToken);
 		}
 
 		#endregion
@@ -2715,14 +2718,14 @@ namespace EpicGames.Perforce
 		/// <returns>Response from the server</returns>
 		public static async Task<PerforceResponse<SubmitRecord>> TrySubmitAsync(this IPerforceConnection Connection, int ChangeNumber, SubmitOptions Options, CancellationToken CancellationToken)
 		{
-			StringBuilder Arguments = new StringBuilder("submit");
+			List<string> Arguments = new List<string>();
 			if((Options & SubmitOptions.ReopenAsEdit) != 0)
 			{
-				Arguments.Append(" -r");
+				Arguments.Add("-r");
 			}
-			Arguments.AppendFormat(" -c {0}", ChangeNumber);
+			Arguments.Add($"-c{ChangeNumber}");
 
-			return (await CommandAsync<SubmitRecord>(Connection, Arguments.ToString(), null, CancellationToken))[0];
+			return (await CommandAsync<SubmitRecord>(Connection, "submit", Arguments, null, CancellationToken))[0];
 		}
 
 		/// <summary>
@@ -2746,10 +2749,10 @@ namespace EpicGames.Perforce
 		/// <returns>Response from the server</returns>
 		public static async Task<PerforceResponse<SubmitRecord>> TrySubmitShelvedAsync(this IPerforceConnection Connection, int ChangeNumber, CancellationToken CancellationToken)
 		{
-			StringBuilder Arguments = new StringBuilder("submit");
-			Arguments.AppendFormat(" -e {0}", ChangeNumber);
+			List<string> Arguments = new List<string>();
+			Arguments.Add($"-e{ChangeNumber}");
 
-			return (await CommandAsync<SubmitRecord>(Connection, Arguments.ToString(), null, CancellationToken))[0];
+			return (await CommandAsync<SubmitRecord>(Connection, "submit", Arguments, null, CancellationToken))[0];
 		}
 
 		#endregion
@@ -2936,50 +2939,50 @@ namespace EpicGames.Perforce
 		/// <returns>Arguments for the command</returns>
 		private static async Task<PerforceResponseList<T>> SyncInternalAsync<T>(this IPerforceConnection Connection, SyncOptions Options, int MaxFiles, int NumThreads, int Batch, int BatchSize, int Min, int MinSize, string[] FileSpecs, bool bQuiet, CancellationToken CancellationToken) where T : class
 		{
-			StringBuilder Arguments = new StringBuilder("sync");
+			List<string> Arguments = new List<string>();
 			if((Options & SyncOptions.Force) != 0)
 			{
-				Arguments.Append(" -f");
+				Arguments.Add("-f");
 			}
 			if((Options & SyncOptions.KeepWorkspaceFiles) != 0)
 			{
-				Arguments.Append(" -k");
+				Arguments.Add("-k");
 			}
 			if((Options & SyncOptions.FullDepotSyntax) != 0)
 			{
-				Arguments.Append(" -L");
+				Arguments.Add("-L");
 			}
 			if((Options & SyncOptions.PreviewOnly) != 0)
 			{
-				Arguments.Append(" -n");
+				Arguments.Add("-n");
 			}
 			if((Options & SyncOptions.NetworkPreviewOnly) != 0)
 			{
-				Arguments.Append(" -N");
+				Arguments.Add("-N");
 			}
 			if((Options & SyncOptions.DoNotUpdateHaveList) != 0)
 			{
-				Arguments.Append(" -p");
+				Arguments.Add("-p");
 			}
 			if(bQuiet)
 			{
-				Arguments.Append(" -q");
+				Arguments.Add("-q");
 			}
 			if((Options & SyncOptions.ReopenMovedFiles) != 0)
 			{
-				Arguments.Append(" -r");
+				Arguments.Add("-r");
 			}
 			if((Options & SyncOptions.Safe) != 0)
 			{
-				Arguments.Append(" -s");
+				Arguments.Add("-s");
 			}
 			if(MaxFiles != -1)
 			{
-				Arguments.AppendFormat(" -m {0}", MaxFiles);
+				Arguments.Add($"-m{MaxFiles}");
 			}
 			/*			if(NumThreads != -1)
 						{
-							Arguments.AppendFormat(" --parallel-threads={0}", NumThreads);
+							Arguments.Add($"--parallel-threads={0}", NumThreads);
 							if(Batch != -1)
 							{
 								Arguments.AppendFormat(",batch={0}", Batch);
@@ -2999,11 +3002,8 @@ namespace EpicGames.Perforce
 						}*/
 			if (FileSpecs.Length < 10)
 			{
-				foreach (string FileSpec in FileSpecs)
-				{
-					Arguments.AppendFormat(" \"{0}\"", FileSpec);
-				}
-				return await CommandAsync<T>(Connection, Arguments.ToString(), null, CancellationToken);
+				Arguments.AddRange(FileSpecs);
+				return await CommandAsync<T>(Connection, "sync", Arguments, null, CancellationToken);
 			}
 			else
 			{
@@ -3012,7 +3012,7 @@ namespace EpicGames.Perforce
 				{
 					await File.WriteAllLinesAsync(TempFileName, FileSpecs);
 					Arguments.Insert(0, $"-x\"{TempFileName}\" ");
-					return await CommandAsync<T>(Connection, Arguments.ToString(), null, CancellationToken);
+					return await CommandAsync<T>(Connection, "sync", Arguments, null, CancellationToken);
 				}
 				finally
 				{
@@ -3058,38 +3058,35 @@ namespace EpicGames.Perforce
 		/// <returns>Response from the server</returns>
 		public static Task<PerforceResponseList<UnshelveRecord>> TryUnshelveAsync(this IPerforceConnection Connection, int ChangeNumber, int IntoChangeNumber, string UsingBranchSpec, string UsingStream, string ForceParentStream, UnshelveOptions Options, string[] FileSpecs, CancellationToken CancellationToken)
 		{
-			StringBuilder Arguments = new StringBuilder("unshelve");
-			Arguments.AppendFormat(" -s {0}", ChangeNumber);
+			List<string> Arguments = new List<string>();
+			Arguments.Add($"-s{ChangeNumber}");
 			if((Options & UnshelveOptions.ForceOverwrite) != 0)
 			{
-				Arguments.Append(" -f");
+				Arguments.Add("-f");
 			}
 			if((Options & UnshelveOptions.PreviewOnly) != 0)
 			{
-				Arguments.Append(" -n");
+				Arguments.Add("-n");
 			}
 			if(IntoChangeNumber != -1)
 			{
-				Arguments.AppendFormat(" -c {0}", IntoChangeNumber);
+				Arguments.Add($"-c{IntoChangeNumber}");
 			}
 			if(UsingBranchSpec != null)
 			{
-				Arguments.AppendFormat(" -b \"{0}\"", UsingBranchSpec);
+				Arguments.Add($"-b{UsingBranchSpec}");
 			}
 			if(UsingStream != null)
 			{
-				Arguments.AppendFormat(" -S \"{0}\"", UsingStream);
+				Arguments.Add($"-S{UsingStream}");
 			}
 			if(ForceParentStream != null)
 			{
-				Arguments.AppendFormat(" -P \"{0}\"", ForceParentStream);
-			}
-			foreach(string FileSpec in FileSpecs)
-			{
-				Arguments.AppendFormat(" \"{0}\"", FileSpec);
+				Arguments.Add($"-P{ForceParentStream}");
 			}
+			Arguments.AddRange(FileSpecs);
 
-			return CommandAsync<UnshelveRecord>(Connection, Arguments.ToString(), null, CancellationToken);
+			return CommandAsync<UnshelveRecord>(Connection, "unshelve", Arguments, null, CancellationToken);
 		}
 
 		#endregion
@@ -3117,9 +3114,9 @@ namespace EpicGames.Perforce
 		/// <returns>Response from the server</returns>
 		public static Task<PerforceResponseList<UserRecord>> TryGetUserAsync(this IPerforceConnection Connection, string UserName, CancellationToken CancellationToken)
 		{
-			StringBuilder Arguments = new StringBuilder("user");
-			Arguments.AppendFormat(" -o \"{0}\"", UserName);
-			return CommandAsync<UserRecord>(Connection, Arguments.ToString(), null, CancellationToken);
+			List<string> Arguments = new List<string>();
+			Arguments.Add($"-o{UserName}");
+			return CommandAsync<UserRecord>(Connection, "user", Arguments, null, CancellationToken);
 		}
 
 		#endregion
@@ -3150,19 +3147,9 @@ namespace EpicGames.Perforce
 		/// <param name="FileSpecs">Patterns for the files to query</param>
 		/// <param name="CancellationToken">Token used to cancel the operation</param>
 		/// <returns>List of responses from the server</returns>
-		public static async Task<PerforceResponseList<WhereRecord>> TryWhereAsync(this IPerforceConnection Connection, string[] FileSpecs, CancellationToken CancellationToken)
+		public static Task<PerforceResponseList<WhereRecord>> TryWhereAsync(this IPerforceConnection Connection, string[] FileSpecs, CancellationToken CancellationToken)
 		{
-			PerforceResponseList<WhereRecord> Responses = new PerforceResponseList<WhereRecord>();
-			for (int Idx = 0; Idx < FileSpecs.Length; )
-			{
-				StringBuilder Arguments = new StringBuilder("where");
-				for (; Arguments.Length < 2048 && Idx < FileSpecs.Length; Idx++)
-				{
-					Arguments.AppendFormat(" \"{0}\"", FileSpecs[Idx]);
-				}
-				Responses.AddRange(await CommandAsync<WhereRecord>(Connection, Arguments.ToString(), null, CancellationToken));
-			}
-			return Responses;
+			return BatchedCommandAsync<WhereRecord>(Connection, "where", new List<string>(), FileSpecs, null, CancellationToken);
 		}
 
 		#endregion
@@ -3192,11 +3179,10 @@ namespace EpicGames.Perforce
 		/// <returns>Response from the server</returns>
 		public static Task<PerforceResponseList<UndoRecord>> TryUndoChangeAsync(this IPerforceConnection Connection, int ChangeNumberToUndo, int ChangeNumber, CancellationToken CancellationToken)
 		{
-			StringBuilder Arguments = new StringBuilder("undo -c");
-			
-			Arguments.AppendFormat(" {0} //...@{1}", ChangeNumber, ChangeNumberToUndo);
-
-			return CommandAsync<UndoRecord>(Connection, Arguments.ToString(), null, CancellationToken);
+			List<string> Arguments = new List<string>();
+			Arguments.Add($"-c{ChangeNumber}");
+			Arguments.Add($"//...@{ChangeNumberToUndo}");
+			return CommandAsync<UndoRecord>(Connection, "undo", Arguments, null, CancellationToken);
 		}
 		#endregion
 
@@ -3225,28 +3211,26 @@ namespace EpicGames.Perforce
 		/// <returns>List of annotate records</returns>
 		public static Task<PerforceResponseList<AnnotateRecord>> TryAnnotateAsync(this IPerforceConnection Connection, string FileSpec, AnnotateOptions Options, CancellationToken CancellationToken)
 		{
-			PerforceResponseList<AnnotateRecord> Responses = new PerforceResponseList<AnnotateRecord>();
-			StringBuilder Arguments = new StringBuilder("annotate ");
+			List<string> Arguments = new List<string>();
 			if ((Options & AnnotateOptions.IncludeDeletedFilesAndLines) != 0)
 			{
-				Arguments.Append(" -a");
+				Arguments.Add("-a");
 			}
 			if ((Options & AnnotateOptions.IgnoreWhiteSpaceChanges) != 0)
 			{
-				Arguments.Append(" -db");
+				Arguments.Add("-db");
 			}
 			if ((Options & AnnotateOptions.OutputUserAndDate) != 0)
 			{
-				Arguments.Append(" -u");
+				Arguments.Add("-u");
 			}
 			if ((Options & AnnotateOptions.FollowIntegrations) != 0)
 			{
-				Arguments.Append(" -I");
+				Arguments.Add("-I");
 			}
+			Arguments.Add(FileSpec);
 
-			Arguments.AppendFormat(" {0}", FileSpec);
-
-			return CommandAsync<AnnotateRecord>(Connection, Arguments.ToString(), null, CancellationToken);
+			return CommandAsync<AnnotateRecord>(Connection, "annotate", Arguments, null, CancellationToken);
 		}
 
 		#endregion
