Quality of life improvements for EpicGames.Perforce.

* Introduce a new FileSpecList class which supports implicit conversion from a single string and arrays/lists of strings.
* Make the cancellation token argument to all commands optional.

diff --git a/Engine/Source/Programs/Shared/EpicGames.Perforce/FileSpecList.cs b/Engine/Source/Programs/Shared/EpicGames.Perforce/FileSpecList.cs
new file mode 100644
--- /dev/null
+++ b/Engine/Source/Programs/Shared/EpicGames.Perforce/FileSpecList.cs
@@ -0,0 +1,60 @@
+// Copyright Epic Games, Inc. All Rights Reserved.
+
+using System;
+using System.Collections.Generic;
+using System.Text;
+
+namespace EpicGames.Perforce
+{
+	/// <summary>
+	/// Wrapper for a list of filespecs. Allows implicit conversion from string (a single entry) or list.
+	/// </summary>
+	public struct FileSpecList
+	{
+		/// <summary>
+		/// Empty filespec list
+		/// </summary>
+		public static FileSpecList Empty { get; } = new FileSpecList(new List<string>());
+
+		/// <summary>
+		/// The list of filespecs
+		/// </summary>
+		public IReadOnlyList<string> List { get; }
+
+		/// <summary>
+		/// Private constructor. Use implicit conversion operators below instead.
+		/// </summary>
+		/// <param name="FileSpecList">List of filespecs</param>
+		private FileSpecList(IReadOnlyList<string> FileSpecList)
+		{
+			this.List = FileSpecList;
+		}
+
+		/// <summary>
+		/// Implicit conversion operator from a list of filespecs
+		/// </summary>
+		/// <param name="List">The list to construct from</param>
+		public static implicit operator FileSpecList(List<string> List)
+		{
+			return new FileSpecList(List);
+		}
+
+		/// <summary>
+		/// Implicit conversion operator from an array of filespecs
+		/// </summary>
+		/// <param name="Array">The array to construct from</param>
+		public static implicit operator FileSpecList(string[] Array)
+		{
+			return new FileSpecList(Array);
+		}
+
+		/// <summary>
+		/// Implicit conversion operator from a single filespec
+		/// </summary>
+		/// <param name="FileSpec">The single filespec to construct from</param>
+		public static implicit operator FileSpecList(string FileSpec)
+		{
+			return new FileSpecList(new string[] { FileSpec });
+		}
+	}
+}
diff --git a/Engine/Source/Programs/Shared/EpicGames.Perforce/PerforceConnection.cs b/Engine/Source/Programs/Shared/EpicGames.Perforce/PerforceConnection.cs
--- a/Engine/Source/Programs/Shared/EpicGames.Perforce/PerforceConnection.cs
+++ b/Engine/Source/Programs/Shared/EpicGames.Perforce/PerforceConnection.cs
@@ -35,7 +35,7 @@ namespace EpicGames.Perforce
 		/// <param name="StatRecordType">The type of records to return for "stat" responses</param>
 		/// <param name="CancellationToken">Token used to cancel the operation</param>
 		/// <returns>List of objects returned by the server</returns>
-		Task<List<PerforceResponse>> CommandAsync(string Command, IReadOnlyList<string> Arguments, byte[]? InputData, Type? StatRecordType, CancellationToken CancellationToken);
+		Task<List<PerforceResponse>> CommandAsync(string Command, IReadOnlyList<string> Arguments, byte[]? InputData, Type? StatRecordType, CancellationToken CancellationToken = default);
 
 		/// <summary>
 		/// Execute a command and parse the response
@@ -46,7 +46,7 @@ namespace EpicGames.Perforce
 		/// <param name="HandleRecord">Delegate used to handle each record</param>
 		/// <param name="CancellationToken">Token used to cancel the operation</param>
 		/// <returns>List of objects returned by the server</returns>
-		Task RecordCommandAsync(string Command, IReadOnlyList<string> Arguments, byte[]? InputData, Action<PerforceRecord> HandleRecord, CancellationToken CancellationToken);
+		Task RecordCommandAsync(string Command, IReadOnlyList<string> Arguments, byte[]? InputData, Action<PerforceRecord> HandleRecord, CancellationToken CancellationToken = default);
 
 		/// <summary>
 		/// Execute the 'login' command
@@ -54,7 +54,7 @@ namespace EpicGames.Perforce
 		/// <param name="Password">Password to use to login</param>
 		/// <param name="CancellationToken">Token used to cancel the operation</param>
 		/// <returns>Response from the server</returns>
-		Task LoginAsync(string Password, CancellationToken CancellationToken);
+		Task LoginAsync(string Password, CancellationToken CancellationToken = default);
 
 		/// <summary>
 		/// Sets an environment variable
@@ -63,7 +63,7 @@ namespace EpicGames.Perforce
 		/// <param name="Value">Value for the variable</param>
 		/// <param name="CancellationToken">Token used to cancel the operation</param>
 		/// <returns>Response from the server</returns>
-		Task SetAsync(string Name, string Value, CancellationToken CancellationToken);
+		Task SetAsync(string Name, string Value, CancellationToken CancellationToken = default);
 
 		/// <summary>
 		/// Gets the setting of a Perforce variable
@@ -71,7 +71,7 @@ namespace EpicGames.Perforce
 		/// <param name="Name">Name of the variable to get</param>
 		/// <param name="CancellationToken">Cancellation token for the request</param>
 		/// <returns>Value of the variable</returns>
-		Task<string?> TryGetSettingAsync(string Name, CancellationToken CancellationToken);
+		Task<string?> TryGetSettingAsync(string Name, CancellationToken CancellationToken = default);
 	}
 
 	/// <summary>
@@ -234,7 +234,7 @@ namespace EpicGames.Perforce
 		/// <param name="StatRecordType">The type of records to return for "stat" responses</param>
 		/// <param name="CancellationToken">Token used to cancel the operation</param>
 		/// <returns>List of objects returned by the server</returns>
-		public async Task<List<PerforceResponse>> CommandAsync(string Command, IReadOnlyList<string> Arguments, byte[]? InputData, Type? StatRecordType, CancellationToken CancellationToken)
+		public async Task<List<PerforceResponse>> CommandAsync(string Command, IReadOnlyList<string> Arguments, byte[]? InputData, Type? StatRecordType, CancellationToken CancellationToken = default)
 		{
 			using (PerforceChildProcess Process = new PerforceChildProcess(Command, Arguments, InputData, GetGlobalArguments(), Logger))
 			{
@@ -252,7 +252,7 @@ namespace EpicGames.Perforce
 		/// <param name="HandleRecord">Delegate used to handle each record</param>
 		/// <param name="CancellationToken">Token used to cancel the operation</param>
 		/// <returns>List of objects returned by the server</returns>
-		public async Task RecordCommandAsync(string Command, IReadOnlyList<string> Arguments, byte[]? InputData, Action<PerforceRecord> HandleRecord, CancellationToken CancellationToken)
+		public async Task RecordCommandAsync(string Command, IReadOnlyList<string> Arguments, byte[]? InputData, Action<PerforceRecord> HandleRecord, CancellationToken CancellationToken = default)
 		{
 			using (PerforceChildProcess Process = new PerforceChildProcess(Command, Arguments, InputData, GetGlobalArguments(), Logger))
 			{
@@ -268,7 +268,7 @@ namespace EpicGames.Perforce
 		/// <param name="Password">Password to use to login</param>
 		/// <param name="CancellationToken">Token used to cancel the operation</param>
 		/// <returns>Response from the server</returns>
-		public async Task LoginAsync(string Password, CancellationToken CancellationToken)
+		public async Task LoginAsync(string Password, CancellationToken CancellationToken = default)
 		{
 			// Some versions of P4.EXE do not support marshalled output for P4 login calls, so we only support this as a basic text query.
 			byte[] PasswordBytes = Encoding.UTF8.GetBytes(Password);
@@ -322,7 +322,7 @@ namespace EpicGames.Perforce
 		/// <param name="Value">Value for the variable</param>
 		/// <param name="CancellationToken">Token used to cancel the operation</param>
 		/// <returns>Response from the server</returns>
-		public async Task SetAsync(string Name, string Value, CancellationToken CancellationToken)
+		public async Task SetAsync(string Name, string Value, CancellationToken CancellationToken = default)
 		{
 			Tuple<bool, string> Result = await TrySetAsync(Name, Value, CancellationToken);
 			if (!Result.Item1)
@@ -338,7 +338,7 @@ namespace EpicGames.Perforce
 		/// <param name="Value">Value for the variable</param>
 		/// <param name="CancellationToken">Token used to cancel the operation</param>
 		/// <returns>Response from the server</returns>
-		public async Task<Tuple<bool, string>> TrySetAsync(string Name, string Value, CancellationToken CancellationToken)
+		public async Task<Tuple<bool, string>> TrySetAsync(string Name, string Value, CancellationToken CancellationToken = default)
 		{
 			List<string> Arguments = new List<string>();
 			Arguments.Append($"{Name}={Value}");
@@ -355,7 +355,7 @@ namespace EpicGames.Perforce
 		/// <param name="Name">Name of the variable to get</param>
 		/// <param name="CancellationToken">Cancellation token for the request</param>
 		/// <returns>Value of the variable</returns>
-		public async Task<string?> TryGetSettingAsync(string Name, CancellationToken CancellationToken)
+		public async Task<string?> TryGetSettingAsync(string Name, CancellationToken CancellationToken = default)
 		{
 			using (PerforceChildProcess ChildProcess = new PerforceChildProcess("set", new List<string> { $"{Name}=" }, null, GetGlobalArguments(), Logger))
 			{
@@ -432,7 +432,7 @@ namespace EpicGames.Perforce
 		/// <param name="InputData">Input data to pass to Perforce</param>
 		/// <param name="CancellationToken">Token used to cancel the operation</param>
 		/// <returns>List of objects returned by the server</returns>
-		public static async Task<PerforceResponseList<T>> CommandAsync<T>(this IPerforceConnection Connection, string Command, List<string> Arguments, byte[]? InputData, CancellationToken CancellationToken) where T : class
+		public static async Task<PerforceResponseList<T>> CommandAsync<T>(this IPerforceConnection Connection, string Command, List<string> Arguments, byte[]? InputData, CancellationToken CancellationToken = default) where T : class
 		{
 			List<PerforceResponse> Responses = await Connection.CommandAsync(Command, Arguments, InputData, typeof(T), CancellationToken);
 
@@ -454,7 +454,7 @@ namespace EpicGames.Perforce
 		/// <param name="InputData">Input data to pass to Perforce</param>
 		/// <param name="CancellationToken">Token used to cancel the operation</param>
 		/// <returns>List of objects returned by the server</returns>
-		public static async Task<PerforceResponseList<T>> BatchedCommandAsync<T>(IPerforceConnection Connection, string Command, IReadOnlyList<string> CommonArguments, IReadOnlyList<string> BatchedArguments, byte[]? InputData, CancellationToken CancellationToken) where T : class
+		public static async Task<PerforceResponseList<T>> BatchedCommandAsync<T>(IPerforceConnection Connection, string Command, IReadOnlyList<string> CommonArguments, IReadOnlyList<string> BatchedArguments, byte[]? InputData, CancellationToken CancellationToken = default) where T : class
 		{
 			PerforceResponseList<T> Responses = new PerforceResponseList<T>();
 			for (int FileSpecIdx = 0; FileSpecIdx < BatchedArguments.Count;)
@@ -485,7 +485,7 @@ namespace EpicGames.Perforce
 		/// <param name="StatRecordType">Type of element to return in the response</param>
 		/// <param name="CancellationToken">Token used to cancel the operation</param>
 		/// <returns>Response from the server; either an object of type T or error.</returns>
-		static async Task<PerforceResponse> SingleResponseCommandAsync(IPerforceConnection Connection, string Command, IReadOnlyList<string> Arguments, byte[]? InputData, Type? StatRecordType, CancellationToken CancellationToken)
+		static async Task<PerforceResponse> SingleResponseCommandAsync(IPerforceConnection Connection, string Command, IReadOnlyList<string> Arguments, byte[]? InputData, Type? StatRecordType, CancellationToken CancellationToken = default)
 		{
 			List<PerforceResponse> Responses = await Connection.CommandAsync(Command, Arguments, InputData, StatRecordType, CancellationToken);
 			if (Responses.Count != 1)
@@ -509,7 +509,7 @@ namespace EpicGames.Perforce
 		/// <param name="InputData">Input data for the command.</param>
 		/// <param name="CancellationToken">Token used to cancel the operation</param>
 		/// <returns>Response from the server; either an object of type T or error.</returns>
-		public static async Task<PerforceResponse<T>> SingleResponseCommandAsync<T>(IPerforceConnection Connection, string Command, List<string> Arguments, byte[]? InputData, CancellationToken CancellationToken) where T : class
+		public static async Task<PerforceResponse<T>> SingleResponseCommandAsync<T>(IPerforceConnection Connection, string Command, List<string> Arguments, byte[]? InputData, CancellationToken CancellationToken = default) where T : class
 		{
 			return new PerforceResponse<T>(await SingleResponseCommandAsync(Connection, Command, Arguments, InputData, typeof(T), CancellationToken));
 		}
@@ -523,12 +523,12 @@ namespace EpicGames.Perforce
 		/// </summary>
 		/// <param name="Connection">Connection to the Perforce server</param>
 		/// <param name="ChangeNumber">Changelist to add files to</param>
-		/// <param name="FileNames">Files to be added</param>
+		/// <param name="FileSpecList">Files to be added</param>
 		/// <param name="CancellationToken">Token used to cancel the operation</param>
 		/// <returns>Response from the server</returns>
-		public static async Task<List<AddRecord>> AddAsync(this IPerforceConnection Connection, int ChangeNumber, string[] FileNames, CancellationToken CancellationToken)
+		public static async Task<List<AddRecord>> AddAsync(this IPerforceConnection Connection, int ChangeNumber, FileSpecList FileSpecList, CancellationToken CancellationToken = default)
 		{
-			return (await TryAddAsync(Connection, ChangeNumber, null, AddOptions.None, FileNames, CancellationToken)).Data;
+			return (await TryAddAsync(Connection, ChangeNumber, null, AddOptions.None, FileSpecList, CancellationToken)).Data;
 		}
 
 		/// <summary>
@@ -538,12 +538,12 @@ namespace EpicGames.Perforce
 		/// <param name="ChangeNumber">Changelist to add files to</param>
 		/// <param name="FileType">Type for new files</param>
 		/// <param name="Options">Options for the command</param>
-		/// <param name="FileNames">Files to be added</param>
+		/// <param name="FileSpecList">Files to be added</param>
 		/// <param name="CancellationToken">Token used to cancel the operation</param>
 		/// <returns>Response from the server</returns>
-		public static async Task<List<AddRecord>> AddAsync(this IPerforceConnection Connection, int ChangeNumber, string? FileType, AddOptions Options, string[] FileNames, CancellationToken CancellationToken)
+		public static async Task<List<AddRecord>> AddAsync(this IPerforceConnection Connection, int ChangeNumber, string? FileType, AddOptions Options, FileSpecList FileSpecList, CancellationToken CancellationToken = default)
 		{
-			return (await TryAddAsync(Connection, ChangeNumber, FileType, Options, FileNames, CancellationToken)).Data;
+			return (await TryAddAsync(Connection, ChangeNumber, FileType, Options, FileSpecList, CancellationToken)).Data;
 		}
 
 		/// <summary>
@@ -551,12 +551,12 @@ namespace EpicGames.Perforce
 		/// </summary>
 		/// <param name="Connection">Connection to the Perforce server</param>
 		/// <param name="ChangeNumber">Changelist to add files to</param>
-		/// <param name="FileNames">Files to be added</param>
+		/// <param name="FileSpecList">Files to be added</param>
 		/// <param name="CancellationToken">Token used to cancel the operation</param>
 		/// <returns>Response from the server</returns>
-		public static Task<PerforceResponseList<AddRecord>> TryAddAsync(this IPerforceConnection Connection, int ChangeNumber, string[] FileNames, CancellationToken CancellationToken)
+		public static Task<PerforceResponseList<AddRecord>> TryAddAsync(this IPerforceConnection Connection, int ChangeNumber, FileSpecList FileSpecList, CancellationToken CancellationToken = default)
 		{
-			return TryAddAsync(Connection, ChangeNumber, null, AddOptions.None, FileNames, CancellationToken);
+			return TryAddAsync(Connection, ChangeNumber, null, AddOptions.None, FileSpecList, CancellationToken);
 		}
 
 		/// <summary>
@@ -569,7 +569,7 @@ namespace EpicGames.Perforce
 		/// <param name="FileNames">Files to be added</param>
 		/// <param name="CancellationToken">Token used to cancel the operation</param>
 		/// <returns>Response from the server</returns>
-		public static Task<PerforceResponseList<AddRecord>> TryAddAsync(this IPerforceConnection Connection, int ChangeNumber, string? FileType, AddOptions Options, IReadOnlyList<string> FileNames, CancellationToken CancellationToken)
+		public static Task<PerforceResponseList<AddRecord>> TryAddAsync(this IPerforceConnection Connection, int ChangeNumber, string? FileType, AddOptions Options, FileSpecList FileNames, CancellationToken CancellationToken = default)
 		{
 			List<string> Arguments = new List<string>();
 			if (ChangeNumber != -1)
@@ -597,7 +597,7 @@ namespace EpicGames.Perforce
 				Arguments.Add($"-t{FileType}");
 			}
 
-			return BatchedCommandAsync<AddRecord>(Connection, "add", Arguments, FileNames, null, CancellationToken);
+			return BatchedCommandAsync<AddRecord>(Connection, "add", Arguments, FileNames.List, null, CancellationToken);
 		}
 
 		#endregion
@@ -611,7 +611,7 @@ namespace EpicGames.Perforce
 		/// <param name="Record">Information for the change to create. The number field should be left set to -1.</param>
 		/// <param name="CancellationToken">Token used to cancel the operation</param>
 		/// <returns>The changelist number, or an error.</returns>
-		public static async Task<ChangeRecord> CreateChangeAsync(this IPerforceConnection Connection, ChangeRecord Record, CancellationToken CancellationToken)
+		public static async Task<ChangeRecord> CreateChangeAsync(this IPerforceConnection Connection, ChangeRecord Record, CancellationToken CancellationToken = default)
 		{
 			return (await TryCreateChangeAsync(Connection, Record, CancellationToken)).Data;
 		}
@@ -623,7 +623,7 @@ namespace EpicGames.Perforce
 		/// <param name="Record">Information for the change to create. The number field should be left set to -1.</param>
 		/// <param name="CancellationToken">Token used to cancel the operation</param>
 		/// <returns>The changelist number, or an error.</returns>
-		public static async Task<PerforceResponse<ChangeRecord>> TryCreateChangeAsync(this IPerforceConnection Connection, ChangeRecord Record, CancellationToken CancellationToken)
+		public static async Task<PerforceResponse<ChangeRecord>> TryCreateChangeAsync(this IPerforceConnection Connection, ChangeRecord Record, CancellationToken CancellationToken = default)
 		{
 			if (Record.Number != -1)
 			{
@@ -662,7 +662,7 @@ namespace EpicGames.Perforce
 		/// <param name="Record">Information for the change to create. The number field should be left set to zero.</param>
 		/// <param name="CancellationToken">Token used to cancel the operation</param>
 		/// <returns>The changelist number, or an error.</returns>
-		public static async Task UpdateChangeAsync(this IPerforceConnection Connection, UpdateChangeOptions Options, ChangeRecord Record, CancellationToken CancellationToken)
+		public static async Task UpdateChangeAsync(this IPerforceConnection Connection, UpdateChangeOptions Options, ChangeRecord Record, CancellationToken CancellationToken = default)
 		{
 			(await TryUpdateChangeAsync(Connection, Options, Record, CancellationToken)).EnsureSuccess();
 		}
@@ -675,7 +675,7 @@ namespace EpicGames.Perforce
 		/// <param name="Record">Information for the change to create. The number field should be left set to zero.</param>
 		/// <param name="CancellationToken">Token used to cancel the operation</param>
 		/// <returns>The changelist number, or an error.</returns>
-		public static Task<PerforceResponse> TryUpdateChangeAsync(this IPerforceConnection Connection, UpdateChangeOptions Options, ChangeRecord Record, CancellationToken CancellationToken)
+		public static Task<PerforceResponse> TryUpdateChangeAsync(this IPerforceConnection Connection, UpdateChangeOptions Options, ChangeRecord Record, CancellationToken CancellationToken = default)
 		{
 			if (Record.Number == -1)
 			{
@@ -704,7 +704,7 @@ namespace EpicGames.Perforce
 		/// <param name="ChangeNumber">Changelist number to delete</param>
 		/// <param name="CancellationToken">Token used to cancel the operation</param>
 		/// <returns>Response from the server</returns>
-		public static async Task DeleteChangeAsync(this IPerforceConnection Connection, DeleteChangeOptions Options, int ChangeNumber, CancellationToken CancellationToken)
+		public static async Task DeleteChangeAsync(this IPerforceConnection Connection, DeleteChangeOptions Options, int ChangeNumber, CancellationToken CancellationToken = default)
 		{
 			(await TryDeleteChangeAsync(Connection, Options, ChangeNumber, CancellationToken)).EnsureSuccess();
 		}
@@ -717,7 +717,7 @@ namespace EpicGames.Perforce
 		/// <param name="ChangeNumber">Changelist number to delete</param>
 		/// <param name="CancellationToken">Token used to cancel the operation</param>
 		/// <returns>Response from the server</returns>
-		public static Task<PerforceResponse> TryDeleteChangeAsync(this IPerforceConnection Connection, DeleteChangeOptions Options, int ChangeNumber, CancellationToken CancellationToken)
+		public static Task<PerforceResponse> TryDeleteChangeAsync(this IPerforceConnection Connection, DeleteChangeOptions Options, int ChangeNumber, CancellationToken CancellationToken = default)
 		{
 			List<string> Arguments = new List<string> { "-d" };
 			if ((Options & DeleteChangeOptions.Submitted) != 0)
@@ -741,7 +741,7 @@ namespace EpicGames.Perforce
 		/// <param name="ChangeNumber">Changelist number to retrieve. -1 is the default changelist for this workspace.</param>
 		/// <param name="CancellationToken">Token used to cancel the operation</param>
 		/// <returns>Response from the server</returns>
-		public static async Task<ChangeRecord> GetChangeAsync(this IPerforceConnection Connection, GetChangeOptions Options, int ChangeNumber, CancellationToken CancellationToken)
+		public static async Task<ChangeRecord> GetChangeAsync(this IPerforceConnection Connection, GetChangeOptions Options, int ChangeNumber, CancellationToken CancellationToken = default)
 		{
 			return (await TryGetChange(Connection, Options, ChangeNumber, CancellationToken)).Data;
 		}
@@ -754,7 +754,7 @@ namespace EpicGames.Perforce
 		/// <param name="ChangeNumber">Changelist number to retrieve. -1 is the default changelist for this workspace.</param>
 		/// <param name="CancellationToken">Token used to cancel the operation</param>
 		/// <returns>Response from the server</returns>
-		public static Task<PerforceResponse<ChangeRecord>> TryGetChange(this IPerforceConnection Connection, GetChangeOptions Options, int ChangeNumber, CancellationToken CancellationToken)
+		public static Task<PerforceResponse<ChangeRecord>> TryGetChange(this IPerforceConnection Connection, GetChangeOptions Options, int ChangeNumber, CancellationToken CancellationToken = default)
 		{
 			List<string> Arguments = new List<string> { "-o" };
 			if ((Options & GetChangeOptions.BeforeRenumber) != 0)
@@ -818,7 +818,7 @@ namespace EpicGames.Perforce
 		/// <param name="FileSpecs">Paths to query changes for</param>
 		/// <param name="CancellationToken">Token used to cancel the operation</param>
 		/// <returns>List of responses from the server.</returns>
-		public static async Task<List<ChangesRecord>> GetChangesAsync(this IPerforceConnection Connection, ChangesOptions Options, int MaxChanges, ChangeStatus Status, string[] FileSpecs, CancellationToken CancellationToken)
+		public static async Task<List<ChangesRecord>> GetChangesAsync(this IPerforceConnection Connection, ChangesOptions Options, int MaxChanges, ChangeStatus Status, FileSpecList FileSpecs, CancellationToken CancellationToken = default)
 		{
 			return (await TryGetChangesAsync(Connection, Options, MaxChanges, Status, FileSpecs, CancellationToken)).Data;
 		}
@@ -835,7 +835,7 @@ namespace EpicGames.Perforce
 		/// <param name="FileSpecs">Paths to query changes for</param>
 		/// <param name="CancellationToken">Token used to cancel the operation</param>
 		/// <returns>List of responses from the server.</returns>
-		public static async Task<List<ChangesRecord>> GetChangesAsync(this IPerforceConnection Connection, ChangesOptions Options, string? ClientName, int MaxChanges, ChangeStatus Status, string? UserName, string[]? FileSpecs, CancellationToken CancellationToken)
+		public static async Task<List<ChangesRecord>> GetChangesAsync(this IPerforceConnection Connection, ChangesOptions Options, string? ClientName, int MaxChanges, ChangeStatus Status, string? UserName, FileSpecList FileSpecs, CancellationToken CancellationToken = default)
 		{
 			return (await TryGetChangesAsync(Connection, Options, ClientName, MaxChanges, Status, UserName, FileSpecs, CancellationToken)).Data;
 		}
@@ -850,7 +850,7 @@ namespace EpicGames.Perforce
 		/// <param name="FileSpecs">Paths to query changes for</param>
 		/// <param name="CancellationToken">Token used to cancel the operation</param>
 		/// <returns>List of responses from the server.</returns>
-		public static Task<PerforceResponseList<ChangesRecord>> TryGetChangesAsync(this IPerforceConnection Connection, ChangesOptions Options, int MaxChanges, ChangeStatus Status, string[] FileSpecs, CancellationToken CancellationToken)
+		public static Task<PerforceResponseList<ChangesRecord>> TryGetChangesAsync(this IPerforceConnection Connection, ChangesOptions Options, int MaxChanges, ChangeStatus Status, FileSpecList FileSpecs, CancellationToken CancellationToken = default)
 		{
 			return TryGetChangesAsync(Connection, Options, null, MaxChanges, Status, null, FileSpecs, CancellationToken);
 		}
@@ -867,7 +867,7 @@ namespace EpicGames.Perforce
 		/// <param name="FileSpecs">Paths to query changes for</param>
 		/// <param name="CancellationToken">Token used to cancel the operation</param>
 		/// <returns>List of responses from the server.</returns>
-		public static Task<PerforceResponseList<ChangesRecord>> TryGetChangesAsync(this IPerforceConnection Connection, ChangesOptions Options, string? ClientName, int MaxChanges, ChangeStatus Status, string? UserName, string[]? FileSpecs, CancellationToken CancellationToken)
+		public static Task<PerforceResponseList<ChangesRecord>> TryGetChangesAsync(this IPerforceConnection Connection, ChangesOptions Options, string? ClientName, int MaxChanges, ChangeStatus Status, string? UserName, FileSpecList FileSpecs, CancellationToken CancellationToken = default)
 		{
 			List<string> Arguments = new List<string>();
 			if ((Options & ChangesOptions.IncludeIntegrations) != 0)
@@ -906,10 +906,8 @@ namespace EpicGames.Perforce
 			{
 				Arguments.Add($"-u{UserName}");
 			}
-			if (FileSpecs != null)
-			{
-				Arguments.AddRange(FileSpecs);
-			}
+			Arguments.AddRange(FileSpecs.List);
+
 			return CommandAsync<ChangesRecord>(Connection, "changes", Arguments, null, CancellationToken);
 		}
 
@@ -925,7 +923,7 @@ namespace EpicGames.Perforce
 		/// <param name="FileSpecs">Files to sync</param>
 		/// <param name="CancellationToken">Token used to cancel the operation</param>
 		/// <returns>Response from the server</returns>
-		public static async Task<List<CleanRecord>> CleanAsync(this IPerforceConnection Connection, CleanOptions Options, string[] FileSpecs, CancellationToken CancellationToken)
+		public static async Task<List<CleanRecord>> CleanAsync(this IPerforceConnection Connection, CleanOptions Options, FileSpecList FileSpecs, CancellationToken CancellationToken = default)
 		{
 			return (await TryCleanAsync(Connection, Options, FileSpecs, CancellationToken)).Data;
 		}
@@ -938,7 +936,7 @@ namespace EpicGames.Perforce
 		/// <param name="FileSpecs">Files to sync</param>
 		/// <param name="CancellationToken">Token used to cancel the operation</param>
 		/// <returns>Response from the server</returns>
-		public static async Task<PerforceResponseList<CleanRecord>> TryCleanAsync(this IPerforceConnection Connection, CleanOptions Options, string[] FileSpecs, CancellationToken CancellationToken)
+		public static async Task<PerforceResponseList<CleanRecord>> TryCleanAsync(this IPerforceConnection Connection, CleanOptions Options, FileSpecList FileSpecs, CancellationToken CancellationToken = default)
 		{
 			List<string> Arguments = new List<string>();
 			if ((Options & CleanOptions.Edited) != 0)
@@ -969,7 +967,7 @@ namespace EpicGames.Perforce
 			{
 				Arguments.Add("-m");
 			}
-			Arguments.AddRange(FileSpecs);
+			Arguments.AddRange(FileSpecs.List);
 
 			PerforceResponseList<CleanRecord> Records = await CommandAsync<CleanRecord>(Connection, "clean", Arguments, null, CancellationToken);
 			Records.RemoveAll(x => x.Error != null && x.Error.Generic == PerforceGenericCode.Empty);
@@ -987,7 +985,7 @@ namespace EpicGames.Perforce
 		/// <param name="Record">The client record</param>
 		/// <param name="CancellationToken">Token used to cancel the operation</param>
 		/// <returns>Response from the server</returns>
-		public static async Task CreateClientAsync(this IPerforceConnection Connection, ClientRecord Record, CancellationToken CancellationToken)
+		public static async Task CreateClientAsync(this IPerforceConnection Connection, ClientRecord Record, CancellationToken CancellationToken = default)
 		{
 			(await TryCreateClientAsync(Connection, Record, CancellationToken)).EnsureSuccess();
 		}
@@ -999,7 +997,7 @@ namespace EpicGames.Perforce
 		/// <param name="Record">The client record</param>
 		/// <param name="CancellationToken">Token used to cancel the operation</param>
 		/// <returns>Response from the server</returns>
-		public static Task<PerforceResponse> TryCreateClientAsync(this IPerforceConnection Connection, ClientRecord Record, CancellationToken CancellationToken)
+		public static Task<PerforceResponse> TryCreateClientAsync(this IPerforceConnection Connection, ClientRecord Record, CancellationToken CancellationToken = default)
 		{
 			return TryUpdateClientAsync(Connection, Record, CancellationToken);
 		}
@@ -1011,7 +1009,7 @@ namespace EpicGames.Perforce
 		/// <param name="Record">The client record</param>
 		/// <param name="CancellationToken">Token used to cancel the operation</param>
 		/// <returns>Response from the server</returns>
-		public static async Task UpdateClientAsync(this IPerforceConnection Connection, ClientRecord Record, CancellationToken CancellationToken)
+		public static async Task UpdateClientAsync(this IPerforceConnection Connection, ClientRecord Record, CancellationToken CancellationToken = default)
 		{
 			(await TryUpdateClientAsync(Connection, Record, CancellationToken)).EnsureSuccess();
 		}
@@ -1023,7 +1021,7 @@ namespace EpicGames.Perforce
 		/// <param name="Record">The client record</param>
 		/// <param name="CancellationToken">Token used to cancel the operation</param>
 		/// <returns>Response from the server</returns>
-		public static Task<PerforceResponse> TryUpdateClientAsync(this IPerforceConnection Connection, ClientRecord Record, CancellationToken CancellationToken)
+		public static Task<PerforceResponse> TryUpdateClientAsync(this IPerforceConnection Connection, ClientRecord Record, CancellationToken CancellationToken = default)
 		{
 			return SingleResponseCommandAsync(Connection, "client", new List<string> { "-i" }, SerializeRecord(Record), null, CancellationToken);
 		}
@@ -1036,7 +1034,7 @@ namespace EpicGames.Perforce
 		/// <param name="ClientName">Name of the client</param>
 		/// <param name="CancellationToken">Token used to cancel the operation</param>
 		/// <returns>Response from the server</returns>
-		public static async Task DeleteClientAsync(this IPerforceConnection Connection, DeleteClientOptions Options, string ClientName, CancellationToken CancellationToken)
+		public static async Task DeleteClientAsync(this IPerforceConnection Connection, DeleteClientOptions Options, string ClientName, CancellationToken CancellationToken = default)
 		{
 			(await TryDeleteClientAsync(Connection, Options, ClientName, CancellationToken)).EnsureSuccess();
 		}
@@ -1049,7 +1047,7 @@ namespace EpicGames.Perforce
 		/// <param name="ClientName">Name of the client</param>
 		/// <param name="CancellationToken">Token used to cancel the operation</param>
 		/// <returns>Response from the server</returns>
-		public static Task<PerforceResponse> TryDeleteClientAsync(this IPerforceConnection Connection, DeleteClientOptions Options, string ClientName, CancellationToken CancellationToken)
+		public static Task<PerforceResponse> TryDeleteClientAsync(this IPerforceConnection Connection, DeleteClientOptions Options, string ClientName, CancellationToken CancellationToken = default)
 		{
 			List<string> Arguments = new List<string> { "-d" };
 			if ((Options & DeleteClientOptions.Force) != 0)
@@ -1074,7 +1072,7 @@ namespace EpicGames.Perforce
 		/// <param name="Options">Options for this command</param>
 		/// <param name="CancellationToken">Token used to cancel the operation</param>
 		/// <returns>Response from the server</returns>
-		public static async Task SwitchClientToStreamAsync(this IPerforceConnection Connection, string ClientName, string StreamName, SwitchClientOptions Options, CancellationToken CancellationToken)
+		public static async Task SwitchClientToStreamAsync(this IPerforceConnection Connection, string ClientName, string StreamName, SwitchClientOptions Options, CancellationToken CancellationToken = default)
 		{
 			(await TrySwitchClientToStreamAsync(Connection, ClientName, StreamName, Options, CancellationToken)).EnsureSuccess();
 		}
@@ -1088,7 +1086,7 @@ namespace EpicGames.Perforce
 		/// <param name="Options">Options for this command</param>
 		/// <param name="CancellationToken">Token used to cancel the operation</param>
 		/// <returns>Response from the server</returns>
-		public static Task<PerforceResponse> TrySwitchClientToStreamAsync(this IPerforceConnection Connection, string ClientName, string StreamName, SwitchClientOptions Options, CancellationToken CancellationToken)
+		public static Task<PerforceResponse> TrySwitchClientToStreamAsync(this IPerforceConnection Connection, string ClientName, string StreamName, SwitchClientOptions Options, CancellationToken CancellationToken = default)
 		{
 			List<string> Arguments = new List<string> { "-s" };
 			if ((Options & SwitchClientOptions.IgnoreOpenFiles) != 0)
@@ -1108,7 +1106,7 @@ namespace EpicGames.Perforce
 		/// <param name="TemplateName">The new stream to be associated with the client</param>
 		/// <param name="CancellationToken">Token used to cancel the operation</param>
 		/// <returns>Response from the server</returns>
-		public static async Task SwitchClientToTemplateAsync(this IPerforceConnection Connection, string ClientName, string TemplateName, CancellationToken CancellationToken)
+		public static async Task SwitchClientToTemplateAsync(this IPerforceConnection Connection, string ClientName, string TemplateName, CancellationToken CancellationToken = default)
 		{
 			(await TrySwitchClientToTemplateAsync(Connection, ClientName, TemplateName, CancellationToken)).EnsureSuccess();
 		}
@@ -1121,7 +1119,7 @@ namespace EpicGames.Perforce
 		/// <param name="TemplateName">The new stream to be associated with the client</param>
 		/// <param name="CancellationToken">Token used to cancel the operation</param>
 		/// <returns>Response from the server</returns>
-		public static Task<PerforceResponse> TrySwitchClientToTemplateAsync(this IPerforceConnection Connection, string ClientName, string TemplateName, CancellationToken CancellationToken)
+		public static Task<PerforceResponse> TrySwitchClientToTemplateAsync(this IPerforceConnection Connection, string ClientName, string TemplateName, CancellationToken CancellationToken = default)
 		{
 			List<string> Arguments = new List<string>();
 			Arguments.Add("-s");
@@ -1137,7 +1135,7 @@ namespace EpicGames.Perforce
 		/// <param name="ClientName">Name of the client. Specify null for the current client.</param>
 		/// <param name="CancellationToken">Token used to cancel the operation</param>
 		/// <returns>Response from the server; either a client record or error code</returns>
-		public static async Task<ClientRecord> GetClientAsync(this IPerforceConnection Connection, string? ClientName, CancellationToken CancellationToken)
+		public static async Task<ClientRecord> GetClientAsync(this IPerforceConnection Connection, string? ClientName, CancellationToken CancellationToken = default)
 		{
 			return (await TryGetClientAsync(Connection, ClientName, CancellationToken)).Data;
 		}
@@ -1149,7 +1147,7 @@ namespace EpicGames.Perforce
 		/// <param name="ClientName">Name of the client. Specify null for the current client.</param>
 		/// <param name="CancellationToken">Token used to cancel the operation</param>
 		/// <returns>Response from the server; either a client record or error code</returns>
-		public static Task<PerforceResponse<ClientRecord>> TryGetClientAsync(this IPerforceConnection Connection, string? ClientName, CancellationToken CancellationToken)
+		public static Task<PerforceResponse<ClientRecord>> TryGetClientAsync(this IPerforceConnection Connection, string? ClientName, CancellationToken CancellationToken = default)
 		{
 			List<string> Arguments = new List<string> { "-o" };
 			if (ClientName != null)
@@ -1167,7 +1165,7 @@ namespace EpicGames.Perforce
 		/// <param name="bIncludeView">Whether to include the stream view in the output</param>
 		/// <param name="CancellationToken">Token used to cancel the operation</param>
 		/// <returns>Stream information record</returns>
-		public static async Task<StreamRecord> GetStreamAsync(this IPerforceConnection Connection, string StreamName, bool bIncludeView, CancellationToken CancellationToken)
+		public static async Task<StreamRecord> GetStreamAsync(this IPerforceConnection Connection, string StreamName, bool bIncludeView, CancellationToken CancellationToken = default)
 		{
 			return (await TryGetStreamAsync(Connection, StreamName, bIncludeView, CancellationToken)).Data;
 		}
@@ -1180,7 +1178,7 @@ namespace EpicGames.Perforce
 		/// <param name="bIncludeView">Whether to include the stream view in the output</param>
 		/// <param name="CancellationToken">Token used to cancel the operation</param>
 		/// <returns>Stream information record</returns>
-		public static Task<PerforceResponse<StreamRecord>> TryGetStreamAsync(this IPerforceConnection Connection, string StreamName, bool bIncludeView, CancellationToken CancellationToken)
+		public static Task<PerforceResponse<StreamRecord>> TryGetStreamAsync(this IPerforceConnection Connection, string StreamName, bool bIncludeView, CancellationToken CancellationToken = default)
 		{
 			List<string> Arguments = new List<string> { "-o" };
 			if (bIncludeView)
@@ -1200,7 +1198,7 @@ namespace EpicGames.Perforce
 		/// <param name="ChangeNumber">Changelist at which to query the stream view</param>
 		/// <param name="CancellationToken">Token used to cancel the operation</param>
 		/// <returns>Response from the server; either a client record or error code</returns>
-		public static async Task<ClientRecord> GetStreamViewAsync(this IPerforceConnection Connection, string StreamName, int ChangeNumber, CancellationToken CancellationToken)
+		public static async Task<ClientRecord> GetStreamViewAsync(this IPerforceConnection Connection, string StreamName, int ChangeNumber, CancellationToken CancellationToken = default)
 		{
 			return (await TryGetStreamViewAsync(Connection, StreamName, ChangeNumber, CancellationToken)).Data;
 		}
@@ -1213,7 +1211,7 @@ namespace EpicGames.Perforce
 		/// <param name="ChangeNumber">Changelist at which to query the stream view</param>
 		/// <param name="CancellationToken">Token used to cancel the operation</param>
 		/// <returns>Response from the server; either a client record or error code</returns>
-		public static Task<PerforceResponse<ClientRecord>> TryGetStreamViewAsync(this IPerforceConnection Connection, string StreamName, int ChangeNumber, CancellationToken CancellationToken)
+		public static Task<PerforceResponse<ClientRecord>> TryGetStreamViewAsync(this IPerforceConnection Connection, string StreamName, int ChangeNumber, CancellationToken CancellationToken = default)
 		{
 			List<string> Arguments = new List<string> { "-o" };
 			Arguments.Add($"-S{StreamName}");
@@ -1292,7 +1290,7 @@ namespace EpicGames.Perforce
 		/// <param name="UserName">List only client workspaces owned by this user.</param>
 		/// <param name="CancellationToken">Token used to cancel the operation</param>
 		/// <returns>Response from the server; either a client record or error code</returns>
-		public static async Task<List<ClientsRecord>> GetClientsAsync(this IPerforceConnection Connection, ClientsOptions Options, string? UserName, CancellationToken CancellationToken)
+		public static async Task<List<ClientsRecord>> GetClientsAsync(this IPerforceConnection Connection, ClientsOptions Options, string? UserName, CancellationToken CancellationToken = default)
 		{
 			return (await TryGetClientsAsync(Connection, Options, UserName, CancellationToken)).Data;
 		}
@@ -1305,7 +1303,7 @@ namespace EpicGames.Perforce
 		/// <param name="UserName">List only client workspaces owned by this user.</param>
 		/// <param name="CancellationToken">Token used to cancel the operation</param>
 		/// <returns>Response from the server; either a client record or error code</returns>
-		public static Task<PerforceResponseList<ClientsRecord>> TryGetClientsAsync(this IPerforceConnection Connection, ClientsOptions Options, string? UserName, CancellationToken CancellationToken)
+		public static Task<PerforceResponseList<ClientsRecord>> TryGetClientsAsync(this IPerforceConnection Connection, ClientsOptions Options, string? UserName, CancellationToken CancellationToken = default)
 		{
 			return TryGetClientsAsync(Connection, Options, null, -1, null, UserName, CancellationToken);
 		}
@@ -1321,7 +1319,7 @@ namespace EpicGames.Perforce
 		/// <param name="UserName">List only client workspaces owned by this user.</param>
 		/// <param name="CancellationToken">Token used to cancel the operation</param>
 		/// <returns>Response from the server; either a client record or error code</returns>
-		public static async Task<List<ClientsRecord>> GetClientsAsync(this IPerforceConnection Connection, ClientsOptions Options, string? Filter, int MaxResults, string? Stream, string? UserName, CancellationToken CancellationToken)
+		public static async Task<List<ClientsRecord>> GetClientsAsync(this IPerforceConnection Connection, ClientsOptions Options, string? Filter, int MaxResults, string? Stream, string? UserName, CancellationToken CancellationToken = default)
 		{
 			return (await TryGetClientsAsync(Connection, Options, Filter, MaxResults, Stream, UserName, CancellationToken)).Data;
 		}
@@ -1337,7 +1335,7 @@ namespace EpicGames.Perforce
 		/// <param name="UserName">List only client workspaces owned by this user.</param>
 		/// <param name="CancellationToken">Token used to cancel the operation</param>
 		/// <returns>Response from the server; either a client record or error code</returns>
-		public static Task<PerforceResponseList<ClientsRecord>> TryGetClientsAsync(this IPerforceConnection Connection, ClientsOptions Options, string? Filter, int MaxResults, string? Stream, string? UserName, CancellationToken CancellationToken)
+		public static Task<PerforceResponseList<ClientsRecord>> TryGetClientsAsync(this IPerforceConnection Connection, ClientsOptions Options, string? Filter, int MaxResults, string? Stream, string? UserName, CancellationToken CancellationToken = default)
 		{
 			List<string> Arguments = new List<string>();
 			if ((Options & ClientsOptions.All) != 0)
@@ -1393,7 +1391,7 @@ namespace EpicGames.Perforce
 		/// <param name="FileSpecs">List of file specifications to query</param>
 		/// <param name="CancellationToken">Token used to cancel the operation</param>
 		/// <returns>List of responses from the server</returns>
-		public static async Task<List<DeleteRecord>> DeleteAsync(this IPerforceConnection Connection, int ChangeNumber, DeleteOptions Options, string[] FileSpecs, CancellationToken CancellationToken)
+		public static async Task<List<DeleteRecord>> DeleteAsync(this IPerforceConnection Connection, int ChangeNumber, DeleteOptions Options, FileSpecList FileSpecs, CancellationToken CancellationToken = default)
 		{
 			return (await TryDeleteAsync(Connection, ChangeNumber, Options, FileSpecs, CancellationToken)).Data;
 		}
@@ -1407,7 +1405,7 @@ namespace EpicGames.Perforce
 		/// <param name="FileSpecs">List of file specifications to query</param>
 		/// <param name="CancellationToken">Token used to cancel the operation</param>
 		/// <returns>List of responses from the server</returns>
-		public static Task<PerforceResponseList<DeleteRecord>> TryDeleteAsync(this IPerforceConnection Connection, int ChangeNumber, DeleteOptions Options, string[] FileSpecs, CancellationToken CancellationToken)
+		public static Task<PerforceResponseList<DeleteRecord>> TryDeleteAsync(this IPerforceConnection Connection, int ChangeNumber, DeleteOptions Options, FileSpecList FileSpecs, CancellationToken CancellationToken = default)
 		{
 			List<string> Arguments = new List<string>();
 			if (ChangeNumber != -1)
@@ -1427,7 +1425,7 @@ namespace EpicGames.Perforce
 				Arguments.Add("-v");
 			}
 
-			return BatchedCommandAsync<DeleteRecord>(Connection, "delete", Arguments, FileSpecs, null, CancellationToken);
+			return BatchedCommandAsync<DeleteRecord>(Connection, "delete", Arguments, FileSpecs.List, null, CancellationToken);
 		}
 
 		#endregion
@@ -1441,7 +1439,7 @@ namespace EpicGames.Perforce
 		/// <param name="ChangeNumber">The changelist number to retrieve description for</param>
 		/// <param name="CancellationToken">Token used to cancel the operation</param>
 		/// <returns>Response from the server; either a describe record or error code</returns>
-		public static async Task<DescribeRecord> DescribeAsync(this IPerforceConnection Connection, int ChangeNumber, CancellationToken CancellationToken)
+		public static async Task<DescribeRecord> DescribeAsync(this IPerforceConnection Connection, int ChangeNumber, CancellationToken CancellationToken = default)
 		{
 			return (await TryDescribeAsync(Connection, ChangeNumber, CancellationToken)).Data;
 		}
@@ -1453,7 +1451,7 @@ namespace EpicGames.Perforce
 		/// <param name="ChangeNumber">The changelist number to retrieve description for</param>
 		/// <param name="CancellationToken">Token used to cancel the operation</param>
 		/// <returns>Response from the server; either a describe record or error code</returns>
-		public static async Task<PerforceResponse<DescribeRecord>> TryDescribeAsync(this IPerforceConnection Connection, int ChangeNumber, CancellationToken CancellationToken)
+		public static async Task<PerforceResponse<DescribeRecord>> TryDescribeAsync(this IPerforceConnection Connection, int ChangeNumber, CancellationToken CancellationToken = default)
 		{
 			PerforceResponseList<DescribeRecord> Records = await TryDescribeAsync(Connection, new int[] { ChangeNumber }, CancellationToken);
 			if (Records.Count != 1)
@@ -1470,7 +1468,7 @@ namespace EpicGames.Perforce
 		/// <param name="ChangeNumbers">The changelist numbers to retrieve descriptions for</param>
 		/// <param name="CancellationToken">Token used to cancel the operation</param>
 		/// <returns>List of responses from the server</returns>
-		public static async Task<List<DescribeRecord>> DescribeAsync(this IPerforceConnection Connection, int[] ChangeNumbers, CancellationToken CancellationToken)
+		public static async Task<List<DescribeRecord>> DescribeAsync(this IPerforceConnection Connection, int[] ChangeNumbers, CancellationToken CancellationToken = default)
 		{
 			return (await TryDescribeAsync(Connection, ChangeNumbers, CancellationToken)).Data;
 		}
@@ -1482,7 +1480,7 @@ namespace EpicGames.Perforce
 		/// <param name="ChangeNumbers">The changelist numbers to retrieve descriptions for</param>
 		/// <param name="CancellationToken">Token used to cancel the operation</param>
 		/// <returns>List of responses from the server</returns>
-		public static Task<PerforceResponseList<DescribeRecord>> TryDescribeAsync(this IPerforceConnection Connection, int[] ChangeNumbers, CancellationToken CancellationToken)
+		public static Task<PerforceResponseList<DescribeRecord>> TryDescribeAsync(this IPerforceConnection Connection, int[] ChangeNumbers, CancellationToken CancellationToken = default)
 		{
 			List<string> Arguments = new List<string> { "-s" };
 			foreach (int ChangeNumber in ChangeNumbers)
@@ -1501,7 +1499,7 @@ namespace EpicGames.Perforce
 		/// <param name="ChangeNumbers">The changelist numbers to retrieve descriptions for</param>
 		/// <param name="CancellationToken">Token used to cancel the operation</param>
 		/// <returns>List of responses from the server</returns>
-		public static async Task<List<DescribeRecord>> DescribeAsync(this IPerforceConnection Connection, DescribeOptions Options, int MaxNumFiles, int[] ChangeNumbers, CancellationToken CancellationToken)
+		public static async Task<List<DescribeRecord>> DescribeAsync(this IPerforceConnection Connection, DescribeOptions Options, int MaxNumFiles, int[] ChangeNumbers, CancellationToken CancellationToken = default)
 		{
 			return (await TryDescribeAsync(Connection, Options, MaxNumFiles, ChangeNumbers, CancellationToken)).Data;
 		}
@@ -1515,7 +1513,7 @@ namespace EpicGames.Perforce
 		/// <param name="ChangeNumbers">The changelist numbers to retrieve descriptions for</param>
 		/// <param name="CancellationToken">Token used to cancel the operation</param>
 		/// <returns>List of responses from the server</returns>
-		public static Task<PerforceResponseList<DescribeRecord>> TryDescribeAsync(this IPerforceConnection Connection, DescribeOptions Options, int MaxNumFiles, int[] ChangeNumbers, CancellationToken CancellationToken)
+		public static Task<PerforceResponseList<DescribeRecord>> TryDescribeAsync(this IPerforceConnection Connection, DescribeOptions Options, int MaxNumFiles, int[] ChangeNumbers, CancellationToken CancellationToken = default)
 		{
 			List<string> Arguments = new List<string> { "-s" };
 			if ((Options & DescribeOptions.ShowDescriptionForRestrictedChanges) != 0)
@@ -1557,7 +1555,7 @@ namespace EpicGames.Perforce
 		/// <param name="FileSpecs">Files to be opened for edit</param>
 		/// <param name="CancellationToken">Token used to cancel the operation</param>
 		/// <returns>Response from the server</returns>
-		public static async Task<List<EditRecord>> EditAsync(this IPerforceConnection Connection, int ChangeNumber, string[] FileSpecs, CancellationToken CancellationToken)
+		public static async Task<List<EditRecord>> EditAsync(this IPerforceConnection Connection, int ChangeNumber, FileSpecList FileSpecs, CancellationToken CancellationToken = default)
 		{
 			return (await TryEditAsync(Connection, ChangeNumber, FileSpecs, CancellationToken)).Data;
 		}
@@ -1570,7 +1568,7 @@ namespace EpicGames.Perforce
 		/// <param name="FileSpecs">Files to be opened for edit</param>
 		/// <param name="CancellationToken">Token used to cancel the operation</param>
 		/// <returns>Response from the server</returns>
-		public static Task<PerforceResponseList<EditRecord>> TryEditAsync(this IPerforceConnection Connection, int ChangeNumber, string[] FileSpecs, CancellationToken CancellationToken)
+		public static Task<PerforceResponseList<EditRecord>> TryEditAsync(this IPerforceConnection Connection, int ChangeNumber, FileSpecList FileSpecs, CancellationToken CancellationToken = default)
 		{
 			return TryEditAsync(Connection, ChangeNumber, null, EditOptions.None, FileSpecs, CancellationToken);
 		}
@@ -1585,7 +1583,7 @@ namespace EpicGames.Perforce
 		/// <param name="FileSpecs">Files to be opened for edit</param>
 		/// <param name="CancellationToken">Token used to cancel the operation</param>
 		/// <returns>Response from the server</returns>
-		public static async Task<List<EditRecord>> EditAsync(this IPerforceConnection Connection, int ChangeNumber, string? FileType, EditOptions Options, string[] FileSpecs, CancellationToken CancellationToken)
+		public static async Task<List<EditRecord>> EditAsync(this IPerforceConnection Connection, int ChangeNumber, string? FileType, EditOptions Options, FileSpecList FileSpecs, CancellationToken CancellationToken = default)
 		{
 			return (await TryEditAsync(Connection, ChangeNumber, FileType, Options, FileSpecs, CancellationToken)).Data;
 		}
@@ -1600,7 +1598,7 @@ namespace EpicGames.Perforce
 		/// <param name="FileSpecs">Files to be opened for edit</param>
 		/// <param name="CancellationToken">Token used to cancel the operation</param>
 		/// <returns>Response from the server</returns>
-		public static Task<PerforceResponseList<EditRecord>> TryEditAsync(this IPerforceConnection Connection, int ChangeNumber, string? FileType, EditOptions Options, string[] FileSpecs, CancellationToken CancellationToken)
+		public static Task<PerforceResponseList<EditRecord>> TryEditAsync(this IPerforceConnection Connection, int ChangeNumber, string? FileType, EditOptions Options, FileSpecList FileSpecs, CancellationToken CancellationToken = default)
 		{
 			List<string> Arguments = new List<string>();
 			if (ChangeNumber != -1)
@@ -1620,7 +1618,7 @@ namespace EpicGames.Perforce
 				Arguments.Add($"-t{FileType}");
 			}
 
-			return BatchedCommandAsync<EditRecord>(Connection, "edit", Arguments, FileSpecs, null, CancellationToken);
+			return BatchedCommandAsync<EditRecord>(Connection, "edit", Arguments, FileSpecs.List, null, CancellationToken);
 		}
 
 		#endregion
@@ -1635,7 +1633,7 @@ namespace EpicGames.Perforce
 		/// <param name="FileSpecs">List of file specifications to query</param>
 		/// <param name="CancellationToken">Token used to cancel the operation</param>
 		/// <returns>List of responses from the server</returns>
-		public static async Task<List<FileLogRecord>> FileLogAsync(this IPerforceConnection Connection, FileLogOptions Options, string[] FileSpecs, CancellationToken CancellationToken)
+		public static async Task<List<FileLogRecord>> FileLogAsync(this IPerforceConnection Connection, FileLogOptions Options, FileSpecList FileSpecs, CancellationToken CancellationToken = default)
 		{
 			return (await TryFileLogAsync(Connection, Options, FileSpecs, CancellationToken)).Data;
 		}
@@ -1648,7 +1646,7 @@ namespace EpicGames.Perforce
 		/// <param name="FileSpecs">List of file specifications to query</param>
 		/// <param name="CancellationToken">Token used to cancel the operation</param>
 		/// <returns>List of responses from the server</returns>
-		public static Task<PerforceResponseList<FileLogRecord>> TryFileLogAsync(this IPerforceConnection Connection, FileLogOptions Options, string[] FileSpecs, CancellationToken CancellationToken)
+		public static Task<PerforceResponseList<FileLogRecord>> TryFileLogAsync(this IPerforceConnection Connection, FileLogOptions Options, FileSpecList FileSpecs, CancellationToken CancellationToken = default)
 		{
 			return TryFileLogAsync(Connection, -1, -1, Options, FileSpecs, CancellationToken);
 		}
@@ -1662,7 +1660,7 @@ namespace EpicGames.Perforce
 		/// <param name="FileSpecs">List of file specifications to query</param>
 		/// <param name="CancellationToken">Token used to cancel the operation</param>
 		/// <returns>List of responses from the server</returns>
-		public static async Task<List<FileLogRecord>> FileLogAsync(this IPerforceConnection Connection, int MaxChanges, FileLogOptions Options, string[] FileSpecs, CancellationToken CancellationToken)
+		public static async Task<List<FileLogRecord>> FileLogAsync(this IPerforceConnection Connection, int MaxChanges, FileLogOptions Options, FileSpecList FileSpecs, CancellationToken CancellationToken = default)
 		{
 			return (await TryFileLogAsync(Connection, MaxChanges, Options, FileSpecs, CancellationToken)).Data;
 		}
@@ -1676,7 +1674,7 @@ namespace EpicGames.Perforce
 		/// <param name="FileSpecs">List of file specifications to query</param>
 		/// <param name="CancellationToken">Token used to cancel the operation</param>
 		/// <returns>List of responses from the server</returns>
-		public static Task<PerforceResponseList<FileLogRecord>> TryFileLogAsync(this IPerforceConnection Connection, int MaxChanges, FileLogOptions Options, string[] FileSpecs, CancellationToken CancellationToken)
+		public static Task<PerforceResponseList<FileLogRecord>> TryFileLogAsync(this IPerforceConnection Connection, int MaxChanges, FileLogOptions Options, FileSpecList FileSpecs, CancellationToken CancellationToken = default)
 		{
 			return TryFileLogAsync(Connection, -1, MaxChanges, Options, FileSpecs, CancellationToken);
 		}
@@ -1691,7 +1689,7 @@ namespace EpicGames.Perforce
 		/// <param name="FileSpecs">List of file specifications to query</param>
 		/// <param name="CancellationToken">Token used to cancel the operation</param>
 		/// <returns>List of responses from the server</returns>
-		public static async Task<List<FileLogRecord>> FileLogAsync(this IPerforceConnection Connection, int ChangeNumber, int MaxChanges, FileLogOptions Options, string[] FileSpecs, CancellationToken CancellationToken)
+		public static async Task<List<FileLogRecord>> FileLogAsync(this IPerforceConnection Connection, int ChangeNumber, int MaxChanges, FileLogOptions Options, FileSpecList FileSpecs, CancellationToken CancellationToken = default)
 		{
 			return (await TryFileLogAsync(Connection, ChangeNumber, MaxChanges, Options, FileSpecs, CancellationToken)).Data;
 		}
@@ -1706,7 +1704,7 @@ namespace EpicGames.Perforce
 		/// <param name="FileSpecs">List of file specifications to query</param>
 		/// <param name="CancellationToken">Token used to cancel the operation</param>
 		/// <returns>List of responses from the server</returns>
-		public static Task<PerforceResponseList<FileLogRecord>> TryFileLogAsync(this IPerforceConnection Connection, int ChangeNumber, int MaxChanges, FileLogOptions Options, string[] FileSpecs, CancellationToken CancellationToken)
+		public static Task<PerforceResponseList<FileLogRecord>> TryFileLogAsync(this IPerforceConnection Connection, int ChangeNumber, int MaxChanges, FileLogOptions Options, FileSpecList FileSpecs, CancellationToken CancellationToken = default)
 		{
 			// Build the argument list
 			List<string> Arguments = new List<string>();
@@ -1747,7 +1745,7 @@ namespace EpicGames.Perforce
 			Arguments.Add("-t");
 
 			// Append all the arguments
-			return BatchedCommandAsync<FileLogRecord>(Connection, "filelog", Arguments, FileSpecs, null, CancellationToken);
+			return BatchedCommandAsync<FileLogRecord>(Connection, "filelog", Arguments, FileSpecs.List, null, CancellationToken);
 		}
 
 		#endregion
@@ -1762,7 +1760,7 @@ namespace EpicGames.Perforce
 		/// <param name="FileSpecs">List of file specifications to query</param>
 		/// <param name="CancellationToken">Token used to cancel the operation</param>
 		/// <returns>List of responses from the server</returns>
-		public static async Task<List<FStatRecord>> FStatAsync(this IPerforceConnection Connection, FStatOptions Options, string[] FileSpecs, CancellationToken CancellationToken)
+		public static async Task<List<FStatRecord>> FStatAsync(this IPerforceConnection Connection, FStatOptions Options, FileSpecList FileSpecs, CancellationToken CancellationToken = default)
 		{
 			return (await TryFStatAsync(Connection, Options, FileSpecs, CancellationToken)).Data;
 		}
@@ -1775,7 +1773,7 @@ namespace EpicGames.Perforce
 		/// <param name="FileSpecs">List of file specifications to query</param>
 		/// <param name="CancellationToken">Token used to cancel the operation</param>
 		/// <returns>List of responses from the server</returns>
-		public static Task<PerforceResponseList<FStatRecord>> TryFStatAsync(this IPerforceConnection Connection, FStatOptions Options, string[] FileSpecs, CancellationToken CancellationToken)
+		public static Task<PerforceResponseList<FStatRecord>> TryFStatAsync(this IPerforceConnection Connection, FStatOptions Options, FileSpecList FileSpecs, CancellationToken CancellationToken = default)
 		{
 			return TryFStatAsync(Connection, -1, Options, FileSpecs, CancellationToken);
 		}
@@ -1789,7 +1787,7 @@ namespace EpicGames.Perforce
 		/// <param name="FileSpecs">List of file specifications to query</param>
 		/// <param name="CancellationToken">Token used to cancel the operation</param>
 		/// <returns>List of responses from the server</returns>
-		public static async Task<List<FStatRecord>> FStatAsync(this IPerforceConnection Connection, int MaxFiles, FStatOptions Options, string[] FileSpecs, CancellationToken CancellationToken)
+		public static async Task<List<FStatRecord>> FStatAsync(this IPerforceConnection Connection, int MaxFiles, FStatOptions Options, FileSpecList FileSpecs, CancellationToken CancellationToken = default)
 		{
 			return (await TryFStatAsync(Connection, MaxFiles, Options, FileSpecs, CancellationToken)).Data;
 		}
@@ -1803,7 +1801,7 @@ namespace EpicGames.Perforce
 		/// <param name="FileSpecs">List of file specifications to query</param>
 		/// <param name="CancellationToken">Token used to cancel the operation</param>
 		/// <returns>List of responses from the server</returns>
-		public static Task<PerforceResponseList<FStatRecord>> TryFStatAsync(this IPerforceConnection Connection, int MaxFiles, FStatOptions Options, string[] FileSpecs, CancellationToken CancellationToken)
+		public static Task<PerforceResponseList<FStatRecord>> TryFStatAsync(this IPerforceConnection Connection, int MaxFiles, FStatOptions Options, FileSpecList FileSpecs, CancellationToken CancellationToken = default)
 		{
 			return TryFStatAsync(Connection, -1, -1, null, MaxFiles, Options, FileSpecs, CancellationToken);
 		}
@@ -1820,7 +1818,7 @@ namespace EpicGames.Perforce
 		/// <param name="FileSpecs">List of file specifications to query</param>
 		/// <param name="CancellationToken">Token used to cancel the operation</param>
 		/// <returns>List of responses from the server</returns>
-		public static async Task<List<FStatRecord>> FStatAsync(this IPerforceConnection Connection, int AfterChangeNumber, int OnlyChangeNumber, string? Filter, int MaxFiles, FStatOptions Options, string[] FileSpecs, CancellationToken CancellationToken)
+		public static async Task<List<FStatRecord>> FStatAsync(this IPerforceConnection Connection, int AfterChangeNumber, int OnlyChangeNumber, string? Filter, int MaxFiles, FStatOptions Options, FileSpecList FileSpecs, CancellationToken CancellationToken = default)
 		{
 			return (await TryFStatAsync(Connection, AfterChangeNumber, OnlyChangeNumber, Filter, MaxFiles, Options, FileSpecs, CancellationToken)).Data;
 		}
@@ -1837,7 +1835,7 @@ namespace EpicGames.Perforce
 		/// <param name="FileSpecs">List of file specifications to query</param>
 		/// <param name="CancellationToken">Token used to cancel the operation</param>
 		/// <returns>List of responses from the server</returns>
-		public static async Task<PerforceResponseList<FStatRecord>> TryFStatAsync(this IPerforceConnection Connection, int AfterChangeNumber, int OnlyChangeNumber, string? Filter, int MaxFiles, FStatOptions Options, string[] FileSpecs, CancellationToken CancellationToken)
+		public static async Task<PerforceResponseList<FStatRecord>> TryFStatAsync(this IPerforceConnection Connection, int AfterChangeNumber, int OnlyChangeNumber, string? Filter, int MaxFiles, FStatOptions Options, FileSpecList FileSpecs, CancellationToken CancellationToken = default)
 		{
 			// Build the argument list
 			List<string> Arguments = new List<string>();
@@ -1936,7 +1934,7 @@ namespace EpicGames.Perforce
 			}
 
 			// Execute the command
-			PerforceResponseList<FStatRecord> Records = await BatchedCommandAsync<FStatRecord>(Connection, "fstat", Arguments, FileSpecs, null, CancellationToken);
+			PerforceResponseList<FStatRecord> Records = await BatchedCommandAsync<FStatRecord>(Connection, "fstat", Arguments, FileSpecs.List, null, CancellationToken);
 			if (OnlyChangeNumber != -1 && Records.Count > 0 && Records[Records.Count - 1].Succeeded && Records[Records.Count - 1].Data.Description != null)
 			{
 				Records.RemoveAt(Records.Count - 1);
@@ -1955,7 +1953,7 @@ namespace EpicGames.Perforce
 		/// <param name="Options">Options for the command</param>
 		/// <param name="CancellationToken">Token used to cancel the operation</param>
 		/// <returns>Response from the server; an InfoRecord or error code</returns>
-		public static async Task<InfoRecord> GetInfoAsync(this IPerforceConnection Connection, InfoOptions Options, CancellationToken CancellationToken)
+		public static async Task<InfoRecord> GetInfoAsync(this IPerforceConnection Connection, InfoOptions Options, CancellationToken CancellationToken = default)
 		{
 			return (await TryGetInfoAsync(Connection, Options, CancellationToken)).Data;
 		}
@@ -1967,7 +1965,7 @@ namespace EpicGames.Perforce
 		/// <param name="Options">Options for the command</param>
 		/// <param name="CancellationToken">Token used to cancel the operation</param>
 		/// <returns>Response from the server; an InfoRecord or error code</returns>
-		public static Task<PerforceResponse<InfoRecord>> TryGetInfoAsync(this IPerforceConnection Connection, InfoOptions Options, CancellationToken CancellationToken)
+		public static Task<PerforceResponse<InfoRecord>> TryGetInfoAsync(this IPerforceConnection Connection, InfoOptions Options, CancellationToken CancellationToken = default)
 		{
 			// Build the argument list
 			List<string> Arguments = new List<string>();
@@ -1993,7 +1991,7 @@ namespace EpicGames.Perforce
 		/// <param name="TargetFileSpec">The target filespec</param>
 		/// <param name="CancellationToken">Cancellation token</param>
 		/// <returns>List of records</returns>
-		public static Task<PerforceResponseList<MergeRecord>> MergeAsync(this IPerforceConnection Connection, MergeOptions Options, int Change, int MaxFiles, string SourceFileSpec, string TargetFileSpec, CancellationToken CancellationToken)
+		public static Task<PerforceResponseList<MergeRecord>> MergeAsync(this IPerforceConnection Connection, MergeOptions Options, int Change, int MaxFiles, string SourceFileSpec, string TargetFileSpec, CancellationToken CancellationToken = default)
 		{
 			List<string> Arguments = new List<string>();
 			if ((Options & MergeOptions.Preview) != 0)
@@ -2029,7 +2027,7 @@ namespace EpicGames.Perforce
 		/// <param name="TargetFileSpec">The target file(s)</param>
 		/// <param name="CancellationToken">Token used to cancel the operation</param>
 		/// <returns>Response from the server</returns>
-		public static async Task<List<MoveRecord>> MoveAsync(this IPerforceConnection Connection, int ChangeNumber, string? FileType, MoveOptions Options, string SourceFileSpec, string TargetFileSpec, CancellationToken CancellationToken)
+		public static async Task<List<MoveRecord>> MoveAsync(this IPerforceConnection Connection, int ChangeNumber, string? FileType, MoveOptions Options, string SourceFileSpec, string TargetFileSpec, CancellationToken CancellationToken = default)
 		{
 			return (await TryMoveAsync(Connection, ChangeNumber, FileType, Options, SourceFileSpec, TargetFileSpec, CancellationToken)).Data;
 		}
@@ -2045,7 +2043,7 @@ namespace EpicGames.Perforce
 		/// <param name="TargetFileSpec">The target file(s)</param>
 		/// <param name="CancellationToken">Token used to cancel the operation</param>
 		/// <returns>Response from the server</returns>
-		public static Task<PerforceResponseList<MoveRecord>> TryMoveAsync(this IPerforceConnection Connection, int ChangeNumber, string? FileType, MoveOptions Options, string SourceFileSpec, string TargetFileSpec, CancellationToken CancellationToken)
+		public static Task<PerforceResponseList<MoveRecord>> TryMoveAsync(this IPerforceConnection Connection, int ChangeNumber, string? FileType, MoveOptions Options, string SourceFileSpec, string TargetFileSpec, CancellationToken CancellationToken = default)
 		{
 			List<string> Arguments = new List<string>();
 			if (ChangeNumber != -1)
@@ -2090,7 +2088,7 @@ namespace EpicGames.Perforce
 		/// <param name="FileSpecs">Specification for the files to list</param>
 		/// <param name="CancellationToken">Token used to cancel the operation</param>
 		/// <returns>Response from the server</returns>
-		public static async Task<List<FStatRecord>> GetOpenFilesAsync(this IPerforceConnection Connection, OpenedOptions Options, int ChangeNumber, string? ClientName, string? UserName, int MaxResults, string[]? FileSpecs, CancellationToken CancellationToken)
+		public static async Task<List<FStatRecord>> GetOpenFilesAsync(this IPerforceConnection Connection, OpenedOptions Options, int ChangeNumber, string? ClientName, string? UserName, int MaxResults, FileSpecList FileSpecs, CancellationToken CancellationToken = default)
 		{
 			return (await TryGetOpenFilesAsync(Connection, Options, ChangeNumber, ClientName, UserName, MaxResults, FileSpecs, CancellationToken)).Data;
 		}
@@ -2107,7 +2105,7 @@ namespace EpicGames.Perforce
 		/// <param name="FileSpecs">Specification for the files to list</param>
 		/// <param name="CancellationToken">Token used to cancel the operation</param>
 		/// <returns>Response from the server</returns>
-		public static Task<PerforceResponseList<FStatRecord>> TryGetOpenFilesAsync(this IPerforceConnection Connection, OpenedOptions Options, int ChangeNumber, string? ClientName, string? UserName, int MaxResults, string[]? FileSpecs, CancellationToken CancellationToken)
+		public static Task<PerforceResponseList<FStatRecord>> TryGetOpenFilesAsync(this IPerforceConnection Connection, OpenedOptions Options, int ChangeNumber, string? ClientName, string? UserName, int MaxResults, FileSpecList FileSpecs, CancellationToken CancellationToken = default)
 		{
 			// Build the argument list
 			List<string> Arguments = new List<string>();
@@ -2139,10 +2137,8 @@ namespace EpicGames.Perforce
 			{
 				Arguments.Add($"-s");
 			}
-			if (FileSpecs != null)
-			{
-				Arguments.AddRange(FileSpecs);
-			}
+			Arguments.AddRange(FileSpecs.List);
+
 			return CommandAsync<FStatRecord>(Connection, "opened", Arguments, null, CancellationToken);
 		}
 
@@ -2158,7 +2154,7 @@ namespace EpicGames.Perforce
 		/// <param name="FileSpec">Specification for the files to print</param>
 		/// <param name="CancellationToken">Token used to cancel the operation</param>
 		/// <returns>Response from the server</returns>
-		public static async Task<PrintRecord> PrintAsync(this IPerforceConnection Connection, string OutputFile, string FileSpec, CancellationToken CancellationToken)
+		public static async Task<PrintRecord> PrintAsync(this IPerforceConnection Connection, string OutputFile, string FileSpec, CancellationToken CancellationToken = default)
 		{
 			return (await TryPrintAsync(Connection, OutputFile, FileSpec, CancellationToken)).Data;
 		}
@@ -2171,7 +2167,7 @@ namespace EpicGames.Perforce
 		/// <param name="FileSpec">Specification for the files to print</param>
 		/// <param name="CancellationToken">Token used to cancel the operation</param>
 		/// <returns>Response from the server</returns>
-		public static Task<PerforceResponse<PrintRecord>> TryPrintAsync(this IPerforceConnection Connection, string OutputFile, string FileSpec, CancellationToken CancellationToken)
+		public static Task<PerforceResponse<PrintRecord>> TryPrintAsync(this IPerforceConnection Connection, string OutputFile, string FileSpec, CancellationToken CancellationToken = default)
 		{
 			// Build the argument list
 			List<string> Arguments = new List<string>();
@@ -2257,7 +2253,7 @@ namespace EpicGames.Perforce
 		/// <param name="FileSpecs">Files to be reverted</param>
 		/// <param name="CancellationToken">Token used to cancel the operation</param>
 		/// <returns>Response from the server</returns>
-		public static async Task<List<ReconcileRecord>> ReconcileAsync(this IPerforceConnection Connection, int ChangeNumber, ReconcileOptions Options, string[] FileSpecs, CancellationToken CancellationToken)
+		public static async Task<List<ReconcileRecord>> ReconcileAsync(this IPerforceConnection Connection, int ChangeNumber, ReconcileOptions Options, FileSpecList FileSpecs, CancellationToken CancellationToken = default)
 		{
 			PerforceResponseList<ReconcileRecord> Records = await TryReconcileAsync(Connection, ChangeNumber, Options, FileSpecs, CancellationToken);
 			Records.RemoveAll(x => x.Info != null);
@@ -2273,7 +2269,7 @@ namespace EpicGames.Perforce
 		/// <param name="FileSpecs">Files to be reverted</param>
 		/// <param name="CancellationToken">Token used to cancel the operation</param>
 		/// <returns>Response from the server</returns>
-		public static Task<PerforceResponseList<ReconcileRecord>> TryReconcileAsync(this IPerforceConnection Connection, int ChangeNumber, ReconcileOptions Options, string[] FileSpecs, CancellationToken CancellationToken)
+		public static Task<PerforceResponseList<ReconcileRecord>> TryReconcileAsync(this IPerforceConnection Connection, int ChangeNumber, ReconcileOptions Options, FileSpecList FileSpecs, CancellationToken CancellationToken = default)
 		{
 			List<string> Arguments = new List<string>();
 			if(ChangeNumber != -1)
@@ -2312,7 +2308,7 @@ namespace EpicGames.Perforce
 			{
 				Arguments.Add("-m");
 			}
-			Arguments.AddRange(FileSpecs);
+			Arguments.AddRange(FileSpecs.List);
 
 			return CommandAsync<ReconcileRecord>(Connection, "reconcile", Arguments, null, CancellationToken);
 		}
@@ -2330,7 +2326,7 @@ namespace EpicGames.Perforce
 		/// <param name="FileSpec">Files to be reverted</param>
 		/// <param name="CancellationToken">Token used to cancel the operation</param>
 		/// <returns>Response from the server</returns>
-		public static async Task<List<ReopenRecord>> ReopenAsync(this IPerforceConnection Connection, int? ChangeNumber, string? FileType, string FileSpec, CancellationToken CancellationToken)
+		public static async Task<List<ReopenRecord>> ReopenAsync(this IPerforceConnection Connection, int? ChangeNumber, string? FileType, string FileSpec, CancellationToken CancellationToken = default)
 		{
 			return (await TryReopenAsync(Connection, ChangeNumber, FileType, FileSpec, CancellationToken)).Data;
 		}
@@ -2344,7 +2340,7 @@ namespace EpicGames.Perforce
 		/// <param name="FileSpec">Files to be reverted</param>
 		/// <param name="CancellationToken">Token used to cancel the operation</param>
 		/// <returns>Response from the server</returns>
-		public static Task<PerforceResponseList<ReopenRecord>> TryReopenAsync(this IPerforceConnection Connection, int? ChangeNumber, string? FileType, string FileSpec, CancellationToken CancellationToken)
+		public static Task<PerforceResponseList<ReopenRecord>> TryReopenAsync(this IPerforceConnection Connection, int? ChangeNumber, string? FileType, string FileSpec, CancellationToken CancellationToken = default)
 		{
 			List<string> Arguments = new List<string>();
 			if (ChangeNumber != null)
@@ -2379,7 +2375,7 @@ namespace EpicGames.Perforce
 		/// <param name="SourceServerId">The source server id</param>
 		/// <param name="CancellationToken">Token used to cancel the operation</param>
 		/// <returns>Response from the server</returns>
-		public static Task<List<PerforceResponse>> ReloadClient(this IPerforceConnection Connection, string ClientName, string SourceServerId, CancellationToken CancellationToken)
+		public static Task<List<PerforceResponse>> ReloadClient(this IPerforceConnection Connection, string ClientName, string SourceServerId, CancellationToken CancellationToken = default)
 		{
 			List<string> Arguments = new List<string>();
 			Arguments.Add($"-c{ClientName}");
@@ -2400,7 +2396,7 @@ namespace EpicGames.Perforce
 		/// <param name="FileSpecs">Files to be reverted</param>
 		/// <param name="CancellationToken">Token used to cancel the operation</param>
 		/// <returns>Response from the server</returns>
-		public static async Task<List<ResolveRecord>> ResolveAsync(this IPerforceConnection Connection, int ChangeNumber, ResolveOptions Options, string[]? FileSpecs, CancellationToken CancellationToken)
+		public static async Task<List<ResolveRecord>> ResolveAsync(this IPerforceConnection Connection, int ChangeNumber, ResolveOptions Options, FileSpecList FileSpecs, CancellationToken CancellationToken = default)
 		{
 			return (await TryResolveAsync(Connection, ChangeNumber, Options, FileSpecs, CancellationToken)).Data;
 		}
@@ -2414,7 +2410,7 @@ namespace EpicGames.Perforce
 		/// <param name="FileSpecs">Files to be reverted</param>
 		/// <param name="CancellationToken">Token used to cancel the operation</param>
 		/// <returns>Response from the server</returns>
-		public static async Task<PerforceResponseList<ResolveRecord>> TryResolveAsync(this IPerforceConnection Connection, int ChangeNumber, ResolveOptions Options, string[]? FileSpecs, CancellationToken CancellationToken)
+		public static async Task<PerforceResponseList<ResolveRecord>> TryResolveAsync(this IPerforceConnection Connection, int ChangeNumber, ResolveOptions Options, FileSpecList FileSpecs, CancellationToken CancellationToken = default)
 		{
 			List<string> Arguments = new List<string>();
 			if ((Options & ResolveOptions.Automatic) != 0)
@@ -2457,10 +2453,7 @@ namespace EpicGames.Perforce
 			{
 				Arguments.Add("-n");
 			}
-			if (FileSpecs != null)
-			{
-				Arguments.AddRange(FileSpecs);
-			}
+			Arguments.AddRange(FileSpecs.List);
 
 			PerforceResponseList<ResolveRecord> Records = await CommandAsync<ResolveRecord>(Connection, "resolve", Arguments, null, CancellationToken);
 			Records.RemoveAll(x => x.Error != null && x.Error.Generic == PerforceGenericCode.Empty);
@@ -2481,7 +2474,7 @@ namespace EpicGames.Perforce
 		/// <param name="FileSpecs">Files to be reverted</param>
 		/// <param name="CancellationToken">Token used to cancel the operation</param>
 		/// <returns>Response from the server</returns>
-		public static async Task<List<RevertRecord>> RevertAsync(this IPerforceConnection Connection, int ChangeNumber, string? ClientName, RevertOptions Options, string[] FileSpecs, CancellationToken CancellationToken)
+		public static async Task<List<RevertRecord>> RevertAsync(this IPerforceConnection Connection, int ChangeNumber, string? ClientName, RevertOptions Options, FileSpecList FileSpecs, CancellationToken CancellationToken = default)
 		{
 			return (await TryRevertAsync(Connection, ChangeNumber, ClientName, Options, FileSpecs, CancellationToken)).Data;
 		}
@@ -2496,7 +2489,7 @@ namespace EpicGames.Perforce
 		/// <param name="FileSpecs">Files to be reverted</param>
 		/// <param name="CancellationToken">Token used to cancel the operation</param>
 		/// <returns>Response from the server</returns>
-		public static async Task<PerforceResponseList<RevertRecord>> TryRevertAsync(this IPerforceConnection Connection, int ChangeNumber, string? ClientName, RevertOptions Options, string[] FileSpecs, CancellationToken CancellationToken)
+		public static async Task<PerforceResponseList<RevertRecord>> TryRevertAsync(this IPerforceConnection Connection, int ChangeNumber, string? ClientName, RevertOptions Options, FileSpecList FileSpecs, CancellationToken CancellationToken = default)
 		{
 			List<string> Arguments = new List<string>();
 			if((Options & RevertOptions.Unchanged) != 0)
@@ -2524,7 +2517,7 @@ namespace EpicGames.Perforce
 				Arguments.Add($"-C{ClientName}");
 			}
 
-			PerforceResponseList<RevertRecord> Records = await BatchedCommandAsync<RevertRecord>(Connection, "revert", Arguments, FileSpecs, null, CancellationToken);
+			PerforceResponseList<RevertRecord> Records = await BatchedCommandAsync<RevertRecord>(Connection, "revert", Arguments, FileSpecs.List, null, CancellationToken);
 			Records.RemoveAll(x => x.Error != null && x.Error.Generic == PerforceGenericCode.Empty);
 			return Records;
 		}
@@ -2542,7 +2535,7 @@ namespace EpicGames.Perforce
 		/// <param name="FileSpecs">Files to sync</param>
 		/// <param name="CancellationToken">Token used to cancel the operation</param>
 		/// <returns>Response from the server</returns>
-		public static async Task<List<ShelveRecord>> ShelveAsync(this IPerforceConnection Connection, int ChangeNumber, ShelveOptions Options, string[] FileSpecs, CancellationToken CancellationToken)
+		public static async Task<List<ShelveRecord>> ShelveAsync(this IPerforceConnection Connection, int ChangeNumber, ShelveOptions Options, FileSpecList FileSpecs, CancellationToken CancellationToken = default)
 		{
 			return (await TryShelveAsync(Connection, ChangeNumber, Options, FileSpecs, CancellationToken)).Data;
 		}
@@ -2556,7 +2549,7 @@ namespace EpicGames.Perforce
 		/// <param name="FileSpecs">Files to sync</param>
 		/// <param name="CancellationToken">Token used to cancel the operation</param>
 		/// <returns>Response from the server</returns>
-		public static async Task<PerforceResponseList<ShelveRecord>> TryShelveAsync(this IPerforceConnection Connection, int ChangeNumber, ShelveOptions Options, string[] FileSpecs, CancellationToken CancellationToken)
+		public static async Task<PerforceResponseList<ShelveRecord>> TryShelveAsync(this IPerforceConnection Connection, int ChangeNumber, ShelveOptions Options, FileSpecList FileSpecs, CancellationToken CancellationToken = default)
 		{
 			List<string> Arguments = new List<string>();
 			Arguments.Add($"-c{ChangeNumber}");
@@ -2568,7 +2561,7 @@ namespace EpicGames.Perforce
 			{
 				Arguments.Add("-f");
 			}
-			Arguments.AddRange(FileSpecs);
+			Arguments.AddRange(FileSpecs.List);
 
 			PerforceResponseList<ShelveRecord> Records = await CommandAsync<ShelveRecord>(Connection, "shelve", Arguments, null, CancellationToken);
 			Records.RemoveAll(x => x.Error != null && x.Error.Generic == PerforceGenericCode.Empty);
@@ -2583,7 +2576,7 @@ namespace EpicGames.Perforce
 		/// <param name="FileSpecs">Files to delete</param>
 		/// <param name="CancellationToken">Token used to cancel the operation</param>
 		/// <returns>Response from the server</returns>
-		public static async Task DeleteShelvedFilesAsync(this IPerforceConnection Connection, int ChangeNumber, string[] FileSpecs, CancellationToken CancellationToken)
+		public static async Task DeleteShelvedFilesAsync(this IPerforceConnection Connection, int ChangeNumber, FileSpecList FileSpecs, CancellationToken CancellationToken = default)
 		{
 			List<PerforceResponse> Responses = await TryDeleteShelvedFilesAsync(Connection, ChangeNumber, FileSpecs, CancellationToken);
 
@@ -2602,7 +2595,7 @@ namespace EpicGames.Perforce
 		/// <param name="FileSpecs">Files to delete</param>
 		/// <param name="CancellationToken">Token used to cancel the operation</param>
 		/// <returns>Response from the server</returns>
-		public static Task<List<PerforceResponse>> TryDeleteShelvedFilesAsync(this IPerforceConnection Connection, int ChangeNumber, string[] FileSpecs, CancellationToken CancellationToken)
+		public static Task<List<PerforceResponse>> TryDeleteShelvedFilesAsync(this IPerforceConnection Connection, int ChangeNumber, FileSpecList FileSpecs, CancellationToken CancellationToken = default)
 		{
 			List<string> Arguments = new List<string>();
 			Arguments.Add("-d");
@@ -2610,7 +2603,7 @@ namespace EpicGames.Perforce
 			{
 				Arguments.Add($"-c{ChangeNumber}");
 			}
-			Arguments.AddRange(FileSpecs);
+			Arguments.AddRange(FileSpecs.List);
 
 			return Connection.CommandAsync("shelve", Arguments, null, null, CancellationToken);
 		}
@@ -2626,7 +2619,7 @@ namespace EpicGames.Perforce
 		/// <param name="StreamPath">The path for streams to enumerate (eg. "//UE4/...")</param>
 		/// <param name="CancellationToken">Token used to cancel the operation</param>
 		/// <returns>List of streams matching the given criteria</returns>
-		public static async Task<List<StreamsRecord>> GetStreamsAsync(this IPerforceConnection Connection, string StreamPath, CancellationToken CancellationToken)
+		public static async Task<List<StreamsRecord>> GetStreamsAsync(this IPerforceConnection Connection, string StreamPath, CancellationToken CancellationToken = default)
 		{
 			return (await TryGetStreamsAsync(Connection, StreamPath, CancellationToken)).Data;
 		}
@@ -2638,7 +2631,7 @@ namespace EpicGames.Perforce
 		/// <param name="StreamPath">The path for streams to enumerate (eg. "//UE4/...")</param>
 		/// <param name="CancellationToken">Token used to cancel the operation</param>
 		/// <returns>List of streams matching the given criteria</returns>
-		public static Task<PerforceResponseList<StreamsRecord>> TryGetStreamsAsync(this IPerforceConnection Connection, string StreamPath, CancellationToken CancellationToken)
+		public static Task<PerforceResponseList<StreamsRecord>> TryGetStreamsAsync(this IPerforceConnection Connection, string StreamPath, CancellationToken CancellationToken = default)
 		{
 			return TryGetStreamsAsync(Connection, StreamPath, -1, null, false, CancellationToken);
 		}
@@ -2653,7 +2646,7 @@ namespace EpicGames.Perforce
 		/// <param name="bUnloaded">Whether to enumerate unloaded workspaces</param>
 		/// <param name="CancellationToken">Token used to cancel the operation</param>
 		/// <returns>List of streams matching the given criteria</returns>
-		public static async Task<List<StreamsRecord>> GetStreamsAsync(this IPerforceConnection Connection, string StreamPath, int MaxResults, string? Filter, bool bUnloaded, CancellationToken CancellationToken)
+		public static async Task<List<StreamsRecord>> GetStreamsAsync(this IPerforceConnection Connection, string StreamPath, int MaxResults, string? Filter, bool bUnloaded, CancellationToken CancellationToken = default)
 		{
 			return (await TryGetStreamsAsync(Connection, StreamPath, MaxResults, Filter, bUnloaded, CancellationToken)).Data;
 		}
@@ -2668,7 +2661,7 @@ namespace EpicGames.Perforce
 		/// <param name="bUnloaded">Whether to enumerate unloaded workspaces</param>
 		/// <param name="CancellationToken">Token used to cancel the operation</param>
 		/// <returns>List of streams matching the given criteria</returns>
-		public static Task<PerforceResponseList<StreamsRecord>> TryGetStreamsAsync(this IPerforceConnection Connection, string StreamPath, int MaxResults, string? Filter, bool bUnloaded, CancellationToken CancellationToken)
+		public static Task<PerforceResponseList<StreamsRecord>> TryGetStreamsAsync(this IPerforceConnection Connection, string StreamPath, int MaxResults, string? Filter, bool bUnloaded, CancellationToken CancellationToken = default)
 		{
 			// Build the command line
 			List<string> Arguments = new List<string>();
@@ -2703,7 +2696,7 @@ namespace EpicGames.Perforce
 		/// <param name="Options">Options for the command</param>
 		/// <param name="CancellationToken">Token used to cancel the operation</param>
 		/// <returns>Response from the server</returns>
-		public static async Task<SubmitRecord> SubmitAsync(this IPerforceConnection Connection, int ChangeNumber, SubmitOptions Options, CancellationToken CancellationToken)
+		public static async Task<SubmitRecord> SubmitAsync(this IPerforceConnection Connection, int ChangeNumber, SubmitOptions Options, CancellationToken CancellationToken = default)
 		{
 			return (await TrySubmitAsync(Connection, ChangeNumber, Options, CancellationToken)).Data;
 		}
@@ -2716,7 +2709,7 @@ namespace EpicGames.Perforce
 		/// <param name="Options">Options for the command</param>
 		/// <param name="CancellationToken">Token used to cancel the operation</param>
 		/// <returns>Response from the server</returns>
-		public static async Task<PerforceResponse<SubmitRecord>> TrySubmitAsync(this IPerforceConnection Connection, int ChangeNumber, SubmitOptions Options, CancellationToken CancellationToken)
+		public static async Task<PerforceResponse<SubmitRecord>> TrySubmitAsync(this IPerforceConnection Connection, int ChangeNumber, SubmitOptions Options, CancellationToken CancellationToken = default)
 		{
 			List<string> Arguments = new List<string>();
 			if((Options & SubmitOptions.ReopenAsEdit) != 0)
@@ -2735,7 +2728,7 @@ namespace EpicGames.Perforce
 		/// <param name="ChangeNumber">The changelist to submit</param>
 		/// <param name="CancellationToken">Token used to cancel the operation</param>
 		/// <returns>Response from the server</returns>
-		public static async Task<SubmitRecord> SubmitShelvedAsync(this IPerforceConnection Connection, int ChangeNumber, CancellationToken CancellationToken)
+		public static async Task<SubmitRecord> SubmitShelvedAsync(this IPerforceConnection Connection, int ChangeNumber, CancellationToken CancellationToken = default)
 		{
 			return (await TrySubmitShelvedAsync(Connection, ChangeNumber, CancellationToken)).Data;
 		}
@@ -2747,7 +2740,7 @@ namespace EpicGames.Perforce
 		/// <param name="ChangeNumber">The changelist to submit</param>
 		/// <param name="CancellationToken">Token used to cancel the operation</param>
 		/// <returns>Response from the server</returns>
-		public static async Task<PerforceResponse<SubmitRecord>> TrySubmitShelvedAsync(this IPerforceConnection Connection, int ChangeNumber, CancellationToken CancellationToken)
+		public static async Task<PerforceResponse<SubmitRecord>> TrySubmitShelvedAsync(this IPerforceConnection Connection, int ChangeNumber, CancellationToken CancellationToken = default)
 		{
 			List<string> Arguments = new List<string>();
 			Arguments.Add($"-e{ChangeNumber}");
@@ -2766,7 +2759,7 @@ namespace EpicGames.Perforce
 		/// <param name="FileSpecs">Files to sync</param>
 		/// <param name="CancellationToken">Token used to cancel the operation</param>
 		/// <returns>Response from the server</returns>
-		public static async Task<List<SyncRecord>> SyncAsync(this IPerforceConnection Connection, string[] FileSpecs, CancellationToken CancellationToken)
+		public static async Task<List<SyncRecord>> SyncAsync(this IPerforceConnection Connection, FileSpecList FileSpecs, CancellationToken CancellationToken = default)
 		{
 			return (await TrySyncAsync(Connection, FileSpecs, CancellationToken)).Data;
 		}
@@ -2778,7 +2771,7 @@ namespace EpicGames.Perforce
 		/// <param name="FileSpecs">Files to sync</param>
 		/// <param name="CancellationToken">Token used to cancel the operation</param>
 		/// <returns>Response from the server</returns>
-		public static Task<PerforceResponseList<SyncRecord>> TrySyncAsync(this IPerforceConnection Connection, string[] FileSpecs, CancellationToken CancellationToken)
+		public static Task<PerforceResponseList<SyncRecord>> TrySyncAsync(this IPerforceConnection Connection, FileSpecList FileSpecs, CancellationToken CancellationToken = default)
 		{
 			return TrySyncAsync(Connection, SyncOptions.None, -1, FileSpecs, CancellationToken);
 		}
@@ -2792,7 +2785,7 @@ namespace EpicGames.Perforce
 		/// <param name="FileSpecs">Files to sync</param>
 		/// <param name="CancellationToken">Token used to cancel the operation</param>
 		/// <returns>Response from the server</returns>
-		public static async Task<List<SyncRecord>> SyncAsync(this IPerforceConnection Connection, SyncOptions Options, int MaxFiles, string[] FileSpecs, CancellationToken CancellationToken)
+		public static async Task<List<SyncRecord>> SyncAsync(this IPerforceConnection Connection, SyncOptions Options, int MaxFiles, FileSpecList FileSpecs, CancellationToken CancellationToken = default)
 		{
 			return (await TrySyncAsync(Connection, Options, MaxFiles, FileSpecs, CancellationToken)).Data;
 		}
@@ -2806,7 +2799,7 @@ namespace EpicGames.Perforce
 		/// <param name="FileSpecs">Files to sync</param>
 		/// <param name="CancellationToken">Token used to cancel the operation</param>
 		/// <returns>Response from the server</returns>
-		public static Task<PerforceResponseList<SyncRecord>> TrySyncAsync(this IPerforceConnection Connection, SyncOptions Options, int MaxFiles, string[] FileSpecs, CancellationToken CancellationToken)
+		public static Task<PerforceResponseList<SyncRecord>> TrySyncAsync(this IPerforceConnection Connection, SyncOptions Options, int MaxFiles, FileSpecList FileSpecs, CancellationToken CancellationToken = default)
 		{
 			return TrySyncAsync(Connection, Options, MaxFiles, -1, -1, -1, -1, -1, FileSpecs, CancellationToken);
 		}
@@ -2825,7 +2818,7 @@ namespace EpicGames.Perforce
 		/// <param name="FileSpecs">Files to sync</param>
 		/// <param name="CancellationToken">Token used to cancel the operation</param>
 		/// <returns>Response from the server</returns>
-		public static async Task<List<SyncRecord>> SyncAsync(this IPerforceConnection Connection, SyncOptions Options, int MaxFiles, int NumThreads, int Batch, int BatchSize, int Min, int MinSize, string[] FileSpecs, CancellationToken CancellationToken)
+		public static async Task<List<SyncRecord>> SyncAsync(this IPerforceConnection Connection, SyncOptions Options, int MaxFiles, int NumThreads, int Batch, int BatchSize, int Min, int MinSize, FileSpecList FileSpecs, CancellationToken CancellationToken = default)
 		{
 			return (await TrySyncAsync(Connection, Options, MaxFiles, NumThreads, Batch, BatchSize, Min, MinSize, FileSpecs, CancellationToken)).Data;
 		}
@@ -2844,7 +2837,7 @@ namespace EpicGames.Perforce
 		/// <param name="FileSpecs">Files to sync</param>
 		/// <param name="CancellationToken">Token used to cancel the operation</param>
 		/// <returns>Response from the server</returns>
-		public static async Task<PerforceResponseList<SyncRecord>> TrySyncAsync(this IPerforceConnection Connection, SyncOptions Options, int MaxFiles, int NumThreads, int Batch, int BatchSize, int Min, int MinSize, string[] FileSpecs, CancellationToken CancellationToken)
+		public static async Task<PerforceResponseList<SyncRecord>> TrySyncAsync(this IPerforceConnection Connection, SyncOptions Options, int MaxFiles, int NumThreads, int Batch, int BatchSize, int Min, int MinSize, FileSpecList FileSpecs, CancellationToken CancellationToken = default)
 		{
 			// Perforce annoyingly returns 'up-to-date' as an error. Ignore it.
 			PerforceResponseList<SyncRecord> Records = await SyncInternalAsync<SyncRecord>(Connection, Options, MaxFiles, NumThreads, Batch, BatchSize, Min, MinSize, FileSpecs, false, CancellationToken);
@@ -2861,7 +2854,7 @@ namespace EpicGames.Perforce
 		/// <param name="FileSpecs">Files to sync</param>
 		/// <param name="CancellationToken">Token used to cancel the operation</param>
 		/// <returns>Response from the server</returns>
-		public static async Task<List<SyncSummaryRecord>> SyncQuietAsync(this IPerforceConnection Connection, SyncOptions Options, int MaxFiles, string[] FileSpecs, CancellationToken CancellationToken)
+		public static async Task<List<SyncSummaryRecord>> SyncQuietAsync(this IPerforceConnection Connection, SyncOptions Options, int MaxFiles, FileSpecList FileSpecs, CancellationToken CancellationToken = default)
 		{
 			PerforceResponseList<SyncSummaryRecord> Records = await TrySyncQuietAsync(Connection, Options, MaxFiles, FileSpecs, CancellationToken);
 			Records.RemoveAll(x => (x.Error != null && x.Error.Generic == PerforceGenericCode.Empty) || x.Info != null);
@@ -2877,7 +2870,7 @@ namespace EpicGames.Perforce
 		/// <param name="FileSpecs">Files to sync</param>
 		/// <param name="CancellationToken">Token used to cancel the operation</param>
 		/// <returns>Response from the server</returns>
-		public static Task<PerforceResponseList<SyncSummaryRecord>> TrySyncQuietAsync(this IPerforceConnection Connection, SyncOptions Options, int MaxFiles, string[] FileSpecs, CancellationToken CancellationToken)
+		public static Task<PerforceResponseList<SyncSummaryRecord>> TrySyncQuietAsync(this IPerforceConnection Connection, SyncOptions Options, int MaxFiles, FileSpecList FileSpecs, CancellationToken CancellationToken = default)
 		{
 			return TrySyncQuietAsync(Connection, Options, MaxFiles, -1, -1, -1, -1, -1, FileSpecs, null, CancellationToken);
 		}
@@ -2897,7 +2890,7 @@ namespace EpicGames.Perforce
 		/// <param name="ResponseFile">Response file for list of files to sync</param>
 		/// <param name="CancellationToken">Token used to cancel the operation</param>
 		/// <returns>Response from the server</returns>
-		public static async Task<List<SyncSummaryRecord>> SyncQuietAsync(this IPerforceConnection Connection, SyncOptions Options, int MaxFiles, int NumThreads, int Batch, int BatchSize, int Min, int MinSize, string[] FileSpecs, FileReference? ResponseFile, CancellationToken CancellationToken)
+		public static async Task<List<SyncSummaryRecord>> SyncQuietAsync(this IPerforceConnection Connection, SyncOptions Options, int MaxFiles, int NumThreads, int Batch, int BatchSize, int Min, int MinSize, FileSpecList FileSpecs, FileReference? ResponseFile, CancellationToken CancellationToken = default)
 		{
 			return (await TrySyncQuietAsync(Connection, Options, MaxFiles, NumThreads, Batch, BatchSize, Min, MinSize, FileSpecs, ResponseFile, CancellationToken)).Data;
 		}
@@ -2917,7 +2910,7 @@ namespace EpicGames.Perforce
 		/// <param name="ResponseFile">Response file for list of files to sync</param>
 		/// <param name="CancellationToken">Token used to cancel the operation</param>
 		/// <returns>Response from the server</returns>
-		public static Task<PerforceResponseList<SyncSummaryRecord>> TrySyncQuietAsync(this IPerforceConnection Connection, SyncOptions Options, int MaxFiles, int NumThreads, int Batch, int BatchSize, int Min, int MinSize, string[] FileSpecs, FileReference? ResponseFile, CancellationToken CancellationToken)
+		public static Task<PerforceResponseList<SyncSummaryRecord>> TrySyncQuietAsync(this IPerforceConnection Connection, SyncOptions Options, int MaxFiles, int NumThreads, int Batch, int BatchSize, int Min, int MinSize, FileSpecList FileSpecs, FileReference? ResponseFile, CancellationToken CancellationToken = default)
 		{
 			return SyncInternalAsync<SyncSummaryRecord>(Connection, Options, MaxFiles, NumThreads, Batch, BatchSize, Min, MinSize, FileSpecs, true, CancellationToken);
 		}
@@ -2937,7 +2930,7 @@ namespace EpicGames.Perforce
 		/// <param name="bQuiet">Whether to use quiet output</param>
 		/// <param name="CancellationToken">Token used to cancel the operation</param>
 		/// <returns>Arguments for the command</returns>
-		private static async Task<PerforceResponseList<T>> SyncInternalAsync<T>(this IPerforceConnection Connection, SyncOptions Options, int MaxFiles, int NumThreads, int Batch, int BatchSize, int Min, int MinSize, string[] FileSpecs, bool bQuiet, CancellationToken CancellationToken) where T : class
+		private static async Task<PerforceResponseList<T>> SyncInternalAsync<T>(this IPerforceConnection Connection, SyncOptions Options, int MaxFiles, int NumThreads, int Batch, int BatchSize, int Min, int MinSize, FileSpecList FileSpecs, bool bQuiet, CancellationToken CancellationToken = default) where T : class
 		{
 			List<string> Arguments = new List<string>();
 			if((Options & SyncOptions.Force) != 0)
@@ -3000,9 +2993,9 @@ namespace EpicGames.Perforce
 								Arguments.AppendFormat(",minsize={0}", MinSize);
 							}
 						}*/
-			if (FileSpecs.Length < 10)
+			if (FileSpecs.List.Count < 10)
 			{
-				Arguments.AddRange(FileSpecs);
+				Arguments.AddRange(FileSpecs.List);
 				return await CommandAsync<T>(Connection, "sync", Arguments, null, CancellationToken);
 			}
 			else
@@ -3010,7 +3003,7 @@ namespace EpicGames.Perforce
 				string TempFileName = Path.GetTempFileName();
 				try
 				{
-					await File.WriteAllLinesAsync(TempFileName, FileSpecs);
+					await File.WriteAllLinesAsync(TempFileName, FileSpecs.List);
 					Arguments.Insert(0, $"-x\"{TempFileName}\" ");
 					return await CommandAsync<T>(Connection, "sync", Arguments, null, CancellationToken);
 				}
@@ -3038,7 +3031,7 @@ namespace EpicGames.Perforce
 		/// <param name="FileSpecs">Files to unshelve</param>
 		/// <param name="CancellationToken">Token used to cancel the operation</param>
 		/// <returns>Response from the server</returns>
-		public static async Task<List<UnshelveRecord>> UnshelveAsync(this IPerforceConnection Connection, int ChangeNumber, int IntoChangeNumber, string UsingBranchSpec, string UsingStream, string ForceParentStream, UnshelveOptions Options, string[] FileSpecs, CancellationToken CancellationToken)
+		public static async Task<List<UnshelveRecord>> UnshelveAsync(this IPerforceConnection Connection, int ChangeNumber, int IntoChangeNumber, string UsingBranchSpec, string UsingStream, string ForceParentStream, UnshelveOptions Options, FileSpecList FileSpecs, CancellationToken CancellationToken = default)
 		{
 			return (await TryUnshelveAsync(Connection, ChangeNumber, IntoChangeNumber, UsingBranchSpec, UsingStream, ForceParentStream, Options, FileSpecs, CancellationToken)).Data;
 		}
@@ -3056,7 +3049,7 @@ namespace EpicGames.Perforce
 		/// <param name="FileSpecs">Files to unshelve</param>
 		/// <param name="CancellationToken">Token used to cancel the operation</param>
 		/// <returns>Response from the server</returns>
-		public static Task<PerforceResponseList<UnshelveRecord>> TryUnshelveAsync(this IPerforceConnection Connection, int ChangeNumber, int IntoChangeNumber, string UsingBranchSpec, string UsingStream, string ForceParentStream, UnshelveOptions Options, string[] FileSpecs, CancellationToken CancellationToken)
+		public static Task<PerforceResponseList<UnshelveRecord>> TryUnshelveAsync(this IPerforceConnection Connection, int ChangeNumber, int IntoChangeNumber, string UsingBranchSpec, string UsingStream, string ForceParentStream, UnshelveOptions Options, FileSpecList FileSpecs, CancellationToken CancellationToken = default)
 		{
 			List<string> Arguments = new List<string>();
 			Arguments.Add($"-s{ChangeNumber}");
@@ -3084,7 +3077,7 @@ namespace EpicGames.Perforce
 			{
 				Arguments.Add($"-P{ForceParentStream}");
 			}
-			Arguments.AddRange(FileSpecs);
+			Arguments.AddRange(FileSpecs.List);
 
 			return CommandAsync<UnshelveRecord>(Connection, "unshelve", Arguments, null, CancellationToken);
 		}
@@ -3100,7 +3093,7 @@ namespace EpicGames.Perforce
 		/// <param name="UserName">Name of the user to fetch information for</param>
 		/// <param name="CancellationToken">Token used to cancel the operation</param>
 		/// <returns>Response from the server</returns>
-		public static async Task<UserRecord> GetUserAsync(this IPerforceConnection Connection, string UserName, CancellationToken CancellationToken)
+		public static async Task<UserRecord> GetUserAsync(this IPerforceConnection Connection, string UserName, CancellationToken CancellationToken = default)
 		{
 			return (await TryGetUserAsync(Connection, UserName, CancellationToken))[0].Data;
 		}
@@ -3112,7 +3105,7 @@ namespace EpicGames.Perforce
 		/// <param name="UserName">Name of the user to fetch information for</param>
 		/// <param name="CancellationToken">Token used to cancel the operation</param>
 		/// <returns>Response from the server</returns>
-		public static Task<PerforceResponseList<UserRecord>> TryGetUserAsync(this IPerforceConnection Connection, string UserName, CancellationToken CancellationToken)
+		public static Task<PerforceResponseList<UserRecord>> TryGetUserAsync(this IPerforceConnection Connection, string UserName, CancellationToken CancellationToken = default)
 		{
 			List<string> Arguments = new List<string>();
 			Arguments.Add($"-o{UserName}");
@@ -3130,12 +3123,12 @@ namespace EpicGames.Perforce
 		/// <param name="FileSpecs">Patterns for the files to query</param>
 		/// <param name="CancellationToken">Token used to cancel the operation</param>
 		/// <returns>List of responses from the server</returns>
-		public static async Task<List<WhereRecord>> WhereAsync(this IPerforceConnection Connection, string[] FileSpecs, CancellationToken CancellationToken)
+		public static async Task<List<WhereRecord>> WhereAsync(this IPerforceConnection Connection, FileSpecList FileSpecs, CancellationToken CancellationToken = default)
 		{
 			List<WhereRecord> Records = (await TryWhereAsync(Connection, FileSpecs, CancellationToken)).Data;
-			if (Records.Count != FileSpecs.Length)
+			if (Records.Count != FileSpecs.List.Count)
 			{
-				throw new PerforceException($"Unexpected response count; expected {FileSpecs.Length}, got {Records.Count}.");
+				throw new PerforceException($"Unexpected response count; expected {FileSpecs.List.Count}, got {Records.Count}.");
 			}
 			return Records;
 		}
@@ -3147,9 +3140,9 @@ namespace EpicGames.Perforce
 		/// <param name="FileSpecs">Patterns for the files to query</param>
 		/// <param name="CancellationToken">Token used to cancel the operation</param>
 		/// <returns>List of responses from the server</returns>
-		public static Task<PerforceResponseList<WhereRecord>> TryWhereAsync(this IPerforceConnection Connection, string[] FileSpecs, CancellationToken CancellationToken)
+		public static Task<PerforceResponseList<WhereRecord>> TryWhereAsync(this IPerforceConnection Connection, FileSpecList FileSpecs, CancellationToken CancellationToken = default)
 		{
-			return BatchedCommandAsync<WhereRecord>(Connection, "where", new List<string>(), FileSpecs, null, CancellationToken);
+			return BatchedCommandAsync<WhereRecord>(Connection, "where", new List<string>(), FileSpecs.List, null, CancellationToken);
 		}
 
 		#endregion
@@ -3164,7 +3157,7 @@ namespace EpicGames.Perforce
 		/// <param name="ChangeNumber">Changelist number to receive the changes</param>
 		/// <param name="CancellationToken">Token used to cancel the operation</param>
 		/// <returns>Response from the server</returns>
-		public static async Task UndoChangeAsync(this IPerforceConnection Connection, int ChangeNumberToUndo, int ChangeNumber, CancellationToken CancellationToken)
+		public static async Task UndoChangeAsync(this IPerforceConnection Connection, int ChangeNumberToUndo, int ChangeNumber, CancellationToken CancellationToken = default)
 		{
 			(await TryUndoChangeAsync(Connection, ChangeNumberToUndo, ChangeNumber, CancellationToken))[0].EnsureSuccess();
 		}
@@ -3177,7 +3170,7 @@ namespace EpicGames.Perforce
 		/// <param name="ChangeNumber">Changelist number to receive the changes</param>
 		/// <param name="CancellationToken">Token used to cancel the operation</param>
 		/// <returns>Response from the server</returns>
-		public static Task<PerforceResponseList<UndoRecord>> TryUndoChangeAsync(this IPerforceConnection Connection, int ChangeNumberToUndo, int ChangeNumber, CancellationToken CancellationToken)
+		public static Task<PerforceResponseList<UndoRecord>> TryUndoChangeAsync(this IPerforceConnection Connection, int ChangeNumberToUndo, int ChangeNumber, CancellationToken CancellationToken = default)
 		{
 			List<string> Arguments = new List<string>();
 			Arguments.Add($"-c{ChangeNumber}");
@@ -3196,7 +3189,7 @@ namespace EpicGames.Perforce
 		/// <param name="Options">Options for the anotate command</param>
 		/// <param name="CancellationToken">Token used to cancel the operation</param>
 		/// <returns>List of annotate records</returns>
-		public static async Task<List<AnnotateRecord>> AnnotateAsync(this IPerforceConnection Connection, string FileSpec, AnnotateOptions Options, CancellationToken CancellationToken)
+		public static async Task<List<AnnotateRecord>> AnnotateAsync(this IPerforceConnection Connection, string FileSpec, AnnotateOptions Options, CancellationToken CancellationToken = default)
 		{
 			return (await TryAnnotateAsync(Connection, FileSpec, Options, CancellationToken)).Data;
 		}
@@ -3209,7 +3202,7 @@ namespace EpicGames.Perforce
 		/// <param name="Options">Options for the anotate command</param>
 		/// <param name="CancellationToken">Token used to cancel the operation</param>
 		/// <returns>List of annotate records</returns>
-		public static Task<PerforceResponseList<AnnotateRecord>> TryAnnotateAsync(this IPerforceConnection Connection, string FileSpec, AnnotateOptions Options, CancellationToken CancellationToken)
+		public static Task<PerforceResponseList<AnnotateRecord>> TryAnnotateAsync(this IPerforceConnection Connection, string FileSpec, AnnotateOptions Options, CancellationToken CancellationToken = default)
 		{
 			List<string> Arguments = new List<string>();
 			if ((Options & AnnotateOptions.IncludeDeletedFilesAndLines) != 0)
