P4: Add an interface for Perforce responses, and move the record deserialization logic into extension methods on it.

diff --git a/Engine/Source/Programs/Shared/EpicGames.Perforce/PerforceChildProcess.cs b/Engine/Source/Programs/Shared/EpicGames.Perforce/PerforceChildProcess.cs
--- a/Engine/Source/Programs/Shared/EpicGames.Perforce/PerforceChildProcess.cs
+++ b/Engine/Source/Programs/Shared/EpicGames.Perforce/PerforceChildProcess.cs
@@ -21,26 +21,33 @@ using System.Threading.Tasks;
 namespace EpicGames.Perforce
 {
 	/// <summary>
-	/// Wraps a call to a p4.exe child process, and allows reading data from it
+	/// Interface for the result of a Perforce operation
 	/// </summary>
-	class PerforceChildProcess : IDisposable
+	public interface IPerforceResponse : IDisposable
 	{
 		/// <summary>
-		/// String constants for records
+		/// Data containing the result
 		/// </summary>
-		static class ReadOnlyUtf8StringConstants
-		{
-			public static Utf8String Code = "code";
-			public static Utf8String Stat = "stat";
-			public static Utf8String Info = "info";
-			public static Utf8String Error = "error";
-		}
+		ReadOnlyMemory<byte> Data { get; }
 
 		/// <summary>
-		/// Standard prefix for a returned record: record indicator, string, 4 bytes, 'code', string, [value]
+		/// Waits until more data has been read into the buffer. 
 		/// </summary>
-		static readonly byte[] RecordPrefix = { (byte)'{', (byte)'s', 4, 0, 0, 0, (byte)'c', (byte)'o', (byte)'d', (byte)'e', (byte)'s' };
+		/// <returns>True if more data was read, false otherwise</returns>
+		Task<bool> ReadAsync(CancellationToken Token);
+
+		/// <summary>
+		/// Discard bytes from the start of the result buffer
+		/// </summary>
+		/// <param name="NumBytes">Number of bytes to discard</param>
+		void Discard(int NumBytes);
+	}
 
+	/// <summary>
+	/// Wraps a call to a p4.exe child process, and allows reading data from it
+	/// </summary>
+	class PerforceChildProcess : IPerforceResponse, IDisposable
+	{
 		/// <summary>
 		/// The process group
 		/// </summary>
@@ -57,24 +64,17 @@ namespace EpicGames.Perforce
 		ITraceSpan Scope;
 
 		/// <summary>
-		/// Buffer for data from the child process
+		/// The buffer data
 		/// </summary>
 		byte[] Buffer;
 
-		/// <summary>
-		/// Current read position within the buffer
-		/// </summary>
-		int BufferPos;
-
 		/// <summary>
 		/// End of the valid portion of the buffer (exclusive)
 		/// </summary>
 		int BufferEnd;
 
-		/// <summary>
-		/// Whether the remaining contents of the buffer represent the end of the file
-		/// </summary>
-		bool bEndOfStream;
+		/// <inheritdoc/>
+		public ReadOnlyMemory<byte> Data => Buffer.AsMemory(0, BufferEnd);
 
 		/// <summary>
 		/// Constructor
@@ -97,7 +97,7 @@ namespace EpicGames.Perforce
 			string FullArgumentList = CommandLineArguments.Join(FullArguments);
 			Logger.LogDebug("Running {0} {1}", PerforceFileName, FullArgumentList);
 
-			Scope = TraceSpan.Create($"{Command}", Service: "perforce");
+			Scope = TraceSpan.Create(Command, Service: "perforce");
 			Scope.AddMetadata("arguments", FullArgumentList);
 
 			ChildProcessGroup = new ManagedProcessGroup();
@@ -146,6 +146,36 @@ namespace EpicGames.Perforce
 			}
 		}
 
+		/// <inheritdoc/>
+		public async Task<bool> ReadAsync(CancellationToken CancellationToken)
+		{
+			// Update the buffer contents
+			if (BufferEnd == Buffer.Length)
+			{
+				Array.Resize(ref Buffer, Buffer.Length + 65536);
+			}
+
+			// Try to read more data
+			int Count = await ChildProcess!.ReadAsync(Buffer, BufferEnd, Buffer.Length - BufferEnd, CancellationToken);
+			if (Count == 0)
+			{
+				return false;
+			}
+
+			BufferEnd += Count;
+			return true;
+		}
+
+		/// <inheritdoc/>
+		public void Discard(int NumBytes)
+		{
+			if (NumBytes > 0)
+			{
+				Array.Copy(Buffer, NumBytes, Buffer, 0, BufferEnd - NumBytes);
+				BufferEnd -= NumBytes;
+			}
+		}
+
 		/// <summary>
 		/// Reads all output from the child process as a string
 		/// </summary>
@@ -154,12 +184,40 @@ namespace EpicGames.Perforce
 		public async Task<Tuple<bool, string>> TryReadToEndAsync(CancellationToken CancellationToken)
 		{
 			MemoryStream Stream = new MemoryStream();
-			await ChildProcess!.CopyToAsync(Stream, CancellationToken);
-			ChildProcess.WaitForExit();
-			string Data = Encoding.Default.GetString(Stream.ToArray());
-			return Tuple.Create(ChildProcess.ExitCode == 0, Data);
+
+			while (await ReadAsync(CancellationToken))
+			{
+				ReadOnlyMemory<byte> DataCopy = Data;
+				Stream.Write(DataCopy.Span);
+				Discard(DataCopy.Length);
+			}
+
+			string String = Encoding.Default.GetString(Stream.ToArray());
+			return Tuple.Create(ChildProcess.ExitCode == 0, String);
+		}
+	}
+
+	/// <summary>
+	/// Wraps a call to a p4.exe child process, and allows reading data from it
+	/// </summary>
+	static class PerforceResultExtensions
+	{
+		/// <summary>
+		/// String constants for records
+		/// </summary>
+		static class ReadOnlyUtf8StringConstants
+		{
+			public static Utf8String Code = "code";
+			public static Utf8String Stat = "stat";
+			public static Utf8String Info = "info";
+			public static Utf8String Error = "error";
 		}
 
+		/// <summary>
+		/// Standard prefix for a returned record: record indicator, string, 4 bytes, 'code', string, [value]
+		/// </summary>
+		static readonly byte[] RecordPrefix = { (byte)'{', (byte)'s', 4, 0, 0, 0, (byte)'c', (byte)'o', (byte)'d', (byte)'e', (byte)'s' };
+
 		/// <summary>
 		/// Formats the current contents of the buffer to a string
 		/// </summary>
@@ -212,11 +270,12 @@ namespace EpicGames.Perforce
 		/// <summary>
 		/// Formats the current contents of the buffer to a string
 		/// </summary>
+		/// <param name="Buffer">The buffer to read from</param>
 		/// <param name="BufferIdx">The next byte that was read</param>
 		/// <returns>String representation of the buffer</returns>
-		private string FormatDataAsHexDump(int BufferIdx)
+		private static string FormatDataAsHexDump(ReadOnlySpan<byte> Buffer, int BufferIdx)
 		{
-			return FormatDataAsHexDump(Buffer.AsSpan(BufferIdx, Math.Min(BufferEnd - BufferIdx, 1024)));
+			return FormatDataAsHexDump(Buffer.Slice(BufferIdx, Math.Min(Buffer.Length - BufferIdx, 1024)));
 		}
 
 		/// <summary>
@@ -224,7 +283,7 @@ namespace EpicGames.Perforce
 		/// </summary>
 		/// <param name="Data">The data to format</param>
 		/// <returns>String representation of the buffer</returns>
-		private static string FormatDataAsHexDump(Span<byte> Data)
+		private static string FormatDataAsHexDump(ReadOnlySpan<byte> Data)
 		{
 			// Format the result
 			StringBuilder Result = new StringBuilder();
@@ -266,231 +325,115 @@ namespace EpicGames.Perforce
 			return Result.ToString();
 		}
 
-		/// <summary>
-		/// Try to read plain text from the output stream. Fails if the line starts with a binary record marker ('{').
-		/// </summary>
-		/// <param name="CancellationToken">Cancellation token</param>
-		/// <returns>The line that was read.</returns>
-		public async Task<string?> TryReadLineAsync(CancellationToken CancellationToken)
-		{
-			for (; ; )
-			{
-				// Move data to the start of the buffer
-				if (BufferPos > 0)
-				{
-					Array.Copy(Buffer, BufferPos, Buffer, 0, BufferEnd - BufferPos);
-					BufferEnd -= BufferPos;
-					BufferPos = 0;
-				}
-
-				// Append data to the end of the buffer
-				if (BufferEnd < Buffer.Length)
-				{
-					int BytesRead = await ChildProcess!.ReadAsync(Buffer, BufferEnd, Buffer.Length - BufferEnd, CancellationToken);
-					if (BytesRead == 0)
-					{
-						return null;
-					}
-					BufferEnd += BytesRead;
-				}
-
-				// Move past any lines of plain text output
-				if (BufferEnd > 0)
-				{
-					// Start of a data record
-					if(Buffer[0] == '{')
-					{
-						return null;
-					}
-
-					// Try to reach the end of the line
-					for(int Idx = 0; Idx < BufferEnd; Idx++)
-					{
-						if(Buffer[Idx] == '\n')
-						{
-							BufferPos = Idx + 1;
-							return Encoding.UTF8.GetString(Buffer, 0, BufferPos);
-						}
-					}
-				}
-			}
-		}
-
 		/// <summary>
 		/// Read a list of responses from the child process
 		/// </summary>
+		/// <param name="Perforce">The response to read from</param>
 		/// <param name="StatRecordInfo">The type of stat record to parse</param>
 		/// <param name="CancellationToken">Cancellation token for the read</param>
 		/// <returns>Async task</returns>
-		public async Task<List<PerforceResponse>> ReadResponsesAsync(CachedRecordInfo? StatRecordInfo, CancellationToken CancellationToken)
+		public static async Task<List<PerforceResponse>> ReadResponsesAsync(this IPerforceResponse Perforce, CachedRecordInfo? StatRecordInfo, CancellationToken CancellationToken)
 		{
 			List<PerforceResponse> Responses = new List<PerforceResponse>();
-			while (BufferPos != BufferEnd || !bEndOfStream)
-			{
-				// Update the buffer contents
-				if (BufferPos > 0)
-				{
-					Array.Copy(Buffer, BufferPos, Buffer, 0, BufferEnd - BufferPos);
-					BufferEnd -= BufferPos;
-					BufferPos = 0;
-				}
-				else if(BufferEnd == Buffer.Length)
-				{
-					Array.Resize(ref Buffer, Buffer.Length + 65536);
-				}
 
+			// Read all the records into a list
+			Task<bool> ReadTask = Perforce.ReadAsync(CancellationToken);
+			while (await ReadTask)
+			{
 				// Start a read to add more data
-				Task<int>? ReadTask;
-				if (BufferEnd < Buffer.Length && !bEndOfStream)
-				{
-					ReadTask = Task.Run(() => ChildProcess!.ReadAsync(Buffer, BufferEnd, Buffer.Length - BufferEnd, CancellationToken));
-				}
-				else
+				ReadTask = Perforce.ReadAsync(CancellationToken);
+
+				// Check for the whole message not being a marshalled python object, and produce a better response in that scenario
+				ReadOnlyMemory<byte> Data = Perforce.Data;
+				if (Data.Length > 0 && Responses.Count == 0 && Data.Span[0] != '{')
 				{
-					ReadTask = null;
+					throw new PerforceException("Unexpected response from server:{0}", FormatDataAsString(Data.Span));
 				}
 
-				if (BufferEnd > BufferPos)
+				// Parse the responses from the current buffer
+				int BufferPos = 0;
+				for (; ; )
 				{
-					if (Responses.Count == 0 && BufferPos == 0 && Buffer[0] != '{')
+					int NewBufferPos = BufferPos;
+					if (!TryReadResponse(Data.Span, ref NewBufferPos, StatRecordInfo, out PerforceResponse? Response))
 					{
-						// Check for the whole message not being a marshalled python object, and produce a better response in that scenario
-						if (bEndOfStream || BufferEnd == Buffer.Length)
-						{
-							throw new PerforceException("Unexpected response from server:{0}", FormatDataAsString(Buffer.AsSpan(0, BufferEnd)));
-						}
+						break;
 					}
-					else
+					if (Response.Error == null || Response.Error.Generic != PerforceGenericCode.Empty)
 					{
-						// Try to process all the existing data
-						for (; ; )
-						{
-							int InitialBufferPos = BufferPos;
-
-							PerforceResponse? Response;
-							if (!TryReadResponse(StatRecordInfo, out Response))
-							{
-								BufferPos = InitialBufferPos;
-								break;
-							}
-
-							if (Response.Error == null || Response.Error.Generic != PerforceGenericCode.Empty)
-							{
-								Responses.Add(Response);
-							}
-						}
-
-						// If the stream is complete but we couldn't parse a response from the server, treat it as an error
-						if (bEndOfStream && BufferPos < BufferEnd)
-						{
-							throw new PerforceException("Unexpected response from server:{0}", FormatDataAsString(Buffer.AsSpan(BufferPos, BufferEnd - BufferPos)));
-						}
+						Responses.Add(Response);
 					}
+					BufferPos = NewBufferPos;
 				}
 
-				// Wait for the read to complete
-				if(ReadTask != null)
-				{
-					int ReadSize = await ReadTask;
-					if(ReadSize == 0)
-					{
-						bEndOfStream = true;
-					}
-					BufferEnd += ReadSize;
-				}
+				// Discard all the data that we've processed
+				Perforce.Discard(BufferPos);
 			}
+
+			// If the stream is complete but we couldn't parse a response from the server, treat it as an error
+			if (Perforce.Data.Length > 0)
+			{
+				throw new PerforceException("Unexpected response from server:{0}", FormatDataAsString(Perforce.Data.Span));
+			}
+
 			return Responses;
 		}
 
 		/// <summary>
 		/// Read a list of responses from the child process
 		/// </summary>
+		/// <param name="Perforce">The Perforce response</param>
 		/// <param name="HandleRecord">Delegate to invoke for each record read</param>
 		/// <param name="CancellationToken">Cancellation token for the read</param>
 		/// <returns>Async task</returns>
-		public async Task ReadRecordsAsync(Action<PerforceRecord> HandleRecord, CancellationToken CancellationToken)
+		public static async Task ReadRecordsAsync(this IPerforceResponse Perforce, Action<PerforceRecord> HandleRecord, CancellationToken CancellationToken)
 		{
 			PerforceRecord Record = new PerforceRecord();
 			Record.Rows = new List<KeyValuePair<Utf8String, PerforceValue>>();
 
-			for(; ;)
+			// Read all the records into a list
+			Task<bool> ReadTask = Perforce.ReadAsync(CancellationToken);
+			while (await ReadTask)
 			{
-				// Update the buffer contents
-				if (BufferPos > 0)
-				{
-					byte[] PrevBuffer = Buffer;
-					Buffer = new byte[(BufferEnd - BufferPos) + 65536];
-					Array.Copy(PrevBuffer, BufferPos, Buffer, 0, BufferEnd - BufferPos);
-					BufferEnd -= BufferPos;
-					BufferPos = 0;
-				}
-				else if (BufferEnd == Buffer.Length)
-				{
-					Array.Resize(ref Buffer, Buffer.Length + 65536);
-				}
-
 				// Start a read to add more data
-				Task<int>? ReadTask;
-				if (BufferEnd < Buffer.Length && !bEndOfStream)
-				{
-					ReadTask = Task.Run(() => ChildProcess!.ReadAsync(Buffer, BufferEnd, Buffer.Length - BufferEnd, CancellationToken));
-				}
-				else
-				{
-					ReadTask = null;
-				}
+				ReadTask = Perforce.ReadAsync(CancellationToken);
+				ReadOnlyMemory<byte> Data = Perforce.Data;
 
-				// Try to process all the existing data
+				// Parse the responses from the current buffer
+				int BufferPos = 0;
 				for (; ; )
 				{
 					int InitialBufferPos = BufferPos;
-
-					if (!ReadRecord(Record.Rows))
+					if(!ReadRecord(Data.Span, ref BufferPos, Record.Rows))
 					{
 						BufferPos = InitialBufferPos;
 						break;
 					}
-
 					HandleRecord(Record);
 				}
+				Perforce.Discard(BufferPos);
+			}
 
-				// If we're at the end of the stream, check we parsed everything
-				if (bEndOfStream)
-				{
-					if (BufferPos == BufferEnd)
-					{
-						break;
-					}
-					else
-					{
-						throw new PerforceException("Unable to parse last chunk of data from response stream." + FormatDataAsHexDump(BufferPos));
-					}
-				}
-
-				// Wait for the read to complete
-				if (ReadTask != null)
-				{
-					int ReadSize = await ReadTask;
-					if (ReadSize == 0)
-					{
-						bEndOfStream = true;
-					}
-					BufferEnd += ReadSize;
-				}
+			// If the stream is complete but we couldn't parse a response from the server, treat it as an error
+			if (Perforce.Data.Length > 0)
+			{
+				throw new PerforceException("Unexpected response from server:{0}", FormatDataAsString(Perforce.Data.Span));
 			}
 		}
 
 		/// <summary>
 		/// Reads from the buffer into a record object
 		/// </summary>
+		/// <param name="Buffer">The buffer to read from</param>
+		/// <param name="BufferPos">Current read position within the buffer</param>
 		/// <param name="Rows">List of rows to read into</param>
 		/// <returns>True if a record could be read; false if more data is required</returns>
-		bool ReadRecord(List<KeyValuePair<Utf8String, PerforceValue>> Rows)
+		static bool ReadRecord(ReadOnlySpan<byte> Buffer, ref int BufferPos, List<KeyValuePair<Utf8String, PerforceValue>> Rows)
 		{
 			Rows.Clear();
 
 			// Check we can read the initial record marker
-			if (BufferPos >= BufferEnd)
+			if (BufferPos >= Buffer.Length)
 			{
 				return false;
 			}
@@ -505,7 +448,7 @@ namespace EpicGames.Perforce
 			for (; ; )
 			{
 				// Check that we've got a string field
-				if (BufferPos >= BufferEnd)
+				if (BufferPos >= Buffer.Length)
 				{
 					return false;
 				}
@@ -518,12 +461,12 @@ namespace EpicGames.Perforce
 				}
 				else if (KeyType != 's')
 				{
-					throw new PerforceException("Unexpected key field type while parsing marshalled output ({0}) - expected 's', got: {1}", (int)KeyType, FormatDataAsHexDump(BufferPos - 1));
+					throw new PerforceException("Unexpected key field type while parsing marshalled output ({0}) - expected 's', got: {1}", (int)KeyType, FormatDataAsHexDump(Buffer.Slice(BufferPos - 1)));
 				}
 
 				// Read the tag
 				Utf8String Key;
-				if (!TryReadString(out Key))
+				if (!TryReadString(Buffer, ref BufferPos, out Key))
 				{
 					return false;
 				}
@@ -533,7 +476,7 @@ namespace EpicGames.Perforce
 
 				// Read the value type
 				byte ValueType;
-				if (!TryReadByte(out ValueType))
+				if (!TryReadByte(Buffer, ref BufferPos, out ValueType))
 				{
 					return false;
 				}
@@ -543,20 +486,20 @@ namespace EpicGames.Perforce
 				if (ValueType == 's')
 				{
 					Utf8String String;
-					if (!TryReadString(out String))
+					if (!TryReadString(Buffer, ref BufferPos, out String))
 					{
 						return false;
 					}
-					Value = new PerforceValue(Buffer.AsMemory(ValueOffset, BufferPos - ValueOffset));
+					Value = new PerforceValue(Buffer.Slice(ValueOffset, BufferPos - ValueOffset).ToArray());
 				}
 				else if (ValueType == 'i')
 				{
 					int Integer;
-					if(!TryReadInt(out Integer))
+					if(!TryReadInt(Buffer, ref BufferPos, out Integer))
 					{
 						return false;
 					}
-					Value = new PerforceValue(Buffer.AsMemory(ValueOffset, BufferPos - ValueOffset));
+					Value = new PerforceValue(Buffer.Slice(ValueOffset, BufferPos - ValueOffset).ToArray());
 				}
 				else
 				{
@@ -572,18 +515,20 @@ namespace EpicGames.Perforce
 		/// <summary>
 		/// Reads a response object from the buffer
 		/// </summary>
+		/// <param name="Buffer">The buffer to read from</param>
+		/// <param name="BufferPos">Current read position within the buffer</param>
 		/// <param name="StatRecordInfo">The type of record expected to parse from the response</param>
 		/// <param name="Response">Receives the response object on success</param>
 		/// <returns>True if a response was read, false if the buffer needs more data</returns>
-		bool TryReadResponse(CachedRecordInfo? StatRecordInfo, [NotNullWhen(true)] out PerforceResponse? Response)
+		static bool TryReadResponse(ReadOnlySpan<byte> Buffer, ref int BufferPos, CachedRecordInfo? StatRecordInfo, [NotNullWhen(true)] out PerforceResponse? Response)
 		{
-			if (BufferPos + RecordPrefix.Length + 4 > BufferEnd)
+			if (BufferPos + RecordPrefix.Length + 4 > Buffer.Length)
 			{
 				Response = null;
 				return false;
 			}
 
-			Span<byte> Prefix = Buffer.AsSpan(BufferPos, RecordPrefix.Length);
+			ReadOnlySpan<byte> Prefix = Buffer.Slice(BufferPos, RecordPrefix.Length);
 			if (!Prefix.SequenceEqual(RecordPrefix))
 			{
 				throw new PerforceException("Expected 'code' field at the start of record");
@@ -591,7 +536,7 @@ namespace EpicGames.Perforce
 			BufferPos += Prefix.Length;
 
 			Utf8String Code;
-			if (!TryReadString(out Code))
+			if (!TryReadString(Buffer, ref BufferPos, out Code))
 			{
 				Response = null;
 				return false;
@@ -601,7 +546,7 @@ namespace EpicGames.Perforce
 			object? Record;
 			if (Code == ReadOnlyUtf8StringConstants.Stat && StatRecordInfo != null)
 			{
-				if(!TryReadTypedRecord(Utf8String.Empty, StatRecordInfo, out Record))
+				if(!TryReadTypedRecord(Buffer, ref BufferPos, Utf8String.Empty, StatRecordInfo, out Record))
 				{
 					Response = null;
 					return false;
@@ -609,7 +554,7 @@ namespace EpicGames.Perforce
 			}
 			else if (Code == ReadOnlyUtf8StringConstants.Info)
 			{
-				if(!TryReadTypedRecord(Utf8String.Empty, PerforceReflection.InfoRecordInfo, out Record))
+				if(!TryReadTypedRecord(Buffer, ref BufferPos, Utf8String.Empty, PerforceReflection.InfoRecordInfo, out Record))
 				{
 					Response = null;
 					return false;
@@ -617,7 +562,7 @@ namespace EpicGames.Perforce
 			}
 			else if (Code == ReadOnlyUtf8StringConstants.Error)
 			{
-				if (!TryReadTypedRecord(Utf8String.Empty, PerforceReflection.ErrorRecordInfo, out Record))
+				if (!TryReadTypedRecord(Buffer, ref BufferPos, Utf8String.Empty, PerforceReflection.ErrorRecordInfo, out Record))
 				{
 					Response = null;
 					return false;
@@ -629,7 +574,7 @@ namespace EpicGames.Perforce
 			}
 
 			// Skip over the record terminator
-			if(BufferPos >= BufferEnd || Buffer[BufferPos] != '0')
+			if(BufferPos >= Buffer.Length || Buffer[BufferPos] != '0')
 			{
 				throw new PerforceException("Unexpected record terminator");
 			}
@@ -643,11 +588,13 @@ namespace EpicGames.Perforce
 		/// <summary>
 		/// Parse an individual record from the server.
 		/// </summary>
+		/// <param name="Buffer">The buffer to read from</param>
+		/// <param name="BufferPos">Current read position within the buffer</param>
 		/// <param name="RequiredSuffix">The required suffix for any subobject arrays.</param>
 		/// <param name="RecordInfo">Reflection information for the type being serialized into.</param>
-		/// <param name="Record">Receives the deserialized record on success</param>
+		/// <param name="Record">Receives the record on success</param>
 		/// <returns>The parsed object.</returns>
-		bool TryReadTypedRecord(Utf8String RequiredSuffix, CachedRecordInfo RecordInfo, [NotNullWhen(true)] out object? Record)
+		static bool TryReadTypedRecord(ReadOnlySpan<byte> Buffer, ref int BufferPos, Utf8String RequiredSuffix, CachedRecordInfo RecordInfo, [NotNullWhen(true)] out object? Record)
 		{
 			// Create a bitmask for all the required tags
 			ulong RequiredTagsBitMask = 0;
@@ -663,7 +610,7 @@ namespace EpicGames.Perforce
 			for(; ;)
 			{
 				// Check that we've got a string field
-				if(BufferPos >= BufferEnd)
+				if(BufferPos >= Buffer.Length)
 				{
 					Record = null;
 					return false;
@@ -677,7 +624,7 @@ namespace EpicGames.Perforce
 				}
 				else if (KeyType != 's')
 				{
-					throw new PerforceException("Unexpected key field type while parsing marshalled output ({0}) - expected 's', got: {1}", (int)KeyType, FormatDataAsHexDump(BufferPos));
+					throw new PerforceException("Unexpected key field type while parsing marshalled output ({0}) - expected 's', got: {1}", (int)KeyType, FormatDataAsHexDump(Buffer.Slice(BufferPos)));
 				}
 
 				// Capture the initial buffer position, in case we have to roll back
@@ -686,7 +633,7 @@ namespace EpicGames.Perforce
 
 				// Read the tag
 				Utf8String Tag;
-				if(!TryReadString(out Tag))
+				if(!TryReadString(Buffer, ref BufferPos, out Tag))
 				{
 					Record = null;
 					return false;
@@ -713,7 +660,7 @@ namespace EpicGames.Perforce
 				// Check whether it's a subobject or part of the current object.
 				if (Suffix == RequiredSuffix)
 				{
-					if(!TryReadValue(NewRecord, TagInfo))
+					if(!TryReadValue(Buffer, ref BufferPos, NewRecord, TagInfo))
 					{
 						Record = null;
 						return false;
@@ -738,7 +685,7 @@ namespace EpicGames.Perforce
 						}
 
 						// Add it to the list
-						if (!TryReadValue(NewRecord, TagInfo))
+						if (!TryReadValue(Buffer, ref BufferPos, NewRecord, TagInfo))
 						{
 							Record = null;
 							return false;
@@ -764,7 +711,7 @@ namespace EpicGames.Perforce
 
 						// Parse the subobject and add it to the list
 						object? SubRecord;
-						if (!TryReadTypedRecord(Suffix, RecordInfo.SubElementRecordInfo!, out SubRecord))
+						if (!TryReadTypedRecord(Buffer, ref BufferPos, Suffix, RecordInfo.SubElementRecordInfo!, out SubRecord))
 						{
 							Record = null;
 							return false;
@@ -774,7 +721,7 @@ namespace EpicGames.Perforce
 					else
 					{
 						// Just discard the value
-						if (!TryReadValue(NewRecord, TagInfo))
+						if (!TryReadValue(Buffer, ref BufferPos, NewRecord, TagInfo))
 						{
 							Record = null;
 							return false;
@@ -804,14 +751,16 @@ namespace EpicGames.Perforce
 		/// <summary>
 		/// Reads a value from the input buffer
 		/// </summary>
+		/// <param name="Buffer">The buffer to read from</param>
+		/// <param name="BufferPos">Current read position within the buffer</param>
 		/// <param name="NewRecord">The new record</param>
 		/// <param name="TagInfo">The current tag</param>
 		/// <returns></returns>
-		bool TryReadValue(object NewRecord, CachedTagInfo? TagInfo)
+		static bool TryReadValue(ReadOnlySpan<byte> Buffer, ref int BufferPos, object NewRecord, CachedTagInfo? TagInfo)
 		{
 			// Read the value type
 			byte ValueType;
-			if (!TryReadByte(out ValueType))
+			if (!TryReadByte(Buffer, ref BufferPos, out ValueType))
 			{
 				return false;
 			}
@@ -820,7 +769,7 @@ namespace EpicGames.Perforce
 			if (ValueType == 's')
 			{
 				Utf8String String;
-				if (!TryReadString(out String))
+				if (!TryReadString(Buffer, ref BufferPos, out String))
 				{
 					return false;
 				}
@@ -832,7 +781,7 @@ namespace EpicGames.Perforce
 			else if (ValueType == 'i')
 			{
 				int Integer;
-				if (!TryReadInt(out Integer))
+				if (!TryReadInt(Buffer, ref BufferPos, out Integer))
 				{
 					return false;
 				}
@@ -852,11 +801,13 @@ namespace EpicGames.Perforce
 		/// <summary>
 		/// Attempts to read a single byte from the buffer
 		/// </summary>
+		/// <param name="Buffer">The buffer to read from</param>
+		/// <param name="BufferPos">Current read position within the buffer</param>
 		/// <param name="Value">Receives the byte that was read</param>
 		/// <returns>True if a byte was read from the buffer, false if there was not enough data</returns>
-		bool TryReadByte(out byte Value)
+		static bool TryReadByte(ReadOnlySpan<byte> Buffer, ref int BufferPos, out byte Value)
 		{
-			if (BufferPos >= BufferEnd)
+			if (BufferPos >= Buffer.Length)
 			{
 				Value = 0;
 				return false;
@@ -870,11 +821,13 @@ namespace EpicGames.Perforce
 		/// <summary>
 		/// Attempts to read a single int from the buffer
 		/// </summary>
+		/// <param name="Buffer">The buffer to read from</param>
+		/// <param name="BufferPos">Current read position within the buffer</param>
 		/// <param name="Value">Receives the value that was read</param>
 		/// <returns>True if an int was read from the buffer, false if there was not enough data</returns>
-		bool TryReadInt(out int Value)
+		static bool TryReadInt(ReadOnlySpan<byte> Buffer, ref int BufferPos, out int Value)
 		{
-			if(BufferPos + 4 > BufferEnd)
+			if(BufferPos + 4 > Buffer.Length)
 			{
 				Value = 0;
 				return false;
@@ -888,12 +841,14 @@ namespace EpicGames.Perforce
 		/// <summary>
 		/// Attempts to read a string with type indicator from the buffer
 		/// </summary>
+		/// <param name="Buffer">The buffer to read from</param>
+		/// <param name="BufferPos">Current read position within the buffer</param>
 		/// <param name="String">Receives the value that was read</param>
 		/// <returns>True if a string was read from the buffer, false if there was not enough data</returns>
-		bool TryReadStringWithType(out Utf8String String)
+		static bool TryReadStringWithType(ReadOnlySpan<byte> Buffer, ref int BufferPos, out Utf8String String)
 		{
 			byte ValueType;
-			if(!TryReadByte(out ValueType))
+			if(!TryReadByte(Buffer, ref BufferPos, out ValueType))
 			{
 				String = new Utf8String();
 				return false;
@@ -902,30 +857,32 @@ namespace EpicGames.Perforce
 			{
 				throw new PerforceException("Expected string value, got '{0}'", ValueType);
 			}
-			return TryReadString(out String);
+			return TryReadString(Buffer, ref BufferPos, out String);
 		}
 
 		/// <summary>
 		/// Attempts to read a string from the buffer
 		/// </summary>
+		/// <param name="Buffer">The buffer to read from</param>
+		/// <param name="BufferPos">Current read position within the buffer</param>
 		/// <param name="String">Receives the value that was read</param>
 		/// <returns>True if a string was read from the buffer, false if there was not enough data</returns>
-		bool TryReadString(out Utf8String String)
+		static bool TryReadString(ReadOnlySpan<byte> Buffer, ref int BufferPos, out Utf8String String)
 		{
 			int Length;
-			if(!TryReadInt(out Length))
+			if(!TryReadInt(Buffer, ref BufferPos, out Length))
 			{
 				String = new Utf8String();
 				return false;
 			}
 
-			if(BufferPos + Length > BufferEnd)
+			if(BufferPos + Length > Buffer.Length)
 			{
 				String = new Utf8String();
 				return false;
 			}
 
-			String = new Utf8String(Buffer, BufferPos, Length);
+			String = new Utf8String(Buffer.Slice(BufferPos, Length).ToArray());
 			BufferPos += Length;
 			return true;
 		}
diff --git a/Engine/Source/Programs/Shared/EpicGames.Perforce/PerforceConnection.cs b/Engine/Source/Programs/Shared/EpicGames.Perforce/PerforceConnection.cs
--- a/Engine/Source/Programs/Shared/EpicGames.Perforce/PerforceConnection.cs
+++ b/Engine/Source/Programs/Shared/EpicGames.Perforce/PerforceConnection.cs
@@ -27,26 +27,13 @@
 		ILogger Logger { get; }
 
 		/// <summary>
-		/// Execute a command and parse the response
+		/// Runs a Perforce command
 		/// </summary>
-		/// <param name="Command">The command to execute</param>
-		/// <param name="Arguments">Command line arguments</param>
-		/// <param name="InputData">Input data to pass to Perforce</param>
-		/// <param name="StatRecordType">The type of records to return for "stat" responses</param>
-		/// <param name="CancellationToken">Token used to cancel the operation</param>
-		/// <returns>List of objects returned by the server</returns>
-		Task<List<PerforceResponse>> CommandAsync(string Command, IReadOnlyList<string> Arguments, byte[]? InputData, Type? StatRecordType, CancellationToken CancellationToken = default);
+		/// <param name="Command">The command name</param>
+		/// <param name="Arguments">Arguments for the command</param>
+		/// <param name="InputData">Input data to be passed to the command</param>
+		/// <returns>Response object</returns>
+		Task<IPerforceResponse> RunAsync(string Command, IReadOnlyList<string> Arguments, byte[]? InputData);
 
 		/// <summary>
-		/// Execute a command and parse the response
-		/// </summary>
-		/// <param name="Command">The command to execute</param>
-		/// <param name="Arguments">Command line arguments</param>
-		/// <param name="InputData">Input data to pass to Perforce</param>
-		/// <param name="HandleRecord">Delegate used to handle each record</param>
-		/// <param name="CancellationToken">Token used to cancel the operation</param>
-		/// <returns>List of objects returned by the server</returns>
-		Task RecordCommandAsync(string Command, IReadOnlyList<string> Arguments, byte[]? InputData, Action<PerforceRecord> HandleRecord, CancellationToken CancellationToken = default);
-
-		/// <summary>
 		/// Execute the 'login' command
@@ -225,39 +212,10 @@
 			return Arguments;
 		}
 
-		/// <summary>
-		/// Execute a command and parse the response
-		/// </summary>
-		/// <param name="Command">Command to execute</param>
-		/// <param name="Arguments">Arguments for the command</param>
-		/// <param name="InputData">Input data to pass to Perforce</param>
-		/// <param name="StatRecordType">The type of records to return for "stat" responses</param>
-		/// <param name="CancellationToken">Token used to cancel the operation</param>
-		/// <returns>List of objects returned by the server</returns>
-		public async Task<List<PerforceResponse>> CommandAsync(string Command, IReadOnlyList<string> Arguments, byte[]? InputData, Type? StatRecordType, CancellationToken CancellationToken = default)
+		/// <inheritdoc/>
+		public Task<IPerforceResponse> RunAsync(string Command, IReadOnlyList<string> Arguments, byte[]? InputData)
 		{
-			using (PerforceChildProcess Process = new PerforceChildProcess(Command, Arguments, InputData, GetGlobalArguments(), Logger))
-			{
-				CachedRecordInfo? StatRecordInfo = (StatRecordType == null) ? null : PerforceReflection.GetCachedRecordInfo(StatRecordType);
-				return await Process.ReadResponsesAsync(StatRecordInfo, CancellationToken);
-			}
+			return Task.FromResult<IPerforceResponse>(new PerforceChildProcess(Command, Arguments, InputData, GetGlobalArguments(), Logger));
 		}
 
-		/// <summary>
-		/// Execute a command and parse the response
-		/// </summary>
-		/// <param name="Command">Command to execute</param>
-		/// <param name="Arguments">Arguments for the command</param>
-		/// <param name="InputData">Input data to pass to Perforce</param>
-		/// <param name="HandleRecord">Delegate used to handle each record</param>
-		/// <param name="CancellationToken">Token used to cancel the operation</param>
-		/// <returns>List of objects returned by the server</returns>
-		public async Task RecordCommandAsync(string Command, IReadOnlyList<string> Arguments, byte[]? InputData, Action<PerforceRecord> HandleRecord, CancellationToken CancellationToken = default)
-		{
-			using (PerforceChildProcess Process = new PerforceChildProcess(Command, Arguments, InputData, GetGlobalArguments(), Logger))
-			{
-				await Process.ReadRecordsAsync(HandleRecord, CancellationToken);
-			}
-		}
-
 		#region p4 login
@@ -389,5 +347,42 @@
 	public static class PerforceConnectionExtensions
 	{
 		/// <summary>
+		/// Execute a command and parse the response
+		/// </summary>
+		/// <param name="Perforce">The Perforce connection</param>
+		/// <param name="Command">Command to execute</param>
+		/// <param name="Arguments">Arguments for the command</param>
+		/// <param name="InputData">Input data to pass to Perforce</param>
+		/// <param name="StatRecordType">The type of records to return for "stat" responses</param>
+		/// <param name="CancellationToken">Token used to cancel the operation</param>
+		/// <returns>List of objects returned by the server</returns>
+		public static async Task<List<PerforceResponse>> CommandAsync(this IPerforceConnection Perforce, string Command, IReadOnlyList<string> Arguments, byte[]? InputData, Type? StatRecordType, CancellationToken CancellationToken = default)
+		{
+			using (IPerforceResponse Response = await Perforce.RunAsync(Command, Arguments, InputData))
+			{
+				CachedRecordInfo? StatRecordInfo = (StatRecordType == null) ? null : PerforceReflection.GetCachedRecordInfo(StatRecordType);
+				return await Response.ReadResponsesAsync(StatRecordInfo, CancellationToken);
+			}
+		}
+
+		/// <summary>
+		/// Execute a command and parse the response
+		/// </summary>
+		/// <param name="Perforce">The Perforce connection</param>
+		/// <param name="Command">Command to execute</param>
+		/// <param name="Arguments">Arguments for the command</param>
+		/// <param name="InputData">Input data to pass to Perforce</param>
+		/// <param name="HandleRecord">Delegate used to handle each record</param>
+		/// <param name="CancellationToken">Token used to cancel the operation</param>
+		/// <returns>List of objects returned by the server</returns>
+		public static async Task RecordCommandAsync(this IPerforceConnection Perforce, string Command, IReadOnlyList<string> Arguments, byte[]? InputData, Action<PerforceRecord> HandleRecord, CancellationToken CancellationToken = default)
+		{
+			using (IPerforceResponse Response = await Perforce.RunAsync(Command, Arguments, InputData))
+			{
+				await Response.ReadRecordsAsync(HandleRecord, CancellationToken);
+			}
+		}
+
+		/// <summary>
 		/// Serializes a list of key/value pairs into binary format.
 		/// </summary>
