Add a native implementation of IPerforceConnection, which uses the Perforce C++ API (via a DLL called EpicGames.Perforce.Native) to output records to a buffer shared with managed code (which is double-buffered with the processing thread in managed code). Interaction with the API is lightweight and done on a background thread, with an async API overlaid on top. Prebuilt binaries for the native component are currently tested on Windows and Linux.

Performance has been tested against the P4 .NET API and the command line client by performing a fstat of non-deleted files in //UE5/Main.

EpicGames.Perforce (this library): 47s
Command line client piping to file: 59s
P4 .NET API: ~40 minutes

diff --git a/Engine/Source/Programs/Shared/EpicGames.Perforce.Native/Client.cpp b/Engine/Source/Programs/Shared/EpicGames.Perforce.Native/Client.cpp
new file mode 100644
--- /dev/null
+++ b/Engine/Source/Programs/Shared/EpicGames.Perforce.Native/Client.cpp
@@ -0,0 +1,412 @@
+// Copyright Epic Games, Inc. All Rights Reserved.
+
+#include "clientapi.h"
+#include <assert.h>
+
+#ifdef _MSC_VER
+	#define NATIVE_API __declspec(dllexport)
+#else
+	#define NATIVE_API
+#endif
+
+struct FSettings
+{
+	const char* Port;
+	const char* User;
+	const char* Client;
+	const char* AppName;
+	const char* AppVersion;
+};
+
+struct FReadBuffer
+{
+	unsigned char* Data;
+	unsigned int Length;
+	unsigned int Count;
+	unsigned int MaxLength;
+	unsigned int MaxCount;
+};
+
+struct FWriteBuffer
+{
+	unsigned char* Data;
+	unsigned int MaxLength;
+	unsigned int MaxCount;
+};
+
+typedef void (FOnBufferReadyFn)(const FReadBuffer* ReadBuffer, FWriteBuffer* WriteBuffer);
+
+class FClientUser : public ClientUser
+{
+private:
+	unsigned char* Data;
+	unsigned int Count;
+	unsigned int Length;
+	unsigned int MaxCount;
+	unsigned int MaxLength;
+	FOnBufferReadyFn* OnBufferReady;
+
+public:
+	const char* PromptResponse = nullptr;
+
+	FClientUser(FWriteBuffer* WriteBuffer, FOnBufferReadyFn* OnBufferReady)
+	{
+		SetWriteBuffer(WriteBuffer);
+		this->OnBufferReady = OnBufferReady;
+	}
+
+	virtual void Prompt(Error* err, StrBuf& rsp, int noEcho, Error* e) override
+	{
+		rsp.Set(PromptResponse);
+	}
+
+	virtual void Prompt(Error* err, StrBuf& rsp, int noEcho, int noOutput, Error* e) override
+	{
+		rsp.Set(PromptResponse);
+	}
+
+	virtual void Prompt(const StrPtr& msg, StrBuf& rsp, int noEcho, Error* e) override
+	{
+		rsp.Set(PromptResponse);
+	}
+
+	virtual void Prompt(const StrPtr& msg, StrBuf& rsp, int noEcho, int noOutput, Error* e) override
+	{
+		rsp.Set(PromptResponse);
+	}
+
+	void SetWriteBuffer(FWriteBuffer* WriteBuffer)
+	{
+		this->Data = WriteBuffer->Data;
+		this->Length = 0;
+		this->Count = 0;
+		this->MaxLength = WriteBuffer->MaxLength;
+		this->MaxCount = (size_t)WriteBuffer->MaxCount;
+	}
+
+	void Flush()
+	{
+		FReadBuffer ReadBuffer;
+		ReadBuffer.Data = Data;
+		ReadBuffer.Length = Length;
+		ReadBuffer.Count = Count;
+		ReadBuffer.MaxLength = MaxLength;
+		ReadBuffer.MaxCount = MaxCount;
+
+		FWriteBuffer WriteBuffer;
+		memset(&WriteBuffer, 0, sizeof(WriteBuffer));
+
+		(*OnBufferReady)(&ReadBuffer, &WriteBuffer);
+
+		SetWriteBuffer(&WriteBuffer);
+	}
+
+	virtual void HandleError(Error* err) override
+	{
+		while (!TryOutputError(err))
+		{
+			Flush();
+		}
+	}
+
+	virtual void Message(Error* err) override
+	{
+		while (!TryOutputError(err))
+		{
+			Flush();
+		}
+	}
+
+	bool TryOutputError(Error* Err)
+	{
+		static const char CodeKey[] = "code";
+		static const char SeverityKey[] = "severity";
+		static const char GenericKey[] = "generic";
+		static const char DataKey[] = "data";
+
+		static const char ErrorCode[] = "error";
+
+		StrBuf MessageBuf;
+		StrPtr& Message = Err->StrError(MessageBuf);
+
+		unsigned int MessageLen = Message.Length();
+		unsigned int RecordLen = 1 + MeasureStringField(CodeKey, ErrorCode) + MeasureIntField(SeverityKey) + MeasureIntField(GenericKey) + MeasureStringField(DataKey, MessageLen) + 1;
+		if (Length + RecordLen > MaxLength)
+		{
+			return false;
+		}
+
+		unsigned char* Pos = Data + Length;
+		*(Pos++) = '{';
+		Pos = WriteStringField(Pos, CodeKey, ErrorCode);
+		Pos = WriteIntField(Pos, SeverityKey, Err->GetSeverity());
+		Pos = WriteIntField(Pos, GenericKey, Err->GetGeneric());
+		Pos = WriteStringField(Pos, DataKey, Message.Text(), MessageLen);
+		*(Pos++) = '0';
+		assert(Pos == Data + Length + RecordLen);
+
+		Length += RecordLen;
+		return true;
+	}
+
+	virtual void OutputError(const char* errBuf) override
+	{
+		assert(false);
+	}
+
+	virtual void OutputInfo(char Level, const char* Data) override
+	{
+		while (!TryOutputInfo(Level, Data))
+		{
+			Flush();
+		}
+	}
+
+	bool TryOutputInfo(char Level, const char* Info)
+	{
+		static const char CodeKey[] = "code";
+		static const char LevelKey[] = "level";
+		static const char DataKey[] = "data";
+
+		static const char InfoCode[] = "info";
+
+		unsigned int InfoLen = (unsigned int)strlen(Info);
+		unsigned int RecordLen = 1 + MeasureStringField(CodeKey, InfoCode) + MeasureIntField(LevelKey) + MeasureStringField(DataKey, InfoLen) + 1;
+		if (Length + RecordLen > MaxLength)
+		{
+			return false;
+		}
+
+		unsigned char* Pos = Data + Length;
+		*(Pos++) = '{';
+		Pos = WriteStringField(Pos, CodeKey, InfoCode);
+		Pos = WriteIntField(Pos, LevelKey, Level);
+		Pos = WriteStringField(Pos, DataKey, Info, InfoLen);
+		*(Pos++) = '0';
+		assert(Pos == Data + Length + RecordLen);
+
+		Length += RecordLen;
+		return true;
+	}
+
+	virtual void OutputBinary(const char* data, int length) override
+	{
+		assert(false);
+	}
+
+	virtual void OutputText(const char* Data, int Length) override
+	{
+		assert(false);
+	}
+
+	virtual void OutputStat(StrDict* VarList) override
+	{
+		while (!TryWriteRecord(VarList))
+		{
+			Flush();
+		}
+		if (++Count > MaxCount)
+		{
+			Flush();
+		}
+	}
+
+	template<int KeyLen>
+	static int MeasureStringField(const char(&Key)[KeyLen], int ValueLen)
+	{
+		return 10 + (KeyLen - 1) + ValueLen;
+	}
+
+	template<int KeyLen, int ValueLen>
+	static int MeasureStringField(const char(&Key)[KeyLen], const char(&Value)[ValueLen])
+	{
+		return MeasureStringField(Key, ValueLen - 1);
+	}
+
+	template<int KeyLen>
+	static int MeasureIntField(const char(&Key)[KeyLen])
+	{
+		return 10 + (KeyLen - 1);
+	}
+
+	template<int KeyLen>
+	static unsigned char* WriteStringField(unsigned char* Data, const char(&Key)[KeyLen], const char* Value, int ValueLen)
+	{
+		Data = WriteString(Data, Key);
+		return WriteString(Data, Value, ValueLen);
+	}
+
+	template<int KeyLen, int ValueLen>
+	static unsigned char* WriteStringField(unsigned char* Data, const char(&Key)[KeyLen], const char(&Value)[ValueLen])
+	{
+		Data = WriteString(Data, Key);
+		return WriteString(Data, Value);
+	}
+
+	template<int KeyLen>
+	static unsigned char* WriteIntField(unsigned char* Data, const char(&Key)[KeyLen], int Value)
+	{
+		Data = WriteString(Data, Key);
+		return WriteInt(Data, Value);
+	}
+
+	static unsigned char* WriteInt(unsigned char* Data, int Value)
+	{
+		*(Data++) = 'i';
+		return WriteIntValue(Data, Value);
+	}
+
+	static unsigned char* WriteIntValue(unsigned char* Data, int Value)
+	{
+		*(Data++) = Value & 0xff;
+		*(Data++) = (Value >> 8) & 0xff;
+		*(Data++) = (Value >> 16) & 0xff;
+		*(Data++) = (Value >> 24) & 0xff;
+		return Data;
+	}
+
+	static unsigned char* WriteString(unsigned char* Data, const char* Text, int Length)
+	{
+		*(Data++) = 's';
+		Data = WriteIntValue(Data, Length);
+		memcpy(Data, Text, Length);
+		return Data + Length;
+	}
+
+	template<int Length>
+	static unsigned char* WriteString(unsigned char* Data, const char(&Text)[Length])
+	{
+		return WriteString(Data, Text, Length - 1);
+	}
+
+	bool TryWriteRecord(StrDict* VarList)
+	{
+		unsigned int Pos = Length;
+
+		// Record prefix: '{' (open record), string, 4 bytes, 'code', string ...
+		unsigned char Prefix[] = { '{', 's', 4, 0, 0, 0, 'c', 'o', 'd', 'e', 's', 4, 0, 0, 0, 's', 't', 'a', 't' };
+		unsigned int PrefixLen = sizeof(Prefix) / sizeof(Prefix[0]);
+
+		if (Pos + PrefixLen > MaxLength - 1)
+		{
+			return false;
+		}
+
+		memcpy(Data + Pos, Prefix, PrefixLen);
+		Pos += PrefixLen;
+
+		for (int Idx = 0;; Idx++)
+		{
+			StrRef Field;
+			StrRef Value;
+			if (VarList->GetVar(Idx, Field, Value) == 0)
+			{
+				break;
+			}
+
+			unsigned int FieldLength = Field.Length();
+			unsigned int ValueLength = Value.Length();
+
+			if (Pos + 1 + 4 + FieldLength + 1 + 4 + ValueLength > MaxLength - 1)
+			{
+				return false;
+			}
+
+			Data[Pos] = 's';
+			Pos++;
+
+			memcpy(Data + Pos, &FieldLength, sizeof(FieldLength));
+			Pos += sizeof(FieldLength);
+
+			memcpy(Data + Pos, Field.Text(), FieldLength);
+			Pos += FieldLength;
+
+			Data[Pos] = 's';
+			Pos++;
+
+			memcpy(Data + Pos, &ValueLength, sizeof(ValueLength));
+			Pos += sizeof(ValueLength);
+
+			memcpy(Data + Pos, Value.Text(), ValueLength);
+			Pos += ValueLength;
+		}
+
+		Data[Pos++] = '0';
+
+		Length = Pos;
+		return true;
+	}
+};
+
+class FClient
+{
+public:
+	ClientApi ClientApi;
+	FClientUser User;
+
+	FClient(FWriteBuffer* WriteBuffer, FOnBufferReadyFn* OnBufferReady)
+		: User(WriteBuffer, OnBufferReady)
+	{
+	}
+};
+
+extern "C" NATIVE_API FClient* Client_Create(const FSettings* Settings, FWriteBuffer* WriteBuffer, FOnBufferReadyFn* OnBufferReady)
+{
+	FClient* Client = new FClient(WriteBuffer, OnBufferReady);
+
+	if (Settings != nullptr)
+	{
+		if (Settings->Port != nullptr)
+		{
+			Client->ClientApi.SetPort(Settings->Port);
+		}
+		if (Settings->User != nullptr)
+		{
+			Client->ClientApi.SetUser(Settings->User);
+		}
+		if (Settings->Client != nullptr)
+		{
+			Client->ClientApi.SetClient(Settings->Client);
+		}
+		if (Settings->AppName != nullptr)
+		{
+			Client->ClientApi.SetProg(Settings->AppName);
+		}
+		if (Settings->AppVersion != nullptr)
+		{
+			Client->ClientApi.SetVersion(Settings->AppVersion);
+		}
+	}
+	Client->ClientApi.SetProtocol("tags", "");
+
+	Error Err;
+	Client->ClientApi.Init(&Err);
+	Client->User.HandleError(&Err);
+	Client->User.Flush();
+
+	return Client;
+}
+
+extern "C" NATIVE_API void Client_Login(FClient * Client, const char* Password)
+{
+	Client->User.PromptResponse = Password;
+	Client->ClientApi.SetArgv(0, nullptr);
+	Client->ClientApi.Run("login", &Client->User);
+	Client->User.Flush();
+	Client->User.PromptResponse = nullptr;
+}
+
+extern "C" NATIVE_API void Client_Command(FClient* Client, const char* Func, int ArgCount, const char** Args)
+{
+	Client->ClientApi.SetArgv(ArgCount, (char* const*)Args);
+	Client->ClientApi.Run(Func, &Client->User);
+	Client->User.Flush();
+}
+
+extern "C" NATIVE_API void Client_Destroy(FClient* Client)
+{
+	Error Err;
+	Client->ClientApi.Final(&Err);
+	Client->User.HandleError(&Err);
+	delete Client;
+}
diff --git a/Engine/Source/Programs/Shared/EpicGames.Perforce.Native/Makefile b/Engine/Source/Programs/Shared/EpicGames.Perforce.Native/Makefile
new file mode 100644
--- /dev/null
+++ b/Engine/Source/Programs/Shared/EpicGames.Perforce.Native/Makefile
@@ -0,0 +1,26 @@
+# EpicGames.Perforce.Native
+
+OBJDIR=obj/linux-x64/Release
+OUTDIR=bin/linux-x64/Release
+ 
+CXXFLAGS=-fPIC -std=c++14 -O1 -I../../../ThirdParty/Perforce/p4api-2018.1/Include/Linux/p4
+ 
+LIBPATH=../../../ThirdParty/Perforce/p4api-2018.1/Lib/Linux
+LIBS=$(LIBPATH)/libclient.a $(LIBPATH)/librpc.a $(LIBPATH)/libsupp.a -lssl -lcrypto -lpthread -ldl -lstdc++ -L../../../ThirdParty/OpenSSL/1.1.1c/lib/Unix/x86_64-unknown-linux-gnu
+ 
+$(OBJDIR):
+	mkdir -p $(OBJDIR)
+ 
+$(OBJDIR)/%.o: %.cpp | $(OBJDIR)
+	$(CC) $(CXXFLAGS) -c -o $@ $<
+ 
+# note by default we strip debug info for size
+client: $(OBJDIR)/Client.o
+	mkdir -p $(OUTDIR)/
+	$(CC) $(LDFLAGS) -shared -o $(OUTDIR)/EpicGames.Perforce.Native.so $^ $(LIBS)
+ 
+clean:
+	rm -rf $(OBJDIR)/ 
+	rm -rf $(OUTDIR)/
+ 
+all: bridge
\ No newline at end of file
diff --git a/Engine/Source/Programs/Shared/EpicGames.Perforce/EpicGames.Perforce.csproj b/Engine/Source/Programs/Shared/EpicGames.Perforce/EpicGames.Perforce.csproj
--- a/Engine/Source/Programs/Shared/EpicGames.Perforce/EpicGames.Perforce.csproj
+++ b/Engine/Source/Programs/Shared/EpicGames.Perforce/EpicGames.Perforce.csproj
@@ -21,6 +21,10 @@
     <DocumentationFile></DocumentationFile>
     <Optimize>true</Optimize>
   </PropertyGroup>
+  <ItemGroup>
+    <None Remove="EpicGames.Perforce.uatbuildrecord" />
+    <None Remove="EpicGames.Perforce.xml" />
+  </ItemGroup>
   <ItemGroup>
     <Compile Include="..\..\Shared\MetaData.cs">
       <Link>Properties\MetaData.cs</Link>
@@ -31,4 +35,21 @@
     <ProjectReference Include="..\EpicGames.Core\EpicGames.Core.csproj" />
   </ItemGroup>
 
+  <PropertyGroup>
+    <WindowsNativeConfig>Release</WindowsNativeConfig>
+    <WindowsNativeConfig Condition=" '$(Configuration)|$(Platform)' == 'Debug|AnyCPU' ">Debug</WindowsNativeConfig>
+  </PropertyGroup>
+
+  <ItemGroup>
+    <None Include="..\..\Shared\EpicGames.Perforce.Native\bin\win-x64\$(WindowsNativeConfig)\EpicGames.Perforce.Native.dll" Visible="false">
+      <CopyToOutputDirectory>PreserveNewest</CopyToOutputDirectory>
+    </None>
+    <None Include="..\..\Shared\EpicGames.Perforce.Native\bin\win-x64\$(WindowsNativeConfig)\EpicGames.Perforce.Native.pdb" Visible="false">
+      <CopyToOutputDirectory>PreserveNewest</CopyToOutputDirectory>
+    </None>
+    <None Include="..\..\Shared\EpicGames.Perforce.Native\bin\linux-x64\Release\EpicGames.Perforce.Native.so" Visible="false">
+      <CopyToOutputDirectory>PreserveNewest</CopyToOutputDirectory>
+    </None>
+  </ItemGroup>
+
 </Project>
diff --git a/Engine/Source/Programs/Shared/EpicGames.Perforce/IPerforceConnection.cs b/Engine/Source/Programs/Shared/EpicGames.Perforce/IPerforceConnection.cs
new file mode 100644
--- /dev/null
+++ b/Engine/Source/Programs/Shared/EpicGames.Perforce/IPerforceConnection.cs
@@ -0,0 +1,56 @@
+// Copyright Epic Games, Inc. All Rights Reserved.
+
+using Microsoft.Extensions.Logging;
+using System;
+using System.Collections.Generic;
+using System.Text;
+using System.Threading;
+using System.Threading.Tasks;
+
+namespace EpicGames.Perforce
+{
+	/// <summary>
+	/// Base interface for Perforce clients
+	/// </summary>
+	public interface IPerforceConnection
+	{
+		/// <summary>
+		/// Logger for this connection
+		/// </summary>
+		ILogger Logger { get; }
+
+		/// <summary>
+		/// Runs a Perforce command
+		/// </summary>
+		/// <param name="Command">The command name</param>
+		/// <param name="Arguments">Arguments for the command</param>
+		/// <param name="InputData">Input data to be passed to the command</param>
+		/// <returns>Response object</returns>
+		Task<IPerforceOutput> CommandAsync(string Command, IReadOnlyList<string> Arguments, byte[]? InputData);
+
+		/// <summary>
+		/// Execute the 'login' command
+		/// </summary>
+		/// <param name="Password">Password to use to login</param>
+		/// <param name="CancellationToken">Token used to cancel the operation</param>
+		/// <returns>Response from the server</returns>
+		Task LoginAsync(string Password, CancellationToken CancellationToken = default);
+
+		/// <summary>
+		/// Sets an environment variable
+		/// </summary>
+		/// <param name="Name">Name of the variable to set</param>
+		/// <param name="Value">Value for the variable</param>
+		/// <param name="CancellationToken">Token used to cancel the operation</param>
+		/// <returns>Response from the server</returns>
+		Task SetAsync(string Name, string Value, CancellationToken CancellationToken = default);
+
+		/// <summary>
+		/// Gets the setting of a Perforce variable
+		/// </summary>
+		/// <param name="Name">Name of the variable to get</param>
+		/// <param name="CancellationToken">Cancellation token for the request</param>
+		/// <returns>Value of the variable</returns>
+		Task<string?> TryGetSettingAsync(string Name, CancellationToken CancellationToken = default);
+	}
+}
diff --git a/Engine/Source/Programs/Shared/EpicGames.Perforce/IPerforceOutput.cs b/Engine/Source/Programs/Shared/EpicGames.Perforce/IPerforceOutput.cs
new file mode 100644
--- /dev/null
+++ b/Engine/Source/Programs/Shared/EpicGames.Perforce/IPerforceOutput.cs
@@ -0,0 +1,748 @@
+// Copyright Epic Games, Inc. All Rights Reserved.
+
+using EpicGames.Core;
+using System;
+using System.Buffers.Text;
+using System.Collections.Generic;
+using System.Diagnostics.CodeAnalysis;
+using System.IO;
+using System.Linq;
+using System.Text;
+using System.Threading;
+using System.Threading.Tasks;
+
+namespace EpicGames.Perforce
+{
+	/// <summary>
+	/// Interface for the result of a Perforce operation
+	/// </summary>
+	public interface IPerforceOutput : IAsyncDisposable
+	{
+		/// <summary>
+		/// Data containing the result
+		/// </summary>
+		ReadOnlyMemory<byte> Data { get; }
+
+		/// <summary>
+		/// Waits until more data has been read into the buffer. 
+		/// </summary>
+		/// <returns>True if more data was read, false otherwise</returns>
+		Task<bool> ReadAsync(CancellationToken Token);
+
+		/// <summary>
+		/// Discard bytes from the start of the result buffer
+		/// </summary>
+		/// <param name="NumBytes">Number of bytes to discard</param>
+		void Discard(int NumBytes);
+	}
+
+	/// <summary>
+	/// Wraps a call to a p4.exe child process, and allows reading data from it
+	/// </summary>
+	public static class PerforceOutputExtensions
+	{
+		/// <summary>
+		/// String constants for records
+		/// </summary>
+		static class ReadOnlyUtf8StringConstants
+		{
+			public static Utf8String Code = "code";
+			public static Utf8String Stat = "stat";
+			public static Utf8String Info = "info";
+			public static Utf8String Error = "error";
+		}
+
+		/// <summary>
+		/// Standard prefix for a returned record: record indicator, string, 4 bytes, 'code', string, [value]
+		/// </summary>
+		static readonly byte[] RecordPrefix = { (byte)'{', (byte)'s', 4, 0, 0, 0, (byte)'c', (byte)'o', (byte)'d', (byte)'e', (byte)'s' };
+
+		/// <summary>
+		/// Formats the current contents of the buffer to a string
+		/// </summary>
+		/// <param name="Data">The next byte that was read</param>
+		/// <returns>String representation of the buffer</returns>
+		private static string FormatDataAsString(ReadOnlySpan<byte> Data)
+		{
+			StringBuilder Result = new StringBuilder();
+			if (Data.Length > 0)
+			{
+				for (int Idx = 0; Idx < Data.Length && Idx < 1024;)
+				{
+					Result.Append("\n   ");
+
+					// Output to the end of the line
+					for (; Idx < Data.Length && Idx < 1024 && Data[Idx] != '\r' && Data[Idx] != '\n'; Idx++)
+					{
+						if (Data[Idx] == '\t')
+						{
+							Result.Append("\t");
+						}
+						else if (Data[Idx] == '\\')
+						{
+							Result.Append("\\");
+						}
+						else if (Data[Idx] >= 0x20 && Data[Idx] <= 0x7f)
+						{
+							Result.Append((char)Data[Idx]);
+						}
+						else
+						{
+							Result.AppendFormat("\\x{0:x2}", Data[Idx]);
+						}
+					}
+
+					// Skip the newline characters
+					if (Idx < Data.Length && Data[Idx] == '\r')
+					{
+						Idx++;
+					}
+					if (Idx < Data.Length && Data[Idx] == '\n')
+					{
+						Idx++;
+					}
+				}
+			}
+			return Result.ToString();
+		}
+
+		/// <summary>
+		/// Formats the current contents of the buffer to a string
+		/// </summary>
+		/// <param name="Data">The data to format</param>
+		/// <returns>String representation of the buffer</returns>
+		private static string FormatDataAsHexDump(ReadOnlySpan<byte> Data)
+		{
+			// Format the result
+			StringBuilder Result = new StringBuilder();
+
+			const int RowLength = 16;
+			for (int BaseIdx = 0; BaseIdx < Data.Length; BaseIdx += RowLength)
+			{
+				Result.Append("\n    ");
+				for (int Offset = 0; Offset < RowLength; Offset++)
+				{
+					int Idx = BaseIdx + Offset;
+					if (Idx >= Data.Length)
+					{
+						Result.Append("   ");
+					}
+					else
+					{
+						Result.AppendFormat("{0:x2} ", Data[Idx]);
+					}
+				}
+				Result.Append("   ");
+				for (int Offset = 0; Offset < RowLength; Offset++)
+				{
+					int Idx = BaseIdx + Offset;
+					if (Idx >= Data.Length)
+					{
+						break;
+					}
+					else if (Data[Idx] < 0x20 || Data[Idx] >= 0x7f)
+					{
+						Result.Append('.');
+					}
+					else
+					{
+						Result.Append((char)Data[Idx]);
+					}
+				}
+			}
+			return Result.ToString();
+		}
+
+		/// <summary>
+		/// Read a list of responses from the child process
+		/// </summary>
+		/// <param name="Perforce">The response to read from</param>
+		/// <param name="StatRecordType">The type of stat record to parse</param>
+		/// <param name="CancellationToken">Cancellation token for the read</param>
+		/// <returns>Async task</returns>
+		public static async Task<List<PerforceResponse>> ReadResponsesAsync(this IPerforceOutput Perforce, Type? StatRecordType, CancellationToken CancellationToken)
+		{
+			CachedRecordInfo? StatRecordInfo = (StatRecordType == null) ? null : PerforceReflection.GetCachedRecordInfo(StatRecordType);
+
+			List<PerforceResponse> Responses = new List<PerforceResponse>();
+
+			// Read all the records into a list
+			while (await Perforce.ReadAsync(CancellationToken))
+			{
+				// Check for the whole message not being a marshalled python object, and produce a better response in that scenario
+				ReadOnlyMemory<byte> Data = Perforce.Data;
+				if (Data.Length > 0 && Responses.Count == 0 && Data.Span[0] != '{')
+				{
+					throw new PerforceException("Unexpected response from server:{0}", FormatDataAsString(Data.Span));
+				}
+
+				// Parse the responses from the current buffer
+				int BufferPos = 0;
+				for (; ; )
+				{
+					int NewBufferPos = BufferPos;
+					if (!TryReadResponse(Data.Span, ref NewBufferPos, StatRecordInfo, out PerforceResponse? Response))
+					{
+						break;
+					}
+					if (Response.Error == null || Response.Error.Generic != PerforceGenericCode.Empty)
+					{
+						Responses.Add(Response);
+					}
+					BufferPos = NewBufferPos;
+				}
+
+				// Discard all the data that we've processed
+				Perforce.Discard(BufferPos);
+			}
+
+			// If the stream is complete but we couldn't parse a response from the server, treat it as an error
+			if (Perforce.Data.Length > 0)
+			{
+				throw new PerforceException("Unexpected response from server:{0}", FormatDataAsString(Perforce.Data.Span));
+			}
+
+			return Responses;
+		}
+
+		/// <summary>
+		/// Read a list of responses from the child process
+		/// </summary>
+		/// <param name="Perforce">The Perforce response</param>
+		/// <param name="HandleRecord">Delegate to invoke for each record read</param>
+		/// <param name="CancellationToken">Cancellation token for the read</param>
+		/// <returns>Async task</returns>
+		public static async Task ReadRecordsAsync(this IPerforceOutput Perforce, Action<PerforceRecord> HandleRecord, CancellationToken CancellationToken)
+		{
+			PerforceRecord Record = new PerforceRecord();
+			Record.Rows = new List<KeyValuePair<Utf8String, PerforceValue>>();
+
+			// Read all the records into a list
+			Task<bool> ReadTask = Perforce.ReadAsync(CancellationToken);
+			while (await ReadTask)
+			{
+				// Start a read to add more data
+				ReadTask = Perforce.ReadAsync(CancellationToken);
+				ReadOnlyMemory<byte> Data = Perforce.Data;
+
+				// Parse the responses from the current buffer
+				int BufferPos = 0;
+				for (; ; )
+				{
+					int InitialBufferPos = BufferPos;
+					if (!ReadRecord(Data.Span, ref BufferPos, Record.Rows))
+					{
+						BufferPos = InitialBufferPos;
+						break;
+					}
+					HandleRecord(Record);
+				}
+				Perforce.Discard(BufferPos);
+			}
+
+			// If the stream is complete but we couldn't parse a response from the server, treat it as an error
+			if (Perforce.Data.Length > 0)
+			{
+				throw new PerforceException("Unexpected response from server:{0}", FormatDataAsString(Perforce.Data.Span));
+			}
+		}
+
+		/// <summary>
+		/// Reads from the buffer into a record object
+		/// </summary>
+		/// <param name="Buffer">The buffer to read from</param>
+		/// <param name="BufferPos">Current read position within the buffer</param>
+		/// <param name="Rows">List of rows to read into</param>
+		/// <returns>True if a record could be read; false if more data is required</returns>
+		static bool ReadRecord(ReadOnlySpan<byte> Buffer, ref int BufferPos, List<KeyValuePair<Utf8String, PerforceValue>> Rows)
+		{
+			Rows.Clear();
+
+			// Check we can read the initial record marker
+			if (BufferPos >= Buffer.Length)
+			{
+				return false;
+			}
+			if (Buffer[BufferPos] != '{')
+			{
+				throw new PerforceException("Invalid record start");
+			}
+			BufferPos++;
+
+			// Capture the start of the string
+			int StartPos = BufferPos;
+			for (; ; )
+			{
+				// Check that we've got a string field
+				if (BufferPos >= Buffer.Length)
+				{
+					return false;
+				}
+
+				// If this is the end of the record, break out
+				byte KeyType = Buffer[BufferPos++];
+				if (KeyType == '0')
+				{
+					break;
+				}
+				else if (KeyType != 's')
+				{
+					throw new PerforceException("Unexpected key field type while parsing marshalled output ({0}) - expected 's', got: {1}", (int)KeyType, FormatDataAsHexDump(Buffer.Slice(BufferPos - 1)));
+				}
+
+				// Read the tag
+				Utf8String Key;
+				if (!TryReadString(Buffer, ref BufferPos, out Key))
+				{
+					return false;
+				}
+
+				// Remember the start of the value
+				int ValueOffset = BufferPos;
+
+				// Read the value type
+				byte ValueType;
+				if (!TryReadByte(Buffer, ref BufferPos, out ValueType))
+				{
+					return false;
+				}
+
+				// Parse the appropriate value
+				PerforceValue Value;
+				if (ValueType == 's')
+				{
+					Utf8String String;
+					if (!TryReadString(Buffer, ref BufferPos, out String))
+					{
+						return false;
+					}
+					Value = new PerforceValue(Buffer.Slice(ValueOffset, BufferPos - ValueOffset).ToArray());
+				}
+				else if (ValueType == 'i')
+				{
+					int Integer;
+					if (!TryReadInt(Buffer, ref BufferPos, out Integer))
+					{
+						return false;
+					}
+					Value = new PerforceValue(Buffer.Slice(ValueOffset, BufferPos - ValueOffset).ToArray());
+				}
+				else
+				{
+					throw new PerforceException("Unrecognized value type {0}", ValueType);
+				}
+
+				// Construct the response object with the record
+				Rows.Add(KeyValuePair.Create(Key, Value));
+			}
+			return true;
+		}
+
+		/// <summary>
+		/// Reads a response object from the buffer
+		/// </summary>
+		/// <param name="Buffer">The buffer to read from</param>
+		/// <param name="BufferPos">Current read position within the buffer</param>
+		/// <param name="StatRecordInfo">The type of record expected to parse from the response</param>
+		/// <param name="Response">Receives the response object on success</param>
+		/// <returns>True if a response was read, false if the buffer needs more data</returns>
+		static bool TryReadResponse(ReadOnlySpan<byte> Buffer, ref int BufferPos, CachedRecordInfo? StatRecordInfo, [NotNullWhen(true)] out PerforceResponse? Response)
+		{
+			if (BufferPos + RecordPrefix.Length + 4 > Buffer.Length)
+			{
+				Response = null;
+				return false;
+			}
+
+			ReadOnlySpan<byte> Prefix = Buffer.Slice(BufferPos, RecordPrefix.Length);
+			if (!Prefix.SequenceEqual(RecordPrefix))
+			{
+				throw new PerforceException("Expected 'code' field at the start of record");
+			}
+			BufferPos += Prefix.Length;
+
+			Utf8String Code;
+			if (!TryReadString(Buffer, ref BufferPos, out Code))
+			{
+				Response = null;
+				return false;
+			}
+
+			// Dispatch it to the appropriate handler
+			object? Record;
+			if (Code == ReadOnlyUtf8StringConstants.Stat && StatRecordInfo != null)
+			{
+				if (!TryReadTypedRecord(Buffer, ref BufferPos, Utf8String.Empty, StatRecordInfo, out Record))
+				{
+					Response = null;
+					return false;
+				}
+			}
+			else if (Code == ReadOnlyUtf8StringConstants.Info)
+			{
+				if (!TryReadTypedRecord(Buffer, ref BufferPos, Utf8String.Empty, PerforceReflection.InfoRecordInfo, out Record))
+				{
+					Response = null;
+					return false;
+				}
+			}
+			else if (Code == ReadOnlyUtf8StringConstants.Error)
+			{
+				if (!TryReadTypedRecord(Buffer, ref BufferPos, Utf8String.Empty, PerforceReflection.ErrorRecordInfo, out Record))
+				{
+					Response = null;
+					return false;
+				}
+			}
+			else
+			{
+				throw new PerforceException("Unknown return code for record: {0}", Code);
+			}
+
+			// Skip over the record terminator
+			if (BufferPos >= Buffer.Length || Buffer[BufferPos] != '0')
+			{
+				throw new PerforceException("Unexpected record terminator");
+			}
+			BufferPos++;
+
+			// Create the response
+			Response = new PerforceResponse(Record);
+			return true;
+		}
+
+		/// <summary>
+		/// Parse an individual record from the server.
+		/// </summary>
+		/// <param name="Buffer">The buffer to read from</param>
+		/// <param name="BufferPos">Current read position within the buffer</param>
+		/// <param name="RequiredSuffix">The required suffix for any subobject arrays.</param>
+		/// <param name="RecordInfo">Reflection information for the type being serialized into.</param>
+		/// <param name="Record">Receives the record on success</param>
+		/// <returns>The parsed object.</returns>
+		static bool TryReadTypedRecord(ReadOnlySpan<byte> Buffer, ref int BufferPos, Utf8String RequiredSuffix, CachedRecordInfo RecordInfo, [NotNullWhen(true)] out object? Record)
+		{
+			// Create a bitmask for all the required tags
+			ulong RequiredTagsBitMask = 0;
+
+			// Create the new record
+			object? NewRecord = RecordInfo.CreateInstance();
+			if (NewRecord == null)
+			{
+				throw new InvalidDataException($"Unable to construct record of type {RecordInfo.Type}");
+			}
+
+			// Get the record info, and parse it into the object
+			for (; ; )
+			{
+				// Check that we've got a string field
+				if (BufferPos >= Buffer.Length)
+				{
+					Record = null;
+					return false;
+				}
+
+				// If this is the end of the record, break out
+				byte KeyType = Buffer[BufferPos];
+				if (KeyType == '0')
+				{
+					break;
+				}
+				else if (KeyType != 's')
+				{
+					throw new PerforceException("Unexpected key field type while parsing marshalled output ({0}) - expected 's', got: {1}", (int)KeyType, FormatDataAsHexDump(Buffer.Slice(BufferPos)));
+				}
+
+				// Capture the initial buffer position, in case we have to roll back
+				int StartBufferPos = BufferPos;
+				BufferPos++;
+
+				// Read the tag
+				Utf8String Tag;
+				if (!TryReadString(Buffer, ref BufferPos, out Tag))
+				{
+					Record = null;
+					return false;
+				}
+
+				// Find the start of the array suffix
+				int SuffixIdx = Tag.Length;
+				while (SuffixIdx > 0 && (Tag[SuffixIdx - 1] == (byte)',' || (Tag[SuffixIdx - 1] >= '0' && Tag[SuffixIdx - 1] <= '9')))
+				{
+					SuffixIdx--;
+				}
+
+				// Separate the key into tag and suffix
+				Utf8String Suffix = Tag.Slice(SuffixIdx);
+				Tag = Tag.Slice(0, SuffixIdx);
+
+				// Try to find the matching field
+				CachedTagInfo? TagInfo;
+				if (RecordInfo.NameToInfo.TryGetValue(Tag, out TagInfo))
+				{
+					RequiredTagsBitMask |= TagInfo.RequiredTagBitMask;
+				}
+
+				// Check whether it's a subobject or part of the current object.
+				if (Suffix == RequiredSuffix)
+				{
+					if (!TryReadValue(Buffer, ref BufferPos, NewRecord, TagInfo))
+					{
+						Record = null;
+						return false;
+					}
+				}
+				else if (Suffix.StartsWith(RequiredSuffix) && (RequiredSuffix.Length == 0 || Suffix[RequiredSuffix.Length] == ','))
+				{
+					// Part of a subobject. If this record doesn't have any listed subobject type, skip the field and continue.
+					if (TagInfo != null)
+					{
+						// Get the list field
+						System.Collections.IList? List = (System.Collections.IList?)TagInfo.Field.GetValue(NewRecord);
+						if (List == null)
+						{
+							throw new PerforceException($"Empty list for {TagInfo.Field.Name}");
+						}
+
+						// Check the suffix matches the index of the next element
+						if (!IsCorrectIndex(Suffix, RequiredSuffix, List.Count))
+						{
+							throw new PerforceException("Subobject element received out of order: got {0}", Suffix);
+						}
+
+						// Add it to the list
+						if (!TryReadValue(Buffer, ref BufferPos, NewRecord, TagInfo))
+						{
+							Record = null;
+							return false;
+						}
+					}
+					else if (RecordInfo.SubElementField != null)
+					{
+						// Move back to the start of this tag
+						BufferPos = StartBufferPos;
+
+						// Get the list field
+						System.Collections.IList? List = (System.Collections.IList?)RecordInfo.SubElementField.GetValue(NewRecord);
+						if (List == null)
+						{
+							throw new PerforceException($"Invalid field for {RecordInfo.SubElementField.Name}");
+						}
+
+						// Check the suffix matches the index of the next element
+						if (!IsCorrectIndex(Suffix, RequiredSuffix, List.Count))
+						{
+							throw new PerforceException("Subobject element received out of order: got {0}", Suffix);
+						}
+
+						// Parse the subobject and add it to the list
+						object? SubRecord;
+						if (!TryReadTypedRecord(Buffer, ref BufferPos, Suffix, RecordInfo.SubElementRecordInfo!, out SubRecord))
+						{
+							Record = null;
+							return false;
+						}
+						List.Add(SubRecord);
+					}
+					else
+					{
+						// Just discard the value
+						if (!TryReadValue(Buffer, ref BufferPos, NewRecord, TagInfo))
+						{
+							Record = null;
+							return false;
+						}
+					}
+				}
+				else
+				{
+					// Roll back
+					BufferPos = StartBufferPos;
+					break;
+				}
+			}
+
+			// Make sure we've got all the required tags we need
+			if (RequiredTagsBitMask != RecordInfo.RequiredTagsBitMask)
+			{
+				string MissingTagNames = String.Join(", ", RecordInfo.NameToInfo.Where(x => (RequiredTagsBitMask | x.Value.RequiredTagBitMask) != RequiredTagsBitMask).Select(x => x.Key));
+				throw new PerforceException("Missing '{0}' tag when parsing '{1}'", MissingTagNames, RecordInfo.Type.Name);
+			}
+
+			// Construct the response object with the record
+			Record = NewRecord;
+			return true;
+		}
+
+		/// <summary>
+		/// Reads a value from the input buffer
+		/// </summary>
+		/// <param name="Buffer">The buffer to read from</param>
+		/// <param name="BufferPos">Current read position within the buffer</param>
+		/// <param name="NewRecord">The new record</param>
+		/// <param name="TagInfo">The current tag</param>
+		/// <returns></returns>
+		static bool TryReadValue(ReadOnlySpan<byte> Buffer, ref int BufferPos, object NewRecord, CachedTagInfo? TagInfo)
+		{
+			// Read the value type
+			byte ValueType;
+			if (!TryReadByte(Buffer, ref BufferPos, out ValueType))
+			{
+				return false;
+			}
+
+			// Parse the appropriate value
+			if (ValueType == 's')
+			{
+				Utf8String String;
+				if (!TryReadString(Buffer, ref BufferPos, out String))
+				{
+					return false;
+				}
+				if (TagInfo != null)
+				{
+					TagInfo.SetFromString(NewRecord, String);
+				}
+			}
+			else if (ValueType == 'i')
+			{
+				int Integer;
+				if (!TryReadInt(Buffer, ref BufferPos, out Integer))
+				{
+					return false;
+				}
+				if (TagInfo != null)
+				{
+					TagInfo.SetFromInteger(NewRecord, Integer);
+				}
+			}
+			else
+			{
+				throw new PerforceException("Unrecognized value type {0}", ValueType);
+			}
+
+			return true;
+		}
+
+		/// <summary>
+		/// Attempts to read a single byte from the buffer
+		/// </summary>
+		/// <param name="Buffer">The buffer to read from</param>
+		/// <param name="BufferPos">Current read position within the buffer</param>
+		/// <param name="Value">Receives the byte that was read</param>
+		/// <returns>True if a byte was read from the buffer, false if there was not enough data</returns>
+		static bool TryReadByte(ReadOnlySpan<byte> Buffer, ref int BufferPos, out byte Value)
+		{
+			if (BufferPos >= Buffer.Length)
+			{
+				Value = 0;
+				return false;
+			}
+
+			Value = Buffer[BufferPos];
+			BufferPos++;
+			return true;
+		}
+
+		/// <summary>
+		/// Attempts to read a single int from the buffer
+		/// </summary>
+		/// <param name="Buffer">The buffer to read from</param>
+		/// <param name="BufferPos">Current read position within the buffer</param>
+		/// <param name="Value">Receives the value that was read</param>
+		/// <returns>True if an int was read from the buffer, false if there was not enough data</returns>
+		static bool TryReadInt(ReadOnlySpan<byte> Buffer, ref int BufferPos, out int Value)
+		{
+			if (BufferPos + 4 > Buffer.Length)
+			{
+				Value = 0;
+				return false;
+			}
+
+			Value = Buffer[BufferPos + 0] | (Buffer[BufferPos + 1] << 8) | (Buffer[BufferPos + 2] << 16) | (Buffer[BufferPos + 3] << 24);
+			BufferPos += 4;
+			return true;
+		}
+
+		/// <summary>
+		/// Attempts to read a string with type indicator from the buffer
+		/// </summary>
+		/// <param name="Buffer">The buffer to read from</param>
+		/// <param name="BufferPos">Current read position within the buffer</param>
+		/// <param name="String">Receives the value that was read</param>
+		/// <returns>True if a string was read from the buffer, false if there was not enough data</returns>
+		static bool TryReadStringWithType(ReadOnlySpan<byte> Buffer, ref int BufferPos, out Utf8String String)
+		{
+			byte ValueType;
+			if (!TryReadByte(Buffer, ref BufferPos, out ValueType))
+			{
+				String = new Utf8String();
+				return false;
+			}
+			if (ValueType != 's')
+			{
+				throw new PerforceException("Expected string value, got '{0}'", ValueType);
+			}
+			return TryReadString(Buffer, ref BufferPos, out String);
+		}
+
+		/// <summary>
+		/// Attempts to read a string from the buffer
+		/// </summary>
+		/// <param name="Buffer">The buffer to read from</param>
+		/// <param name="BufferPos">Current read position within the buffer</param>
+		/// <param name="String">Receives the value that was read</param>
+		/// <returns>True if a string was read from the buffer, false if there was not enough data</returns>
+		static bool TryReadString(ReadOnlySpan<byte> Buffer, ref int BufferPos, out Utf8String String)
+		{
+			int Length;
+			if (!TryReadInt(Buffer, ref BufferPos, out Length))
+			{
+				String = new Utf8String();
+				return false;
+			}
+
+			if (BufferPos + Length > Buffer.Length)
+			{
+				String = new Utf8String();
+				return false;
+			}
+
+			String = new Utf8String(Buffer.Slice(BufferPos, Length).ToArray());
+			BufferPos += Length;
+			return true;
+		}
+
+		/// <summary>
+		/// Determines if the given text contains the expected prefix followed by an array index
+		/// </summary>
+		/// <param name="Text">The text to check</param>
+		/// <param name="Prefix">The required prefix</param>
+		/// <param name="Index">The required index</param>
+		/// <returns>True if the index is correct</returns>
+		static bool IsCorrectIndex(Utf8String Text, Utf8String Prefix, int Index)
+		{
+			if (Prefix.Length > 0)
+			{
+				return Text.StartsWith(Prefix) && Text.Length > Prefix.Length && Text[Prefix.Length] == (byte)',' && IsCorrectIndex(Text.Span.Slice(Prefix.Length + 1), Index);
+			}
+			else
+			{
+				return IsCorrectIndex(Text.Span, Index);
+			}
+		}
+
+		/// <summary>
+		/// Determines if the given text matches the expected array index
+		/// </summary>
+		/// <param name="Span">The text to check</param>
+		/// <param name="ExpectedIndex">The expected array index</param>
+		/// <returns>True if the span matches</returns>
+		static bool IsCorrectIndex(ReadOnlySpan<byte> Span, int ExpectedIndex)
+		{
+			int Index;
+			int BytesConsumed;
+			return Utf8Parser.TryParse(Span, out Index, out BytesConsumed) && BytesConsumed == Span.Length && Index == ExpectedIndex;
+		}
+	}
+}
diff --git a/Engine/Source/Programs/Shared/EpicGames.Perforce/PerforceChildProcess.cs b/Engine/Source/Programs/Shared/EpicGames.Perforce/PerforceChildProcess.cs
--- a/Engine/Source/Programs/Shared/EpicGames.Perforce/PerforceChildProcess.cs
+++ b/Engine/Source/Programs/Shared/EpicGames.Perforce/PerforceChildProcess.cs
@@ -20,33 +20,10 @@ using System.Threading.Tasks;
 
 namespace EpicGames.Perforce
 {
-	/// <summary>
-	/// Interface for the result of a Perforce operation
-	/// </summary>
-	public interface IPerforceResponse : IDisposable
-	{
-		/// <summary>
-		/// Data containing the result
-		/// </summary>
-		ReadOnlyMemory<byte> Data { get; }
-
-		/// <summary>
-		/// Waits until more data has been read into the buffer. 
-		/// </summary>
-		/// <returns>True if more data was read, false otherwise</returns>
-		Task<bool> ReadAsync(CancellationToken Token);
-
-		/// <summary>
-		/// Discard bytes from the start of the result buffer
-		/// </summary>
-		/// <param name="NumBytes">Number of bytes to discard</param>
-		void Discard(int NumBytes);
-	}
-
 	/// <summary>
 	/// Wraps a call to a p4.exe child process, and allows reading data from it
 	/// </summary>
-	class PerforceChildProcess : IPerforceResponse, IDisposable
+	class PerforceChildProcess : IPerforceOutput, IDisposable
 	{
 		/// <summary>
 		/// The process group
@@ -124,9 +101,14 @@ namespace EpicGames.Perforce
 			return PerforceFileName;
 		}
 
-		/// <summary>
-		/// Dispose of resources owned by this object
-		/// </summary>
+		/// <inheritdoc/>
+		public ValueTask DisposeAsync()
+		{
+			DisposeAsync();
+			return new ValueTask(Task.CompletedTask);
+		}
+
+		/// <inheritdoc/>
 		public void Dispose()
 		{
 			if (ChildProcess != null)
@@ -196,727 +178,4 @@ namespace EpicGames.Perforce
 			return Tuple.Create(ChildProcess.ExitCode == 0, String);
 		}
 	}
-
-	/// <summary>
-	/// Wraps a call to a p4.exe child process, and allows reading data from it
-	/// </summary>
-	static class PerforceResultExtensions
-	{
-		/// <summary>
-		/// String constants for records
-		/// </summary>
-		static class ReadOnlyUtf8StringConstants
-		{
-			public static Utf8String Code = "code";
-			public static Utf8String Stat = "stat";
-			public static Utf8String Info = "info";
-			public static Utf8String Error = "error";
-		}
-
-		/// <summary>
-		/// Standard prefix for a returned record: record indicator, string, 4 bytes, 'code', string, [value]
-		/// </summary>
-		static readonly byte[] RecordPrefix = { (byte)'{', (byte)'s', 4, 0, 0, 0, (byte)'c', (byte)'o', (byte)'d', (byte)'e', (byte)'s' };
-
-		/// <summary>
-		/// Formats the current contents of the buffer to a string
-		/// </summary>
-		/// <param name="Data">The next byte that was read</param>
-		/// <returns>String representation of the buffer</returns>
-		private static string FormatDataAsString(ReadOnlySpan<byte> Data)
-		{
-			StringBuilder Result = new StringBuilder();
-			if (Data.Length > 0)
-			{
-				for (int Idx = 0; Idx < Data.Length && Idx < 1024;)
-				{
-					Result.Append("\n   ");
-
-					// Output to the end of the line
-					for (; Idx < Data.Length && Idx < 1024 && Data[Idx] != '\r' && Data[Idx] != '\n'; Idx++)
-					{
-						if (Data[Idx] == '\t')
-						{
-							Result.Append("\t");
-						}
-						else if (Data[Idx] == '\\')
-						{
-							Result.Append("\\");
-						}
-						else if (Data[Idx] >= 0x20 && Data[Idx] <= 0x7f)
-						{
-							Result.Append((char)Data[Idx]);
-						}
-						else
-						{
-							Result.AppendFormat("\\x{0:x2}", Data[Idx]);
-						}
-					}
-
-					// Skip the newline characters
-					if (Idx < Data.Length && Data[Idx] == '\r')
-					{
-						Idx++;
-					}
-					if (Idx < Data.Length && Data[Idx] == '\n')
-					{
-						Idx++;
-					}
-				}
-			}
-			return Result.ToString();
-		}
-
-		/// <summary>
-		/// Formats the current contents of the buffer to a string
-		/// </summary>
-		/// <param name="Buffer">The buffer to read from</param>
-		/// <param name="BufferIdx">The next byte that was read</param>
-		/// <returns>String representation of the buffer</returns>
-		private static string FormatDataAsHexDump(ReadOnlySpan<byte> Buffer, int BufferIdx)
-		{
-			return FormatDataAsHexDump(Buffer.Slice(BufferIdx, Math.Min(Buffer.Length - BufferIdx, 1024)));
-		}
-
-		/// <summary>
-		/// Formats the current contents of the buffer to a string
-		/// </summary>
-		/// <param name="Data">The data to format</param>
-		/// <returns>String representation of the buffer</returns>
-		private static string FormatDataAsHexDump(ReadOnlySpan<byte> Data)
-		{
-			// Format the result
-			StringBuilder Result = new StringBuilder();
-
-			const int RowLength = 16;
-			for (int BaseIdx = 0; BaseIdx < Data.Length; BaseIdx += RowLength)
-			{
-				Result.Append("\n    ");
-				for (int Offset = 0; Offset < RowLength; Offset++)
-				{
-					int Idx = BaseIdx + Offset;
-					if (Idx >= Data.Length)
-					{
-						Result.Append("   ");
-					}
-					else
-					{
-						Result.AppendFormat("{0:x2} ", Data[Idx]);
-					}
-				}
-				Result.Append("   ");
-				for (int Offset = 0; Offset < RowLength; Offset++)
-				{
-					int Idx = BaseIdx + Offset;
-					if (Idx >= Data.Length)
-					{
-						break;
-					}
-					else if (Data[Idx] < 0x20 || Data[Idx] >= 0x7f)
-					{
-						Result.Append('.');
-					}
-					else
-					{
-						Result.Append((char)Data[Idx]);
-					}
-				}
-			}
-			return Result.ToString();
-		}
-
-		/// <summary>
-		/// Read a list of responses from the child process
-		/// </summary>
-		/// <param name="Perforce">The response to read from</param>
-		/// <param name="StatRecordInfo">The type of stat record to parse</param>
-		/// <param name="CancellationToken">Cancellation token for the read</param>
-		/// <returns>Async task</returns>
-		public static async Task<List<PerforceResponse>> ReadResponsesAsync(this IPerforceResponse Perforce, CachedRecordInfo? StatRecordInfo, CancellationToken CancellationToken)
-		{
-			List<PerforceResponse> Responses = new List<PerforceResponse>();
-
-			// Read all the records into a list
-			Task<bool> ReadTask = Perforce.ReadAsync(CancellationToken);
-			while (await ReadTask)
-			{
-				// Start a read to add more data
-				ReadTask = Perforce.ReadAsync(CancellationToken);
-
-				// Check for the whole message not being a marshalled python object, and produce a better response in that scenario
-				ReadOnlyMemory<byte> Data = Perforce.Data;
-				if (Data.Length > 0 && Responses.Count == 0 && Data.Span[0] != '{')
-				{
-					throw new PerforceException("Unexpected response from server:{0}", FormatDataAsString(Data.Span));
-				}
-
-				// Parse the responses from the current buffer
-				int BufferPos = 0;
-				for (; ; )
-				{
-					int NewBufferPos = BufferPos;
-					if (!TryReadResponse(Data.Span, ref NewBufferPos, StatRecordInfo, out PerforceResponse? Response))
-					{
-						break;
-					}
-					if (Response.Error == null || Response.Error.Generic != PerforceGenericCode.Empty)
-					{
-						Responses.Add(Response);
-					}
-					BufferPos = NewBufferPos;
-				}
-
-				// Discard all the data that we've processed
-				Perforce.Discard(BufferPos);
-			}
-
-			// If the stream is complete but we couldn't parse a response from the server, treat it as an error
-			if (Perforce.Data.Length > 0)
-			{
-				throw new PerforceException("Unexpected response from server:{0}", FormatDataAsString(Perforce.Data.Span));
-			}
-
-			return Responses;
-		}
-
-		/// <summary>
-		/// Read a list of responses from the child process
-		/// </summary>
-		/// <param name="Perforce">The Perforce response</param>
-		/// <param name="HandleRecord">Delegate to invoke for each record read</param>
-		/// <param name="CancellationToken">Cancellation token for the read</param>
-		/// <returns>Async task</returns>
-		public static async Task ReadRecordsAsync(this IPerforceResponse Perforce, Action<PerforceRecord> HandleRecord, CancellationToken CancellationToken)
-		{
-			PerforceRecord Record = new PerforceRecord();
-			Record.Rows = new List<KeyValuePair<Utf8String, PerforceValue>>();
-
-			// Read all the records into a list
-			Task<bool> ReadTask = Perforce.ReadAsync(CancellationToken);
-			while (await ReadTask)
-			{
-				// Start a read to add more data
-				ReadTask = Perforce.ReadAsync(CancellationToken);
-				ReadOnlyMemory<byte> Data = Perforce.Data;
-
-				// Parse the responses from the current buffer
-				int BufferPos = 0;
-				for (; ; )
-				{
-					int InitialBufferPos = BufferPos;
-					if(!ReadRecord(Data.Span, ref BufferPos, Record.Rows))
-					{
-						BufferPos = InitialBufferPos;
-						break;
-					}
-					HandleRecord(Record);
-				}
-				Perforce.Discard(BufferPos);
-			}
-
-			// If the stream is complete but we couldn't parse a response from the server, treat it as an error
-			if (Perforce.Data.Length > 0)
-			{
-				throw new PerforceException("Unexpected response from server:{0}", FormatDataAsString(Perforce.Data.Span));
-			}
-		}
-
-		/// <summary>
-		/// Reads from the buffer into a record object
-		/// </summary>
-		/// <param name="Buffer">The buffer to read from</param>
-		/// <param name="BufferPos">Current read position within the buffer</param>
-		/// <param name="Rows">List of rows to read into</param>
-		/// <returns>True if a record could be read; false if more data is required</returns>
-		static bool ReadRecord(ReadOnlySpan<byte> Buffer, ref int BufferPos, List<KeyValuePair<Utf8String, PerforceValue>> Rows)
-		{
-			Rows.Clear();
-
-			// Check we can read the initial record marker
-			if (BufferPos >= Buffer.Length)
-			{
-				return false;
-			}
-			if (Buffer[BufferPos] != '{')
-			{
-				throw new PerforceException("Invalid record start");
-			}
-			BufferPos++;
-
-			// Capture the start of the string
-			int StartPos = BufferPos;
-			for (; ; )
-			{
-				// Check that we've got a string field
-				if (BufferPos >= Buffer.Length)
-				{
-					return false;
-				}
-
-				// If this is the end of the record, break out
-				byte KeyType = Buffer[BufferPos++];
-				if (KeyType == '0')
-				{
-					break;
-				}
-				else if (KeyType != 's')
-				{
-					throw new PerforceException("Unexpected key field type while parsing marshalled output ({0}) - expected 's', got: {1}", (int)KeyType, FormatDataAsHexDump(Buffer.Slice(BufferPos - 1)));
-				}
-
-				// Read the tag
-				Utf8String Key;
-				if (!TryReadString(Buffer, ref BufferPos, out Key))
-				{
-					return false;
-				}
-
-				// Remember the start of the value
-				int ValueOffset = BufferPos;
-
-				// Read the value type
-				byte ValueType;
-				if (!TryReadByte(Buffer, ref BufferPos, out ValueType))
-				{
-					return false;
-				}
-
-				// Parse the appropriate value
-				PerforceValue Value;
-				if (ValueType == 's')
-				{
-					Utf8String String;
-					if (!TryReadString(Buffer, ref BufferPos, out String))
-					{
-						return false;
-					}
-					Value = new PerforceValue(Buffer.Slice(ValueOffset, BufferPos - ValueOffset).ToArray());
-				}
-				else if (ValueType == 'i')
-				{
-					int Integer;
-					if(!TryReadInt(Buffer, ref BufferPos, out Integer))
-					{
-						return false;
-					}
-					Value = new PerforceValue(Buffer.Slice(ValueOffset, BufferPos - ValueOffset).ToArray());
-				}
-				else
-				{
-					throw new PerforceException("Unrecognized value type {0}", ValueType);
-				}
-
-				// Construct the response object with the record
-				Rows.Add(KeyValuePair.Create(Key, Value));
-			}
-			return true;
-		}
-
-		/// <summary>
-		/// Reads a response object from the buffer
-		/// </summary>
-		/// <param name="Buffer">The buffer to read from</param>
-		/// <param name="BufferPos">Current read position within the buffer</param>
-		/// <param name="StatRecordInfo">The type of record expected to parse from the response</param>
-		/// <param name="Response">Receives the response object on success</param>
-		/// <returns>True if a response was read, false if the buffer needs more data</returns>
-		static bool TryReadResponse(ReadOnlySpan<byte> Buffer, ref int BufferPos, CachedRecordInfo? StatRecordInfo, [NotNullWhen(true)] out PerforceResponse? Response)
-		{
-			if (BufferPos + RecordPrefix.Length + 4 > Buffer.Length)
-			{
-				Response = null;
-				return false;
-			}
-
-			ReadOnlySpan<byte> Prefix = Buffer.Slice(BufferPos, RecordPrefix.Length);
-			if (!Prefix.SequenceEqual(RecordPrefix))
-			{
-				throw new PerforceException("Expected 'code' field at the start of record");
-			}
-			BufferPos += Prefix.Length;
-
-			Utf8String Code;
-			if (!TryReadString(Buffer, ref BufferPos, out Code))
-			{
-				Response = null;
-				return false;
-			}
-
-			// Dispatch it to the appropriate handler
-			object? Record;
-			if (Code == ReadOnlyUtf8StringConstants.Stat && StatRecordInfo != null)
-			{
-				if(!TryReadTypedRecord(Buffer, ref BufferPos, Utf8String.Empty, StatRecordInfo, out Record))
-				{
-					Response = null;
-					return false;
-				}
-			}
-			else if (Code == ReadOnlyUtf8StringConstants.Info)
-			{
-				if(!TryReadTypedRecord(Buffer, ref BufferPos, Utf8String.Empty, PerforceReflection.InfoRecordInfo, out Record))
-				{
-					Response = null;
-					return false;
-				}
-			}
-			else if (Code == ReadOnlyUtf8StringConstants.Error)
-			{
-				if (!TryReadTypedRecord(Buffer, ref BufferPos, Utf8String.Empty, PerforceReflection.ErrorRecordInfo, out Record))
-				{
-					Response = null;
-					return false;
-				}
-			}
-			else
-			{
-				throw new PerforceException("Unknown return code for record: {0}", Code);
-			}
-
-			// Skip over the record terminator
-			if(BufferPos >= Buffer.Length || Buffer[BufferPos] != '0')
-			{
-				throw new PerforceException("Unexpected record terminator");
-			}
-			BufferPos++;
-
-			// Create the response
-			Response = new PerforceResponse(Record);
-			return true;
-		}
-
-		/// <summary>
-		/// Parse an individual record from the server.
-		/// </summary>
-		/// <param name="Buffer">The buffer to read from</param>
-		/// <param name="BufferPos">Current read position within the buffer</param>
-		/// <param name="RequiredSuffix">The required suffix for any subobject arrays.</param>
-		/// <param name="RecordInfo">Reflection information for the type being serialized into.</param>
-		/// <param name="Record">Receives the record on success</param>
-		/// <returns>The parsed object.</returns>
-		static bool TryReadTypedRecord(ReadOnlySpan<byte> Buffer, ref int BufferPos, Utf8String RequiredSuffix, CachedRecordInfo RecordInfo, [NotNullWhen(true)] out object? Record)
-		{
-			// Create a bitmask for all the required tags
-			ulong RequiredTagsBitMask = 0;
-
-			// Create the new record
-			object? NewRecord = RecordInfo.CreateInstance();
-			if (NewRecord == null)
-			{
-				throw new InvalidDataException($"Unable to construct record of type {RecordInfo.Type}");
-			}
-
-			// Get the record info, and parse it into the object
-			for(; ;)
-			{
-				// Check that we've got a string field
-				if(BufferPos >= Buffer.Length)
-				{
-					Record = null;
-					return false;
-				}
-
-				// If this is the end of the record, break out
-				byte KeyType = Buffer[BufferPos];
-				if (KeyType == '0')
-				{
-					break;
-				}
-				else if (KeyType != 's')
-				{
-					throw new PerforceException("Unexpected key field type while parsing marshalled output ({0}) - expected 's', got: {1}", (int)KeyType, FormatDataAsHexDump(Buffer.Slice(BufferPos)));
-				}
-
-				// Capture the initial buffer position, in case we have to roll back
-				int StartBufferPos = BufferPos;
-				BufferPos++;
-
-				// Read the tag
-				Utf8String Tag;
-				if(!TryReadString(Buffer, ref BufferPos, out Tag))
-				{
-					Record = null;
-					return false;
-				}
-
-				// Find the start of the array suffix
-				int SuffixIdx = Tag.Length;
-				while (SuffixIdx > 0 && (Tag[SuffixIdx - 1] == (byte)',' || (Tag[SuffixIdx - 1] >= '0' && Tag[SuffixIdx - 1] <= '9')))
-				{
-					SuffixIdx--;
-				}
-
-				// Separate the key into tag and suffix
-				Utf8String Suffix = Tag.Slice(SuffixIdx);
-				Tag = Tag.Slice(0, SuffixIdx);
-
-				// Try to find the matching field
-				CachedTagInfo? TagInfo;
-				if (RecordInfo.NameToInfo.TryGetValue(Tag, out TagInfo))
-				{
-					RequiredTagsBitMask |= TagInfo.RequiredTagBitMask;
-				}
-
-				// Check whether it's a subobject or part of the current object.
-				if (Suffix == RequiredSuffix)
-				{
-					if(!TryReadValue(Buffer, ref BufferPos, NewRecord, TagInfo))
-					{
-						Record = null;
-						return false;
-					}
-				}
-				else if (Suffix.StartsWith(RequiredSuffix) && (RequiredSuffix.Length == 0 || Suffix[RequiredSuffix.Length] == ','))
-				{
-					// Part of a subobject. If this record doesn't have any listed subobject type, skip the field and continue.
-					if (TagInfo != null)
-					{
-						// Get the list field
-						System.Collections.IList? List = (System.Collections.IList?)TagInfo.Field.GetValue(NewRecord);
-						if (List == null)
-						{
-							throw new PerforceException($"Empty list for {TagInfo.Field.Name}");
-						}
-
-						// Check the suffix matches the index of the next element
-						if (!IsCorrectIndex(Suffix, RequiredSuffix, List.Count))
-						{
-							throw new PerforceException("Subobject element received out of order: got {0}", Suffix);
-						}
-
-						// Add it to the list
-						if (!TryReadValue(Buffer, ref BufferPos, NewRecord, TagInfo))
-						{
-							Record = null;
-							return false;
-						}
-					}
-					else if (RecordInfo.SubElementField != null)
-					{
-						// Move back to the start of this tag
-						BufferPos = StartBufferPos;
-
-						// Get the list field
-						System.Collections.IList? List = (System.Collections.IList?)RecordInfo.SubElementField.GetValue(NewRecord);
-						if (List == null)
-						{
-							throw new PerforceException($"Invalid field for {RecordInfo.SubElementField.Name}");
-						}
-
-						// Check the suffix matches the index of the next element
-						if (!IsCorrectIndex(Suffix, RequiredSuffix, List.Count))
-						{
-							throw new PerforceException("Subobject element received out of order: got {0}", Suffix);
-						}
-
-						// Parse the subobject and add it to the list
-						object? SubRecord;
-						if (!TryReadTypedRecord(Buffer, ref BufferPos, Suffix, RecordInfo.SubElementRecordInfo!, out SubRecord))
-						{
-							Record = null;
-							return false;
-						}
-						List.Add(SubRecord);
-					}
-					else
-					{
-						// Just discard the value
-						if (!TryReadValue(Buffer, ref BufferPos, NewRecord, TagInfo))
-						{
-							Record = null;
-							return false;
-						}
-					}
-				}
-				else
-				{
-					// Roll back
-					BufferPos = StartBufferPos;
-					break;
-				}
-			}
-
-			// Make sure we've got all the required tags we need
-			if (RequiredTagsBitMask != RecordInfo.RequiredTagsBitMask)
-			{
-				string MissingTagNames = String.Join(", ", RecordInfo.NameToInfo.Where(x => (RequiredTagsBitMask | x.Value.RequiredTagBitMask) != RequiredTagsBitMask).Select(x => x.Key));
-				throw new PerforceException("Missing '{0}' tag when parsing '{1}'", MissingTagNames, RecordInfo.Type.Name);
-			}
-
-			// Construct the response object with the record
-			Record = NewRecord;
-			return true;
-		}
-
-		/// <summary>
-		/// Reads a value from the input buffer
-		/// </summary>
-		/// <param name="Buffer">The buffer to read from</param>
-		/// <param name="BufferPos">Current read position within the buffer</param>
-		/// <param name="NewRecord">The new record</param>
-		/// <param name="TagInfo">The current tag</param>
-		/// <returns></returns>
-		static bool TryReadValue(ReadOnlySpan<byte> Buffer, ref int BufferPos, object NewRecord, CachedTagInfo? TagInfo)
-		{
-			// Read the value type
-			byte ValueType;
-			if (!TryReadByte(Buffer, ref BufferPos, out ValueType))
-			{
-				return false;
-			}
-
-			// Parse the appropriate value
-			if (ValueType == 's')
-			{
-				Utf8String String;
-				if (!TryReadString(Buffer, ref BufferPos, out String))
-				{
-					return false;
-				}
-				if (TagInfo != null)
-				{
-					TagInfo.SetFromString(NewRecord, String);
-				}
-			}
-			else if (ValueType == 'i')
-			{
-				int Integer;
-				if (!TryReadInt(Buffer, ref BufferPos, out Integer))
-				{
-					return false;
-				}
-				if (TagInfo != null)
-				{
-					TagInfo.SetFromInteger(NewRecord, Integer);
-				}
-			}
-			else
-			{
-				throw new PerforceException("Unrecognized value type {0}", ValueType);
-			}
-
-			return true;
-		}
-
-		/// <summary>
-		/// Attempts to read a single byte from the buffer
-		/// </summary>
-		/// <param name="Buffer">The buffer to read from</param>
-		/// <param name="BufferPos">Current read position within the buffer</param>
-		/// <param name="Value">Receives the byte that was read</param>
-		/// <returns>True if a byte was read from the buffer, false if there was not enough data</returns>
-		static bool TryReadByte(ReadOnlySpan<byte> Buffer, ref int BufferPos, out byte Value)
-		{
-			if (BufferPos >= Buffer.Length)
-			{
-				Value = 0;
-				return false;
-			}
-
-			Value = Buffer[BufferPos];
-			BufferPos++;
-			return true;
-		}
-
-		/// <summary>
-		/// Attempts to read a single int from the buffer
-		/// </summary>
-		/// <param name="Buffer">The buffer to read from</param>
-		/// <param name="BufferPos">Current read position within the buffer</param>
-		/// <param name="Value">Receives the value that was read</param>
-		/// <returns>True if an int was read from the buffer, false if there was not enough data</returns>
-		static bool TryReadInt(ReadOnlySpan<byte> Buffer, ref int BufferPos, out int Value)
-		{
-			if(BufferPos + 4 > Buffer.Length)
-			{
-				Value = 0;
-				return false;
-			}
-
-			Value = Buffer[BufferPos + 0] | (Buffer[BufferPos + 1] << 8) | (Buffer[BufferPos + 2] << 16) | (Buffer[BufferPos + 3] << 24);
-			BufferPos += 4;
-			return true;
-		}
-
-		/// <summary>
-		/// Attempts to read a string with type indicator from the buffer
-		/// </summary>
-		/// <param name="Buffer">The buffer to read from</param>
-		/// <param name="BufferPos">Current read position within the buffer</param>
-		/// <param name="String">Receives the value that was read</param>
-		/// <returns>True if a string was read from the buffer, false if there was not enough data</returns>
-		static bool TryReadStringWithType(ReadOnlySpan<byte> Buffer, ref int BufferPos, out Utf8String String)
-		{
-			byte ValueType;
-			if(!TryReadByte(Buffer, ref BufferPos, out ValueType))
-			{
-				String = new Utf8String();
-				return false;
-			}
-			if (ValueType != 's')
-			{
-				throw new PerforceException("Expected string value, got '{0}'", ValueType);
-			}
-			return TryReadString(Buffer, ref BufferPos, out String);
-		}
-
-		/// <summary>
-		/// Attempts to read a string from the buffer
-		/// </summary>
-		/// <param name="Buffer">The buffer to read from</param>
-		/// <param name="BufferPos">Current read position within the buffer</param>
-		/// <param name="String">Receives the value that was read</param>
-		/// <returns>True if a string was read from the buffer, false if there was not enough data</returns>
-		static bool TryReadString(ReadOnlySpan<byte> Buffer, ref int BufferPos, out Utf8String String)
-		{
-			int Length;
-			if(!TryReadInt(Buffer, ref BufferPos, out Length))
-			{
-				String = new Utf8String();
-				return false;
-			}
-
-			if(BufferPos + Length > Buffer.Length)
-			{
-				String = new Utf8String();
-				return false;
-			}
-
-			String = new Utf8String(Buffer.Slice(BufferPos, Length).ToArray());
-			BufferPos += Length;
-			return true;
-		}
-
-		/// <summary>
-		/// Determines if the given text contains the expected prefix followed by an array index
-		/// </summary>
-		/// <param name="Text">The text to check</param>
-		/// <param name="Prefix">The required prefix</param>
-		/// <param name="Index">The required index</param>
-		/// <returns>True if the index is correct</returns>
-		static bool IsCorrectIndex(Utf8String Text, Utf8String Prefix, int Index)
-		{
-			if (Prefix.Length > 0)
-			{
-				return Text.StartsWith(Prefix) && Text.Length > Prefix.Length && Text[Prefix.Length] == (byte)',' && IsCorrectIndex(Text.Span.Slice(Prefix.Length + 1), Index);
-			}
-			else
-			{
-				return IsCorrectIndex(Text.Span, Index);
-			}
-		}
-
-		/// <summary>
-		/// Determines if the given text matches the expected array index
-		/// </summary>
-		/// <param name="Span">The text to check</param>
-		/// <param name="ExpectedIndex">The expected array index</param>
-		/// <returns>True if the span matches</returns>
-		static bool IsCorrectIndex(ReadOnlySpan<byte> Span, int ExpectedIndex)
-		{
-			int Index;
-			int BytesConsumed;
-			return Utf8Parser.TryParse(Span, out Index, out BytesConsumed) && BytesConsumed == Span.Length && Index == ExpectedIndex;
-		}
-	}
 }
diff --git a/Engine/Source/Programs/Shared/EpicGames.Perforce/PerforceConnection.cs b/Engine/Source/Programs/Shared/EpicGames.Perforce/PerforceConnection.cs
--- a/Engine/Source/Programs/Shared/EpicGames.Perforce/PerforceConnection.cs
+++ b/Engine/Source/Programs/Shared/EpicGames.Perforce/PerforceConnection.cs
@@ -16,51 +16,6 @@ using System.Threading.Tasks;
 
 namespace EpicGames.Perforce
 {
-	/// <summary>
-	/// Base interface for Perforce clients
-	/// </summary>
-	public interface IPerforceConnection
-	{
-		/// <summary>
-		/// Logger for this connection
-		/// </summary>
-		ILogger Logger { get; }
-
-		/// <summary>
-		/// Runs a Perforce command
-		/// </summary>
-		/// <param name="Command">The command name</param>
-		/// <param name="Arguments">Arguments for the command</param>
-		/// <param name="InputData">Input data to be passed to the command</param>
-		/// <returns>Response object</returns>
-		Task<IPerforceResponse> RunAsync(string Command, IReadOnlyList<string> Arguments, byte[]? InputData);
-
-		/// <summary>
-		/// Execute the 'login' command
-		/// </summary>
-		/// <param name="Password">Password to use to login</param>
-		/// <param name="CancellationToken">Token used to cancel the operation</param>
-		/// <returns>Response from the server</returns>
-		Task LoginAsync(string Password, CancellationToken CancellationToken = default);
-
-		/// <summary>
-		/// Sets an environment variable
-		/// </summary>
-		/// <param name="Name">Name of the variable to set</param>
-		/// <param name="Value">Value for the variable</param>
-		/// <param name="CancellationToken">Token used to cancel the operation</param>
-		/// <returns>Response from the server</returns>
-		Task SetAsync(string Name, string Value, CancellationToken CancellationToken = default);
-
-		/// <summary>
-		/// Gets the setting of a Perforce variable
-		/// </summary>
-		/// <param name="Name">Name of the variable to get</param>
-		/// <param name="CancellationToken">Cancellation token for the request</param>
-		/// <returns>Value of the variable</returns>
-		Task<string?> TryGetSettingAsync(string Name, CancellationToken CancellationToken = default);
-	}
-
 	/// <summary>
 	/// Stores settings for communicating with a Perforce server.
 	/// </summary>
@@ -213,9 +168,9 @@ namespace EpicGames.Perforce
 		}
 
 		/// <inheritdoc/>
-		public Task<IPerforceResponse> RunAsync(string Command, IReadOnlyList<string> Arguments, byte[]? InputData)
+		public Task<IPerforceOutput> CommandAsync(string Command, IReadOnlyList<string> Arguments, byte[]? InputData)
 		{
-			return Task.FromResult<IPerforceResponse>(new PerforceChildProcess(Command, Arguments, InputData, GetGlobalArguments(), Logger));
+			return Task.FromResult<IPerforceOutput>(new PerforceChildProcess(Command, Arguments, InputData, GetGlobalArguments(), Logger));
 		}
 
 		#region p4 login
@@ -358,10 +313,9 @@ namespace EpicGames.Perforce
 		/// <returns>List of objects returned by the server</returns>
 		public static async Task<List<PerforceResponse>> CommandAsync(this IPerforceConnection Perforce, string Command, IReadOnlyList<string> Arguments, byte[]? InputData, Type? StatRecordType, CancellationToken CancellationToken = default)
 		{
-			using (IPerforceResponse Response = await Perforce.RunAsync(Command, Arguments, InputData))
+			await using (IPerforceOutput Response = await Perforce.CommandAsync(Command, Arguments, InputData))
 			{
-				CachedRecordInfo? StatRecordInfo = (StatRecordType == null) ? null : PerforceReflection.GetCachedRecordInfo(StatRecordType);
-				return await Response.ReadResponsesAsync(StatRecordInfo, CancellationToken);
+				return await Response.ReadResponsesAsync(StatRecordType, CancellationToken);
 			}
 		}
 
@@ -377,7 +331,7 @@ namespace EpicGames.Perforce
 		/// <returns>List of objects returned by the server</returns>
 		public static async Task RecordCommandAsync(this IPerforceConnection Perforce, string Command, IReadOnlyList<string> Arguments, byte[]? InputData, Action<PerforceRecord> HandleRecord, CancellationToken CancellationToken = default)
 		{
-			using (IPerforceResponse Response = await Perforce.RunAsync(Command, Arguments, InputData))
+			await using (IPerforceOutput Response = await Perforce.CommandAsync(Command, Arguments, InputData))
 			{
 				await Response.ReadRecordsAsync(HandleRecord, CancellationToken);
 			}
@@ -1747,6 +1701,18 @@ namespace EpicGames.Perforce
 
 		#region p4 fstat
 
+		/// <summary>
+		/// Execute the 'fstat' command
+		/// </summary>
+		/// <param name="Connection">Connection to the Perforce server</param>
+		/// <param name="FileSpecs">List of file specifications to query</param>
+		/// <param name="CancellationToken">Token used to cancel the operation</param>
+		/// <returns>List of responses from the server</returns>
+		public static async Task<List<FStatRecord>> FStatAsync(this IPerforceConnection Connection, FileSpecList FileSpecs, CancellationToken CancellationToken = default)
+		{
+			return (await TryFStatAsync(Connection, FStatOptions.None, FileSpecs, CancellationToken)).Data;
+		}
+
 		/// <summary>
 		/// Execute the 'fstat' command
 		/// </summary>
@@ -1798,7 +1764,7 @@ namespace EpicGames.Perforce
 		/// <returns>List of responses from the server</returns>
 		public static Task<PerforceResponseList<FStatRecord>> TryFStatAsync(this IPerforceConnection Connection, int MaxFiles, FStatOptions Options, FileSpecList FileSpecs, CancellationToken CancellationToken = default)
 		{
-			return TryFStatAsync(Connection, -1, -1, null, MaxFiles, Options, FileSpecs, CancellationToken);
+			return TryFStatAsync(Connection, -1, -1, null, null, MaxFiles, Options, FileSpecs, CancellationToken);
 		}
 
 		/// <summary>
@@ -1808,14 +1774,15 @@ namespace EpicGames.Perforce
 		/// <param name="AfterChangeNumber">Return only files affected after the given changelist number.</param>
 		/// <param name="OnlyChangeNumber">Return only files affected by the given changelist number.</param>
 		/// <param name="Filter">List only those files that match the criteria specified.</param>
+		/// <param name="Fields">Fields to return in the output</param>
 		/// <param name="MaxFiles">Produce fstat output for only the first max files.</param>
 		/// <param name="Options">Options for the command</param>
 		/// <param name="FileSpecs">List of file specifications to query</param>
 		/// <param name="CancellationToken">Token used to cancel the operation</param>
 		/// <returns>List of responses from the server</returns>
-		public static async Task<List<FStatRecord>> FStatAsync(this IPerforceConnection Connection, int AfterChangeNumber, int OnlyChangeNumber, string? Filter, int MaxFiles, FStatOptions Options, FileSpecList FileSpecs, CancellationToken CancellationToken = default)
+		public static async Task<List<FStatRecord>> FStatAsync(this IPerforceConnection Connection, int AfterChangeNumber, int OnlyChangeNumber, string? Filter, string? Fields, int MaxFiles, FStatOptions Options, FileSpecList FileSpecs, CancellationToken CancellationToken = default)
 		{
-			return (await TryFStatAsync(Connection, AfterChangeNumber, OnlyChangeNumber, Filter, MaxFiles, Options, FileSpecs, CancellationToken)).Data;
+			return (await TryFStatAsync(Connection, AfterChangeNumber, OnlyChangeNumber, Filter, Fields, MaxFiles, Options, FileSpecs, CancellationToken)).Data;
 		}
 
 		/// <summary>
@@ -1825,12 +1792,13 @@ namespace EpicGames.Perforce
 		/// <param name="AfterChangeNumber">Return only files affected after the given changelist number.</param>
 		/// <param name="OnlyChangeNumber">Return only files affected by the given changelist number.</param>
 		/// <param name="Filter">List only those files that match the criteria specified.</param>
+		/// <param name="Fields">Fields to return in the output</param>
 		/// <param name="MaxFiles">Produce fstat output for only the first max files.</param>
 		/// <param name="Options">Options for the command</param>
 		/// <param name="FileSpecs">List of file specifications to query</param>
 		/// <param name="CancellationToken">Token used to cancel the operation</param>
 		/// <returns>List of responses from the server</returns>
-		public static async Task<PerforceResponseList<FStatRecord>> TryFStatAsync(this IPerforceConnection Connection, int AfterChangeNumber, int OnlyChangeNumber, string? Filter, int MaxFiles, FStatOptions Options, FileSpecList FileSpecs, CancellationToken CancellationToken = default)
+		public static async Task<PerforceResponseList<FStatRecord>> TryFStatAsync(this IPerforceConnection Connection, int AfterChangeNumber, int OnlyChangeNumber, string? Filter, string? Fields, int MaxFiles, FStatOptions Options, FileSpecList FileSpecs, CancellationToken CancellationToken = default)
 		{
 			// Build the argument list
 			List<string> Arguments = new List<string>();
@@ -1847,6 +1815,11 @@ namespace EpicGames.Perforce
 				Arguments.Add("-F");
 				Arguments.Add(Filter);
 			}
+			if (Fields != null)
+			{
+				Arguments.Add("-T");
+				Arguments.Add(Fields);
+			}
 			if((Options & FStatOptions.ReportDepotSyntax) != 0)
 			{
 				Arguments.Add("-L");
