UnrealBuildTool, AutomationTool:
Print path to log file, print a timestamp at the top of each log file, don't print (wrong) callstack location.

diff --git a/Engine/Source/Programs/AutomationTool/AutomationUtils/Automation.cs b/Engine/Source/Programs/AutomationTool/AutomationUtils/Automation.cs
--- a/Engine/Source/Programs/AutomationTool/AutomationUtils/Automation.cs
+++ b/Engine/Source/Programs/AutomationTool/AutomationUtils/Automation.cs
@@ -99,6 +99,12 @@ namespace AutomationTool
 				StartupListener.CopyTo(LogTraceListener);
 				Trace.Listeners.Remove(StartupListener);
 				
+				Log.TraceInformation($"Log location: {LogUtils.LogFileName}");
+				if (!String.Equals(LogUtils.FinalLogFileName, LogUtils.LogFileName))
+				{
+					Log.TraceInformation($"Final log location: {LogUtils.FinalLogFileName}");	
+				}
+				
 				// Initialize UBT
 				if (!UnrealBuildTool.PlatformExports.Initialize())
 				{
diff --git a/Engine/Source/Programs/AutomationTool/Program.cs b/Engine/Source/Programs/AutomationTool/Program.cs
--- a/Engine/Source/Programs/AutomationTool/Program.cs
+++ b/Engine/Source/Programs/AutomationTool/Program.cs
@@ -308,7 +308,8 @@ namespace AutomationToolDriver
 		public static int Main(string[] Arguments)
 		{
 			// Initialize the log system, buffering the output until we can create the log file
-			Trace.Listeners.Add(StartupListener);
+			Log.IncludeCallingMethod = false;
+			Log.AddTraceListener(StartupListener);
 
 			// Populate AutomationToolCommandLine and CommandsToExecute
 			try
diff --git a/Engine/Source/Programs/Shared/EpicGames.Core/Log.cs b/Engine/Source/Programs/Shared/EpicGames.Core/Log.cs
--- a/Engine/Source/Programs/Shared/EpicGames.Core/Log.cs
+++ b/Engine/Source/Programs/Shared/EpicGames.Core/Log.cs
@@ -205,6 +205,15 @@ namespace EpicGames.Core
 			set => DefaultLogger.ColorConsoleOutput = value;
 		}
 
+		/// <summary>
+		/// When true, a timestamp will be written to the log file when the first listener is added
+		/// </summary>
+		public static bool IncludeStartingTimestamp
+		{
+			get => DefaultLogger.IncludeStartingTimestamp;
+			set => DefaultLogger.IncludeStartingTimestamp = value;
+		}
+
 		/// <summary>
 		/// Path to the log file being written to. May be null.
 		/// </summary>
@@ -711,6 +720,15 @@ namespace EpicGames.Core
 			get; set;
 		}
 
+		/// <summary>
+		/// When true, a timestamp will be written to the log file when the first listener is added
+		/// </summary>
+		public bool IncludeStartingTimestamp
+		{
+			get; set;
+		}
+		private bool IncludeStartingTimestampWritten = false;
+
 		/// <summary>
 		/// Path to the log file being written to. May be null.
 		/// </summary>
@@ -763,6 +781,7 @@ namespace EpicGames.Core
 			IncludeCallingMethod = true;
 			IncludeCallingMethodForConsole = false;
 			ColorConsoleOutput = true;
+			IncludeStartingTimestamp = true;
 		}
 
 		/// <summary>
@@ -770,6 +789,7 @@ namespace EpicGames.Core
 		/// </summary>
 		/// <param name="OutputFile">The file to write to</param>
 		/// <returns>The created trace listener</returns>
+		[MethodImplAttribute(MethodImplOptions.NoInlining)]
 		public TextWriterTraceListener AddFileWriter(string Name, FileReference OutputFile)
 		{
 			try
@@ -780,6 +800,7 @@ namespace EpicGames.Core
 				lock (SyncObject)
 				{
 					Trace.Listeners.Add(LogTraceListener);
+					WriteInitialTimestamp();
 				}
 				return LogTraceListener;
 			}
@@ -793,6 +814,7 @@ namespace EpicGames.Core
 		/// Adds a <see cref="TraceListener"/> to the collection in a safe manner.
 		/// </summary>
 		/// <param name="TraceListener">The <see cref="TraceListener"/> to add.</param>
+		[MethodImplAttribute(MethodImplOptions.NoInlining)]
 		public void AddTraceListener(TraceListener TraceListener)
 		{
 			lock (SyncObject)
@@ -800,10 +822,25 @@ namespace EpicGames.Core
 				if (!Trace.Listeners.Contains(TraceListener))
 				{
 					Trace.Listeners.Add(TraceListener);
+					WriteInitialTimestamp();
 				}
 			}
 		}
 
+		/// <summary>
+		/// Write a timestamp to the log, once. To be called when a new listener is added.
+		/// </summary>
+		[MethodImplAttribute(MethodImplOptions.NoInlining)]
+		private void WriteInitialTimestamp()
+		{
+			if (IncludeStartingTimestamp && !IncludeStartingTimestampWritten)
+			{
+				DateTime Now = DateTime.Now;
+				WriteLinePrivate(2, false, LogEventType.Log, LogFormatOptions.NoConsoleOutput, $"Log started at {DateTime.Now} ({DateTime.Now.ToUniversalTime():yyyy-MM-ddTHH\\:mm\\:ssZ})");
+				IncludeStartingTimestampWritten = true;
+			}
+		}
+
 		/// <summary>
 		/// Removes a <see cref="TraceListener"/> from the collection in a safe manner.
 		/// </summary>
diff --git a/Engine/Source/Programs/UnrealBuildTool/Modes/BuildMode.cs b/Engine/Source/Programs/UnrealBuildTool/Modes/BuildMode.cs
--- a/Engine/Source/Programs/UnrealBuildTool/Modes/BuildMode.cs
+++ b/Engine/Source/Programs/UnrealBuildTool/Modes/BuildMode.cs
@@ -129,6 +129,7 @@ namespace UnrealBuildTool
 					LogFile = LogFile.ChangeExtension(null) + "_" + LogSuffix + LogFile.GetExtension();
 				}
 
+				Log.TraceInformation($"Log will be written to {LogFile}");
 				TextWriterTraceListener LogTraceListener = Log.AddFileWriter("DefaultLogTraceListener", LogFile);
 				StartupListener.CopyTo(LogTraceListener);
 			}
diff --git a/Engine/Source/Programs/UnrealBuildTool/Modes/GenerateProjectFilesMode.cs b/Engine/Source/Programs/UnrealBuildTool/Modes/GenerateProjectFilesMode.cs
--- a/Engine/Source/Programs/UnrealBuildTool/Modes/GenerateProjectFilesMode.cs
+++ b/Engine/Source/Programs/UnrealBuildTool/Modes/GenerateProjectFilesMode.cs
@@ -73,6 +73,7 @@ namespace UnrealBuildTool
 			FileReference LogFile = FileReference.Combine(UnrealBuildTool.EngineProgramSavedDirectory, "UnrealBuildTool", "Log_GPF.txt");
 			// find the StartupTraceListener in the listeners that was added super early on
 			StartupTraceListener StartupListener = Trace.Listeners.OfType<StartupTraceListener>().First();
+			Log.TraceInformation($"Log will be written to {LogFile}");
 			TextWriterTraceListener LogTraceListener = Log.AddFileWriter("DefaultLogTraceListener", LogFile);
 			StartupListener.CopyTo(LogTraceListener);
 			Trace.Listeners.Remove(StartupListener);
@@ -162,7 +163,7 @@ namespace UnrealBuildTool
 			}
 			if (BadPlatformNames.Count > 0)
 			{
-				Log.TraceInformationOnce("\nSome Platforms were skipped due to invalid SDK setup: {0}.\nSee the log file for detailed information ({1})\n\n", string.Join(", ", BadPlatformNames), Log.OutputFile);
+				Log.TraceInformationOnce("\nSome Platforms were skipped due to invalid SDK setup: {0}.\nSee the log file for detailed information\n\n", string.Join(", ", BadPlatformNames));
 				Log.TraceInformation("");
 			}
 			Log.TraceLog("---   SDK INFO END   ---");
diff --git a/Engine/Source/Programs/UnrealBuildTool/UnrealBuildTool.cs b/Engine/Source/Programs/UnrealBuildTool/UnrealBuildTool.cs
--- a/Engine/Source/Programs/UnrealBuildTool/UnrealBuildTool.cs
+++ b/Engine/Source/Programs/UnrealBuildTool/UnrealBuildTool.cs
@@ -449,6 +449,7 @@ namespace UnrealBuildTool
 				Log.OutputLevel = Options.LogOutputLevel;
 				Log.IncludeTimestamps = Options.bLogTimestamps;
 				Log.IncludeProgramNameWithSeverityPrefix = Options.bLogFromMsBuild;
+				Log.IncludeCallingMethod = false;
 
 				// Always start capturing logs as early as possible to later copy to a log file if the ToolMode desires it (we have to start capturing before we get the ToolModeOptions below)
 				StartupTraceListener StartupTrace = new StartupTraceListener();
@@ -460,6 +461,7 @@ namespace UnrealBuildTool
 				// Add the log writer if requested. When building a target, we'll create the writer for the default log file later.
 				if(Options.LogFileName != null)
 				{
+					Log.TraceInformation($"Log will be written to {Options.LogFileName}");
 					Log.AddFileWriter("LogTraceListener", Options.LogFileName);
 				}
 
