Horde: Fix DateTime CB serialization tests.

diff --git a/Engine/Source/Programs/Shared/EpicGames.Serialization.Tests/EngineTests.cs b/Engine/Source/Programs/Shared/EpicGames.Serialization.Tests/EngineTests.cs
--- a/Engine/Source/Programs/Shared/EpicGames.Serialization.Tests/EngineTests.cs
+++ b/Engine/Source/Programs/Shared/EpicGames.Serialization.Tests/EngineTests.cs
@@ -33,6 +33,11 @@ namespace EpicGames.Serialization.Tests
 				return new CbFieldAccessors(new T(), IsType, x => AsTypeWithDefault(x, new T()), (x, y) => AsTypeWithDefault(x, (T)y));
 			}
 
+			public static CbFieldAccessors FromStruct<T>(T Default, Func<CbField, bool> IsType, Func<CbField, T, object> AsTypeWithDefault) where T : struct
+			{
+				return new CbFieldAccessors(Default, IsType, x => AsTypeWithDefault(x, new T()), (x, y) => AsTypeWithDefault(x, (T)y));
+			}
+
 			public static CbFieldAccessors FromStruct<T>(Func<CbField, bool> IsType, Func<CbField, T, object> AsTypeWithDefault, Func<T, T, bool> Comparer) where T : struct
 			{
 				return new CbFieldAccessors(new T(), IsType, x => AsTypeWithDefault(x, new T()), (x, y) => AsTypeWithDefault(x, (T)y)) { Comparer = (x, y) => Comparer((T)x, (T)y) };
@@ -57,7 +62,7 @@ namespace EpicGames.Serialization.Tests
 			[CbFieldType.BinaryAttachment] = new CbFieldAccessors(IoHash.Zero, x => x.IsBinaryAttachment(), x => x.AsBinaryAttachment(), (x, Default) => x.AsBinaryAttachment((IoHash)Default)),
 			[CbFieldType.Hash] = new CbFieldAccessors(IoHash.Zero, x => x.IsHash(), x => x.AsHash(), (x, Default) => x.AsHash((IoHash)Default)),
 			[CbFieldType.Uuid] = CbFieldAccessors.FromStruct<Guid>(x => x.IsUuid(), (x, y) => x.AsUuid(y)),
-			[CbFieldType.DateTime] = CbFieldAccessors.FromStruct<DateTime>(x => x.IsDateTime(), (x, y) => x.AsDateTime(y)),
+			[CbFieldType.DateTime] = CbFieldAccessors.FromStruct<DateTime>(new DateTime(0, DateTimeKind.Utc), x => x.IsDateTime(), (x, y) => x.AsDateTime(y)),
 			[CbFieldType.TimeSpan] = CbFieldAccessors.FromStruct<TimeSpan>(x => x.IsTimeSpan(), (x, y) => x.AsTimeSpan(y)),
 		};
 
@@ -935,7 +940,7 @@ namespace EpicGames.Serialization.Tests
 			TestField(CbFieldType.DateTime, new byte[]{0, 0, 0, 0, 0, 0, 0, 0});
 
 			// Test CbField(DateTime, 0x1020_3040_5060_7080)
-			TestField(CbFieldType.DateTime, new byte[]{0x10, 0x20, 0x30, 0x40, 0x50, 0x60, 0x70, 0x80}, new DateTime(0x1020_3040_5060_7080L));
+			TestField(CbFieldType.DateTime, new byte[]{0x10, 0x20, 0x30, 0x40, 0x50, 0x60, 0x70, 0x80}, new DateTime(0x1020_3040_5060_7080L, DateTimeKind.Utc));
 
 			// Test CbField(DateTime, Zero) as FDateTime
 			{
@@ -948,7 +953,7 @@ namespace EpicGames.Serialization.Tests
 			{
 				CbField DefaultField = new CbField();
 				TestFieldError(CbFieldType.DateTime, DefaultField, CbFieldError.TypeError);
-				DateTime DefaultValue = new DateTime(0x1020_3040_5060_7080L);
+				DateTime DefaultValue = new DateTime(0x1020_3040_5060_7080L, DateTimeKind.Utc);
 				Assert.AreEqual(DefaultField.AsDateTime(DefaultValue), DefaultValue);
 			}
 		}
diff --git a/Engine/Source/Programs/Shared/EpicGames.Serialization/CompactBinary.cs b/Engine/Source/Programs/Shared/EpicGames.Serialization/CompactBinary.cs
--- a/Engine/Source/Programs/Shared/EpicGames.Serialization/CompactBinary.cs
+++ b/Engine/Source/Programs/Shared/EpicGames.Serialization/CompactBinary.cs
@@ -1124,7 +1124,7 @@ namespace EpicGames.Serialization
 		/// <returns></returns>
 		public DateTime AsDateTime()
 		{
-			return AsDateTime(default);
+			return AsDateTime(new DateTime(0, DateTimeKind.Utc));
 		}
 
 		/// <summary>
