Horde: Add a more specific error when preflight files are not in the current stream.

diff --git a/Engine/Source/Programs/Horde/HordeAgent/Commands/Workspace/SyncCommand.cs b/Engine/Source/Programs/Horde/HordeAgent/Commands/Workspace/SyncCommand.cs
--- a/Engine/Source/Programs/Horde/HordeAgent/Commands/Workspace/SyncCommand.cs
+++ b/Engine/Source/Programs/Horde/HordeAgent/Commands/Workspace/SyncCommand.cs
@@ -59,7 +59,7 @@ namespace HordeAgent.Commands.Workspace
 
 			if (PreflightChange != -1)
 			{
-				await Repo.UnshelveAsync(PerforceClient, PreflightChange, CancellationToken.None);
+				await Repo.UnshelveAsync(PerforceClient, StreamName, PreflightChange, CancellationToken.None);
 			}
 		}
 	}
diff --git a/Engine/Source/Programs/Horde/HordeAgent/Utility/WorkspaceInfo.cs b/Engine/Source/Programs/Horde/HordeAgent/Utility/WorkspaceInfo.cs
--- a/Engine/Source/Programs/Horde/HordeAgent/Utility/WorkspaceInfo.cs
+++ b/Engine/Source/Programs/Horde/HordeAgent/Utility/WorkspaceInfo.cs
@@ -342,7 +342,7 @@ namespace HordeAgent.Utility
 			// if we're running a preflight
 			if (PreflightChange > 0)
 			{
-				await Repository.UnshelveAsync(PerforceClient, PreflightChange, CancellationToken);
+				await Repository.UnshelveAsync(PerforceClient, StreamName, PreflightChange, CancellationToken);
 			}
 		}
 
diff --git a/Engine/Source/Programs/Shared/EpicGames.Perforce.Managed/ManagedWorkspace.cs b/Engine/Source/Programs/Shared/EpicGames.Perforce.Managed/ManagedWorkspace.cs
--- a/Engine/Source/Programs/Shared/EpicGames.Perforce.Managed/ManagedWorkspace.cs
+++ b/Engine/Source/Programs/Shared/EpicGames.Perforce.Managed/ManagedWorkspace.cs
@@ -887,7 +887,7 @@ namespace EpicGames.Perforce.Managed
 		/// Replays the effects of unshelving a changelist, but clobbering files in the workspace rather than actually unshelving them (to prevent problems with multiple machines locking them)
 		/// </summary>
 		/// <returns>Async task</returns>
-		public async Task UnshelveAsync(PerforceClientConnection Perforce, int UnshelveChangelist, CancellationToken CancellationToken)
+		public async Task UnshelveAsync(PerforceClientConnection Perforce, string StreamName, int UnshelveChangelist, CancellationToken CancellationToken)
 		{
 			// Need to mark those files as dirty - update the workspace with those files 
 			// Delete is fine, but need to flag anything added
@@ -915,6 +915,11 @@ namespace EpicGames.Perforce.Managed
 			List<WhereRecord> WriteFiles = new List<WhereRecord>();
 			for(int FileIdx = 0; FileIdx < LastRecord.Files.Count; FileIdx++)
 			{
+				if (FileIdx >= WhereRecords.Count)
+				{
+					throw new PerforceException("Unable to get location of {File} within {Stream}. Check the correct stream is specified.", LastRecord.Files[FileIdx].DepotFile, StreamName);
+				}
+
 				PerforceResponse<WhereRecord> Response = WhereRecords[FileIdx];
 				if (!Response.Succeeded)
 				{
