ManagedProcess: Flush stdout & err on process exit when merging output streams (Mac & Linux)

#rnx

diff --git a/Engine/Source/Programs/Shared/EpicGames.Core/ManagedProcess.cs b/Engine/Source/Programs/Shared/EpicGames.Core/ManagedProcess.cs
--- a/Engine/Source/Programs/Shared/EpicGames.Core/ManagedProcess.cs
+++ b/Engine/Source/Programs/Shared/EpicGames.Core/ManagedProcess.cs
@@ -735,11 +735,15 @@ namespace EpicGames.Core
 					{
 						if (FrameworkStdOutThread != null)
 						{
+							FrameworkProcess.StandardOutput.BaseStream.Flush();
+							FrameworkProcess.StandardOutput.Close();
 							FrameworkStdOutThread.Join();
 							FrameworkStdOutThread = null;
 						}
 						if (FrameworkStdErrThread != null)
 						{
+							FrameworkProcess.StandardError.BaseStream.Flush();
+							FrameworkProcess.StandardError.Close();
 							FrameworkStdErrThread.Join();
 							FrameworkStdErrThread = null;
 						}
@@ -751,7 +755,6 @@ namespace EpicGames.Core
 						}
 					}
 				};
-
 				// AnonymousPipes block reading even if the stream has been fully read, until the writer pipe handle is closed.
 				FrameworkMergedStdWriter = new AnonymousPipeServerStream(PipeDirection.Out);
 				StdOut = new AnonymousPipeClientStream(PipeDirection.In, FrameworkMergedStdWriter.ClientSafePipeHandle);
@@ -761,7 +764,6 @@ namespace EpicGames.Core
 				StdErrText = StdOutText;
 			}
 
-			// Ensure threads start before Exited is called, if the process terminates instantly
 			lock (this)
 			{
 				FrameworkStartTime = DateTime.Now;
@@ -769,9 +771,13 @@ namespace EpicGames.Core
 				if ((ManagedFlags & ManagedProcessFlags.MergeOutputPipes) != 0)
 				{
 					FrameworkStdOutThread = new Thread(() => CopyPipe(FrameworkProcess.StandardOutput.BaseStream, FrameworkMergedStdWriter!));
+					FrameworkStdOutThread.Name = $"ManagedProcess {FrameworkProcess.Id} Merge StdOut";
+					FrameworkStdOutThread.IsBackground = true;
 					FrameworkStdOutThread.Start();
 
 					FrameworkStdErrThread = new Thread(() => CopyPipe(FrameworkProcess.StandardError.BaseStream, FrameworkMergedStdWriter!));
+					FrameworkStdErrThread.Name = $"ManagedProcess {FrameworkProcess.Id} Merge StdErr";
+					FrameworkStdErrThread.IsBackground = true;
 					FrameworkStdErrThread.Start();
 				}
 			}
