GitSync: Fix stack overflow when disposing of child process objects.

diff --git a/Engine/Source/Programs/Shared/EpicGames.Perforce/PerforceChildProcess.cs b/Engine/Source/Programs/Shared/EpicGames.Perforce/PerforceChildProcess.cs
--- a/Engine/Source/Programs/Shared/EpicGames.Perforce/PerforceChildProcess.cs
+++ b/Engine/Source/Programs/Shared/EpicGames.Perforce/PerforceChildProcess.cs
@@ -104,7 +104,7 @@ namespace EpicGames.Perforce
 		/// <inheritdoc/>
 		public ValueTask DisposeAsync()
 		{
-			DisposeAsync();
+			Dispose();
 			return new ValueTask(Task.CompletedTask);
 		}
 
