Double the size of the P4 response buffer until a limit, then increase it by 32MB chunks. Also improve error reporting for corrupt buffers.

diff --git a/Engine/Source/Programs/Shared/EpicGames.Perforce/IPerforceOutput.cs b/Engine/Source/Programs/Shared/EpicGames.Perforce/IPerforceOutput.cs
--- a/Engine/Source/Programs/Shared/EpicGames.Perforce/IPerforceOutput.cs
+++ b/Engine/Source/Programs/Shared/EpicGames.Perforce/IPerforceOutput.cs
@@ -167,13 +167,14 @@ namespace EpicGames.Perforce
 			List<PerforceResponse> Responses = new List<PerforceResponse>();
 
 			// Read all the records into a list
+			long ParsedLen = 0;
 			while (await Perforce.ReadAsync(CancellationToken))
 			{
 				// Check for the whole message not being a marshalled python object, and produce a better response in that scenario
 				ReadOnlyMemory<byte> Data = Perforce.Data;
 				if (Data.Length > 0 && Responses.Count == 0 && Data.Span[0] != '{')
 				{
-					throw new PerforceException("Unexpected response from server:{0}", FormatDataAsString(Data.Span));
+					throw new PerforceException("Unexpected response from server (expected '{'):{0}", FormatDataAsString(Data.Span));
 				}
 
 				// Parse the responses from the current buffer
@@ -194,12 +195,13 @@ namespace EpicGames.Perforce
 
 				// Discard all the data that we've processed
 				Perforce.Discard(BufferPos);
+				ParsedLen += BufferPos;
 			}
 
 			// If the stream is complete but we couldn't parse a response from the server, treat it as an error
 			if (Perforce.Data.Length > 0)
 			{
-				throw new PerforceException("Unexpected response from server:{0}", FormatDataAsString(Perforce.Data.Span));
+				throw new PerforceException("Unparsable data at offset {0}:{1}", ParsedLen, FormatDataAsString(Perforce.Data.Span));
 			}
 
 			return Responses;
@@ -243,7 +245,7 @@ namespace EpicGames.Perforce
 			// If the stream is complete but we couldn't parse a response from the server, treat it as an error
 			if (Perforce.Data.Length > 0)
 			{
-				throw new PerforceException("Unexpected response from server:{0}", FormatDataAsString(Perforce.Data.Span));
+				throw new PerforceException("Unexpected trailing response data from server:{0}", FormatDataAsString(Perforce.Data.Span));
 			}
 		}
 
diff --git a/Engine/Source/Programs/Shared/EpicGames.Perforce/PerforceChildProcess.cs b/Engine/Source/Programs/Shared/EpicGames.Perforce/PerforceChildProcess.cs
--- a/Engine/Source/Programs/Shared/EpicGames.Perforce/PerforceChildProcess.cs
+++ b/Engine/Source/Programs/Shared/EpicGames.Perforce/PerforceChildProcess.cs
@@ -134,7 +134,7 @@ namespace EpicGames.Perforce
 			// Update the buffer contents
 			if (BufferEnd == Buffer.Length)
 			{
-				Array.Resize(ref Buffer, Buffer.Length + 65536);
+				Array.Resize(ref Buffer, Math.Min(Buffer.Length + (32 * 1024 * 1024), Buffer.Length * 2));
 			}
 
 			// Try to read more data
