Additional data for malformed Perforce response data.

diff --git a/Engine/Source/Programs/Shared/EpicGames.Perforce/IPerforceOutput.cs b/Engine/Source/Programs/Shared/EpicGames.Perforce/IPerforceOutput.cs
--- a/Engine/Source/Programs/Shared/EpicGames.Perforce/IPerforceOutput.cs
+++ b/Engine/Source/Programs/Shared/EpicGames.Perforce/IPerforceOutput.cs
@@ -168,6 +168,7 @@ namespace EpicGames.Perforce
 
 			// Read all the records into a list
 			long ParsedLen = 0;
+			long MaxParsedLen = 0;
 			while (await Perforce.ReadAsync(CancellationToken))
 			{
 				// Check for the whole message not being a marshalled python object, and produce a better response in that scenario
@@ -184,6 +185,7 @@ namespace EpicGames.Perforce
 					int NewBufferPos = BufferPos;
 					if (!TryReadResponse(Data.Span, ref NewBufferPos, StatRecordInfo, out PerforceResponse? Response))
 					{
+						MaxParsedLen = ParsedLen + NewBufferPos;
 						break;
 					}
 					if (Response.Error == null || Response.Error.Generic != PerforceGenericCode.Empty)
@@ -201,7 +203,8 @@ namespace EpicGames.Perforce
 			// If the stream is complete but we couldn't parse a response from the server, treat it as an error
 			if (Perforce.Data.Length > 0)
 			{
-				throw new PerforceException("Unparsable data at offset {0}:{1}", ParsedLen, FormatDataAsString(Perforce.Data.Span));
+				string HexDump = FormatDataAsString(Perforce.Data.Span.Slice((int)(MaxParsedLen - ParsedLen)));
+				throw new PerforceException("Unparsable data at offset {0}+{1}:{2}", ParsedLen, MaxParsedLen - ParsedLen, HexDump);
 			}
 
 			return Responses;
