More debug output for corrupt Perforce response data.

diff --git a/Engine/Source/Programs/Shared/EpicGames.Perforce/IPerforceOutput.cs b/Engine/Source/Programs/Shared/EpicGames.Perforce/IPerforceOutput.cs
--- a/Engine/Source/Programs/Shared/EpicGames.Perforce/IPerforceOutput.cs
+++ b/Engine/Source/Programs/Shared/EpicGames.Perforce/IPerforceOutput.cs
@@ -203,8 +203,10 @@ namespace EpicGames.Perforce
 			// If the stream is complete but we couldn't parse a response from the server, treat it as an error
 			if (Perforce.Data.Length > 0)
 			{
-				string HexDump = FormatDataAsString(Perforce.Data.Span.Slice((int)(MaxParsedLen - ParsedLen)));
-				throw new PerforceException("Unparsable data at offset {0}+{1}:{2}", ParsedLen, MaxParsedLen - ParsedLen, HexDump);
+				long DumpOffset = Math.Max(MaxParsedLen - 32, ParsedLen);
+				string StrDump = FormatDataAsString(Perforce.Data.Span.Slice((int)(DumpOffset - ParsedLen)));
+				string HexDump = FormatDataAsHexDump(Perforce.Data.Span.Slice((int)(DumpOffset - ParsedLen)));
+				throw new PerforceException("Unparsable data at offset {0}+{1}/{2}.\nString data from offset {3}:{4}\nHex data from offset {3}:{5}", ParsedLen, MaxParsedLen - ParsedLen, ParsedLen + Perforce.Data.Length, DumpOffset, StrDump, HexDump);
 			}
 
 			return Responses;
