Partially revert CL 17374535. Closing stdout/stderr before joining with background threads means that we lose whatever is still in the pipe. Causing exceptions on GitSync service due to corrupt output.

diff --git a/Engine/Source/Programs/Shared/EpicGames.Core/ManagedProcess.cs b/Engine/Source/Programs/Shared/EpicGames.Core/ManagedProcess.cs
--- a/Engine/Source/Programs/Shared/EpicGames.Core/ManagedProcess.cs
+++ b/Engine/Source/Programs/Shared/EpicGames.Core/ManagedProcess.cs
@@ -735,15 +735,11 @@ namespace EpicGames.Core
 					{
 						if (FrameworkStdOutThread != null)
 						{
-							FrameworkProcess.StandardOutput.BaseStream.Flush();
-							FrameworkProcess.StandardOutput.Close();
 							FrameworkStdOutThread.Join();
 							FrameworkStdOutThread = null;
 						}
 						if (FrameworkStdErrThread != null)
 						{
-							FrameworkProcess.StandardError.BaseStream.Flush();
-							FrameworkProcess.StandardError.Close();
 							FrameworkStdErrThread.Join();
 							FrameworkStdErrThread = null;
 						}
