EpicGames.Perforce: Completely fill read buffer before returning to the caller. Linux returns 64k chunks of data and requires multiple calls to fill the buffer. In situations where we want to parse very large records (such as large p4 describe results), this results in attempting to parse the buffer many times more than necessary in order to succeed.

diff --git a/Engine/Source/Programs/Shared/EpicGames.Perforce/PerforceChildProcess.cs b/Engine/Source/Programs/Shared/EpicGames.Perforce/PerforceChildProcess.cs
--- a/Engine/Source/Programs/Shared/EpicGames.Perforce/PerforceChildProcess.cs
+++ b/Engine/Source/Programs/Shared/EpicGames.Perforce/PerforceChildProcess.cs
@@ -138,14 +138,17 @@ namespace EpicGames.Perforce
 			}
 
 			// Try to read more data
-			int Count = await ChildProcess!.ReadAsync(Buffer, BufferEnd, Buffer.Length - BufferEnd, CancellationToken);
-			if (Count == 0)
+			int PrevBufferEnd = BufferEnd;
+			while (BufferEnd < Buffer.Length)
 			{
-				return false;
+				int Count = await ChildProcess!.ReadAsync(Buffer, BufferEnd, Buffer.Length - BufferEnd, CancellationToken);
+				if (Count == 0)
+				{
+					break;
+				}
+				BufferEnd += Count;
 			}
-
-			BufferEnd += Count;
-			return true;
+			return BufferEnd > PrevBufferEnd;
 		}
 
 		/// <inheritdoc/>
