EpicGames.Perforce Fix argument for -x when using child process. Delegate to the connection implementation to write arguments to a response file.

diff --git a/Engine/Source/Programs/Shared/EpicGames.Perforce/IPerforceConnection.cs b/Engine/Source/Programs/Shared/EpicGames.Perforce/IPerforceConnection.cs
--- a/Engine/Source/Programs/Shared/EpicGames.Perforce/IPerforceConnection.cs
+++ b/Engine/Source/Programs/Shared/EpicGames.Perforce/IPerforceConnection.cs
@@ -24,9 +24,10 @@ namespace EpicGames.Perforce
 		/// </summary>
 		/// <param name="Command">The command name</param>
 		/// <param name="Arguments">Arguments for the command</param>
+		/// <param name="FileArguments">File arguments (may be put into a response file)</param>
 		/// <param name="InputData">Input data to be passed to the command</param>
 		/// <returns>Response object</returns>
-		Task<IPerforceOutput> CommandAsync(string Command, IReadOnlyList<string> Arguments, byte[]? InputData);
+		Task<IPerforceOutput> CommandAsync(string Command, IReadOnlyList<string> Arguments, IReadOnlyList<string>? FileArguments, byte[]? InputData);
 
 		/// <summary>
 		/// Execute the 'login' command
diff --git a/Engine/Source/Programs/Shared/EpicGames.Perforce/PerforceChildProcess.cs b/Engine/Source/Programs/Shared/EpicGames.Perforce/PerforceChildProcess.cs
--- a/Engine/Source/Programs/Shared/EpicGames.Perforce/PerforceChildProcess.cs
+++ b/Engine/Source/Programs/Shared/EpicGames.Perforce/PerforceChildProcess.cs
@@ -53,21 +53,33 @@ namespace EpicGames.Perforce
 		/// <inheritdoc/>
 		public ReadOnlyMemory<byte> Data => Buffer.AsMemory(0, BufferEnd);
 
+		/// <summary>
+		/// Temp file containing file arguments
+		/// </summary>
+		string? TempFileName;
+
 		/// <summary>
 		/// Constructor
 		/// </summary>
 		/// <param name="Command"></param>
 		/// <param name="Arguments">Command line arguments</param>
+		/// <param name="FileArguments">File arguments, which may be placed in a response file</param>
 		/// <param name="InputData">Input data to pass to the child process</param>
 		/// <param name="GlobalOptions"></param>
 		/// <param name="Logger">Logging device</param>
-		public PerforceChildProcess(string Command, IReadOnlyList<string> Arguments, byte[]? InputData, IReadOnlyList<string> GlobalOptions, ILogger Logger)
+		public PerforceChildProcess(string Command, IReadOnlyList<string> Arguments, IReadOnlyList<string>? FileArguments, byte[]? InputData, IReadOnlyList<string> GlobalOptions, ILogger Logger)
 		{
 			string PerforceFileName = GetExecutable();
 
 			List<string> FullArguments = new List<string>();
 			FullArguments.Add("-G");
 			FullArguments.AddRange(GlobalOptions);
+			if (FileArguments != null)
+			{
+				TempFileName = Path.GetTempFileName();
+				File.WriteAllLines(TempFileName, FileArguments);
+				FullArguments.Add($"-x{TempFileName}");
+			}
 			FullArguments.Add(Command);
 			FullArguments.AddRange(Arguments);
 
@@ -126,6 +138,11 @@ namespace EpicGames.Perforce
 				Scope.Dispose();
 				Scope = null!;
 			}
+			if (TempFileName != null)
+			{
+				try { File.Delete(TempFileName); } catch { }
+				TempFileName = null;
+			}
 		}
 
 		/// <inheritdoc/>
diff --git a/Engine/Source/Programs/Shared/EpicGames.Perforce/PerforceConnection.cs b/Engine/Source/Programs/Shared/EpicGames.Perforce/PerforceConnection.cs
--- a/Engine/Source/Programs/Shared/EpicGames.Perforce/PerforceConnection.cs
+++ b/Engine/Source/Programs/Shared/EpicGames.Perforce/PerforceConnection.cs
@@ -168,9 +168,9 @@ namespace EpicGames.Perforce
 		}
 
 		/// <inheritdoc/>
-		public Task<IPerforceOutput> CommandAsync(string Command, IReadOnlyList<string> Arguments, byte[]? InputData)
+		public Task<IPerforceOutput> CommandAsync(string Command, IReadOnlyList<string> Arguments, IReadOnlyList<string>? FileArguments, byte[]? InputData)
 		{
-			return Task.FromResult<IPerforceOutput>(new PerforceChildProcess(Command, Arguments, InputData, GetGlobalArguments(), Logger));
+			return Task.FromResult<IPerforceOutput>(new PerforceChildProcess(Command, Arguments, FileArguments, InputData, GetGlobalArguments(), Logger));
 		}
 
 		#region p4 login
@@ -256,7 +256,7 @@ namespace EpicGames.Perforce
 			List<string> Arguments = new List<string>();
 			Arguments.Append($"{Name}={Value}");
 
-			using (PerforceChildProcess ChildProcess = new PerforceChildProcess("set", Arguments, null, GetGlobalArguments(), Logger))
+			using (PerforceChildProcess ChildProcess = new PerforceChildProcess("set", Arguments, null, null, GetGlobalArguments(), Logger))
 			{
 				return await ChildProcess.TryReadToEndAsync(CancellationToken);
 			}
@@ -270,7 +270,7 @@ namespace EpicGames.Perforce
 		/// <returns>Value of the variable</returns>
 		public async Task<string?> TryGetSettingAsync(string Name, CancellationToken CancellationToken = default)
 		{
-			using (PerforceChildProcess ChildProcess = new PerforceChildProcess("set", new List<string> { $"{Name}=" }, null, GetGlobalArguments(), Logger))
+			using (PerforceChildProcess ChildProcess = new PerforceChildProcess("set", new List<string> { $"{Name}=" }, null, null, GetGlobalArguments(), Logger))
 			{
 				Tuple<bool, string> Response = await ChildProcess.TryReadToEndAsync(CancellationToken);
 				if (Response.Item1)
@@ -311,9 +311,25 @@ namespace EpicGames.Perforce
 		/// <param name="StatRecordType">The type of records to return for "stat" responses</param>
 		/// <param name="CancellationToken">Token used to cancel the operation</param>
 		/// <returns>List of objects returned by the server</returns>
-		public static async Task<List<PerforceResponse>> CommandAsync(this IPerforceConnection Perforce, string Command, IReadOnlyList<string> Arguments, byte[]? InputData, Type? StatRecordType, CancellationToken CancellationToken = default)
+		public static Task<List<PerforceResponse>> CommandAsync(this IPerforceConnection Perforce, string Command, IReadOnlyList<string> Arguments, byte[]? InputData, Type? StatRecordType, CancellationToken CancellationToken = default)
 		{
-			await using (IPerforceOutput Response = await Perforce.CommandAsync(Command, Arguments, InputData))
+			return CommandAsync(Perforce, Command, Arguments, InputData, StatRecordType, CancellationToken);
+		}
+
+		/// <summary>
+		/// Execute a command and parse the response
+		/// </summary>
+		/// <param name="Perforce">The Perforce connection</param>
+		/// <param name="Command">Command to execute</param>
+		/// <param name="Arguments">Arguments for the command</param>
+		/// <param name="FileArguments">File arguments for the command</param>
+		/// <param name="InputData">Input data to pass to Perforce</param>
+		/// <param name="StatRecordType">The type of records to return for "stat" responses</param>
+		/// <param name="CancellationToken">Token used to cancel the operation</param>
+		/// <returns>List of objects returned by the server</returns>
+		public static async Task<List<PerforceResponse>> CommandAsync(this IPerforceConnection Perforce, string Command, IReadOnlyList<string> Arguments, IReadOnlyList<string>? FileArguments, byte[]? InputData, Type? StatRecordType, CancellationToken CancellationToken = default)
+		{
+			await using (IPerforceOutput Response = await Perforce.CommandAsync(Command, Arguments, FileArguments, InputData))
 			{
 				return await Response.ReadResponsesAsync(StatRecordType, CancellationToken);
 			}
@@ -331,7 +347,7 @@ namespace EpicGames.Perforce
 		/// <returns>List of objects returned by the server</returns>
 		public static async Task RecordCommandAsync(this IPerforceConnection Perforce, string Command, IReadOnlyList<string> Arguments, byte[]? InputData, Action<PerforceRecord> HandleRecord, CancellationToken CancellationToken = default)
 		{
-			await using (IPerforceOutput Response = await Perforce.CommandAsync(Command, Arguments, InputData))
+			await using (IPerforceOutput Response = await Perforce.CommandAsync(Command, Arguments, null, InputData))
 			{
 				await Response.ReadRecordsAsync(HandleRecord, CancellationToken);
 			}
@@ -378,12 +394,13 @@ namespace EpicGames.Perforce
 		/// <param name="Connection">Connection to the Perforce server</param>
 		/// <param name="Command">The command to execute</param>
 		/// <param name="Arguments">Arguments for the command</param>
+		/// <param name="FileArguments">Arguments which can be passed into a -x argument</param>
 		/// <param name="InputData">Input data to pass to Perforce</param>
 		/// <param name="CancellationToken">Token used to cancel the operation</param>
 		/// <returns>List of objects returned by the server</returns>
-		public static async Task<PerforceResponseList<T>> CommandAsync<T>(this IPerforceConnection Connection, string Command, List<string> Arguments, byte[]? InputData, CancellationToken CancellationToken = default) where T : class
+		public static async Task<PerforceResponseList<T>> CommandAsync<T>(this IPerforceConnection Connection, string Command, IReadOnlyList<string> Arguments, IReadOnlyList<string>? FileArguments, byte[]? InputData, CancellationToken CancellationToken = default) where T : class
 		{
-			List<PerforceResponse> Responses = await Connection.CommandAsync(Command, Arguments, InputData, typeof(T), CancellationToken);
+			List<PerforceResponse> Responses = await Connection.CommandAsync(Command, Arguments, FileArguments, InputData, typeof(T), CancellationToken);
 
 			PerforceResponseList<T> TypedResponses = new PerforceResponseList<T>();
 			foreach (PerforceResponse Response in Responses)
@@ -393,6 +410,20 @@ namespace EpicGames.Perforce
 			return TypedResponses;
 		}
 
+		/// <summary>
+		/// Execute a command and parse the response
+		/// </summary>
+		/// <param name="Connection">Connection to the Perforce server</param>
+		/// <param name="Command">The command to execute</param>
+		/// <param name="Arguments">Arguments for the command</param>
+		/// <param name="InputData">Input data to pass to Perforce</param>
+		/// <param name="CancellationToken">Token used to cancel the operation</param>
+		/// <returns>List of objects returned by the server</returns>
+		public static Task<PerforceResponseList<T>> CommandAsync<T>(this IPerforceConnection Connection, string Command, List<string> Arguments, byte[]? InputData, CancellationToken CancellationToken = default) where T : class
+		{
+			return CommandAsync<T>(Connection, Command, Arguments, null, InputData, CancellationToken);
+		}
+
 		/// <summary>
 		/// Execute a command and parse the response
 		/// </summary>
@@ -2961,25 +2992,7 @@ namespace EpicGames.Perforce
 								Arguments.AppendFormat(",minsize={0}", MinSize);
 							}
 						}*/
-			if (FileSpecs.List.Count < 10)
-			{
-				Arguments.AddRange(FileSpecs.List);
-				return await CommandAsync<T>(Connection, "sync", Arguments, null, CancellationToken);
-			}
-			else
-			{
-				string TempFileName = Path.GetTempFileName();
-				try
-				{
-					await File.WriteAllLinesAsync(TempFileName, FileSpecs.List);
-					Arguments.Insert(0, $"-x\"{TempFileName}\" ");
-					return await CommandAsync<T>(Connection, "sync", Arguments, null, CancellationToken);
-				}
-				finally
-				{
-					File.Delete(TempFileName);
-				}
-			}
+			return await CommandAsync<T>(Connection, "sync", Arguments, FileSpecs.List, null, CancellationToken);
 		}
 
 		#endregion
