EpicGames.Perforce: Fix infinite loop in CommandAsync().

diff --git a/Engine/Source/Programs/Shared/EpicGames.Perforce/PerforceConnection.cs b/Engine/Source/Programs/Shared/EpicGames.Perforce/PerforceConnection.cs
--- a/Engine/Source/Programs/Shared/EpicGames.Perforce/PerforceConnection.cs
+++ b/Engine/Source/Programs/Shared/EpicGames.Perforce/PerforceConnection.cs
@@ -313,7 +313,7 @@ namespace EpicGames.Perforce
 		/// <returns>List of objects returned by the server</returns>
 		public static Task<List<PerforceResponse>> CommandAsync(this IPerforceConnection Perforce, string Command, IReadOnlyList<string> Arguments, byte[]? InputData, Type? StatRecordType, CancellationToken CancellationToken = default)
 		{
-			return CommandAsync(Perforce, Command, Arguments, InputData, StatRecordType, CancellationToken);
+			return CommandAsync(Perforce, Command, Arguments, null, InputData, StatRecordType, CancellationToken);
 		}
 
 		/// <summary>
