Fixing Perforce Fetching user info failing.

user command must have a space after the -o before the user name, restore the space and quoting the user name

diff --git a/Engine/Source/Programs/Shared/EpicGames.Perforce/PerforceConnection.cs b/Engine/Source/Programs/Shared/EpicGames.Perforce/PerforceConnection.cs
--- a/Engine/Source/Programs/Shared/EpicGames.Perforce/PerforceConnection.cs
+++ b/Engine/Source/Programs/Shared/EpicGames.Perforce/PerforceConnection.cs
@@ -3089,7 +3089,8 @@ namespace EpicGames.Perforce
 		public static Task<PerforceResponseList<UserRecord>> TryGetUserAsync(this IPerforceConnection Connection, string UserName, CancellationToken CancellationToken = default)
 		{
 			List<string> Arguments = new List<string>();
-			Arguments.Add($"-o{UserName}");
+			Arguments.Add("-o");
+			Arguments.Add(UserName);
 			return CommandAsync<UserRecord>(Connection, "user", Arguments, null, CancellationToken);
 		}
 
