Horde: Add CbObjectAttachment and CbBinaryAttachment classes to reference typed data by class. We need to distinguish between the two types rather than relying solely on a hash value.

diff --git a/Engine/Source/Programs/Shared/EpicGames.Horde/Compute/ComputeTask.cs b/Engine/Source/Programs/Shared/EpicGames.Horde/Compute/ComputeTask.cs
--- a/Engine/Source/Programs/Shared/EpicGames.Horde/Compute/ComputeTask.cs
+++ b/Engine/Source/Programs/Shared/EpicGames.Horde/Compute/ComputeTask.cs
@@ -41,13 +41,13 @@ namespace EpicGames.Horde.Compute
 		/// Hash of a <see cref="DirectoryTree"/> object encoded as a CbObject and stored in the CAS
 		/// </summary>
 		[CbField("s")]
-		public IoHash SandboxHash { get; set; }
+		public CbObjectAttachment SandboxHash { get; set; }
 
 		/// <summary>
 		/// Requirements for the agent to execute the work
 		/// </summary>
 		[CbField("r")]
-		public IoHash RequirementsHash { get; set; } = IoHash.Zero;
+		public CbObjectAttachment RequirementsHash { get; set; } = CbObjectAttachment.Zero;
 
 		/// <summary>
 		/// List of output paths to be captured on completion of the action. These may be files or directories.
@@ -69,7 +69,7 @@ namespace EpicGames.Horde.Compute
 		/// <param name="Arguments">Arguments for the executable to run</param>
 		/// <param name="WorkingDirectory">Working directory for execution</param>
 		/// <param name="SandboxHash">Hash of the sandbox</param>
-		public ComputeTask(Utf8String Executable, List<Utf8String> Arguments, Utf8String WorkingDirectory, IoHash SandboxHash)
+		public ComputeTask(Utf8String Executable, List<Utf8String> Arguments, Utf8String WorkingDirectory, CbObjectAttachment SandboxHash)
 		{
 			this.Executable = Executable;
 			this.Arguments = Arguments;
diff --git a/Engine/Source/Programs/Shared/EpicGames.Horde/Compute/ComputeTaskResult.cs b/Engine/Source/Programs/Shared/EpicGames.Horde/Compute/ComputeTaskResult.cs
--- a/Engine/Source/Programs/Shared/EpicGames.Horde/Compute/ComputeTaskResult.cs
+++ b/Engine/Source/Programs/Shared/EpicGames.Horde/Compute/ComputeTaskResult.cs
@@ -23,19 +23,19 @@ namespace EpicGames.Horde.Compute
 		/// Hash of the stdout data
 		/// </summary>
 		[CbField("so")]
-		public IoHash? StdOutHash { get; set; }
+		public CbBinaryAttachment? StdOutHash { get; set; }
 
 		/// <summary>
 		/// Hash of the stderr data
 		/// </summary>
 		[CbField("se")]
-		public IoHash? StdErrHash { get; set; }
+		public CbBinaryAttachment? StdErrHash { get; set; }
 
 		/// <summary>
 		/// Hash of the output <see cref="DirectoryTree"/> stored in the CAS. Only items matching the output paths in the action will be included.
 		/// </summary>
 		[CbField("o")]
-		public IoHash? OutputHash { get; set; }
+		public CbObjectAttachment? OutputHash { get; set; }
 
 		/// <summary>
 		/// Constructor
diff --git a/Engine/Source/Programs/Shared/EpicGames.Horde/Compute/DirectoryTree.cs b/Engine/Source/Programs/Shared/EpicGames.Horde/Compute/DirectoryTree.cs
--- a/Engine/Source/Programs/Shared/EpicGames.Horde/Compute/DirectoryTree.cs
+++ b/Engine/Source/Programs/Shared/EpicGames.Horde/Compute/DirectoryTree.cs
@@ -41,7 +41,7 @@ namespace EpicGames.Horde.Compute
 		/// Hash of bulk data for the file stored in CAS.
 		/// </summary>
 		[CbField("h")]
-		public IoHash Hash { get; set; }
+		public CbBinaryAttachment Hash { get; set; }
 
 		/// <summary>
 		/// Size of the file
@@ -65,7 +65,7 @@ namespace EpicGames.Horde.Compute
 		/// <summary>
 		/// Constructor
 		/// </summary>
-		public FileNode(Utf8String Name, IoHash Hash, long Size, int Attributes)
+		public FileNode(Utf8String Name, CbBinaryAttachment Hash, long Size, int Attributes)
 		{
 			this.Name = Name;
 			this.Hash = Hash;
@@ -89,7 +89,7 @@ namespace EpicGames.Horde.Compute
 		/// Hash of a <see cref="CbObject"/> encoded <see cref="DirectoryTree"/> stored in CAS.
 		/// </summary>
 		[CbField("h")]
-		public IoHash Hash { get; set; }
+		public CbObjectAttachment Hash { get; set; }
 
 		/// <summary>
 		/// Private constructor for serialization
@@ -101,7 +101,7 @@ namespace EpicGames.Horde.Compute
 		/// <summary>
 		/// Constructor
 		/// </summary>
-		public DirectoryNode(Utf8String Name, IoHash Hash)
+		public DirectoryNode(Utf8String Name, CbObjectAttachment Hash)
 		{
 			this.Name = Name;
 			this.Hash = Hash;
diff --git a/Engine/Source/Programs/Shared/EpicGames.Serialization/CompactBinary.cs b/Engine/Source/Programs/Shared/EpicGames.Serialization/CompactBinary.cs
--- a/Engine/Source/Programs/Shared/EpicGames.Serialization/CompactBinary.cs
+++ b/Engine/Source/Programs/Shared/EpicGames.Serialization/CompactBinary.cs
@@ -101,7 +101,7 @@ namespace EpicGames.Serialization
 		/// Boolean false value. Payload is empty. 
 		/// </summary>
 		BoolFalse = 0x0c,
-		
+
 		/// <summary>
 		/// Boolean true value. Payload is empty. 
 		/// </summary>
@@ -188,6 +188,94 @@ namespace EpicGames.Serialization
 		HasFieldName = 0x80,
 	}
 
+	/// <summary>
+	/// A binary attachment, referenced by <see cref="IoHash"/>
+	/// </summary>
+	[DebuggerDisplay("{Hash}")]
+	public struct CbBinaryAttachment
+	{
+		/// <summary>
+		/// Attachment with a hash of zero
+		/// </summary>
+		public static CbBinaryAttachment Zero { get; } = new CbBinaryAttachment(IoHash.Zero);
+
+		/// <summary>
+		/// Hash of the referenced object
+		/// </summary>
+		public IoHash Hash;
+
+		/// <summary>
+		/// Constructor
+		/// </summary>
+		/// <param name="Hash">Hash of the referenced object</param>
+		public CbBinaryAttachment(IoHash Hash)
+		{
+			this.Hash = Hash;
+		}
+
+		/// <summary>
+		/// Convert a hash to a binary attachment 
+		/// </summary>
+		/// <param name="Hash">The attachment to convert</param>
+		public static implicit operator CbBinaryAttachment(IoHash Hash)
+		{
+			return new CbBinaryAttachment(Hash);
+		}
+
+		/// <summary>
+		/// Use a binary attachment as a hash
+		/// </summary>
+		/// <param name="Attachment">The attachment to convert</param>
+		public static implicit operator IoHash(CbBinaryAttachment Attachment)
+		{
+			return Attachment.Hash;
+		}
+	}
+
+	/// <summary>
+	/// An object attachment, referenced by <see cref="IoHash"/>
+	/// </summary>
+	[DebuggerDisplay("{Hash}")]
+	public struct CbObjectAttachment
+	{
+		/// <summary>
+		/// Attachment with a hash of zero
+		/// </summary>
+		public static CbObjectAttachment Zero { get; } = new CbObjectAttachment(IoHash.Zero);
+
+		/// <summary>
+		/// Hash of the referenced object
+		/// </summary>
+		public IoHash Hash;
+
+		/// <summary>
+		/// Constructor
+		/// </summary>
+		/// <param name="Hash">Hash of the referenced object</param>
+		public CbObjectAttachment(IoHash Hash)
+		{
+			this.Hash = Hash;
+		}
+
+		/// <summary>
+		/// Use an object attachment as a hash
+		/// </summary>
+		/// <param name="Attachment">The attachment to convert</param>
+		public static implicit operator CbObjectAttachment(IoHash Hash)
+		{
+			return new CbObjectAttachment(Hash);
+		}
+
+		/// <summary>
+		/// Use an object attachment as a hash
+		/// </summary>
+		/// <param name="Attachment">The attachment to convert</param>
+		public static implicit operator IoHash(CbObjectAttachment Attachment)
+		{
+			return Attachment.Hash;
+		}
+	}
+
 	/// <summary>
 	/// Methods that operate on <see cref="CbFieldType"/>.
 	/// </summary>
@@ -647,9 +735,11 @@ namespace EpicGames.Serialization
 					case CbFieldType.BoolTrue:
 						return true;
 					case CbFieldType.ObjectAttachment:
+						return AsObjectAttachment();
 					case CbFieldType.BinaryAttachment:
+						return AsBinaryAttachment();
 					case CbFieldType.Hash:
-						return AsAttachment();
+						return AsHash();
 					case CbFieldType.Uuid:
 						return AsUuid();
 					case CbFieldType.DateTime:
@@ -976,14 +1066,14 @@ namespace EpicGames.Serialization
 		/// Access the field as a hash referencing an object attachment. Returns the provided default on error.
 		/// </summary>
 		/// <returns>Value of the field</returns>
-		public IoHash AsObjectAttachment() => AsObjectAttachment(IoHash.Zero);
+		public CbObjectAttachment AsObjectAttachment() => AsObjectAttachment(CbObjectAttachment.Zero);
 
 		/// <summary>
 		/// Access the field as a hash referencing an object attachment. Returns the provided default on error.
 		/// </summary>
 		/// <param name="Default">Default value</param>
 		/// <returns>Value of the field</returns>
-		public IoHash AsObjectAttachment(IoHash Default)
+		public CbObjectAttachment AsObjectAttachment(CbObjectAttachment Default)
 		{
 			if (CbFieldUtils.IsObjectAttachment(TypeWithFlags))
 			{
@@ -1002,14 +1092,14 @@ namespace EpicGames.Serialization
 		/// </summary>
 		/// <param name="Default">Default value</param>
 		/// <returns>Value of the field</returns>
-		public IoHash AsBinaryAttachment() => AsBinaryAttachment(IoHash.Zero);
+		public CbBinaryAttachment AsBinaryAttachment() => AsBinaryAttachment(CbBinaryAttachment.Zero);
 
 		/// <summary>
 		/// Access the field as a hash referencing a binary attachment. Returns the provided default on error.
 		/// </summary>
 		/// <param name="Default">Default value</param>
 		/// <returns>Value of the field</returns>
-		public IoHash AsBinaryAttachment(IoHash Default)
+		public CbBinaryAttachment AsBinaryAttachment(CbBinaryAttachment Default)
 		{
 			if (CbFieldUtils.IsBinaryAttachment(TypeWithFlags))
 			{
diff --git a/Engine/Source/Programs/Shared/EpicGames.Serialization/Converters/CbDefaultConverter.cs b/Engine/Source/Programs/Shared/EpicGames.Serialization/Converters/CbDefaultConverter.cs
--- a/Engine/Source/Programs/Shared/EpicGames.Serialization/Converters/CbDefaultConverter.cs
+++ b/Engine/Source/Programs/Shared/EpicGames.Serialization/Converters/CbDefaultConverter.cs
@@ -114,6 +114,8 @@ namespace EpicGames.Serialization.Converters
 			GetPodConverterInfo(x => x.AsInt64(), (w, v) => WriteInt64Value(w, v), (w, n, v) => WriteInt64(w, n, v), null),
 			GetPodConverterInfo(x => x.AsString(), (w, v) => w.WriteStringValue(v), (w, n, v) => w.WriteString(n, v), null),
 			GetPodConverterInfo(x => x.AsHash(), (w, v) => w.WriteHashValue(v), (w, n, v) => w.WriteHash(n, v), null),
+			GetPodConverterInfo(x => x.AsObjectAttachment(), (w, v) => w.WriteObjectAttachmentValue(v.Hash), (w, n, v) => w.WriteObjectAttachment(n, v.Hash), null),
+			GetPodConverterInfo(x => x.AsBinaryAttachment(), (w, v) => w.WriteBinaryAttachmentValue(v.Hash), (w, n, v) => w.WriteBinaryAttachment(n, v.Hash), null),
 			GetPodConverterInfo(x => x.AsDateTime(), (w, v) => w.WriteDateTimeValue(v), (w, n, v) => w.WriteDateTime(n, v), null),
 			GetPodConverterInfo(x => ReadString(x), (w, v) => WriteStringValue(w, v), (w, n, v) => WriteString(w, n, v), SkipIfNullOrZero)
 		});
