Fix Perforce hexdump output not being capped.

diff --git a/Engine/Source/Programs/Shared/EpicGames.Perforce/IPerforceOutput.cs b/Engine/Source/Programs/Shared/EpicGames.Perforce/IPerforceOutput.cs
--- a/Engine/Source/Programs/Shared/EpicGames.Perforce/IPerforceOutput.cs
+++ b/Engine/Source/Programs/Shared/EpicGames.Perforce/IPerforceOutput.cs
@@ -204,8 +204,9 @@ namespace EpicGames.Perforce
 			if (Perforce.Data.Length > 0)
 			{
 				long DumpOffset = Math.Max(MaxParsedLen - 32, ParsedLen);
-				string StrDump = FormatDataAsString(Perforce.Data.Span.Slice((int)(DumpOffset - ParsedLen)));
-				string HexDump = FormatDataAsHexDump(Perforce.Data.Span.Slice((int)(DumpOffset - ParsedLen)));
+				int SliceOffset = (int)(DumpOffset - ParsedLen);
+				string StrDump = FormatDataAsString(Perforce.Data.Span.Slice(SliceOffset));
+				string HexDump = FormatDataAsHexDump(Perforce.Data.Span.Slice(SliceOffset, Math.Min(1024, Perforce.Data.Length - SliceOffset)));
 				throw new PerforceException("Unparsable data at offset {0}+{1}/{2}.\nString data from offset {3}:{4}\nHex data from offset {3}:{5}", ParsedLen, MaxParsedLen - ParsedLen, ParsedLen + Perforce.Data.Length, DumpOffset, StrDump, HexDump);
 			}
 
