EpicGames.Core ManagedProcess:

On systems with multiple process groups, distribute launched processes explicitly between groups to improve overall utiliization.

diff --git a/Engine/Source/Programs/Shared/EpicGames.Core/ManagedProcess.cs b/Engine/Source/Programs/Shared/EpicGames.Core/ManagedProcess.cs
--- a/Engine/Source/Programs/Shared/EpicGames.Core/ManagedProcess.cs
+++ b/Engine/Source/Programs/Shared/EpicGames.Core/ManagedProcess.cs
@@ -309,6 +309,25 @@ namespace EpicGames.Core
 			out System.Runtime.InteropServices.ComTypes.FILETIME lpKernelTime,
 			out System.Runtime.InteropServices.ComTypes.FILETIME lpUserTime);
 
+		[DllImport("kernel32.dll")]
+		static extern UInt16 GetActiveProcessorGroupCount();
+
+		[DllImport("kernel32.dll")]
+		static extern UInt32 GetActiveProcessorCount(UInt16 GroupNumber);
+		
+		[StructLayout(LayoutKind.Sequential)]
+		class GROUP_AFFINITY
+		{
+			public UInt64 Mask;
+			public UInt16 Group;
+			public UInt16 Reserved0;
+			public UInt16 Reserved1;
+			public UInt16 Reserved2;
+		}
+
+		[DllImport("kernel32.dll", SetLastError = true)]
+		static extern int SetThreadGroupAffinity(IntPtr hThread, GROUP_AFFINITY GroupAffinity, GROUP_AFFINITY? PreviousGroupAffinity);
+
 		/// <summary>
 		/// Converts FILETIME to DateTime.
 		/// </summary>
@@ -463,6 +482,13 @@ namespace EpicGames.Core
 			}
 		}
 
+		static readonly int ProcessorGroupCount = GetActiveProcessorGroupCount();
+
+		/// <summary>
+		/// A counter used to decide which processor group the next process should be assigned to
+		/// </summary>
+		static int ProcessorGroupCounter = 0;
+
 		/// <summary>
 		/// Spawns a new managed process using Win32 native functions. 
 		/// </summary>
@@ -637,6 +663,32 @@ namespace EpicGames.Core
 						}
 					}
 
+					// On systems with more than one processor group (more than one CPU socket, or more than 64 cores), it is possible
+					// that processes launched from here may be scheduled by the operating system in a way that impedes overall throughput.
+					//
+					// From https://docs.microsoft.com/en-us/windows/win32/procthread/processor-groups
+					// > The operating system initially assigns each process to a single group in a round-robin manner across the groups in the system
+					//
+					// When UBT launches a process here, it is not uncommon for more than one may be created (for example: clang and conhost)
+					// If the number of processes created is the same as the number of processor groups in the system, this can lead
+					// to high workload processes (like clang) predominantly assigned to one processor group, and lower workload processes
+					// (like conhost) to the other.
+					//
+					// To reduce the chance of pathological process scheduling, we will explicitly distribute processes to each process
+					// group. Doing so has been observed to reduce overall compile times for large builds by as much as 10%.
+					if (ProcessorGroupCount > 1)
+					{
+						ushort ProcessorGroup = (ushort)(Interlocked.Increment(ref ProcessorGroupCounter) % ProcessorGroupCount);
+
+						uint GroupProcessorCount = GetActiveProcessorCount(ProcessorGroup);
+
+						GROUP_AFFINITY GroupAffinity = new GROUP_AFFINITY();
+						GroupAffinity.Mask = ~0ul >> (int)(64 - GroupProcessorCount);
+						GroupAffinity.Group = ProcessorGroup;
+						SetThreadGroupAffinity(ProcessInfo.hThread, GroupAffinity, null);
+					}
+					
+
 					// Allow the thread to start running
 					if (ResumeThread(ProcessInfo.hThread) == -1)
 					{
