Cap the maximum output from P4 error messages.

diff --git a/Engine/Source/Programs/Shared/EpicGames.Perforce/IPerforceOutput.cs b/Engine/Source/Programs/Shared/EpicGames.Perforce/IPerforceOutput.cs
--- a/Engine/Source/Programs/Shared/EpicGames.Perforce/IPerforceOutput.cs
+++ b/Engine/Source/Programs/Shared/EpicGames.Perforce/IPerforceOutput.cs
@@ -110,14 +110,15 @@ namespace EpicGames.Perforce
 		/// Formats the current contents of the buffer to a string
 		/// </summary>
 		/// <param name="Data">The data to format</param>
+		/// <param name="MaxLength"></param>
 		/// <returns>String representation of the buffer</returns>
-		private static string FormatDataAsHexDump(ReadOnlySpan<byte> Data)
+		private static string FormatDataAsHexDump(ReadOnlySpan<byte> Data, int MaxLength = 1024)
 		{
 			// Format the result
 			StringBuilder Result = new StringBuilder();
 
 			const int RowLength = 16;
-			for (int BaseIdx = 0; BaseIdx < Data.Length; BaseIdx += RowLength)
+			for (int BaseIdx = 0; BaseIdx < Data.Length && BaseIdx < MaxLength; BaseIdx += RowLength)
 			{
 				Result.Append("\n    ");
 				for (int Offset = 0; Offset < RowLength; Offset++)
