Fix buffer corruption when enumerating Perforce records.

diff --git a/Engine/Source/Programs/Shared/EpicGames.Perforce/IPerforceOutput.cs b/Engine/Source/Programs/Shared/EpicGames.Perforce/IPerforceOutput.cs
--- a/Engine/Source/Programs/Shared/EpicGames.Perforce/IPerforceOutput.cs
+++ b/Engine/Source/Programs/Shared/EpicGames.Perforce/IPerforceOutput.cs
@@ -227,11 +227,9 @@ namespace EpicGames.Perforce
 			Record.Rows = new List<KeyValuePair<Utf8String, PerforceValue>>();
 
 			// Read all the records into a list
-			Task<bool> ReadTask = Perforce.ReadAsync(CancellationToken);
-			while (await ReadTask)
+			while (await Perforce.ReadAsync(CancellationToken))
 			{
 				// Start a read to add more data
-				ReadTask = Perforce.ReadAsync(CancellationToken);
 				ReadOnlyMemory<byte> Data = Perforce.Data;
 
 				// Parse the responses from the current buffer
