Horde: Fixup compute API to use CbObjectAttachment/CbBinaryAttachment types where appropriate.

diff --git a/Engine/Source/Programs/Horde/HordeAgent/Commands/Execution/ExecuteCommandV2.cs b/Engine/Source/Programs/Horde/HordeAgent/Commands/Execution/ExecuteCommandV2.cs
--- a/Engine/Source/Programs/Horde/HordeAgent/Commands/Execution/ExecuteCommandV2.cs
+++ b/Engine/Source/Programs/Horde/HordeAgent/Commands/Execution/ExecuteCommandV2.cs
@@ -196,7 +196,7 @@ namespace HordeAgent.Commands
 			return Hash;
 		}
 
-		private async Task ExecuteAction(ILogger Logger, GrpcService GrpcService, IoHash TaskHash, IoHash RequirementsHash, Dictionary<IoHash, byte[]> UploadList)
+		private async Task ExecuteAction(ILogger Logger, GrpcService GrpcService, CbObjectAttachment TaskHash, CbObjectAttachment RequirementsHash, Dictionary<IoHash, byte[]> UploadList)
 		{
 			using GrpcChannel Channel = GrpcService.CreateGrpcChannel();
 
@@ -216,7 +216,7 @@ namespace HordeAgent.Commands
 			// Execute the action
 			ComputeRpc.ComputeRpcClient ComputeClient = new ComputeRpc.ComputeRpcClient(Channel);
 
-			AddTasksRpcRequest AddTasksRequest = new AddTasksRpcRequest(Guid.NewGuid().ToString(), NamespaceId, RequirementsHash, new List<IoHash> { TaskHash }, SkipCacheLookup);
+			AddTasksRpcRequest AddTasksRequest = new AddTasksRpcRequest(Guid.NewGuid().ToString(), NamespaceId, RequirementsHash, new List<CbObjectAttachment> { TaskHash }, SkipCacheLookup);
 			await ComputeClient.AddTasksAsync(AddTasksRequest);
 
 			using (AsyncDuplexStreamingCall<GetTaskUpdatesRpcRequest, GetTaskUpdatesRpcResponse> Call = ComputeClient.GetTaskUpdates())
@@ -229,10 +229,10 @@ namespace HordeAgent.Commands
 				while (await Call.ResponseStream.MoveNext())
 				{
 					GetTaskUpdatesRpcResponse Response = Call.ResponseStream.Current;
-					Logger.LogInformation("{OperationName}: Execution state: {State}", (IoHash)Response.TaskHash, Response.State.ToString());
+					Logger.LogInformation("{OperationName}: Execution state: {State}", (IoHash)(CbObjectAttachment)Response.TaskHash, Response.State.ToString());
 					if (!String.IsNullOrEmpty(Response.AgentId) || !String.IsNullOrEmpty(Response.LeaseId))
 					{
-						Logger.LogInformation("{OperationName}: Running on agent {AgentId} under lease {LeaseId}", (IoHash)Response.TaskHash, Response.AgentId, Response.LeaseId);
+						Logger.LogInformation("{OperationName}: Running on agent {AgentId} under lease {LeaseId}", (IoHash)(CbObjectAttachment)Response.TaskHash, Response.AgentId, Response.LeaseId);
 					}
 					if (Response.ResultHash != null)
 					{
@@ -248,7 +248,7 @@ namespace HordeAgent.Commands
 			}
 		}
 
-		async Task HandleCompleteTask(BlobStore.BlobStoreClient Client, string NamespaceId, IoHash Hash, ILogger Logger)
+		async Task HandleCompleteTask(BlobStore.BlobStoreClient Client, string NamespaceId, CbObjectAttachment Hash, ILogger Logger)
 		{
 			ComputeTaskResult Result = await Client.GetObjectAsync<ComputeTaskResult>(NamespaceId, Hash);
 			Logger.LogInformation("exit: {ExitCode}", Result.ExitCode);
diff --git a/Engine/Source/Programs/Horde/HordeAgent/Execution/ComputeTaskExecutor.cs b/Engine/Source/Programs/Horde/HordeAgent/Execution/ComputeTaskExecutor.cs
--- a/Engine/Source/Programs/Horde/HordeAgent/Execution/ComputeTaskExecutor.cs
+++ b/Engine/Source/Programs/Horde/HordeAgent/Execution/ComputeTaskExecutor.cs
@@ -26,6 +26,7 @@ using HordeCommon.Rpc;
 using EpicGames.Horde.Common.RemoteExecution;
 using EpicGames.Serialization;
 using EpicGames.Horde.Compute;
+using EpicGames.Horde.Common;
 
 namespace HordeAgent
 {
@@ -103,12 +104,12 @@ namespace HordeAgent
 		/// <returns>The action result</returns>
 		public async Task<ComputeTaskResultMessage> ExecuteAsync(string LeaseId, ComputeTaskMessage ComputeTaskMessage, DirectoryReference SandboxDir, CancellationToken CancellationToken)
 		{
-			ComputeTask? Task = await BlobStore.GetObjectAsync<ComputeTask>(ComputeTaskMessage.NamespaceId, ComputeTaskMessage.TaskHash);
+			ComputeTask? Task = await BlobStore.GetObjectAsync<ComputeTask>(ComputeTaskMessage.NamespaceId, (CbObjectAttachment)ComputeTaskMessage.TaskHash);
 			if (Task == null)
 			{
 				throw new RpcException(new Grpc.Core.Status(StatusCode.NotFound, "Unable to find action message"));
 			}
-			Logger.LogInformation("Executing task {Hash} for lease ID {LeaseId}", ComputeTaskMessage.TaskHash.Hash, LeaseId);
+			Logger.LogInformation("Executing task {Hash} for lease ID {LeaseId}", (IoHash)(CbObjectAttachment)ComputeTaskMessage.TaskHash, LeaseId);
 
 			DirectoryReference.CreateDirectory(SandboxDir);
 			FileUtils.ForceDeleteDirectoryContents(SandboxDir);
@@ -169,7 +170,7 @@ namespace HordeAgent
 					}
 
 					ResultMessage = new ComputeTaskResultMessage();
-					ResultMessage.OutputHash = await BlobStore.PutObjectAsync(ComputeTaskMessage.NamespaceId, Result);
+					ResultMessage.OutputHash = (CbObjectAttachment)await BlobStore.PutObjectAsync(ComputeTaskMessage.NamespaceId, Result);
 
 
 					/*
diff --git a/Engine/Source/Programs/Horde/HordeServer/Compute/IComputeService.cs b/Engine/Source/Programs/Horde/HordeServer/Compute/IComputeService.cs
--- a/Engine/Source/Programs/Horde/HordeServer/Compute/IComputeService.cs
+++ b/Engine/Source/Programs/Horde/HordeServer/Compute/IComputeService.cs
@@ -2,6 +2,7 @@
 
 using EpicGames.Core;
 using EpicGames.Horde.Compute;
+using EpicGames.Serialization;
 using HordeServer.Models;
 using HordeServer.Storage;
 using HordeServer.Tasks;
@@ -36,7 +37,7 @@ namespace HordeServer.Compute
 		/// <summary>
 		/// The input hash
 		/// </summary>
-		public IoHash TaskHash { get; }
+		public CbObjectAttachment TaskHash { get; }
 
 		/// <summary>
 		/// Timestamp for the event
@@ -61,7 +62,7 @@ namespace HordeServer.Compute
 		/// <summary>
 		/// Hash of the result, if complete
 		/// </summary>
-		public IoHash? ResultHash { get; }
+		public CbObjectAttachment? ResultHash { get; }
 	}
 
 	/// <summary>
@@ -77,7 +78,7 @@ namespace HordeServer.Compute
 		/// <param name="TaskHashes">List of task hashes</param>
 		/// <param name="ChannelId">Unique identifier of the client</param>
 		/// <returns>Async task</returns>
-		Task AddTasksAsync(NamespaceId NamespaceId, IoHash RequirementsHash, List<IoHash> TaskHashes, ChannelId ChannelId);
+		Task AddTasksAsync(NamespaceId NamespaceId, CbObjectAttachment RequirementsHash, List<CbObjectAttachment> TaskHashes, ChannelId ChannelId);
 
 		/// <summary>
 		/// Dequeue completed items from a queue and return immediately
@@ -109,9 +110,9 @@ namespace HordeServer.Compute
 		/// <param name="TaskHash">The task hash</param>
 		/// <param name="ChannelId"></param>
 		/// <returns></returns>
-		public static Task AddTaskAsync(this IComputeService ComputeService, NamespaceId NamespaceId, IoHash RequirementsHash, IoHash TaskHash, ChannelId ChannelId)
+		public static Task AddTaskAsync(this IComputeService ComputeService, NamespaceId NamespaceId, CbObjectAttachment RequirementsHash, CbObjectAttachment TaskHash, ChannelId ChannelId)
 		{
-			return ComputeService.AddTasksAsync(NamespaceId, RequirementsHash, new List<IoHash> { TaskHash }, ChannelId);
+			return ComputeService.AddTasksAsync(NamespaceId, RequirementsHash, new List<CbObjectAttachment> { TaskHash }, ChannelId);
 		}
 	}
 }
diff --git a/Engine/Source/Programs/Horde/HordeServer/Compute/Impl/ComputeService.cs b/Engine/Source/Programs/Horde/HordeServer/Compute/Impl/ComputeService.cs
--- a/Engine/Source/Programs/Horde/HordeServer/Compute/Impl/ComputeService.cs
+++ b/Engine/Source/Programs/Horde/HordeServer/Compute/Impl/ComputeService.cs
@@ -48,7 +48,7 @@ namespace HordeServer.Compute.Impl
 	class ComputeTaskInfo
 	{
 		[CbField("h")]
-		public IoHash TaskHash { get; set; }
+		public CbObjectAttachment TaskHash { get; set; }
 
 		[CbField("c")]
 		public ChannelId ChannelId { get; set; }
@@ -57,7 +57,7 @@ namespace HordeServer.Compute.Impl
 		{
 		}
 
-		public ComputeTaskInfo(IoHash TaskHash, ChannelId ChannelId)
+		public ComputeTaskInfo(CbObjectAttachment TaskHash, ChannelId ChannelId)
 		{
 			this.TaskHash = TaskHash;
 			this.ChannelId = ChannelId;
@@ -72,7 +72,7 @@ namespace HordeServer.Compute.Impl
 	{
 		/// <inheritdoc/>
 		[CbField("h")]
-		public IoHash TaskHash { get; set; }
+		public CbObjectAttachment TaskHash { get; set; }
 
 		/// <inheritdoc/>
 		[CbField("t")]
@@ -92,7 +92,7 @@ namespace HordeServer.Compute.Impl
 
 		/// <inheritdoc/>
 		[CbField("r")]
-		public IoHash? ResultHash { get; set; }
+		public CbObjectAttachment? ResultHash { get; set; }
 
 		/// <inheritdoc/>
 		AgentId? IComputeTaskStatus.AgentId => AgentId.IsEmpty ? (AgentId?)null : new AgentId(AgentId.ToString());
@@ -108,7 +108,7 @@ namespace HordeServer.Compute.Impl
 		{
 		}
 
-		public ComputeTaskStatus(IoHash TaskHash, ComputeTaskState State)
+		public ComputeTaskStatus(CbObjectAttachment TaskHash, ComputeTaskState State)
 		{
 			this.TaskHash = TaskHash;
 			this.Time = DateTime.UtcNow;
@@ -149,10 +149,10 @@ namespace HordeServer.Compute.Impl
 		}
 
 		/// <inheritdoc/>
-		public async Task AddTasksAsync(NamespaceId NamespaceId, IoHash RequirementsHash, List<IoHash> TaskHashes, ChannelId ChannelId)
+		public async Task AddTasksAsync(NamespaceId NamespaceId, CbObjectAttachment RequirementsHash, List<CbObjectAttachment> TaskHashes, ChannelId ChannelId)
 		{
 			List<Task> Tasks = new List<Task>();
-			foreach (IoHash TaskHash in TaskHashes)
+			foreach (CbObjectAttachment TaskHash in TaskHashes)
 			{
 				ComputeTaskInfo TaskInfo = new ComputeTaskInfo(TaskHash, ChannelId);
 				Tasks.Add(TaskScheduler.EnqueueAsync(TaskInfo, RequirementsHash));
@@ -212,8 +212,8 @@ namespace HordeServer.Compute.Impl
 		{
 			ComputeTaskMessage ComputeTask = Payload.Unpack<ComputeTaskMessage>();
 
-			Logger.LogInformation("Task complete: {Hash}", (IoHash)ComputeTask.TaskHash);
 			ComputeTaskResultMessage Result = ComputeTaskResultMessage.Parser.ParseFrom(Output.ToArray());
+			Logger.LogInformation("Task complete: {Hash} -> {OutputHash}", (IoHash)(CbObjectAttachment)ComputeTask.TaskHash, (Result.OutputHash == null)? IoHash.Zero : (IoHash)(CbObjectAttachment)Result.OutputHash);
 
 			ComputeTaskStatus Status = new ComputeTaskStatus(ComputeTask.TaskHash, ComputeTaskState.Complete);
 			Status.LeaseId = LeaseId;
@@ -242,7 +242,7 @@ namespace HordeServer.Compute.Impl
 		{
 			NamespaceId NamespaceId = new NamespaceId(RpcRequest.NamespaceId);
 			ChannelId ChannelId = new ChannelId(RpcRequest.ChannelId);
-			await ComputeService.AddTasksAsync(NamespaceId, RpcRequest.RequirementsHash, RpcRequest.TaskHashes.Select(x => (IoHash)x).ToList(), ChannelId);
+			await ComputeService.AddTasksAsync(NamespaceId, RpcRequest.RequirementsHash, RpcRequest.TaskHashes.Select(x => (CbObjectAttachment)x).ToList(), ChannelId);
 			return new AddTasksRpcResponse();
 		}
 
diff --git a/Engine/Source/Programs/Horde/HordeServer/Compute/Impl/RedisTaskScheduler.cs b/Engine/Source/Programs/Horde/HordeServer/Compute/Impl/RedisTaskScheduler.cs
--- a/Engine/Source/Programs/Horde/HordeServer/Compute/Impl/RedisTaskScheduler.cs
+++ b/Engine/Source/Programs/Horde/HordeServer/Compute/Impl/RedisTaskScheduler.cs
@@ -38,14 +38,14 @@ namespace HordeServer.Compute.Impl
 		/// <summary>
 		/// Hash of the <see cref="Requirements"/> object in CAS for the agent to execute the task
 		/// </summary>
-		public IoHash RequirementsHash { get; }
+		public CbObjectAttachment RequirementsHash { get; }
 
 		/// <summary>
 		/// Constructor
 		/// </summary>
 		/// <param name="Item">The task item</param>
 		/// <param name="RequirementsHash">Requirements for executing the task</param>
-		public TaskSchedulerEntry(T Item, IoHash RequirementsHash)
+		public TaskSchedulerEntry(T Item, CbObjectAttachment RequirementsHash)
 		{
 			this.Item = Item;
 			this.RequirementsHash = RequirementsHash;
@@ -62,14 +62,14 @@ namespace HordeServer.Compute.Impl
 		/// </summary>
 		/// <param name="Item">The item to add</param>
 		/// <param name="RequirementsHash">Hash of a <see cref="Requirements"/> object stored in the CAS which describes the agent to execute the task</param>
-		Task EnqueueAsync(T Item, IoHash RequirementsHash);
+		Task EnqueueAsync(T Item, CbObjectAttachment RequirementsHash);
 
 		/// <summary>
 		/// Inserts a previously dequeued task back at the front of the queue
 		/// </summary>
 		/// <param name="Item">The item to add</param>
 		/// <param name="RequirementsHash">Hash of a <see cref="Requirements"/> object stored in the CAS which describes the agent to execute the task</param>
-		Task RequeueAsync(T Item, IoHash RequirementsHash);
+		Task RequeueAsync(T Item, CbObjectAttachment RequirementsHash);
 
 		/// <summary>
 		/// Dequeues a task that the given agent can execute
@@ -157,7 +157,7 @@ namespace HordeServer.Compute.Impl
 		/// </summary>
 		/// <param name="RequirementsHash"></param>
 		/// <returns></returns>
-		RedisKey GetQueueKey(IoHash RequirementsHash)
+		RedisKey GetQueueKey(CbObjectAttachment RequirementsHash)
 		{
 			return BaseKey.Append(RedisSerializer.Serialize(RequirementsHash).AsKey());
 		}
@@ -167,7 +167,7 @@ namespace HordeServer.Compute.Impl
 		/// </summary>
 		/// <param name="RequirementsHash"></param>
 		/// <returns></returns>
-		RedisList<T> GetQueue(IoHash RequirementsHash)
+		RedisList<T> GetQueue(CbObjectAttachment RequirementsHash)
 		{
 			return new RedisList<T>(GetQueueKey(RequirementsHash));
 		}
@@ -176,7 +176,7 @@ namespace HordeServer.Compute.Impl
 		/// Adds a queue to the index, atomically checking that the queue exists
 		/// </summary>
 		/// <param name="RequirementsHash">Hash of the queue requirements objectadd</param>
-		async ValueTask AddQueueToIndexAsync(IoHash RequirementsHash)
+		async ValueTask AddQueueToIndexAsync(CbObjectAttachment RequirementsHash)
 		{
 			ITransaction Transaction = Redis.CreateTransaction();
 			Transaction.AddCondition(Condition.KeyExists(GetQueueKey(RequirementsHash)));
@@ -190,7 +190,7 @@ namespace HordeServer.Compute.Impl
 		/// Removes a queue from the index, atomically checking that it is empty
 		/// </summary>
 		/// <param name="RequirementsHash">Hash of the queue requirements objectadd</param>
-		async ValueTask RemoveQueueFromIndexAsync(IoHash RequirementsHash)
+		async ValueTask RemoveQueueFromIndexAsync(CbObjectAttachment RequirementsHash)
 		{
 			ITransaction Transaction = Redis.CreateTransaction();
 			Transaction.AddCondition(Condition.KeyNotExists(GetQueueKey(RequirementsHash)));
@@ -203,7 +203,7 @@ namespace HordeServer.Compute.Impl
 		/// </summary>
 		/// <param name="Item">The item to be added</param>
 		/// <param name="RequirementsHash">Requirements for the task</param>
-		public async Task EnqueueAsync(T Item, IoHash RequirementsHash)
+		public async Task EnqueueAsync(T Item, CbObjectAttachment RequirementsHash)
 		{
 			RedisList<T> List = GetQueue(RequirementsHash);
 
@@ -219,7 +219,7 @@ namespace HordeServer.Compute.Impl
 		/// </summary>
 		/// <param name="Item">The item to be added</param>
 		/// <param name="RequirementsHash">Requirements for the task</param>
-		public async Task RequeueAsync(T Item, IoHash RequirementsHash)
+		public async Task RequeueAsync(T Item, CbObjectAttachment RequirementsHash)
 		{
 			RedisList<T> Queue = GetQueue(RequirementsHash);
 
diff --git a/Engine/Source/Programs/Shared/EpicGames.Horde/Common/CbBinaryAttachmentWrapper.cs b/Engine/Source/Programs/Shared/EpicGames.Horde/Common/CbBinaryAttachmentWrapper.cs
new file mode 100644
--- /dev/null
+++ b/Engine/Source/Programs/Shared/EpicGames.Horde/Common/CbBinaryAttachmentWrapper.cs
@@ -0,0 +1,45 @@
+// Copyright Epic Games, Inc. All Rights Reserved.
+
+using EpicGames.Core;
+using EpicGames.Serialization;
+using Google.Protobuf;
+using System;
+using System.Collections.Generic;
+using System.Text;
+using System.Threading.Tasks;
+
+namespace EpicGames.Horde.Common
+{
+	/// <summary>
+	/// Methods for CbBinaryAttachmentWrapper
+	/// </summary>
+	public partial class CbBinaryAttachmentWrapper
+	{
+		/// <summary>
+		/// Constructor
+		/// </summary>
+		/// <param name="Hash"></param>
+		public CbBinaryAttachmentWrapper(IoHash Hash)
+		{
+			this.Hash = ByteString.CopyFrom(Hash.Memory.ToArray());
+		}
+
+		/// <summary>
+		/// Convert from an IoHashWrapper to an IoHash
+		/// </summary>
+		/// <param name="Hash"></param>
+		public static implicit operator CbBinaryAttachment(CbBinaryAttachmentWrapper Attachment)
+		{
+			return new CbBinaryAttachment(new IoHash(Attachment.Hash.ToByteArray()));
+		}
+
+		/// <summary>
+		/// Convert from an IoHash to an IoHashWrapper
+		/// </summary>
+		/// <param name="Hash"></param>
+		public static implicit operator CbBinaryAttachmentWrapper(CbBinaryAttachment Attachment)
+		{
+			return new CbBinaryAttachmentWrapper(Attachment.Hash);
+		}
+	}
+}
diff --git a/Engine/Source/Programs/Shared/EpicGames.Horde/Common/CbObjectAttachmentWrapper.cs b/Engine/Source/Programs/Shared/EpicGames.Horde/Common/CbObjectAttachmentWrapper.cs
new file mode 100644
--- /dev/null
+++ b/Engine/Source/Programs/Shared/EpicGames.Horde/Common/CbObjectAttachmentWrapper.cs
@@ -0,0 +1,45 @@
+// Copyright Epic Games, Inc. All Rights Reserved.
+
+using EpicGames.Core;
+using EpicGames.Serialization;
+using Google.Protobuf;
+using System;
+using System.Collections.Generic;
+using System.Text;
+using System.Threading.Tasks;
+
+namespace EpicGames.Horde.Common
+{
+	/// <summary>
+	/// Methods for CbObjectAttachmentWrapper
+	/// </summary>
+	public partial class CbObjectAttachmentWrapper
+	{
+		/// <summary>
+		/// Constructor
+		/// </summary>
+		/// <param name="Hash"></param>
+		public CbObjectAttachmentWrapper(IoHash Hash)
+		{
+			this.Hash = ByteString.CopyFrom(Hash.Memory.ToArray());
+		}
+
+		/// <summary>
+		/// Convert from an IoHashWrapper to an IoHash
+		/// </summary>
+		/// <param name="Hash"></param>
+		public static implicit operator CbObjectAttachment(CbObjectAttachmentWrapper Attachment)
+		{
+			return new CbObjectAttachment(new IoHash(Attachment.Hash.ToByteArray()));
+		}
+
+		/// <summary>
+		/// Convert from an IoHash to an IoHashWrapper
+		/// </summary>
+		/// <param name="Hash"></param>
+		public static implicit operator CbObjectAttachmentWrapper(CbObjectAttachment Attachment)
+		{
+			return new CbObjectAttachmentWrapper(Attachment.Hash);
+		}
+	}
+}
diff --git a/Engine/Source/Programs/Shared/EpicGames.Horde/Compute/Api.cs b/Engine/Source/Programs/Shared/EpicGames.Horde/Compute/Api.cs
--- a/Engine/Source/Programs/Shared/EpicGames.Horde/Compute/Api.cs
+++ b/Engine/Source/Programs/Shared/EpicGames.Horde/Compute/Api.cs
@@ -18,11 +18,11 @@ namespace EpicGames.Horde.Compute
 	{
 		/// <inheritdoc cref="AddTasksRpcRequest.RequirementsHash"/>
 		[Required, CbField("r")]
-		public IoHash RequirementsHash { get; set; }
+		public CbObjectAttachment RequirementsHash { get; set; }
 
 		/// <inheritdoc cref="AddTasksRpcRequest.TaskHashes"/>
 		[Required, CbField("t")]
-		public List<IoHash> TaskHashes { get; set; } = new List<IoHash>();
+		public List<CbObjectAttachment> TaskHashes { get; set; } = new List<CbObjectAttachment>();
 
 		/// <inheritdoc cref="AddTasksRpcRequest.DoNotCache"/>
 		[CbField("nc")]
@@ -48,7 +48,7 @@ namespace EpicGames.Horde.Compute
 	{
 		/// <inheritdoc cref="GetTaskUpdatesRpcResponse.TaskHash"/>
 		[CbField("h")]
-		public IoHash TaskHash { get; set; }
+		public CbObjectAttachment TaskHash { get; set; }
 
 		/// <inheritdoc cref="GetTaskUpdatesRpcResponse.Time"/>
 		[CbField("t")]
@@ -60,7 +60,7 @@ namespace EpicGames.Horde.Compute
 
 		/// <inheritdoc cref="GetTaskUpdatesRpcResponse.ResultHash"/>
 		[CbField("r")]
-		public IoHash? ResultHash { get; set; }
+		public CbObjectAttachment? ResultHash { get; set; }
 
 		/// <inheritdoc cref="GetTaskUpdatesRpcResponse.AgentId"/>
 		[CbField("a")]
@@ -73,12 +73,12 @@ namespace EpicGames.Horde.Compute
 
 	partial class AddTasksRpcRequest
 	{
-		public AddTasksRpcRequest(string ChannelId, string NamespaceId, IoHash RequirementsHash, List<IoHash> TaskHashes, bool DoNotCache)
+		public AddTasksRpcRequest(string ChannelId, string NamespaceId, CbObjectAttachment RequirementsHash, List<CbObjectAttachment> TaskHashes, bool DoNotCache)
 		{
 			this.ChannelId = ChannelId;
 			this.NamespaceId = NamespaceId;
 			this.RequirementsHash = RequirementsHash;
-			this.TaskHashes.Add(TaskHashes.Select(x => (IoHashWrapper)x));
+			this.TaskHashes.Add(TaskHashes.Select(x => (CbObjectAttachmentWrapper)x));
 			this.DoNotCache = DoNotCache;
 		}
 	}
diff --git a/Engine/Source/Programs/Shared/EpicGames.Horde/EpicGames.Horde.csproj b/Engine/Source/Programs/Shared/EpicGames.Horde/EpicGames.Horde.csproj
--- a/Engine/Source/Programs/Shared/EpicGames.Horde/EpicGames.Horde.csproj
+++ b/Engine/Source/Programs/Shared/EpicGames.Horde/EpicGames.Horde.csproj
@@ -42,7 +42,9 @@
     <ProjectReference Include="..\EpicGames.Serialization\EpicGames.Serialization.csproj" />
   </ItemGroup>
   <ItemGroup>
+    <Protobuf Include="Protos\horde\common\cbbinaryattachment_wrapper.proto" ProtoRoot="Protos" />
     <Protobuf Include="Protos\horde\common\cbobject_wrapper.proto" ProtoRoot="Protos" />
+    <Protobuf Include="Protos\horde\common\cbobjectattachment_wrapper.proto" ProtoRoot="Protos" />
     <Protobuf Include="Protos\horde\common\iohash_wrapper.proto" ProtoRoot="Protos" />
     <Protobuf Include="Protos\horde\compute\compute_rpc_messages.proto" ProtoRoot="Protos" />
     <Protobuf Include="Protos\horde\storage\blob_store_messages.proto" ProtoRoot="Protos" />
diff --git a/Engine/Source/Programs/Shared/EpicGames.Serialization/CompactBinary.cs b/Engine/Source/Programs/Shared/EpicGames.Serialization/CompactBinary.cs
--- a/Engine/Source/Programs/Shared/EpicGames.Serialization/CompactBinary.cs
+++ b/Engine/Source/Programs/Shared/EpicGames.Serialization/CompactBinary.cs
@@ -6,7 +6,9 @@ using System.Buffers;
 using System.Buffers.Binary;
 using System.Collections;
 using System.Collections.Generic;
+using System.ComponentModel;
 using System.Diagnostics;
+using System.Globalization;
 using System.Linq;
 using System.Numerics;
 using System.Text;
@@ -194,6 +196,7 @@ namespace EpicGames.Serialization
 	/// </summary>
 	[DebuggerDisplay("{Hash}")]
 	[JsonConverter(typeof(CbBinaryAttachmentJsonConverter))]
+	[TypeConverter(typeof(CbBinaryAttachmentTypeConverter))]
 	public struct CbBinaryAttachment
 	{
 		/// <summary>
@@ -215,23 +218,20 @@ namespace EpicGames.Serialization
 			this.Hash = Hash;
 		}
 
+		/// <inheritdoc/>
+		public override string ToString() => Hash.ToString();
+
 		/// <summary>
 		/// Convert a hash to a binary attachment 
 		/// </summary>
 		/// <param name="Hash">The attachment to convert</param>
-		public static implicit operator CbBinaryAttachment(IoHash Hash)
-		{
-			return new CbBinaryAttachment(Hash);
-		}
+		public static implicit operator CbBinaryAttachment(IoHash Hash) => new CbBinaryAttachment(Hash);
 
 		/// <summary>
 		/// Use a binary attachment as a hash
 		/// </summary>
 		/// <param name="Attachment">The attachment to convert</param>
-		public static implicit operator IoHash(CbBinaryAttachment Attachment)
-		{
-			return Attachment.Hash;
-		}
+		public static implicit operator IoHash(CbBinaryAttachment Attachment) => Attachment.Hash;
 	}
 
 	/// <summary>
@@ -246,11 +246,24 @@ namespace EpicGames.Serialization
 		public override void Write(Utf8JsonWriter Writer, CbBinaryAttachment Value, JsonSerializerOptions Options) => Writer.WriteStringValue(Value.Hash.ToUtf8String().Span);
 	}
 
+	/// <summary>
+	/// Type converter from strings to IoHash objects
+	/// </summary>
+	sealed class CbBinaryAttachmentTypeConverter : TypeConverter
+	{
+		/// <inheritdoc/>
+		public override bool CanConvertFrom(ITypeDescriptorContext Context, Type SourceType) => SourceType == typeof(string);
+
+		/// <inheritdoc/>
+		public override object ConvertFrom(ITypeDescriptorContext Context, CultureInfo Culture, object Value) => new CbBinaryAttachment(IoHash.Parse((string)Value));
+	}
+
 	/// <summary>
 	/// An object attachment, referenced by <see cref="IoHash"/>
 	/// </summary>
 	[DebuggerDisplay("{Hash}")]
 	[JsonConverter(typeof(CbObjectAttachmentJsonConverter))]
+	[TypeConverter(typeof(CbObjectAttachmentTypeConverter))]
 	public struct CbObjectAttachment
 	{
 		/// <summary>
@@ -272,23 +285,20 @@ namespace EpicGames.Serialization
 			this.Hash = Hash;
 		}
 
+		/// <inheritdoc/>
+		public override string ToString() => Hash.ToString();
+
 		/// <summary>
 		/// Use an object attachment as a hash
 		/// </summary>
 		/// <param name="Attachment">The attachment to convert</param>
-		public static implicit operator CbObjectAttachment(IoHash Hash)
-		{
-			return new CbObjectAttachment(Hash);
-		}
+		public static implicit operator CbObjectAttachment(IoHash Hash) => new CbObjectAttachment(Hash);
 
 		/// <summary>
 		/// Use an object attachment as a hash
 		/// </summary>
 		/// <param name="Attachment">The attachment to convert</param>
-		public static implicit operator IoHash(CbObjectAttachment Attachment)
-		{
-			return Attachment.Hash;
-		}
+		public static implicit operator IoHash(CbObjectAttachment Attachment) => Attachment.Hash;
 	}
 
 	/// <summary>
@@ -303,6 +313,18 @@ namespace EpicGames.Serialization
 		public override void Write(Utf8JsonWriter Writer, CbObjectAttachment Value, JsonSerializerOptions Options) => Writer.WriteStringValue(Value.Hash.ToUtf8String().Span);
 	}
 
+	/// <summary>
+	/// Type converter from strings to IoHash objects
+	/// </summary>
+	sealed class CbObjectAttachmentTypeConverter : TypeConverter
+	{
+		/// <inheritdoc/>
+		public override bool CanConvertFrom(ITypeDescriptorContext Context, Type SourceType) => SourceType == typeof(string);
+
+		/// <inheritdoc/>
+		public override object ConvertFrom(ITypeDescriptorContext Context, CultureInfo Culture, object Value) => new CbObjectAttachment(IoHash.Parse((string)Value));
+	}
+
 	/// <summary>
 	/// Methods that operate on <see cref="CbFieldType"/>.
 	/// </summary>
