Horde: Allow requesting objects in JSON format using a ".json" suffix on the URL. Compact binary objects will automatically be transcoded.

diff --git a/Engine/Source/Programs/Horde/HordeServer/Startup.cs b/Engine/Source/Programs/Horde/HordeServer/Startup.cs
--- a/Engine/Source/Programs/Horde/HordeServer/Startup.cs
+++ b/Engine/Source/Programs/Horde/HordeServer/Startup.cs
@@ -34,7 +34,6 @@ using Microsoft.OpenApi.Models;
 using MongoDB.Bson;
 using MongoDB.Bson.Serialization;
 using MongoDB.Bson.Serialization.Conventions;
-using System.Net.Http;
 using Microsoft.AspNetCore.HttpOverrides;
 using HordeServer.Utilities;
 using Grpc.Core.Interceptors;
@@ -82,6 +81,7 @@ using OpenTracing.Util;
 using EpicGames.Horde.Compute;
 using HordeServer.Compute.Impl;
 using HordeServer.Compute;
+using System.Net.Http.Headers;
 
 namespace HordeServer
 {
@@ -569,6 +569,8 @@ namespace HordeServer
 			{
 				Options.InputFormatters.Add(new CbInputFormatter());
 				Options.OutputFormatters.Add(new CbOutputFormatter());
+				Options.OutputFormatters.Insert(0, new CbPreferredOutputFormatter());
+				Options.FormatterMappings.SetMediaTypeMappingForFormat("uecb", CustomMediaTypeNames.UnrealCompactBinary);
 			});
 
 			Services.AddSwaggerGen(Config =>
diff --git a/Engine/Source/Programs/Horde/HordeServer/Storage/Controllers/ObjectsController.cs b/Engine/Source/Programs/Horde/HordeServer/Storage/Controllers/ObjectsController.cs
--- a/Engine/Source/Programs/Horde/HordeServer/Storage/Controllers/ObjectsController.cs
+++ b/Engine/Source/Programs/Horde/HordeServer/Storage/Controllers/ObjectsController.cs
@@ -18,6 +18,7 @@ using EpicGames.Core;
 using HordeServer.Storage.Collections;
 using EpicGames.Serialization;
 using System.Net;
+using System.Net.Mime;
 
 namespace HordeServer.Storage.Controllers
 {
@@ -92,6 +93,7 @@ namespace HordeServer.Storage.Controllers
 	/// </summary>
 	[ApiController]
 	[Authorize]
+	[FormatFilter]
 	[Route("[controller]")]
 	public class ObjectsController : ControllerBase
 	{
@@ -305,7 +307,7 @@ namespace HordeServer.Storage.Controllers
 		/// <param name="Hash">Hash of the blob. Must match the submitted data.</param>
 		/// <returns></returns>
 		[HttpGet]
-		[Route("/api/v1/objects/{NamespaceId}/{Hash}")]
+		[Route("/api/v1/objects/{NamespaceId}/{Hash}.{format?}")]
 		public async Task<ActionResult> GetObjectAsync(NamespaceId NamespaceId, IoHash Hash)
 		{
 			if (!await NamespaceCollection.AuthorizeAsync(NamespaceId, User, AclAction.ReadBlobs))
@@ -319,7 +321,7 @@ namespace HordeServer.Storage.Controllers
 				return NotFound();
 			}
 
-			return Ok(Object.GetView().ToArray());
+			return Ok(Object);
 		}
 
 		/// <summary>
diff --git a/Engine/Source/Programs/Horde/HordeServer/Utilities/CompactBinaryFormatter.cs b/Engine/Source/Programs/Horde/HordeServer/Utilities/CompactBinaryFormatter.cs
--- a/Engine/Source/Programs/Horde/HordeServer/Utilities/CompactBinaryFormatter.cs
+++ b/Engine/Source/Programs/Horde/HordeServer/Utilities/CompactBinaryFormatter.cs
@@ -10,9 +10,21 @@ using System.IO;
 using EpicGames.Serialization;
 using EpicGames.Serialization.Converters;
 using System.Reflection;
+using System.Net.Mime;
 
 namespace HordeServer.Utilities
 {
+	/// <summary>
+	/// Global constants
+	/// </summary>
+	public static partial class CustomMediaTypeNames
+	{
+		/// <summary>
+		/// Media type for compact binary
+		/// </summary>
+		public const string UnrealCompactBinary = "application/x-ue-cb";
+	}
+
 	/// <summary>
 	/// Converter to allow reading compact binary objects as request bodies
 	/// </summary>
@@ -23,7 +35,7 @@ namespace HordeServer.Utilities
 		/// </summary>
 		public CbInputFormatter()
 		{
-			SupportedMediaTypes.Add(MediaTypeHeaderValue.Parse("application/x-ue-cb"));
+			SupportedMediaTypes.Add(MediaTypeHeaderValue.Parse(CustomMediaTypeNames.UnrealCompactBinary));
 		}
 
 		/// <inheritdoc/>
@@ -58,7 +70,7 @@ namespace HordeServer.Utilities
 		/// </summary>
 		public CbOutputFormatter()
 		{
-			SupportedMediaTypes.Add(MediaTypeHeaderValue.Parse("application/x-ue-cb"));
+			SupportedMediaTypes.Add(MediaTypeHeaderValue.Parse(CustomMediaTypeNames.UnrealCompactBinary));
 		}
 
 		/// <inheritdoc/>
@@ -70,8 +82,29 @@ namespace HordeServer.Utilities
 		/// <inheritdoc/>
 		public override async Task WriteResponseBodyAsync(OutputFormatterWriteContext Context)
 		{
-			ReadOnlyMemory<byte> Data = CbSerializer.Serialize(Context.ObjectType, Context.Object).GetView();
+			ReadOnlyMemory<byte> Data;
+			if (Context.Object is CbObject Object)
+			{
+				Data = Object.GetView();
+			}
+			else
+			{
+				Data = CbSerializer.Serialize(Context.ObjectType, Context.Object).GetView();
+			}
 			await Context.HttpContext.Response.BodyWriter.WriteAsync(Data);
 		}
 	}
+
+	/// <summary>
+	/// Special version of <see cref="CbOutputFormatter"/> which returns native CbObject encoded data. Can be 
+	/// inserted at a high priority in the output formatter list to prevent transcoding to json.
+	/// </summary>
+	public class CbPreferredOutputFormatter : CbOutputFormatter
+	{
+		/// <inheritdoc/>
+		protected override bool CanWriteType(Type Type)
+		{
+			return Type == typeof(CbObject);
+		}
+	}
 }
diff --git a/Engine/Source/Programs/Shared/EpicGames.Serialization/CompactBinary.cs b/Engine/Source/Programs/Shared/EpicGames.Serialization/CompactBinary.cs
--- a/Engine/Source/Programs/Shared/EpicGames.Serialization/CompactBinary.cs
+++ b/Engine/Source/Programs/Shared/EpicGames.Serialization/CompactBinary.cs
@@ -644,6 +644,7 @@ namespace EpicGames.Serialization
 	/// </summary>
 	[DebuggerDisplay("{DebugValue,nq}")]
 	[DebuggerTypeProxy(typeof(CbFieldDebugView))]
+	[JsonConverter(typeof(CbFieldJsonConverter))]
 	public class CbField : IEquatable<CbField>, IEnumerable<CbField>
 	{
 		/// <summary>
@@ -1622,6 +1623,176 @@ namespace EpicGames.Serialization
 		#endregion
 	}
 
+	/// <summary>
+	/// Converter to and from JSON objects
+	/// </summary>
+	public class CbFieldJsonConverter : JsonConverter<CbField>
+	{
+		/// <inheritdoc/>
+		public override CbField Read(ref Utf8JsonReader reader, Type typeToConvert, JsonSerializerOptions options)
+		{
+			throw new NotImplementedException();
+		}
+
+		/// <inheritdoc/>
+		public override void Write(Utf8JsonWriter Writer, CbField Field, JsonSerializerOptions Options)
+		{
+			switch (Field.GetType())
+			{
+				case CbFieldType.Null:
+					if (Field.HasName())
+					{
+						Writer.WriteNull(Field.Name.Span);
+					}
+					else
+					{
+						Writer.WriteNullValue();
+					}
+					break;
+				case CbFieldType.Object:
+				case CbFieldType.UniformObject:
+					if (Field.HasName())
+					{
+						Writer.WriteStartObject(Field.Name.Span);
+					}
+					else
+					{
+						Writer.WriteStartObject();
+					}
+
+					foreach (CbField Member in Field.AsObject())
+					{
+						Write(Writer, Member, Options);
+					}
+					Writer.WriteEndObject();
+					break;
+				case CbFieldType.Array:
+				case CbFieldType.UniformArray:
+					if (Field.HasName())
+					{
+						Writer.WriteStartArray(Field.Name.Span);
+					}
+					else
+					{
+						Writer.WriteStartArray();
+					}
+
+					foreach (CbField Element in Field.AsArray())
+					{
+						Write(Writer, Element, Options);
+					}
+					Writer.WriteEndArray();
+					break;
+				case CbFieldType.Binary:
+					if (Field.HasName())
+					{
+						Writer.WriteBase64String(Field.Name.Span, Field.AsBinary().Span);
+					}
+					else
+					{
+						Writer.WriteBase64StringValue(Field.AsBinary().Span);
+					}
+					break;
+				case CbFieldType.String:
+					if (Field.HasName())
+					{
+						Writer.WriteString(Field.Name.Span, Field.AsString().Span);
+					}
+					else
+					{
+						Writer.WriteStringValue(Field.AsString().Span);
+					}
+					break;
+				case CbFieldType.IntegerPositive:
+					if (Field.HasName())
+					{
+						Writer.WriteNumber(Field.Name.Span, Field.AsUInt64());
+					}
+					else
+					{
+						Writer.WriteNumberValue(Field.AsUInt64());
+					}
+					break;
+				case CbFieldType.IntegerNegative:
+					if (Field.HasName())
+					{
+						Writer.WriteNumber(Field.Name.Span, Field.AsInt64());
+					}
+					else
+					{
+						Writer.WriteNumberValue(Field.AsInt64());
+					}
+					break;
+				case CbFieldType.Float32:
+				case CbFieldType.Float64:
+					if (Field.HasName())
+					{
+						Writer.WriteNumber(Field.Name.Span, Field.AsDouble());
+					}
+					else
+					{
+						Writer.WriteNumberValue(Field.AsDouble());
+					}
+					break;
+				case CbFieldType.BoolFalse:
+				case CbFieldType.BoolTrue:
+					if (Field.HasName())
+					{
+						Writer.WriteBoolean(Field.Name.Span, Field.AsBool());
+					}
+					else
+					{
+						Writer.WriteBooleanValue(Field.AsBool());
+					}
+					break;
+				case CbFieldType.ObjectAttachment:
+				case CbFieldType.BinaryAttachment:
+				case CbFieldType.Hash:
+					if (Field.HasName())
+					{
+						Writer.WriteString(Field.Name.Span, Field.AsHash().ToString());
+					}
+					else
+					{
+						Writer.WriteStringValue(Field.AsHash().ToString());
+					}
+					break;
+				case CbFieldType.Uuid:
+					if (Field.HasName())
+					{
+						Writer.WriteString(Field.Name.Span, Field.AsUuid().ToString());
+					}
+					else
+					{
+						Writer.WriteStringValue(Field.AsUuid().ToString());
+					}
+					break;
+				case CbFieldType.DateTime:
+					if (Field.HasName())
+					{
+						Writer.WriteNumber(Field.Name.Span, Field.AsDateTimeTicks());
+					}
+					else
+					{
+						Writer.WriteNumberValue(Field.AsDateTimeTicks());
+					}
+					break;
+				case CbFieldType.TimeSpan:
+					if (Field.HasName())
+					{
+						Writer.WriteNumber(Field.Name.Span, Field.AsTimeSpanTicks());
+					}
+					else
+					{
+						Writer.WriteNumberValue(Field.AsTimeSpanTicks());
+					}
+					break;
+				default:
+					throw new NotImplementedException($"Unhandled type in cb-json converter");
+			}
+		}
+	}
+
 	public class CbFieldEnumerator : IEnumerator<CbField>
 	{
 		/// <summary>
