Horde: Add support for serializing ReadOnlyMemory<> and byte arrays to reflection-based CB serializer.

diff --git a/Engine/Source/Programs/Shared/EpicGames.Serialization/CompactBinarySerializer.cs b/Engine/Source/Programs/Shared/EpicGames.Serialization/CompactBinarySerializer.cs
--- a/Engine/Source/Programs/Shared/EpicGames.Serialization/CompactBinarySerializer.cs
+++ b/Engine/Source/Programs/Shared/EpicGames.Serialization/CompactBinarySerializer.cs
@@ -17,6 +17,7 @@ namespace EpicGames.Serialization
 	/// <summary>
 	/// Attribute used to mark a property that should be serialized to compact binary
 	/// </summary>
+	[AttributeUsage(AttributeTargets.Property)]
 	public class CbFieldAttribute : Attribute
 	{
 		/// <summary>
diff --git a/Engine/Source/Programs/Shared/EpicGames.Serialization/Converters/CbDefaultConverter.cs b/Engine/Source/Programs/Shared/EpicGames.Serialization/Converters/CbDefaultConverter.cs
--- a/Engine/Source/Programs/Shared/EpicGames.Serialization/Converters/CbDefaultConverter.cs
+++ b/Engine/Source/Programs/Shared/EpicGames.Serialization/Converters/CbDefaultConverter.cs
@@ -107,6 +107,14 @@ namespace EpicGames.Serialization.Converters
 		static void WriteString(CbWriter Writer, Utf8String Name, string Value) => Writer.WriteString(Name, Value);
 		static void WriteStringValue(CbWriter Writer, string Value) => Writer.WriteStringValue(Value);
 
+		static ReadOnlyMemory<byte> ReadMemory(CbField Field) => Field.AsBinary();
+		static void WriteMemory(CbWriter Writer, Utf8String Name, ReadOnlyMemory<byte> Value) => Writer.WriteBinary(Name, Value.Span);
+		static void WriteMemoryValue(CbWriter Writer, ReadOnlyMemory<byte> Value) => Writer.WriteBinaryValue(Value.Span);
+
+		static byte[] ReadByteArray(CbField Field) => Field.AsBinary().ToArray();
+		static void WriteByteArray(CbWriter Writer, Utf8String Name, byte[] Value) => Writer.WriteBinary(Name, Value);
+		static void WriteByteArrayValue(CbWriter Writer, byte[] Value) => Writer.WriteBinaryValue(Value);
+
 		static Dictionary<Type, CbConverterInfo> TypeToConverterInfo = new Dictionary<Type, CbConverterInfo>(new KeyValuePair<Type, CbConverterInfo>[]
 		{
 			GetPodConverterInfo(x => x.AsBool(), (w, v) => w.WriteBoolValue(v), (w, n, v) => w.WriteBool(n, v), null),
@@ -117,7 +125,9 @@ namespace EpicGames.Serialization.Converters
 			GetPodConverterInfo(x => x.AsObjectAttachment(), (w, v) => w.WriteObjectAttachmentValue(v.Hash), (w, n, v) => w.WriteObjectAttachment(n, v.Hash), null),
 			GetPodConverterInfo(x => x.AsBinaryAttachment(), (w, v) => w.WriteBinaryAttachmentValue(v.Hash), (w, n, v) => w.WriteBinaryAttachment(n, v.Hash), null),
 			GetPodConverterInfo(x => x.AsDateTime(), (w, v) => w.WriteDateTimeValue(v), (w, n, v) => w.WriteDateTime(n, v), null),
-			GetPodConverterInfo(x => ReadString(x), (w, v) => WriteStringValue(w, v), (w, n, v) => WriteString(w, n, v), SkipIfNullOrZero)
+			GetPodConverterInfo(x => ReadString(x), (w, v) => WriteStringValue(w, v), (w, n, v) => WriteString(w, n, v), SkipIfNullOrZero),
+			GetPodConverterInfo(x => ReadMemory(x), (w, v) => WriteMemoryValue(w, v), (w, n, v) => WriteMemory(w, n, v), null),
+			GetPodConverterInfo(x => ReadByteArray(x), (w, v) => WriteByteArrayValue(w, v), (w, n, v) => WriteByteArray(w, n, v), null),
 		});
 
 		static KeyValuePair<Type, CbConverterInfo> GetPodConverterInfo<T>(Expression<Func<CbField, T>> Read, Expression<Action<CbWriter, T>> Write, Expression<Action<CbWriter, Utf8String, T>> WriteNamed, SkipDefaultDelegate? SkipDefault)
@@ -558,7 +568,7 @@ namespace EpicGames.Serialization.Converters
 		static (Utf8String, PropertyInfo)[] GetProperties(Type Type)
 		{
 			List<(Utf8String, PropertyInfo)> PropertyList = new List<(Utf8String, PropertyInfo)>();
-			foreach (PropertyInfo Property in Type.GetProperties(BindingFlags.Public | BindingFlags.Instance))
+			foreach (PropertyInfo Property in Type.GetProperties(BindingFlags.Public | BindingFlags.NonPublic | BindingFlags.Instance))
 			{
 				CbFieldAttribute? Attribute = Property.GetCustomAttribute<CbFieldAttribute>();
 				if (Attribute != null)
