Add support to maintain zip perm for Linux

diff --git a/Engine/Source/Programs/Shared/EpicGames.Core/FileUtils.cs b/Engine/Source/Programs/Shared/EpicGames.Core/FileUtils.cs
--- a/Engine/Source/Programs/Shared/EpicGames.Core/FileUtils.cs
+++ b/Engine/Source/Programs/Shared/EpicGames.Core/FileUtils.cs
@@ -365,6 +365,28 @@ namespace EpicGames.Core
 			chmod(FileName, Mode);
 		}
 
+		/// <summary>
+		/// Gets the file mode on Linux
+		/// </summary>
+		/// <param name="FileName"></param>
+		/// <returns></returns>
+		public static int GetFileMode_Linux(string FileName)
+		{
+			stat64_linux_t stat = new stat64_linux_t();
+			int Result = stat64_linux(1, FileName, stat);
+			return (Result >= 0)? (int)stat.st_mode : -1;
+		}
+
+		/// <summary>
+		/// Sets the file mode on Linux
+		/// </summary>
+		/// <param name="FileName"></param>
+		/// <param name="Mode"></param>
+		public static void SetFileMode_Linux(string FileName, ushort Mode)
+		{
+			chmod_linux(FileName, Mode);
+		}
+
 		#region Win32 Native File Methods
 
 		[StructLayout(LayoutKind.Sequential, CharSet = CharSet.Unicode)]
@@ -778,6 +800,41 @@ namespace EpicGames.Core
 		[DllImport("libSystem.dylib")]
 		static extern int chmod(string path, ushort mode);
 
+#pragma warning restore CS0649
+		#endregion
+
+		#region Linux Native File Methods
+#pragma warning disable CS0649
+
+		[StructLayout(LayoutKind.Sequential)]
+		class stat64_linux_t
+		{
+			public ulong st_dev;
+			public ulong st_ino;
+			public ulong st_nlink;
+			public uint st_mode;
+			public uint st_uid;
+			public uint st_gid;
+			public int pad0;
+			public ulong st_rdev;
+			public long st_size;
+			public long st_blksize;
+			public long st_blocks;
+			public timespec_t st_atime;
+			public timespec_t st_mtime;
+			public timespec_t st_ctime;
+			public long glibc_reserved0;
+			public long glibc_reserved1;
+			public long glibc_reserved2;
+		};
+
+		/* stat tends to get compiled to another symbol and libc doesnt directly have that entry point */
+		[DllImport("libc", EntryPoint="__xstat64")]
+		static extern int stat64_linux(int ver, string pathname, stat64_linux_t stat);
+
+		[DllImport("libc", EntryPoint="chmod")]
+		static extern int chmod_linux(string path, ushort mode);
+
 #pragma warning restore CS0649
 		#endregion
 	}
diff --git a/Engine/Source/Programs/Shared/EpicGames.Core/ZipArchiveExtensions.cs b/Engine/Source/Programs/Shared/EpicGames.Core/ZipArchiveExtensions.cs
--- a/Engine/Source/Programs/Shared/EpicGames.Core/ZipArchiveExtensions.cs
+++ b/Engine/Source/Programs/Shared/EpicGames.Core/ZipArchiveExtensions.cs
@@ -25,14 +25,22 @@ namespace EpicGames.Core
 		public static ZipArchiveEntry CreateEntryFromFile_CrossPlatform(this ZipArchive Destination, string SourceFileName, string EntryName, CompressionLevel CompressionLevel)
 		{
 			ZipArchiveEntry Entry = ZipFileExtensions.CreateEntryFromFile(Destination, SourceFileName, EntryName, CompressionLevel);
-			if (RuntimeInformation.IsOSPlatform(OSPlatform.OSX))
+			int Result = -1;
+
+			if (RuntimeInformation.IsOSPlatform(OSPlatform.Linux))
 			{
-				int Result = FileUtils.GetFileMode_Mac(SourceFileName);
-				if(Result >= 0)
-				{
-					Entry.ExternalAttributes = (int)Result << 16;
-				}
+				Result = FileUtils.GetFileMode_Linux(SourceFileName);
+			}
+			else if (RuntimeInformation.IsOSPlatform(OSPlatform.OSX))
+			{
+				Result = FileUtils.GetFileMode_Mac(SourceFileName);
+			}
+
+			if(Result >= 0)
+			{
+				Entry.ExternalAttributes = (int)Result << 16;
 			}
+
 			return Entry;
 		}
 
@@ -50,7 +58,16 @@ namespace EpicGames.Core
 		public static void ExtractToFile_CrossPlatform(this ZipArchiveEntry Entry, string TargetFileName, bool Overwrite)
 		{
 			ZipFileExtensions.ExtractToFile(Entry, TargetFileName, Overwrite);
-			if (RuntimeInformation.IsOSPlatform(OSPlatform.OSX))
+
+			if (RuntimeInformation.IsOSPlatform(OSPlatform.Linux))
+			{
+				int MadeByPlatform = Convert.ToInt32(VersionMadeByPlatformField.GetValue(Entry)!);
+				if (MadeByPlatform == 3 || MadeByPlatform == 19) // Unix or OSX
+				{
+					FileUtils.SetFileMode_Linux(TargetFileName, (ushort)(Entry.ExternalAttributes >> 16));
+				}
+			}
+			else if (RuntimeInformation.IsOSPlatform(OSPlatform.OSX))
 			{
 				int MadeByPlatform = Convert.ToInt32(VersionMadeByPlatformField.GetValue(Entry)!);
 				if (MadeByPlatform == 3 || MadeByPlatform == 19) // Unix or OSX
