Horde: Add compact binary and string conversion classes for condition objects, and embed them directly in agent requirements. Simplifies unification of agent specifications.

diff --git a/Engine/Source/Programs/Horde/HordeAgent/Commands/Execution/ExecuteCommandV2.cs b/Engine/Source/Programs/Horde/HordeAgent/Commands/Execution/ExecuteCommandV2.cs
--- a/Engine/Source/Programs/Horde/HordeAgent/Commands/Execution/ExecuteCommandV2.cs
+++ b/Engine/Source/Programs/Horde/HordeAgent/Commands/Execution/ExecuteCommandV2.cs
@@ -386,7 +386,7 @@ namespace HordeAgent.Commands
 
 				JsonRequirements JsonRequirements = JsonComputeTask.Requirements;
 				Requirements Requirements = new Requirements(JsonRequirements.Condition ?? String.Empty);
-				Requirements.Resources = JsonRequirements.Resources.ToDictionary(x => (Utf8String)x.Key, x => x.Value);
+				Requirements.Resources = JsonRequirements.Resources.ToDictionary(x => x.Key, x => x.Value);
 				Requirements.Exclusive = JsonRequirements.Exclusive;
 				IoHash RequirementsHash = AddCbObject(Blobs, Requirements);
 
diff --git a/Engine/Source/Programs/Horde/HordeServer/Compute/Impl/ComputeService.cs b/Engine/Source/Programs/Horde/HordeServer/Compute/Impl/ComputeService.cs
--- a/Engine/Source/Programs/Horde/HordeServer/Compute/Impl/ComputeService.cs
+++ b/Engine/Source/Programs/Horde/HordeServer/Compute/Impl/ComputeService.cs
@@ -41,7 +41,6 @@ namespace HordeServer.Compute.Impl
 {
 	using ChannelId = StringId<IComputeChannel>;
 	using NamespaceId = StringId<INamespace>;
-	using Condition = HordeServer.Utilities.Condition;
 
 	/// <summary>
 	/// Information about a particular task
@@ -125,12 +124,6 @@ namespace HordeServer.Compute.Impl
 	/// </summary>
 	class ComputeService : IComputeService, IDisposable
 	{
-		class QueueInfo
-		{
-			public Requirements? Requirements;
-			public Condition? Condition;
-		}
-
 		public MessageDescriptor Descriptor => ComputeTaskMessage.Descriptor;
 
 		public static NamespaceId DefaultNamespaceId { get; } = new NamespaceId("default");
@@ -299,21 +292,16 @@ namespace HordeServer.Compute.Impl
 		/// <returns></returns>
 		async ValueTask<bool> CheckRequirements(IoHash RequirementsHash, Lazy<Task<List<string>>> LazyProperties)
 		{
-			QueueInfo QueueInfo = await GetQueueInfoAsync(RequirementsHash);
-			if (QueueInfo.Requirements == null)
+			Requirements? Requirements = await GetCachedRequirementsAsync(RequirementsHash);
+			if (Requirements == null)
 			{
 				return false;
 			}
 
-			if (!QueueInfo.Requirements.Condition.IsEmpty)
+			if (Requirements.Condition != null)
 			{
-				if (QueueInfo.Condition == null)
-				{
-					return false;
-				}
-
 				List<string> Properties = await LazyProperties.Value;
-				if (!QueueInfo.Condition.Evaluate(x => FindAgentProperty(x, Properties)))
+				if (!Requirements.Condition.Evaluate(x => FindAgentProperty(x, Properties)))
 				{
 					return false;
 				}
@@ -376,29 +364,15 @@ namespace HordeServer.Compute.Impl
 		/// </summary>
 		/// <param name="RequirementsHash"></param>
 		/// <returns></returns>
-		async ValueTask<QueueInfo> GetQueueInfoAsync(IoHash RequirementsHash)
+		async ValueTask<Requirements?> GetCachedRequirementsAsync(IoHash RequirementsHash)
 		{
-			QueueInfo? Info;
-			if (!RequirementsCache.TryGetValue(RequirementsHash, out Info))
+			Requirements? Requirements;
+			if (!RequirementsCache.TryGetValue(RequirementsHash, out Requirements))
 			{
-				Info = new QueueInfo();
-				Info.Requirements = await ObjectCollection.GetAsync<Requirements>(DefaultNamespaceId, RequirementsHash);
-
-				if (Info.Requirements != null && !Info.Requirements.Condition.IsEmpty)
-				{
-					try
-					{
-						Info.Condition = Condition.Parse(Info.Requirements.Condition.ToString());
-					}
-					catch (Exception Ex)
-					{
-						Logger.LogError(Ex, "Unable to parse condition '{Condition}': {Message}", Info.Requirements.Condition.ToString(), Ex.Message);
-					}
-				}
-
+				Requirements = await ObjectCollection.GetAsync<Requirements>(DefaultNamespaceId, RequirementsHash);
 				using (ICacheEntry Entry = RequirementsCache.CreateEntry(RequirementsHash))
 				{
-					if (Info.Requirements == null)
+					if (Requirements == null)
 					{
 						Entry.SetAbsoluteExpiration(TimeSpan.FromSeconds(10.0));
 					}
@@ -406,10 +380,10 @@ namespace HordeServer.Compute.Impl
 					{
 						Entry.SetSlidingExpiration(TimeSpan.FromMinutes(10.0));
 					}
-					Entry.SetValue(Info);
+					Entry.SetValue(Requirements);
 				}
 			}
-			return Info!;
+			return Requirements;
 		}
 	}
 
diff --git a/Engine/Source/Programs/Horde/HordeServer/Utilities/Condition.cs b/Engine/Source/Programs/Horde/HordeServer/Utilities/Condition.cs
deleted file mode 100644
diff --git a/Engine/Source/Programs/Shared/EpicGames.Horde/Common/Condition.cs b/Engine/Source/Programs/Shared/EpicGames.Horde/Common/Condition.cs
new file mode 100644
--- /dev/null
+++ b/Engine/Source/Programs/Shared/EpicGames.Horde/Common/Condition.cs
@@ -0,0 +1,591 @@
+using EpicGames.Core;
+using EpicGames.Serialization;
+using EpicGames.Serialization.Converters;
+using System;
+using System.Collections.Generic;
+using System.ComponentModel;
+using System.Diagnostics;
+using System.Diagnostics.CodeAnalysis;
+using System.Globalization;
+using System.Linq;
+using System.Text;
+using System.Text.RegularExpressions;
+
+namespace EpicGames.Horde.Common
+{
+	enum TokenType : byte
+	{
+		Error,
+		End,
+		Identifier,
+		Scalar,
+		Not,
+		Eq,
+		Neq,
+		LogicalOr,
+		LogicalAnd,
+		Lt,
+		Lte,
+		Gt,
+		Gte,
+		Lparen,
+		Rparen,
+		Regex
+	}
+
+	class TokenReader
+	{
+		public string Input { get; private set; }
+		public int Offset { get; private set; }
+		public int Length { get; private set; }
+
+		public StringView Token => new StringView(Input, Offset, Length);
+		public TokenType Type { get; private set; }
+		public string? Scalar { get; private set; }
+
+		public TokenReader(string Input)
+		{
+			this.Input = Input;
+			MoveNext();
+		}
+
+		private void SetCurrent(int Length, TokenType Type, string? Scalar = null)
+		{
+			this.Length = Length;
+			this.Type = Type;
+			this.Scalar = Scalar;
+		}
+
+		public void MoveNext()
+		{
+			Offset += Length;
+
+			while (Offset < Input.Length)
+			{
+				switch (Input[Offset])
+				{
+					case ' ':
+					case '\t':
+						Offset++;
+						break;
+					case '(':
+						SetCurrent(1, TokenType.Lparen);
+						return;
+					case ')':
+						SetCurrent(1, TokenType.Rparen);
+						return;
+					case '!':
+						if (Offset + 1 < Input.Length && Input[Offset + 1] == '=')
+						{
+							SetCurrent(2, TokenType.Neq);
+						}
+						else
+						{
+							SetCurrent(1, TokenType.Not);
+						}
+						return;
+					case '&':
+						if (RequireCharacter(Input, Offset + 1, '&'))
+						{
+							SetCurrent(2, TokenType.LogicalAnd);
+						}
+						return;
+					case '|':
+						if (RequireCharacter(Input, Offset + 1, '|'))
+						{
+							SetCurrent(2, TokenType.LogicalOr);
+						}
+						return;
+					case '=':
+						if (RequireCharacter(Input, Offset + 1, '='))
+						{
+							SetCurrent(2, TokenType.Eq);
+						}
+						return;
+					case '~':
+						if (RequireCharacter(Input, Offset + 1, '='))
+						{
+							SetCurrent(2, TokenType.Regex);
+						}
+						return;
+					case '>':
+						if (Offset + 1 < Input.Length && Input[Offset + 1] == '=')
+						{
+							SetCurrent(2, TokenType.Gte);
+						}
+						else
+						{
+							SetCurrent(1, TokenType.Gt);
+						}
+						return;
+					case '<':
+						if (Offset + 1 < Input.Length && Input[Offset + 1] == '=')
+						{
+							SetCurrent(2, TokenType.Lte);
+						}
+						else
+						{
+							SetCurrent(1, TokenType.Lt);
+						}
+						return;
+					case '\"':
+					case '\'':
+						ParseString();
+						return;
+					case char Char when IsNumber(Char):
+						ParseNumber();
+						return;
+					case char Char when IsIdentifier(Char):
+						int EndIdx = Offset + 1;
+						while (EndIdx < Input.Length && IsIdentifierTail(Input[EndIdx]))
+						{
+							EndIdx++;
+						}
+						SetCurrent(EndIdx - Offset, TokenType.Identifier);
+						return;
+					default:
+						SetCurrent(1, TokenType.Error, $"Invalid character at offset {Offset}: '{Input[Offset]}'");
+						return;
+				}
+			}
+
+			SetCurrent(0, TokenType.End);
+		}
+
+		bool ParseString()
+		{
+			char QuoteChar = Input[Offset];
+			if (QuoteChar != '\'' && QuoteChar != '\"')
+			{
+				SetCurrent(1, TokenType.Error, $"Invalid quote character '{(char)QuoteChar}' at offset {Offset}");
+				return false;
+			}
+
+			int NumEscapeChars = 0;
+
+			int EndIdx = Offset + 1;
+			for (; ; EndIdx++)
+			{
+				if (EndIdx >= Input.Length)
+				{
+					SetCurrent(EndIdx - Offset, TokenType.Error, "Unterminated string in expression");
+					return false;
+				}
+				else if (Input[EndIdx] == '\\')
+				{
+					NumEscapeChars++;
+					EndIdx++;
+				}
+				else if (Input[EndIdx] == QuoteChar)
+				{
+					break;
+				}
+			}
+
+			char[] Copy = new char[EndIdx - (Offset + 1) - NumEscapeChars];
+
+			int InputIdx = Offset + 1;
+			int OutputIdx = 0;
+			while (OutputIdx < Copy.Length)
+			{
+				if (Input[InputIdx] == '\\')
+				{
+					InputIdx++;
+				}
+				Copy[OutputIdx++] = Input[InputIdx++];
+			}
+
+			SetCurrent(EndIdx + 1 - Offset, TokenType.Scalar, new string(Copy));
+			return true;
+		}
+
+		static Dictionary<StringView, ulong> SizeSuffixes = new Dictionary<StringView, ulong>(StringViewComparer.OrdinalIgnoreCase)
+		{
+			["kb"] = 1024UL,
+			["Mb"] = 1024UL * 1024,
+			["Gb"] = 1024UL * 1024 * 1024,
+			["Tb"] = 1024UL * 1024 * 1024 * 1024,
+		};
+
+		bool ParseNumber()
+		{
+			ulong Value = 0;
+
+			int EndIdx = Offset;
+			while (EndIdx < Input.Length && IsNumber(Input[EndIdx]))
+			{
+				Value = (Value * 10) + (uint)(Input[EndIdx] - '0');
+				EndIdx++;
+			}
+
+			if (EndIdx < Input.Length && IsIdentifier(Input[EndIdx]))
+			{
+				int Offset = EndIdx++;
+				while (EndIdx < Input.Length && IsIdentifierTail(Input[EndIdx]))
+				{
+					EndIdx++;
+				}
+
+				StringView Suffix = new StringView(Input, Offset, EndIdx - Offset);
+
+				ulong Size;
+				if (!SizeSuffixes.TryGetValue(Suffix, out Size))
+				{
+					SetCurrent((EndIdx + 1) - Offset, TokenType.Error, $"'{Suffix}' is not a valid numerical suffix");
+					return false;
+				}
+				Value *= Size;
+			}
+
+			SetCurrent(EndIdx - Offset, TokenType.Scalar, Value.ToString(CultureInfo.InvariantCulture));
+			return true;
+		}
+
+		bool RequireCharacter(string Text, int Idx, char Character)
+		{
+			if (Idx == Text.Length || Text[Idx] != Character)
+			{
+				SetCurrent(1, TokenType.Error, $"Invalid character at position {Idx}; expected '{Character}'.");
+				return false;
+			}
+			return true;
+		}
+
+		static bool IsNumber(char Character)
+		{
+			return Character >= '0' && Character <= '9';
+		}
+
+		static bool IsIdentifier(char Character)
+		{
+			return (Character >= 'a' && Character <= 'z') || (Character >= 'A' && Character <= 'Z') || Character == '_';
+		}
+
+		static bool IsIdentifierTail(char Character)
+		{
+			return IsIdentifier(Character) || IsNumber(Character) || Character == '.';
+		}
+	}
+
+	/// <summary>
+	/// Exception thrown when a condition is not valid
+	/// </summary>
+	public class ConditionException : Exception
+	{
+		/// <summary>
+		/// Constructor
+		/// </summary>
+		/// <param name="Error"></param>
+		internal ConditionException(string Error)
+			: base(Error)
+		{
+		}
+	}
+
+	/// <summary>
+	/// A conditional expression that can be evaluated against a particular object
+	/// </summary>
+	[JsonSchemaString]
+	[TypeConverter(typeof(ConditionTypeConverter))]
+	[CbConverter(typeof(ConditionCbConverter))]
+	public class Condition
+	{
+		[DebuggerDisplay("{Type}")]
+		readonly struct Token
+		{
+			public readonly TokenType Type;
+			public readonly byte Index;
+
+			public Token(TokenType Type, int Index)
+			{
+				this.Type = Type;
+				this.Index = (byte)Index;
+			}
+		}
+
+		public string Text { get; }
+		public string? Error { get; private set; }
+		List<Token> Tokens = new List<Token>();
+		List<string> Strings = new List<string>();
+
+		private Condition(string Text)
+		{
+			this.Text = Text;
+
+			TokenReader Reader = new TokenReader(Text);
+			Error = ParseOrExpr(Reader);
+
+			if (Reader.Type == TokenType.Error)
+			{
+				Error = Reader.Scalar;
+			}
+			else if (Error == null && Reader.Type != TokenType.End)
+			{
+				Error = $"Unexpected token at offset {Reader.Offset}: {Reader.Token}";
+			}
+		}
+
+		public bool IsValid() => Error == null;
+
+		public static Condition Parse(string Text)
+		{
+			Condition Condition = new Condition(Text);
+			if(!Condition.IsValid())
+			{
+				throw new ConditionException(Condition.Error!);
+			}
+			return Condition;
+		}
+
+		public static Condition TryParse(string Text)
+		{
+			return new Condition(Text);
+		}
+
+		string? ParseOrExpr(TokenReader Reader)
+		{
+			int StartCount = Tokens.Count;
+
+			string? Error = ParseAndExpr(Reader);
+			while (Error == null && Reader.Type == TokenType.LogicalOr)
+			{
+				Tokens.Insert(StartCount++, new Token(TokenType.LogicalOr, 0));
+				Reader.MoveNext();
+				Error = ParseAndExpr(Reader);
+			}
+
+			return Error;
+		}
+
+		string? ParseAndExpr(TokenReader Reader)
+		{
+			int StartCount = Tokens.Count;
+
+			string? Error = ParseBooleanExpr(Reader);
+			while (Error == null && Reader.Type == TokenType.LogicalAnd)
+			{
+				Tokens.Insert(StartCount++, new Token(TokenType.LogicalAnd, 0));
+				Reader.MoveNext();
+				Error = ParseBooleanExpr(Reader);
+			}
+
+			return Error;
+		}
+
+		string? ParseBooleanExpr(TokenReader Reader)
+		{
+			switch (Reader.Type)
+			{
+				case TokenType.Not:
+					Tokens.Add(new Token(Reader.Type, 0));
+					Reader.MoveNext();
+					if (Reader.Type != TokenType.Lparen)
+					{
+						return $"Expected '(' at offset {Reader.Offset}";
+					}
+					return ParseSubExpr(Reader);
+				case TokenType.Lparen:
+					return ParseSubExpr(Reader);
+				default:
+					return ParseComparisonExpr(Reader);
+			}
+		}
+
+		string? ParseSubExpr(TokenReader Reader)
+		{
+			Reader.MoveNext();
+
+			string? Error = ParseOrExpr(Reader);
+			if (Error == null)
+			{
+				if (Reader.Type == TokenType.Rparen)
+				{
+					Reader.MoveNext();
+				}
+				else
+				{
+					Error = $"Missing ')' at offset {Reader.Offset}";
+				}
+			}
+			return Error;
+		}
+
+		string? ParseComparisonExpr(TokenReader Reader)
+		{
+			int StartCount = Tokens.Count;
+
+			string? Error = ParseScalarExpr(Reader);
+			if (Error == null)
+			{
+				switch(Reader.Type)
+				{
+					case TokenType.Lt:
+					case TokenType.Lte:
+					case TokenType.Gt:
+					case TokenType.Gte:
+					case TokenType.Eq:
+					case TokenType.Neq:
+					case TokenType.Regex:
+						Tokens.Insert(StartCount, new Token(Reader.Type, 0));
+						Reader.MoveNext();
+						Error = ParseScalarExpr(Reader);
+						break;
+				}
+			}
+
+			return Error;
+		}
+
+		string? ParseScalarExpr(TokenReader Reader)
+		{
+			switch (Reader.Type)
+			{
+				case TokenType.Identifier:
+					Strings.Add(Reader.Token.ToString());
+					Tokens.Add(new Token(TokenType.Identifier, Strings.Count - 1));
+					Reader.MoveNext();
+					return null;
+				case TokenType.Scalar:
+					Strings.Add(Reader.Scalar!);
+					Tokens.Add(new Token(TokenType.Scalar, Strings.Count - 1));
+					Reader.MoveNext();
+					return null;
+				default:
+					return $"Unexpected token '{Reader.Token}' at offset {Reader.Offset}";
+			}
+		}
+
+		public bool Evaluate(Func<string, IEnumerable<string>> GetPropertyValues)
+		{
+			int Idx = 0;
+			return IsValid() && EvaluateCondition(ref Idx, GetPropertyValues);
+		}
+
+		bool EvaluateCondition(ref int Idx, Func<string, IEnumerable<string>> GetPropertyValues)
+		{
+			bool LhsBool;
+			bool RhsBool;
+			IEnumerable<string> LhsScalar;
+			IEnumerable<string> RhsScalar;
+
+			Token Token = Tokens[Idx++];
+			switch (Token.Type)
+			{
+				case TokenType.Not:
+					return !EvaluateCondition(ref Idx, GetPropertyValues);
+				case TokenType.Eq:
+					LhsScalar = EvaluateScalar(ref Idx, GetPropertyValues);
+					RhsScalar = EvaluateScalar(ref Idx, GetPropertyValues);
+					return LhsScalar.Any(x => RhsScalar.Contains(x, StringComparer.OrdinalIgnoreCase));
+				case TokenType.Neq:
+					LhsScalar = EvaluateScalar(ref Idx, GetPropertyValues);
+					RhsScalar = EvaluateScalar(ref Idx, GetPropertyValues);
+					return !LhsScalar.Any(x => RhsScalar.Contains(x, StringComparer.OrdinalIgnoreCase));
+				case TokenType.LogicalOr:
+					LhsBool = EvaluateCondition(ref Idx, GetPropertyValues);
+					RhsBool = EvaluateCondition(ref Idx, GetPropertyValues);
+					return LhsBool || RhsBool;
+				case TokenType.LogicalAnd:
+					LhsBool = EvaluateCondition(ref Idx, GetPropertyValues);
+					RhsBool = EvaluateCondition(ref Idx, GetPropertyValues);
+					return LhsBool && RhsBool;
+				case TokenType.Lt:
+					LhsScalar = EvaluateScalar(ref Idx, GetPropertyValues);
+					RhsScalar = EvaluateScalar(ref Idx, GetPropertyValues);
+					return AsIntegers(LhsScalar).Any(x => AsIntegers(RhsScalar).Any(y => x < y));
+				case TokenType.Lte:
+					LhsScalar = EvaluateScalar(ref Idx, GetPropertyValues);
+					RhsScalar = EvaluateScalar(ref Idx, GetPropertyValues);
+					return AsIntegers(LhsScalar).Any(x => AsIntegers(RhsScalar).Any(y => x <= y));
+				case TokenType.Gt:
+					LhsScalar = EvaluateScalar(ref Idx, GetPropertyValues);
+					RhsScalar = EvaluateScalar(ref Idx, GetPropertyValues);
+					return AsIntegers(LhsScalar).Any(x => AsIntegers(RhsScalar).Any(y => x > y));
+				case TokenType.Gte:
+					LhsScalar = EvaluateScalar(ref Idx, GetPropertyValues);
+					RhsScalar = EvaluateScalar(ref Idx, GetPropertyValues);
+					return AsIntegers(LhsScalar).Any(x => AsIntegers(RhsScalar).Any(y => x >= y));
+				case TokenType.Regex:
+					LhsScalar = EvaluateScalar(ref Idx, GetPropertyValues);
+					RhsScalar = EvaluateScalar(ref Idx, GetPropertyValues);
+					return LhsScalar.Any(x => RhsScalar.Any(y => Regex.IsMatch(x, y, RegexOptions.IgnoreCase)));
+				default:
+					throw new InvalidOperationException("Invalid token type");
+			}
+		}
+
+		IEnumerable<string> EvaluateScalar(ref int Idx, Func<string, IEnumerable<string>> GetPropertyValues)
+		{
+			Token Token = Tokens[Idx++];
+			switch (Token.Type)
+			{
+				case TokenType.Identifier:
+					return GetPropertyValues(Strings[Token.Index]);
+				case TokenType.Scalar:
+					return new string[] { Strings[Token.Index] };
+				default:
+					throw new InvalidOperationException("Invalid token type");
+			}
+		}
+
+		static IEnumerable<long> AsIntegers(IEnumerable<string> Scalars)
+		{
+			foreach (string Scalar in Scalars)
+			{
+				long Value;
+				if (long.TryParse(Scalar, out Value))
+				{
+					yield return Value;
+				}
+			}
+		}
+
+		[return: NotNullIfNotNull("Text")]
+		public static implicit operator Condition?(string? Text)
+		{
+			if (Text == null)
+			{
+				return null;
+			}
+			else
+			{
+				return new Condition(Text);
+			}
+		}
+
+		public override string ToString() => (Error != null) ? $"[Error] {Text}" : Text;
+	}
+
+	/// <summary>
+	/// Converter from conditions to compact binary objects
+	/// </summary>
+	public class ConditionCbConverter : CbConverter<Condition>
+	{
+		/// <inheritdoc/>
+		public override Condition Read(CbField Field) => Condition.TryParse(Field.AsString().ToString());
+
+		/// <inheritdoc/>
+		public override void Write(CbWriter Writer, Condition Value) => Writer.WriteStringValue(Value.Text);
+
+		/// <inheritdoc/>
+		public override void WriteNamed(CbWriter Writer, Utf8String Name, Condition Value) => Writer.WriteString(Name, Value.Text);
+	}
+
+	/// <summary>
+	/// Type converter from strings to PropertyFilter objects
+	/// </summary>
+	sealed class ConditionTypeConverter : TypeConverter
+	{
+		/// <inheritdoc/>
+		public override bool CanConvertFrom(ITypeDescriptorContext Context, Type SourceType) => SourceType == typeof(string);
+
+		/// <inheritdoc/>
+		public override object ConvertFrom(ITypeDescriptorContext Context, CultureInfo Culture, object Value) => Condition.TryParse((string)Value);
+
+		/// <inheritdoc/>
+		public override bool CanConvertTo(ITypeDescriptorContext Context, Type DestinationType) => DestinationType == typeof(string);
+
+		/// <inheritdoc/>
+		public override object? ConvertTo(ITypeDescriptorContext Context, CultureInfo Culture, object Value, Type DestinationType) => ((Condition)Value).Text;
+	}
+}
diff --git a/Engine/Source/Programs/Shared/EpicGames.Horde/Compute/Requirements.cs b/Engine/Source/Programs/Shared/EpicGames.Horde/Compute/Requirements.cs
--- a/Engine/Source/Programs/Shared/EpicGames.Horde/Compute/Requirements.cs
+++ b/Engine/Source/Programs/Shared/EpicGames.Horde/Compute/Requirements.cs
@@ -1,6 +1,7 @@
 // Copyright Epic Games, Inc. All Rights Reserved.
 
 using EpicGames.Core;
+using EpicGames.Horde.Common;
 using EpicGames.Serialization;
 using System;
 using System.Collections.Generic;
@@ -17,13 +18,13 @@ namespace EpicGames.Horde.Compute
 		/// Condition string to be evaluated against the machine spec, eg. cpu-cores >= 10 && ram.mb >= 200 && pool == 'worker'
 		/// </summary>
 		[CbField("c")]
-		public Utf8String Condition { get; set; }
+		public Condition? Condition { get; set; }
 
 		/// <summary>
 		/// Resources used by the process
 		/// </summary>
 		[CbField("r")]
-		public Dictionary<Utf8String, int> Resources { get; set; } = new Dictionary<Utf8String, int>();
+		public Dictionary<string, int> Resources { get; set; } = new Dictionary<string, int>();
 
 		/// <summary>
 		/// Whether we require exclusive access to the device
@@ -42,7 +43,7 @@ namespace EpicGames.Horde.Compute
 		/// Construct a requirements object with a condition
 		/// </summary>
 		/// <param name="Condition"></param>
-		public Requirements(Utf8String Condition)
+		public Requirements(Condition? Condition)
 		{
 			this.Condition = Condition;
 		}
diff --git a/Engine/Source/Programs/Shared/EpicGames.Serialization/Converters/CbDefaultConverter.cs b/Engine/Source/Programs/Shared/EpicGames.Serialization/Converters/CbDefaultConverter.cs
--- a/Engine/Source/Programs/Shared/EpicGames.Serialization/Converters/CbDefaultConverter.cs
+++ b/Engine/Source/Programs/Shared/EpicGames.Serialization/Converters/CbDefaultConverter.cs
@@ -117,9 +117,9 @@ namespace EpicGames.Serialization.Converters
 
 		static Dictionary<Type, CbConverterInfo> TypeToConverterInfo = new Dictionary<Type, CbConverterInfo>(new KeyValuePair<Type, CbConverterInfo>[]
 		{
-			GetPodConverterInfo(x => x.AsBool(), (w, v) => w.WriteBoolValue(v), (w, n, v) => w.WriteBool(n, v), null),
-			GetPodConverterInfo(x => x.AsInt32(), (w, v) => WriteInt32Value(w, v), (w, n, v) => WriteInt32(w, n, v), null),
-			GetPodConverterInfo(x => x.AsInt64(), (w, v) => WriteInt64Value(w, v), (w, n, v) => WriteInt64(w, n, v), null),
+			GetPodConverterInfo(x => x.AsBool(), (w, v) => w.WriteBoolValue(v), (w, n, v) => w.WriteBool(n, v), SkipIfNullOrZero),
+			GetPodConverterInfo(x => x.AsInt32(), (w, v) => WriteInt32Value(w, v), (w, n, v) => WriteInt32(w, n, v), SkipIfNullOrZero),
+			GetPodConverterInfo(x => x.AsInt64(), (w, v) => WriteInt64Value(w, v), (w, n, v) => WriteInt64(w, n, v), SkipIfNullOrZero),
 			GetPodConverterInfo(x => x.AsString(), (w, v) => w.WriteStringValue(v), (w, n, v) => w.WriteString(n, v), null),
 			GetPodConverterInfo(x => x.AsHash(), (w, v) => w.WriteHashValue(v), (w, n, v) => w.WriteHash(n, v), null),
 			GetPodConverterInfo(x => x.AsObjectAttachment(), (w, v) => w.WriteObjectAttachmentValue(v.Hash), (w, n, v) => w.WriteObjectAttachment(n, v.Hash), null),
