Horde: Missing copyright notice.

diff --git a/Engine/Source/Programs/Shared/EpicGames.Horde/Common/Condition.cs b/Engine/Source/Programs/Shared/EpicGames.Horde/Common/Condition.cs
--- a/Engine/Source/Programs/Shared/EpicGames.Horde/Common/Condition.cs
+++ b/Engine/Source/Programs/Shared/EpicGames.Horde/Common/Condition.cs
@@ -1,3 +1,5 @@
+// Copyright Epic Games, Inc. All Rights Reserved.
+
 using EpicGames.Core;
 using EpicGames.Serialization;
 using EpicGames.Serialization.Converters;
@@ -313,18 +315,23 @@ namespace EpicGames.Horde.Common
 			this.Text = Text;
 
 			TokenReader Reader = new TokenReader(Text);
-			Error = ParseOrExpr(Reader);
-
-			if (Reader.Type == TokenType.Error)
-			{
-				Error = Reader.Scalar;
-			}
-			else if (Error == null && Reader.Type != TokenType.End)
+			if (Reader.Type != TokenType.End)
 			{
-				Error = $"Unexpected token at offset {Reader.Offset}: {Reader.Token}";
+				Error = ParseOrExpr(Reader);
+
+				if (Reader.Type == TokenType.Error)
+				{
+					Error = Reader.Scalar;
+				}
+				else if (Error == null && Reader.Type != TokenType.End)
+				{
+					Error = $"Unexpected token at offset {Reader.Offset}: {Reader.Token}";
+				}
 			}
 		}
 
+		public bool IsEmpty() => Tokens.Count == 0 && IsValid();
+
 		public bool IsValid() => Error == null;
 
 		public static Condition Parse(string Text)
@@ -457,6 +464,11 @@ namespace EpicGames.Horde.Common
 
 		public bool Evaluate(Func<string, IEnumerable<string>> GetPropertyValues)
 		{
+			if (IsEmpty())
+			{
+				return true;
+			}
+
 			int Idx = 0;
 			return IsValid() && EvaluateCondition(ref Idx, GetPropertyValues);
 		}
