ManagedProcess: Remove unnecessary locks from merged stream copy when creating background threads to merge output streams

#rnx

diff --git a/Engine/Source/Programs/Shared/EpicGames.Core/ManagedProcess.cs b/Engine/Source/Programs/Shared/EpicGames.Core/ManagedProcess.cs
--- a/Engine/Source/Programs/Shared/EpicGames.Core/ManagedProcess.cs
+++ b/Engine/Source/Programs/Shared/EpicGames.Core/ManagedProcess.cs
@@ -406,7 +406,7 @@ namespace EpicGames.Core
 		Process? FrameworkProcess;
 
 		/// <summary>
-		/// Thread to read from stdin
+		/// Thread to read from stdout
 		/// </summary>
 		Thread? FrameworkStdOutThread;
 
@@ -420,6 +420,11 @@ namespace EpicGames.Core
 		/// </summary>
 		AnonymousPipeServerStream? FrameworkMergedStdWriter;
 
+		/// <summary>
+		/// Tracks how many pipes have been copied to the merged writer.
+		/// </summary>
+		int FrameworkMergedStdWriterThreadCount;
+
 		/// <summary>
 		/// Static lock object. This is used to synchronize the creation of child processes - in particular, the inheritance of stdout/stderr write pipes. If processes
 		/// inherit pipes meant for other processes, they won't be closed until both terminate.
@@ -780,29 +785,6 @@ namespace EpicGames.Core
 			// Merge StdOut with StdErr
 			if ((ManagedFlags & ManagedProcessFlags.MergeOutputPipes) != 0)
 			{
-				FrameworkProcess.EnableRaisingEvents = true;
-				FrameworkProcess.Exited += (Sender, Args) =>
-				{
-					lock (this)
-					{
-						if (FrameworkStdOutThread != null)
-						{
-							FrameworkStdOutThread.Join();
-							FrameworkStdOutThread = null;
-						}
-						if (FrameworkStdErrThread != null)
-						{
-							FrameworkStdErrThread.Join();
-							FrameworkStdErrThread = null;
-						}
-						if (FrameworkMergedStdWriter != null)
-						{
-							FrameworkMergedStdWriter.Flush();
-							FrameworkMergedStdWriter.Dispose();
-							FrameworkMergedStdWriter = null;
-						}
-					}
-				};
 				// AnonymousPipes block reading even if the stream has been fully read, until the writer pipe handle is closed.
 				FrameworkMergedStdWriter = new AnonymousPipeServerStream(PipeDirection.Out);
 				StdOut = new AnonymousPipeClientStream(PipeDirection.In, FrameworkMergedStdWriter.ClientSafePipeHandle);
@@ -810,26 +792,36 @@ namespace EpicGames.Core
 
 				StdErr = StdOut;
 				StdErrText = StdOutText;
-			}
+				FrameworkMergedStdWriterThreadCount = 2;
 
-			lock (this)
-			{
-				FrameworkStartTime = DateTime.Now;
-				FrameworkProcess.Start();
-				if ((ManagedFlags & ManagedProcessFlags.MergeOutputPipes) != 0)
-				{
-					FrameworkStdOutThread = new Thread(() => CopyPipe(FrameworkProcess.StandardOutput.BaseStream, FrameworkMergedStdWriter!));
-					FrameworkStdOutThread.Name = $"ManagedProcess {FrameworkProcess.Id} Merge StdOut";
-					FrameworkStdOutThread.IsBackground = true;
-					FrameworkStdOutThread.Start();
-
-					FrameworkStdErrThread = new Thread(() => CopyPipe(FrameworkProcess.StandardError.BaseStream, FrameworkMergedStdWriter!));
-					FrameworkStdErrThread.Name = $"ManagedProcess {FrameworkProcess.Id} Merge StdErr";
-					FrameworkStdErrThread.IsBackground = true;
-					FrameworkStdErrThread.Start();
-				}
+				FrameworkStdOutThread = new Thread(() => {
+					CopyPipe(FrameworkProcess.StandardOutput.BaseStream, FrameworkMergedStdWriter!);
+					if (Interlocked.Decrement(ref FrameworkMergedStdWriterThreadCount) == 0)
+					{
+						// Dispose AnonymousPipe to unblock readers.
+						FrameworkMergedStdWriter?.Dispose();
+					}
+				});
+				FrameworkStdOutThread.Name = $"ManagedProcess Merge StdOut";
+				FrameworkStdOutThread.IsBackground = true;
+
+				FrameworkStdErrThread = new Thread(() => {
+					CopyPipe(FrameworkProcess.StandardError.BaseStream, FrameworkMergedStdWriter!);
+					if (Interlocked.Decrement(ref FrameworkMergedStdWriterThreadCount) == 0)
+					{
+						// Dispose AnonymousPipe to unblock readers.
+						FrameworkMergedStdWriter?.Dispose();
+					}
+				});
+				FrameworkStdErrThread.Name = $"ManagedProcess Merge StdErr";
+				FrameworkStdErrThread.IsBackground = true;
 			}
 
+			FrameworkStartTime = DateTime.Now;
+			FrameworkProcess.Start();
+			FrameworkStdOutThread?.Start();
+			FrameworkStdErrThread?.Start();
+
 			try
 			{
 				FrameworkStartTime = FrameworkProcess.StartTime;
@@ -867,6 +859,25 @@ namespace EpicGames.Core
 			try
 			{
 				Source.CopyTo(Target);
+				Target.Flush();
+			}
+			catch
+			{
+			}
+		}
+
+		/// <summary>
+		/// Copy data from one pipe to another.
+		/// </summary>
+		/// <param name="Source"></param>
+		/// <param name="Target"></param>
+		/// <param name="CancellationToken"></param>
+		async Task CopyPipeAsync(Stream Source, Stream Target, CancellationToken CancellationToken)
+		{
+			try
+			{
+				await Source.CopyToAsync(Target, CancellationToken);
+				await Target.FlushAsync(CancellationToken);
 			}
 			catch
 			{
@@ -909,7 +920,7 @@ namespace EpicGames.Core
 			}
 			if(FrameworkProcess != null)
 			{
-				if(!FrameworkProcess.HasExited)
+				if (!FrameworkProcess.HasExited)
 				{
 					try
 					{
@@ -920,15 +931,18 @@ namespace EpicGames.Core
 					{
 					}
 				}
+				
+				FrameworkMergedStdWriter?.Dispose();
+				FrameworkMergedStdWriter = null;
+
+				FrameworkStdOutThread?.Join();
+				FrameworkStdOutThread = null;
+				FrameworkStdErrThread?.Join();
+				FrameworkStdErrThread = null;
 
 				FrameworkProcess.Dispose();
 				FrameworkProcess = null;
 			}
-			if(FrameworkMergedStdWriter != null)
-			{
-				FrameworkMergedStdWriter.Dispose();
-				FrameworkMergedStdWriter = null;
-			}
 		}
 
 		/// <summary>
