Horde: Add native implementation of IPerforceConnection.

diff --git a/Engine/Source/Programs/Shared/EpicGames.Perforce/NativePerforceConnection.cs b/Engine/Source/Programs/Shared/EpicGames.Perforce/NativePerforceConnection.cs
new file mode 100644
--- /dev/null
+++ b/Engine/Source/Programs/Shared/EpicGames.Perforce/NativePerforceConnection.cs
@@ -0,0 +1,498 @@
+// Copyright Epic Games, Inc. All Rights Reserved.
+
+using EpicGames.Perforce;
+using Microsoft.Extensions.Logging;
+using System;
+using System.Collections.Concurrent;
+using System.Collections.Generic;
+using System.Diagnostics;
+using System.Linq;
+using System.Runtime.InteropServices;
+using System.Threading;
+using System.Threading.Channels;
+using System.Threading.Tasks;
+
+namespace EpicGames.Perforce
+{
+	/// <summary>
+	/// Experimental implementation of <see cref="IPerforceConnection"/> which wraps the native C++ API.
+	/// </summary>
+	public class NativePerforceConnection : IPerforceConnection, IDisposable
+	{
+		const string NativeDll = "EpicGames.Perforce.Native";
+
+		[StructLayout(LayoutKind.Sequential)]
+		class NativeSettings
+		{
+			[MarshalAs(UnmanagedType.LPStr)]
+			public string? Port;
+
+			[MarshalAs(UnmanagedType.LPStr)]
+			public string? User;
+
+			[MarshalAs(UnmanagedType.LPStr)]
+			public string? Client;
+
+			[MarshalAs(UnmanagedType.LPStr)]
+			public string? AppName;
+
+			[MarshalAs(UnmanagedType.LPStr)]
+			public string? AppVersion;
+		}
+
+		[StructLayout(LayoutKind.Sequential)]
+		class NativeReadBuffer
+		{
+			public IntPtr Data;
+			public int Length;
+			public int Count;
+			public int MaxLength;
+			public int MaxCount;
+		};
+
+		[StructLayout(LayoutKind.Sequential)]
+		class NativeWriteBuffer
+		{
+			public IntPtr Data;
+			public int MaxLength;
+			public int MaxCount;
+		};
+
+		[UnmanagedFunctionPointer(CallingConvention.Cdecl)]
+		delegate void OnBufferReadyFn(NativeReadBuffer ReadBuffer, [In, Out] NativeWriteBuffer WriteBuffer);
+
+		[DllImport(NativeDll, CallingConvention = CallingConvention.Cdecl)]
+		static extern IntPtr Client_Create(NativeSettings? Settings, NativeWriteBuffer WriteBuffer, IntPtr OnBufferReadyFnPtr);
+
+		[DllImport(NativeDll, CallingConvention = CallingConvention.Cdecl, CharSet = CharSet.Ansi, BestFitMapping = false, ThrowOnUnmappableChar = true)]
+		static extern void Client_Login(IntPtr Client, [MarshalAs(UnmanagedType.LPStr)] string Password);
+
+		[DllImport(NativeDll, CallingConvention = CallingConvention.Cdecl, CharSet = CharSet.Ansi, BestFitMapping = false, ThrowOnUnmappableChar = true)]
+		static extern void Client_Command(IntPtr Client, [MarshalAs(UnmanagedType.LPStr)] string Command, int NumArgs, string[] Args);
+
+		[DllImport(NativeDll, CallingConvention = CallingConvention.Cdecl)]
+		static extern void Client_Destroy(IntPtr Client);
+
+		/// <summary>
+		/// A buffer used for native code to stream data into
+		/// </summary>
+		class PinnedBuffer : IDisposable
+		{
+			public byte[] Data { get; private set; }
+			public GCHandle Handle { get; private set; }
+			public IntPtr BasePtr { get; private set; }
+			public int MaxLength => Data.Length;
+
+			public PinnedBuffer(int MaxLength)
+			{
+				Data = new byte[MaxLength];
+				Handle = GCHandle.Alloc(Data, GCHandleType.Pinned);
+				BasePtr = Handle.AddrOfPinnedObject();
+			}
+
+			public void Resize(int MaxLength)
+			{
+				byte[] OldData = Data;
+				Handle.Free();
+
+				Data = new byte[MaxLength];
+				Handle = GCHandle.Alloc(Data, GCHandleType.Pinned);
+				BasePtr = Handle.AddrOfPinnedObject();
+
+				OldData.CopyTo(Data, 0);
+			}
+
+			public void Dispose()
+			{
+				Handle.Free();
+			}
+		}
+
+		/// <summary>
+		/// Response object for a request
+		/// </summary>
+		class Response : IPerforceOutput
+		{
+			NativePerforceConnection Outer;
+
+			PinnedBuffer? Buffer;
+			int BufferPos;
+			int BufferLen;
+
+			PinnedBuffer? NextBuffer;
+			int NextBufferPos;
+			int NextBufferLen;
+
+			public Channel<(PinnedBuffer Buffer, int Length)> ReadBuffers = Channel.CreateUnbounded<(PinnedBuffer, int)>();
+
+			public ReadOnlyMemory<byte> Data => (Buffer == null) ? ReadOnlyMemory<byte>.Empty : Buffer.Data.AsMemory(BufferPos, BufferLen);
+
+			public Response(NativePerforceConnection Outer)
+			{
+				this.Outer = Outer;
+			}
+
+			public async ValueTask DisposeAsync()
+			{
+				if (Buffer != null)
+				{
+					Outer.WriteBuffers.Add(Buffer);
+				}
+				if (NextBuffer != null)
+				{
+					Outer.WriteBuffers.Add(NextBuffer);
+				}
+
+				while (await ReadBuffers.Reader.WaitToReadAsync())
+				{
+					(PinnedBuffer, int) Buffer;
+					if (ReadBuffers.Reader.TryRead(out Buffer))
+					{
+						Outer.WriteBuffers.Add(Buffer.Item1);
+					}
+				}
+
+				Outer.ResponseCompleteEvent.Set();
+			}
+
+			async Task<(PinnedBuffer?, int)> GetNextReadBufferAsync()
+			{
+				for (; ; )
+				{
+					if (!await ReadBuffers.Reader.WaitToReadAsync())
+					{
+						return (null, 0);
+					}
+
+					(PinnedBuffer, int) Pair;
+					if (ReadBuffers.Reader.TryRead(out Pair))
+					{
+						return Pair;
+					}
+				}
+			}
+
+			public async Task<bool> ReadAsync(CancellationToken Token)
+			{
+				// If we don't have any data yet, wait until a read completes
+				if (Buffer == null)
+				{
+					(Buffer, BufferLen) = await GetNextReadBufferAsync();
+					return Buffer != null;
+				}
+
+				// Ensure there's some space in the current buffer. In order to handle cases where we want to read data straddling both buffers, copy 16k chunks
+				// back to the first buffer until we can read entirely from the second buffer.
+				int MaxAppend = Buffer.MaxLength - BufferLen;
+				if (MaxAppend == 0)
+				{
+					if (BufferPos > 0)
+					{
+						Buffer.Data.AsSpan(BufferPos, BufferLen - BufferPos).CopyTo(Buffer.Data);
+						BufferLen -= BufferPos;
+						BufferPos = 0;
+					}
+					else
+					{
+						Buffer.Resize(Buffer.MaxLength + 16384);
+					}
+					MaxAppend = Buffer.MaxLength - BufferLen;
+				}
+
+				// Read the next buffer
+				if (NextBuffer == null)
+				{
+					(NextBuffer, NextBufferLen) = await GetNextReadBufferAsync();
+					if (NextBuffer == null)
+					{
+						return false;
+					}
+				}
+
+				// Try to copy some data from the next buffer
+				int CopyLen = Math.Min(NextBufferLen - NextBufferPos, Math.Min(MaxAppend, 16384));
+				NextBuffer.Data.AsSpan(NextBufferPos, CopyLen).CopyTo(Buffer.Data.AsSpan(BufferLen));
+				BufferLen += CopyLen;
+				NextBufferPos += CopyLen;
+
+				// If we've read everything from the next buffer, return it to the write list
+				if (NextBufferPos == NextBufferLen)
+				{
+					Outer.WriteBuffers.Add(NextBuffer);
+
+					NextBuffer = null;
+					NextBufferPos = 0;
+					NextBufferLen = 0;
+				}
+
+				return true;
+			}
+
+			public void Discard(int NumBytes)
+			{
+				if (NumBytes > 0)
+				{
+					// Update the read position
+					BufferPos += NumBytes;
+					Debug.Assert(BufferPos <= BufferLen);
+
+					// If we've used up all the data in the buffer, return it to the write list and move to the next one.
+					int OriginalBufferLen = BufferLen - NextBufferPos;
+					if (BufferPos >= OriginalBufferLen)
+					{
+						Outer.WriteBuffers.Add(Buffer!);
+
+						Buffer = NextBuffer;
+						BufferPos -= OriginalBufferLen;
+						BufferLen = NextBufferLen;
+
+						NextBuffer = null;
+						NextBufferPos = 0;
+						NextBufferLen = 0;
+					}
+				}
+			}
+		}
+
+		IntPtr Client;
+		PinnedBuffer[] Buffers;
+		OnBufferReadyFn OnBufferReadyInst;
+		IntPtr OnBufferReadyFnPtr;
+		Thread? BackgroundThread;
+		BlockingCollection<(Action, Response)?> Requests = new BlockingCollection<(Action, Response)?>();
+		BlockingCollection<PinnedBuffer> WriteBuffers = new BlockingCollection<PinnedBuffer>();
+		Response? CurrentResponse;
+		ManualResetEvent ResponseCompleteEvent;
+
+		/// <inheritdoc/>
+		public ILogger Logger { get; }
+
+		/// <summary>
+		/// Constructor
+		/// </summary>
+		/// <param name="Logger">Logger for messages</param>
+		public NativePerforceConnection(ILogger Logger)
+			: this(2, 64 * 1024, Logger)
+		{
+		}
+
+		/// <summary>
+		/// Constructor
+		/// </summary>
+		/// <param name="BufferCount">Number of buffers to create for streaming response data</param>
+		/// <param name="BufferSize">Size of each buffer</param>
+		/// <param name="Logger">Logger for messages</param>
+		public NativePerforceConnection(int BufferCount, int BufferSize, ILogger Logger)
+		{
+			this.Logger = Logger;
+
+			Buffers = new PinnedBuffer[BufferCount];
+			for (int Idx = 0; Idx < BufferCount; Idx++)
+			{
+				Buffers[Idx] = new PinnedBuffer(BufferSize);
+				WriteBuffers.TryAdd(Buffers[Idx]);
+			}
+
+			OnBufferReadyInst = new OnBufferReadyFn(OnBufferReady);
+			OnBufferReadyFnPtr = Marshal.GetFunctionPointerForDelegate(OnBufferReadyInst);
+
+			ResponseCompleteEvent = new ManualResetEvent(false);
+
+			BackgroundThread = new Thread(BackgroundThreadProc);
+			BackgroundThread.Start();
+		}
+
+		void GetNextWriteBuffer(NativeWriteBuffer NativeWriteBuffer)
+		{
+			PinnedBuffer Buffer = WriteBuffers.Take();
+			NativeWriteBuffer.Data = Buffer.BasePtr;
+			NativeWriteBuffer.MaxLength = Buffer.Data.Length;
+			NativeWriteBuffer.MaxCount = int.MaxValue;
+		}
+
+		/// <summary>
+		/// Finalizer
+		/// </summary>
+		~NativePerforceConnection()
+		{
+			Dispose();
+		}
+
+		/// <inheritdoc/>
+		public void Dispose()
+		{
+			if (BackgroundThread != null)
+			{
+				Requests.Add(null);
+				Requests.CompleteAdding();
+
+				BackgroundThread.Join();
+				BackgroundThread = null!;
+			}
+
+			Requests.Dispose();
+
+			if (Client != IntPtr.Zero)
+			{
+				Client_Destroy(Client);
+				Client = IntPtr.Zero;
+			}
+
+			WriteBuffers.Dispose();
+
+			ResponseCompleteEvent.Dispose();
+
+			foreach (PinnedBuffer PinnedBuffer in Buffers)
+			{
+				PinnedBuffer.Dispose();
+			}
+
+			GC.SuppressFinalize(this);
+		}
+
+		/// <summary>
+		/// Initializes the connection, throwing an error on failure
+		/// </summary>
+		public async Task ConnectAsync(PerforceSettings? Settings)
+		{
+			PerforceError? Error = await TryConnectAsync(Settings);
+			if (Error != null)
+			{
+				throw new PerforceException(Error);
+			}
+		}
+
+		/// <summary>
+		/// Tries to initialize the connection
+		/// </summary>
+		/// <returns>Error returned when attempting to connect</returns>
+		public async Task<PerforceError?> TryConnectAsync(PerforceSettings? Settings)
+		{
+			await using Response Response = new Response(this);
+
+			NativeSettings? NativeSettings = null;
+			if (Settings != null)
+			{
+				NativeSettings = new NativeSettings();
+				NativeSettings.Port = Settings.Port;
+				NativeSettings.User = Settings.User;
+				NativeSettings.Client = Settings.Client;
+				NativeSettings.AppName = Settings.AppName;
+				NativeSettings.AppVersion = Settings.AppVersion;
+			}
+
+			Requests.Add((() =>
+			{
+				NativeWriteBuffer WriteBuffer = new NativeWriteBuffer();
+				GetNextWriteBuffer(WriteBuffer);
+				Client = Client_Create(NativeSettings, WriteBuffer, OnBufferReadyFnPtr);
+			}, Response));
+
+			List<PerforceResponse> Records = await ((IPerforceOutput)Response).ReadResponsesAsync(null, default);
+			if (Records.Count != 1)
+			{
+				throw new PerforceException("Expected at least one record to be returned from Init() call.");
+			}
+
+			PerforceError? Error = Records[0].Error;
+			if (Error == null)
+			{
+				throw new PerforceException("Unexpected response from init call");
+			}
+			if (Error.Severity != PerforceSeverityCode.Empty)
+			{
+				return Error;
+			}
+			return null;
+		}
+
+		/// <summary>
+		/// Background thread which sequences requests on a single thread. The Perforce API isn't async aware, but is primarily
+		/// I/O bound, so this thread will mostly be idle. All processing C#-side is done using async tasks, whereas this thread
+		/// blocks.
+		/// </summary>
+		void BackgroundThreadProc()
+		{
+			for (; ; )
+			{
+				(Action Action, Response Response)? Request = Requests.Take();
+				if (Request == null)
+				{
+					break;
+				}
+
+				CurrentResponse = Request.Value.Response;
+				ResponseCompleteEvent.Reset();
+
+				Request.Value.Action();
+
+				CurrentResponse.ReadBuffers.Writer.TryComplete();
+				ResponseCompleteEvent.WaitOne();
+
+				CurrentResponse = null;
+			}
+		}
+
+		/// <summary>
+		/// Callback for switching buffers
+		/// </summary>
+		/// <param name="ReadBuffer">The complete buffer</param>
+		/// <param name="WriteBuffer">Receives information about the next buffer to write to</param>
+		void OnBufferReady(NativeReadBuffer ReadBuffer, [In, Out] NativeWriteBuffer WriteBuffer)
+		{
+			PinnedBuffer Buffer = Buffers.First(x => x.BasePtr == ReadBuffer.Data);
+			CurrentResponse!.ReadBuffers.Writer.TryWrite((Buffer, ReadBuffer.Length)); // Unbounded; will always succeed
+			GetNextWriteBuffer(WriteBuffer);
+		}
+
+		/// <inheritdoc/>
+		public Task<IPerforceOutput> CommandAsync(string Command, IReadOnlyList<string> Arguments, IReadOnlyList<string>? FileArguments, byte[]? InputData)
+		{
+			if (InputData != null)
+			{
+				throw new NotImplementedException();
+			}
+
+			List<string> AllArguments = new List<string>(Arguments);
+			if (FileArguments != null)
+			{
+				AllArguments.AddRange(FileArguments);
+			}
+
+			Response Response = new Response(this);
+			Requests.Add((() => Client_Command(Client, Command, AllArguments.Count, AllArguments.ToArray()), Response));
+			return Task.FromResult<IPerforceOutput>(Response);
+		}
+
+		/// <inheritdoc/>
+		public async Task LoginAsync(string Password, CancellationToken CancellationToken = default)
+		{
+			await using Response Response = new Response(this);
+			Requests.Add((() => Client_Login(Client, Password), Response));
+
+			List<PerforceResponse> Records = await ((IPerforceOutput)Response).ReadResponsesAsync(null, default);
+			if (Records.Count != 1)
+			{
+				throw new PerforceException("Expected at least one record to be returned from Init() call.");
+			}
+
+			PerforceError? Error = Records[0].Error;
+			if (Error != null && Error.Severity != PerforceSeverityCode.Info)
+			{
+				throw new PerforceException(Error);
+			}
+		}
+
+		/// <inheritdoc/>
+		public Task SetAsync(string Name, string Value, CancellationToken CancellationToken = default)
+		{
+			throw new NotImplementedException();
+		}
+
+		/// <inheritdoc/>
+		public Task<string?> TryGetSettingAsync(string Name, CancellationToken CancellationToken = default)
+		{
+			throw new NotImplementedException();
+		}
+	}
+}
