Horde: Add accessors for string comparison types to perform on Perforce servers, depending on settings returned by the 'info' command.

diff --git a/Engine/Source/Programs/Shared/EpicGames.Perforce/Records/InfoRecord.cs b/Engine/Source/Programs/Shared/EpicGames.Perforce/Records/InfoRecord.cs
--- a/Engine/Source/Programs/Shared/EpicGames.Perforce/Records/InfoRecord.cs
+++ b/Engine/Source/Programs/Shared/EpicGames.Perforce/Records/InfoRecord.cs
@@ -61,6 +61,16 @@ namespace EpicGames.Perforce
 		[PerforceTag("caseHandling", Optional = true)]
 		public string? CaseHandling;
 
+		/// <summary>
+		/// How to compare paths on this server
+		/// </summary>
+		public StringComparer PathComparer => String.Equals(CaseHandling, "insensitive", StringComparison.OrdinalIgnoreCase) ? StringComparer.OrdinalIgnoreCase : StringComparer.Ordinal;
+
+		/// <summary>
+		/// How to compare paths on this server
+		/// </summary>
+		public StringComparison PathComparison => String.Equals(CaseHandling, "insensitive", StringComparison.OrdinalIgnoreCase) ? StringComparison.OrdinalIgnoreCase : StringComparison.Ordinal;
+
 		/// <summary>
 		/// List of services provided by this server
 		/// </summary>
