Horde: Add support for the "p4 dirs" command.

This requires special handling for the native Perforce connection, since it returns information as message records that the client formats as strings through an "Error" object. To work around this, the C++ library converts such messages into stat records that are more easily parsed by C# code.

diff --git a/Engine/Source/Programs/Shared/EpicGames.Perforce.Native/Client.cpp b/Engine/Source/Programs/Shared/EpicGames.Perforce.Native/Client.cpp
--- a/Engine/Source/Programs/Shared/EpicGames.Perforce.Native/Client.cpp
+++ b/Engine/Source/Programs/Shared/EpicGames.Perforce.Native/Client.cpp
@@ -1,10 +1,12 @@
 // Copyright Epic Games, Inc. All Rights Reserved.
 
 #include "clientapi.h"
+#include "strtable.h"
 #include <assert.h>
 
 #ifdef _MSC_VER
 	#define NATIVE_API __declspec(dllexport)
+	#define strcasecmp _stricmp
 #else
 	#define NATIVE_API
 #endif
@@ -47,6 +49,7 @@ private:
 	FOnBufferReadyFn* OnBufferReady;
 
 public:
+	const char* Func = nullptr;
 	const char* PromptResponse = nullptr;
 
 	FClientUser(FWriteBuffer* WriteBuffer, FOnBufferReadyFn* OnBufferReady)
@@ -109,14 +112,6 @@ public:
 		}
 	}
 
-	virtual void Message(Error* err) override
-	{
-		while (!TryOutputError(err))
-		{
-			Flush();
-		}
-	}
-
 	bool TryOutputError(Error* Err)
 	{
 		static const char CodeKey[] = "code";
@@ -149,6 +144,38 @@ public:
 		return true;
 	}
 
+	virtual void Message(Error* Err) override
+	{
+		while (!TryOutputMessage(Err))
+		{
+			Flush();
+		}
+	}
+
+	bool TryOutputMessage(Error* Err)
+	{
+		// Try to translate these into stat messages depending on the command
+		if (Func != nullptr)
+		{
+			if (Err->GetSeverity() == E_INFO && strcasecmp(Func, "dirs") == 0)
+			{
+				StrDict* Dict = Err->GetDict();
+				if (Dict != nullptr)
+				{
+					StrPtr* DirName = Dict->GetVar("dirName");
+					if (DirName != nullptr)
+					{
+						StrBufDict StatDict;
+						StatDict.VSetVar(StrBuf("dir"), *DirName);
+						return TryWriteRecord(&StatDict);
+					}
+				}
+			}
+		}
+
+		return TryOutputError(Err);
+	}
+
 	virtual void OutputError(const char* errBuf) override
 	{
 		assert(false);
@@ -398,9 +425,11 @@ extern "C" NATIVE_API void Client_Login(FClient * Client, const char* Password)
 
 extern "C" NATIVE_API void Client_Command(FClient* Client, const char* Func, int ArgCount, const char** Args)
 {
+	Client->User.Func = Func;
 	Client->ClientApi.SetArgv(ArgCount, (char* const*)Args);
 	Client->ClientApi.Run(Func, &Client->User);
 	Client->User.Flush();
+	Client->User.Func = nullptr;
 }
 
 extern "C" NATIVE_API void Client_Destroy(FClient* Client)
diff --git a/Engine/Source/Programs/Shared/EpicGames.Perforce.Native/Makefile b/Engine/Source/Programs/Shared/EpicGames.Perforce.Native/Makefile
--- a/Engine/Source/Programs/Shared/EpicGames.Perforce.Native/Makefile
+++ b/Engine/Source/Programs/Shared/EpicGames.Perforce.Native/Makefile
@@ -23,4 +23,4 @@ clean:
 	rm -rf $(OBJDIR)/ 
 	rm -rf $(OUTDIR)/
  
-all: bridge
\ No newline at end of file
+all: client
\ No newline at end of file
diff --git a/Engine/Source/Programs/Shared/EpicGames.Perforce/PerforceConnection.cs b/Engine/Source/Programs/Shared/EpicGames.Perforce/PerforceConnection.cs
--- a/Engine/Source/Programs/Shared/EpicGames.Perforce/PerforceConnection.cs
+++ b/Engine/Source/Programs/Shared/EpicGames.Perforce/PerforceConnection.cs
@@ -820,6 +820,24 @@ namespace EpicGames.Perforce
 			return (await TryGetChangesAsync(Connection, Options, ClientName, MaxChanges, Status, UserName, FileSpecs, CancellationToken)).Data;
 		}
 
+		/// <summary>
+		/// Enumerates changes on the server
+		/// </summary>
+		/// <param name="Connection">Connection to the Perforce server</param>
+		/// <param name="Options">Options for the command</param>
+		/// <param name="ClientName">List only changes made from the named client workspace.</param>
+		/// <param name="MinChangeNumber">The minimum changelist number</param>
+		/// <param name="MaxChanges">List only the highest numbered changes</param>
+		/// <param name="Status">Limit the list to the changelists with the given status (pending, submitted or shelved)</param>
+		/// <param name="UserName">List only changes made by the named user</param>
+		/// <param name="FileSpecs">Paths to query changes for</param>
+		/// <param name="CancellationToken">Token used to cancel the operation</param>
+		/// <returns>List of responses from the server.</returns>
+		public static async Task<List<ChangesRecord>> GetChangesAsync(this IPerforceConnection Connection, ChangesOptions Options, string? ClientName, int MinChangeNumber, int MaxChanges, ChangeStatus Status, string? UserName, FileSpecList FileSpecs, CancellationToken CancellationToken = default)
+		{
+			return (await TryGetChangesAsync(Connection, Options, ClientName, MinChangeNumber, MaxChanges, Status, UserName, FileSpecs, CancellationToken)).Data;
+		}
+
 		/// <summary>
 		/// Enumerates changes on the server
 		/// </summary>
@@ -848,6 +866,24 @@ namespace EpicGames.Perforce
 		/// <param name="CancellationToken">Token used to cancel the operation</param>
 		/// <returns>List of responses from the server.</returns>
 		public static Task<PerforceResponseList<ChangesRecord>> TryGetChangesAsync(this IPerforceConnection Connection, ChangesOptions Options, string? ClientName, int MaxChanges, ChangeStatus Status, string? UserName, FileSpecList FileSpecs, CancellationToken CancellationToken = default)
+		{
+			return TryGetChangesAsync(Connection, Options, ClientName, -1, MaxChanges, Status, UserName, FileSpecs, CancellationToken);
+		}
+
+		/// <summary>
+		/// Enumerates changes on the server
+		/// </summary>
+		/// <param name="Connection">Connection to the Perforce server</param>
+		/// <param name="Options">Options for the command</param>
+		/// <param name="ClientName">List only changes made from the named client workspace.</param>
+		/// <param name="MinChangeNumber">The minimum changelist number</param>
+		/// <param name="MaxChanges">List only the highest numbered changes</param>
+		/// <param name="Status">Limit the list to the changelists with the given status (pending, submitted or shelved)</param>
+		/// <param name="UserName">List only changes made by the named user</param>
+		/// <param name="FileSpecs">Paths to query changes for</param>
+		/// <param name="CancellationToken">Token used to cancel the operation</param>
+		/// <returns>List of responses from the server.</returns>
+		public static Task<PerforceResponseList<ChangesRecord>> TryGetChangesAsync(this IPerforceConnection Connection, ChangesOptions Options, string? ClientName, int MinChangeNumber, int MaxChanges, ChangeStatus Status, string? UserName, FileSpecList FileSpecs, CancellationToken CancellationToken = default)
 		{
 			List<string> Arguments = new List<string>();
 			if ((Options & ChangesOptions.IncludeIntegrations) != 0)
@@ -874,6 +910,10 @@ namespace EpicGames.Perforce
 			{
 				Arguments.Add($"-c{ClientName}");
 			}
+			if (MinChangeNumber != -1)
+			{
+				Arguments.Add($"-e{MinChangeNumber}");
+			}
 			if (MaxChanges != -1)
 			{
 				Arguments.Add($"-m{MaxChanges}");
@@ -1525,6 +1565,61 @@ namespace EpicGames.Perforce
 
 		#endregion
 
+		#region p4 dirs
+
+		/// <summary>
+		/// List directories under a depot path
+		/// </summary>
+		/// <param name="Connection">Connection to the Perforce server</param>
+		/// <param name="Options">Options for the command</param>
+		/// <param name="Stream">List directories mapped for the specified stream</param>
+		/// <param name="FileSpecs">Files to be opened for edit</param>
+		/// <param name="CancellationToken">Token used to cancel the operation</param>
+		/// <returns>Response from the server</returns>
+		public static async Task<List<DirsRecord>> GetDirsAsync(this IPerforceConnection Connection, DirsOptions Options, string? Stream, FileSpecList FileSpecs, CancellationToken CancellationToken = default)
+		{
+			return (await TryGetDirsAsync(Connection, Options, Stream, FileSpecs, CancellationToken)).Data;
+		}
+
+		/// <summary>
+		/// List directories under a depot path
+		/// </summary>
+		/// <param name="Connection">Connection to the Perforce server</param>
+		/// <param name="Options">Options for the command</param>
+		/// <param name="Stream">List directories mapped for the specified stream</param>
+		/// <param name="FileSpecs">Files to be opened for edit</param>
+		/// <param name="CancellationToken">Token used to cancel the operation</param>
+		/// <returns>Response from the server</returns>
+		public static Task<PerforceResponseList<DirsRecord>> TryGetDirsAsync(this IPerforceConnection Connection, DirsOptions Options, string? Stream, FileSpecList FileSpecs, CancellationToken CancellationToken = default)
+		{
+			List<string> Arguments = new List<string>();
+			if ((Options & DirsOptions.OnlyMapped) != 0)
+			{
+				Arguments.Add("-C");
+			}
+			if ((Options & DirsOptions.IncludeDeleted) != 0)
+			{
+				Arguments.Add("-D");
+			}
+			if ((Options & DirsOptions.OnlyHave) != 0)
+			{
+				Arguments.Add("-H");
+			}
+			if (Stream != null)
+			{
+				Arguments.Add("-S");
+				Arguments.Add(Stream);
+			}
+			if ((Options & DirsOptions.IgnoreCase) != 0)
+			{
+				Arguments.Add("-i");
+			}
+
+			return CommandAsync<DirsRecord>(Connection, "dirs", Arguments, FileSpecs.List, null, CancellationToken);
+		}
+
+		#endregion
+
 		#region p4 edit
 
 		/// <summary>
diff --git a/Engine/Source/Programs/Shared/EpicGames.Perforce/Records/DirsOptions.cs b/Engine/Source/Programs/Shared/EpicGames.Perforce/Records/DirsOptions.cs
new file mode 100644
--- /dev/null
+++ b/Engine/Source/Programs/Shared/EpicGames.Perforce/Records/DirsOptions.cs
@@ -0,0 +1,42 @@
+// Copyright Epic Games, Inc. All Rights Reserved.
+
+using System;
+using System.Collections.Generic;
+using System.Linq;
+using System.Text;
+using System.Threading.Tasks;
+
+namespace EpicGames.Perforce
+{
+	/// <summary>
+	/// Options for the p4 dirs command
+	/// </summary>
+	[Flags]
+	public enum DirsOptions
+	{
+		/// <summary>
+		/// No options specified
+		/// </summary>
+		None = 0,
+
+		/// <summary>
+		/// Display only those directories that are mapped through the current client view.
+		/// </summary>
+		OnlyMapped = 1,
+
+		/// <summary>
+		/// Include only those directories that contain files on the current client workspace’s p4 have list.
+		/// </summary>
+		OnlyHave = 2,
+
+		/// <summary>
+		/// Include subdirectories that contain only deleted files. By default, these directories are not displayed.
+		/// </summary>
+		IncludeDeleted = 4,
+
+		/// <summary>
+		/// Ignore the case of the directory argument when listing directories in a case-sensitive server.
+		/// </summary>
+		IgnoreCase = 8,
+	}
+}
diff --git a/Engine/Source/Programs/Shared/EpicGames.Perforce/Records/DirsRecord.cs b/Engine/Source/Programs/Shared/EpicGames.Perforce/Records/DirsRecord.cs
new file mode 100644
--- /dev/null
+++ b/Engine/Source/Programs/Shared/EpicGames.Perforce/Records/DirsRecord.cs
@@ -0,0 +1,39 @@
+// Copyright Epic Games, Inc. All Rights Reserved.
+
+using System;
+using System.Collections.Generic;
+using System.Linq;
+using System.Text;
+using System.Threading.Tasks;
+
+namespace EpicGames.Perforce
+{
+	/// <summary>
+	/// Subdirectory under a requested depot path
+	/// </summary>
+	public class DirsRecord
+	{
+		/// <summary>
+		/// The directory name
+		/// </summary>
+		[PerforceTag("dir")]
+		public string Dir;
+
+		/// <summary>
+		/// Private constructor for serialization
+		/// </summary>
+		private DirsRecord()
+		{
+			Dir = null!;
+		}
+
+		/// <summary>
+		/// Format this record for display in the debugger
+		/// </summary>
+		/// <returns>Summary of this revision</returns>
+		public override string ToString()
+		{
+			return Dir;
+		}
+	}
+}
