Missing file.

diff --git a/Engine/Source/Programs/Shared/EpicGames.Perforce/PerforceSettings.cs b/Engine/Source/Programs/Shared/EpicGames.Perforce/PerforceSettings.cs
new file mode 100644
--- /dev/null
+++ b/Engine/Source/Programs/Shared/EpicGames.Perforce/PerforceSettings.cs
@@ -0,0 +1,39 @@
+// Copyright Epic Games, Inc. All Rights Reserved.
+
+using System;
+using System.Collections.Generic;
+using System.Text;
+
+namespace EpicGames.Perforce
+{
+	/// <summary>
+	/// Settings for a new Perforce connection
+	/// </summary>
+	public class PerforceSettings
+	{
+		/// <summary>
+		/// Server and port to connect to
+		/// </summary>
+		public string? Port { get; set; }
+
+		/// <summary>
+		/// Username to log in with
+		/// </summary>
+		public string? User { get; set; }
+
+		/// <summary>
+		/// Name of the client to use
+		/// </summary>
+		public string? Client { get; set; }
+
+		/// <summary>
+		/// The invoking application name
+		/// </summary>
+		public string? AppName { get; set; }
+
+		/// <summary>
+		/// The invoking application version
+		/// </summary>
+		public string? AppVersion { get; set; }
+	}
+}
