BinaryArchive: Add support for reading and writing doubles

#trivial

diff --git a/Engine/Source/Programs/Shared/EpicGames.Core/BinaryArchiveReader.cs b/Engine/Source/Programs/Shared/EpicGames.Core/BinaryArchiveReader.cs
--- a/Engine/Source/Programs/Shared/EpicGames.Core/BinaryArchiveReader.cs
+++ b/Engine/Source/Programs/Shared/EpicGames.Core/BinaryArchiveReader.cs
@@ -167,6 +167,15 @@ namespace EpicGames.Core
 			return Value;
 		}
 
+		/// <summary>
+		/// Reads a double (64 bit floating point value) from the stream
+		/// </summary>
+		/// <returns>The value that was read</returns>
+		public double ReadDouble()
+		{
+			return BitConverter.Int64BitsToDouble(ReadLong());
+		}
+
 		/// <summary>
 		/// Reads a string from the stream
 		/// </summary>
diff --git a/Engine/Source/Programs/Shared/EpicGames.Core/BinaryArchiveWriter.cs b/Engine/Source/Programs/Shared/EpicGames.Core/BinaryArchiveWriter.cs
--- a/Engine/Source/Programs/Shared/EpicGames.Core/BinaryArchiveWriter.cs
+++ b/Engine/Source/Programs/Shared/EpicGames.Core/BinaryArchiveWriter.cs
@@ -225,6 +225,15 @@ namespace EpicGames.Core
 			BufferPos += 8;
 		}
 
+		/// <summary>
+		/// Writes a double (64 bit floating point value) to the stream
+		/// </summary>
+		/// <param name="Value">Value to write</param>
+		public void WriteDouble(double Value)
+		{
+			WriteLong(BitConverter.DoubleToInt64Bits(Value));
+		}
+
 		/// <summary>
 		/// Writes a string to the output
 		/// </summary>
