Horde: Change Redis wrapper to a more object-oriented API. Individual data types now include a reference to the database, allowing methods to directly modify their values.

diff --git a/Engine/Source/Programs/Horde/HordeServer/Compute/Impl/RedisMessageQueue.cs b/Engine/Source/Programs/Horde/HordeServer/Compute/Impl/RedisMessageQueue.cs
--- a/Engine/Source/Programs/Horde/HordeServer/Compute/Impl/RedisMessageQueue.cs
+++ b/Engine/Source/Programs/Horde/HordeServer/Compute/Impl/RedisMessageQueue.cs
@@ -101,7 +101,7 @@ namespace HordeServer.Compute.Impl
 		/// <returns></returns>
 		RedisList<T> GetChannel(string ChannelId)
 		{
-			return new RedisList<T>(KeyPrefix.Append(ChannelId));
+			return new RedisList<T>(Redis, KeyPrefix.Append(ChannelId));
 		}
 
 		/// <inheritdoc/>
@@ -109,7 +109,7 @@ namespace HordeServer.Compute.Impl
 		{
 			RedisList<T> Channel = GetChannel(ChannelId);
 
-			long Length = await Redis.ListRightPushAsync(Channel, Message);
+			long Length = await Channel.RightPushAsync(Message);
 			if (Length == 1)
 			{
 				await Redis.PublishAsync(UpdateChannel, ChannelId, CommandFlags.FireAndForget);
@@ -117,13 +117,13 @@ namespace HordeServer.Compute.Impl
 			await Redis.KeyExpireAsync(Channel.Key, ExpireTime, CommandFlags.FireAndForget);
 		}
 
-		async Task<bool> ReadMessagesAsync(RedisList<T> List, List<T> Messages)
+		static async Task<bool> ReadMessagesAsync(RedisList<T> List, List<T> Messages)
 		{
-			T? Message = await Redis.ListLeftPopAsync(List);
+			T? Message = await List.LeftPopAsync();
 			while(Message != null)
 			{
 				Messages.Add(Message);
-				Message = await Redis.ListLeftPopAsync(List);
+				Message = await List.LeftPopAsync();
 			}
 			return Messages.Count > 0;
 		}
diff --git a/Engine/Source/Programs/Horde/HordeServer/Compute/Impl/RedisTaskScheduler.cs b/Engine/Source/Programs/Horde/HordeServer/Compute/Impl/RedisTaskScheduler.cs
--- a/Engine/Source/Programs/Horde/HordeServer/Compute/Impl/RedisTaskScheduler.cs
+++ b/Engine/Source/Programs/Horde/HordeServer/Compute/Impl/RedisTaskScheduler.cs
@@ -114,8 +114,8 @@ namespace HordeServer.Compute.Impl
 		{
 			this.Redis = Redis;
 			this.BaseKey = BaseKey;
-			this.QueueIndex = new RedisSet<TQueueId>(BaseKey.Append("index"));
-			this.ActiveQueues = new RedisHash<TQueueId, DateTime>(BaseKey.Append("active"));
+			this.QueueIndex = new RedisSet<TQueueId>(Redis, BaseKey.Append("index"));
+			this.ActiveQueues = new RedisHash<TQueueId, DateTime>(Redis, BaseKey.Append("active"));
 			this.NewQueueChannel = new RedisChannel<TQueueId>(BaseKey.Append("new_queues").ToString());
 			this.Logger = Logger;
 
@@ -156,7 +156,7 @@ namespace HordeServer.Compute.Impl
 		/// <returns></returns>
 		RedisList<TTask> GetQueue(TQueueId QueueId)
 		{
-			return new RedisList<TTask>(GetQueueKey(QueueId));
+			return new RedisList<TTask>(Redis, GetQueueKey(QueueId));
 		}
 
 		/// <summary>
@@ -167,7 +167,7 @@ namespace HordeServer.Compute.Impl
 		{
 			ITransaction Transaction = Redis.CreateTransaction();
 			Transaction.AddCondition(Condition.KeyExists(GetQueueKey(QueueId)));
-			_ = Transaction.SetAddAsync(QueueIndex, QueueId);
+			_ = Transaction.With(QueueIndex).AddAsync(QueueId);
 			await Transaction.ExecuteAsync();
 
 			await Redis.PublishAsync(NewQueueChannel, QueueId);
@@ -181,7 +181,7 @@ namespace HordeServer.Compute.Impl
 		{
 			ITransaction Transaction = Redis.CreateTransaction();
 			Transaction.AddCondition(Condition.KeyNotExists(GetQueueKey(QueueId)));
-			_ = Transaction.SetRemoveAsync(QueueIndex, QueueId);
+			_ = Transaction.With(QueueIndex).RemoveAsync(QueueId);
 			await Transaction.ExecuteAsync(CommandFlags.FireAndForget);
 		}
 
@@ -194,7 +194,7 @@ namespace HordeServer.Compute.Impl
 		{
 			RedisList<TTask> List = GetQueue(QueueId);
 
-			long NewLength = await Redis.ListRightPushAsync(List, Task);
+			long NewLength = await List.RightPushAsync(Task);
 			if (NewLength == 1)
 			{
 				await AddQueueToIndexAsync(QueueId);
@@ -210,7 +210,7 @@ namespace HordeServer.Compute.Impl
 		{
 			RedisList<TTask> List = GetQueue(QueueId);
 
-			long NewLength = await Redis.ListLeftPushAsync(List, Task);
+			long NewLength = await List.LeftPushAsync(Task);
 			if (NewLength == 1)
 			{
 				await AddQueueToIndexAsync(QueueId);
@@ -229,7 +229,7 @@ namespace HordeServer.Compute.Impl
 			Listener? Listener = null;
 			try
 			{
-				TQueueId[] Queues = await Redis.SetMembersAsync(QueueIndex);
+				TQueueId[] Queues = await QueueIndex.MembersAsync();
 				while (!Token.IsCancellationRequested)
 				{
 					// Try to dequeue an item from the list
@@ -301,7 +301,7 @@ namespace HordeServer.Compute.Impl
 		{
 			await AddActiveQueue(QueueId);
 
-			TTask? Item = await Redis.ListLeftPopAsync(GetQueue(QueueId));
+			TTask? Item = await GetQueue(QueueId).LeftPopAsync();
 			if (Item == null)
 			{
 				await RemoveQueueFromIndexAsync(QueueId);
@@ -345,7 +345,7 @@ namespace HordeServer.Compute.Impl
 					}
 					if (Interlocked.CompareExchange(ref LocalActiveQueues, NewLocalActiveKeys, LocalActiveKeysCopy) == LocalActiveKeysCopy)
 					{
-						await Redis.HashSetAsync(ActiveQueues, QueueId, DateTime.UtcNow);
+						await ActiveQueues.SetAsync(QueueId, DateTime.UtcNow);
 						break;
 					}
 				}
@@ -358,12 +358,12 @@ namespace HordeServer.Compute.Impl
 		/// <returns></returns>
 		public async Task<List<TQueueId>> GetInactiveQueuesAsync()
 		{
-			HashSet<TQueueId> Keys = new HashSet<TQueueId>(await Redis.SetMembersAsync(QueueIndex));
+			HashSet<TQueueId> Keys = new HashSet<TQueueId>(await QueueIndex.MembersAsync());
 			HashSet<TQueueId> InvalidKeys = new HashSet<TQueueId>();
 
 			DateTime MinTime = DateTime.UtcNow - TimeSpan.FromMinutes(1.0);
 
-			HashEntry<TQueueId, DateTime>[] Entries = await Redis.HashGetAllAsync(ActiveQueues);
+			HashEntry<TQueueId, DateTime>[] Entries = await ActiveQueues.GetAllAsync();
 			foreach (HashEntry<TQueueId, DateTime> Entry in Entries)
 			{
 				if (Entry.Value < MinTime)
@@ -378,7 +378,7 @@ namespace HordeServer.Compute.Impl
 
 			if (InvalidKeys.Count > 0)
 			{
-				await Redis.HashDeleteAsync(ActiveQueues, InvalidKeys.ToArray());
+				await ActiveQueues.DeleteAsync(InvalidKeys.ToArray());
 			}
 
 			return Keys.ToList();
diff --git a/Engine/Source/Programs/Shared/EpicGames.Redis/RedisHash.cs b/Engine/Source/Programs/Shared/EpicGames.Redis/RedisHash.cs
--- a/Engine/Source/Programs/Shared/EpicGames.Redis/RedisHash.cs
+++ b/Engine/Source/Programs/Shared/EpicGames.Redis/RedisHash.cs
@@ -14,8 +14,10 @@ namespace EpicGames.Redis
 	/// </summary>
 	/// <typeparam name="TName">The key type for the hash</typeparam>
 	/// <typeparam name="TValue">The value type for the hash</typeparam>
-	public readonly struct RedisHash<TName, TValue> : IEquatable<RedisHash<TName, TValue>>
+	public readonly struct RedisHash<TName, TValue>
 	{
+		readonly IDatabaseAsync Database;
+
 		/// <summary>
 		/// The key for the hash
 		/// </summary>
@@ -25,22 +27,48 @@ namespace EpicGames.Redis
 		/// Constructor
 		/// </summary>
 		/// <param name="Key"></param>
-		public RedisHash(RedisKey Key) => this.Key = Key;
+		public RedisHash(IDatabaseAsync Database, RedisKey Key)
+		{
+			this.Database = Database;
+			this.Key = Key;
+		}
+
+		/// <inheritdoc cref="IDatabaseAsync.HashDeleteAsync(RedisKey, RedisValue, CommandFlags)"/>
+		public Task<bool> DeleteAsync(TName Name, CommandFlags Flags = CommandFlags.None)
+		{
+			return Database.HashDeleteAsync(Key, RedisSerializer.Serialize(Name), Flags);
+		}
 
-		/// <inheritdoc/>
-		public override bool Equals(object? Obj) => Obj is RedisHash<TName, TValue> List && Key == List.Key;
+		/// <inheritdoc cref="IDatabaseAsync.HashDeleteAsync(RedisKey, RedisValue[], CommandFlags)"/>
+		public Task<long> DeleteAsync(TName[] Names, CommandFlags Flags = CommandFlags.None)
+		{
+			return Database.HashDeleteAsync(Key, Array.ConvertAll(Names, x => RedisSerializer.Serialize(x)), Flags);
+		}
 
-		/// <inheritdoc/>
-		public bool Equals(RedisHash<TName, TValue> Other) => Key == Other.Key;
+		/// <inheritdoc cref="IDatabaseAsync.HashGetAllAsync(RedisKey, CommandFlags)"/>
+		public async Task<HashEntry<TName, TValue>[]> GetAllAsync(CommandFlags Flags = CommandFlags.None)
+		{
+			HashEntry[] Entries = await Database.HashGetAllAsync(Key, Flags);
+			return Array.ConvertAll(Entries, x => new HashEntry<TName, TValue>(RedisSerializer.Deserialize<TName>(x.Name), RedisSerializer.Deserialize<TValue>(x.Value)));
+		}
 
-		/// <inheritdoc/>
-		public override int GetHashCode() => Key.GetHashCode();
+		/// <inheritdoc cref="IDatabaseAsync.HashLengthAsync(RedisKey, CommandFlags)"/>
+		public Task<long> LengthAsync(CommandFlags Flags = CommandFlags.None)
+		{
+			return Database.HashLengthAsync(Key, Flags);
+		}
 
-		/// <summary>Compares two instances for equality</summary>
-		public static bool operator ==(RedisHash<TName, TValue> Left, RedisHash<TName, TValue> Right) => Left.Key == Right.Key;
+		/// <inheritdoc cref="IDatabaseAsync.HashSetAsync(RedisKey, RedisValue, RedisValue, When, CommandFlags)"/>
+		public Task SetAsync(TName Name, TValue Value, When When = When.Always, CommandFlags Flags = CommandFlags.None)
+		{
+			return Database.HashSetAsync(Key, RedisSerializer.Serialize(Name), RedisSerializer.Serialize(Value), When, Flags);
+		}
 
-		/// <summary>Compares two instances for equality</summary>
-		public static bool operator !=(RedisHash<TName, TValue> Left, RedisHash<TName, TValue> Right) => Left.Key != Right.Key;
+		/// <inheritdoc cref="IDatabaseAsync.HashSetAsync(RedisKey, HashEntry[], CommandFlags)"/>
+		public Task SetAsync(HashEntry[] Entries, CommandFlags Flags = CommandFlags.None)
+		{
+			return Database.HashSetAsync(Key, Array.ConvertAll(Entries, x => (HashEntry)x), Flags);
+		}
 	}
 
 	/// <inheritdoc cref="HashEntry"/>
@@ -83,45 +111,18 @@ namespace EpicGames.Redis
 	}
 
 	/// <summary>
-	/// Extension methods for typed lists
+	/// Extension methods for hashes
 	/// </summary>
 	public static class RedisHashExtensions
 	{
-		/// <inheritdoc cref="IDatabaseAsync.HashDeleteAsync(RedisKey, RedisValue, CommandFlags)"/>
-		public static Task<bool> HashDeleteAsync<TName, TValue>(this IDatabase Redis, RedisHash<TName, TValue> Hash, TName Name, CommandFlags Flags = CommandFlags.None)
-		{
-			return Redis.HashDeleteAsync(Hash.Key, RedisSerializer.Serialize(Name), Flags);
-		}
-
-		/// <inheritdoc cref="IDatabaseAsync.HashDeleteAsync(RedisKey, RedisValue[], CommandFlags)"/>
-		public static Task<long> HashDeleteAsync<TName, TValue>(this IDatabase Redis, RedisHash<TName, TValue> Hash, TName[] Names, CommandFlags Flags = CommandFlags.None)
-		{
-			return Redis.HashDeleteAsync(Hash.Key, Array.ConvertAll(Names, x => RedisSerializer.Serialize(x)), Flags);
-		}
-
-		/// <inheritdoc cref="IDatabaseAsync.HashGetAllAsync(RedisKey, CommandFlags)"/>
-		public static async Task<HashEntry<TName, TValue>[]> HashGetAllAsync<TName, TValue>(this IDatabase Redis, RedisHash<TName, TValue> Hash, CommandFlags Flags = CommandFlags.None)
-		{
-			HashEntry[] Entries = await Redis.HashGetAllAsync(Hash.Key, Flags);
-			return Array.ConvertAll(Entries, x => new HashEntry<TName, TValue>(RedisSerializer.Deserialize<TName>(x.Name), RedisSerializer.Deserialize<TValue>(x.Value)));
-		}
-
-		/// <inheritdoc cref="IDatabaseAsync.HashLengthAsync(RedisKey, CommandFlags)"/>
-		public static Task<long> HashLengthAsync<TName, TValue>(this IDatabase Redis, RedisHash<TName, TValue> Hash, CommandFlags Flags = CommandFlags.None)
-		{
-			return Redis.HashLengthAsync(Hash.Key, Flags);
-		}
-
-		/// <inheritdoc cref="IDatabaseAsync.HashSetAsync(RedisKey, RedisValue, RedisValue, When, CommandFlags)"/>
-		public static Task HashSetAsync<TName, TValue>(this IDatabase Redis, RedisHash<TName, TValue> Hash, TName Name, TValue Value, When When = When.Always, CommandFlags Flags = CommandFlags.None)
-		{
-			return Redis.HashSetAsync(Hash.Key, RedisSerializer.Serialize(Name), RedisSerializer.Serialize(Value), When, Flags);
-		}
-
-		/// <inheritdoc cref="IDatabaseAsync.HashSetAsync(RedisKey, HashEntry[], CommandFlags)"/>
-		public static Task HashSetAsync<TName, TValue>(this IDatabase Redis, RedisHash<TName, TValue> Hash, HashEntry<TName, TValue>[] Entries, CommandFlags Flags = CommandFlags.None)
+		/// <summary>
+		/// Creates a version of this set which modifies a transaction rather than the direct DB
+		/// </summary>
+		/// <param name="Transaction"></param>
+		/// <returns></returns>
+		public static RedisHash<TName, TValue> With<TName, TValue>(this ITransaction Transaction, RedisHash<TName, TValue> Set)
 		{
-			return Redis.HashSetAsync(Hash.Key, Array.ConvertAll(Entries, x => (HashEntry)x), Flags);
+			return new RedisHash<TName, TValue>(Transaction, Set.Key);
 		}
 	}
 }
diff --git a/Engine/Source/Programs/Shared/EpicGames.Redis/RedisList.cs b/Engine/Source/Programs/Shared/EpicGames.Redis/RedisList.cs
--- a/Engine/Source/Programs/Shared/EpicGames.Redis/RedisList.cs
+++ b/Engine/Source/Programs/Shared/EpicGames.Redis/RedisList.cs
@@ -3,6 +3,7 @@
 using StackExchange.Redis;
 using System;
 using System.Collections.Generic;
+using System.Diagnostics.CodeAnalysis;
 using System.Linq;
 using System.Text;
 using System.Threading.Tasks;
@@ -13,8 +14,13 @@ namespace EpicGames.Redis
 	/// Represents a typed Redis list with a given key
 	/// </summary>
 	/// <typeparam name="TElement">The type of element stored in the list</typeparam>
-	public readonly struct RedisList<TElement> : IEquatable<RedisList<TElement>>
+	public readonly struct RedisList<TElement>
 	{
+		/// <summary>
+		/// The database containing this object
+		/// </summary>
+		IDatabaseAsync Database { get; }
+
 		/// <summary>
 		/// The key for the list
 		/// </summary>
@@ -24,79 +30,78 @@ namespace EpicGames.Redis
 		/// Constructor
 		/// </summary>
 		/// <param name="Key"></param>
-		public RedisList(RedisKey Key) => this.Key = Key;
-
-		/// <inheritdoc/>
-		public override bool Equals(object? Obj) => Obj is RedisList<TElement> List && Key == List.Key;
-
-		/// <inheritdoc/>
-		public bool Equals(RedisList<TElement> Other) => Key == Other.Key;
-
-		/// <inheritdoc/>
-		public override int GetHashCode() => Key.GetHashCode();
-
-		/// <summary>Compares two instances for equality</summary>
-		public static bool operator ==(RedisList<TElement> Left, RedisList<TElement> Right) => Left.Key == Right.Key;
-
-		/// <summary>Compares two instances for equality</summary>
-		public static bool operator !=(RedisList<TElement> Left, RedisList<TElement> Right) => Left.Key != Right.Key;
-	}
+		public RedisList(IDatabaseAsync Database, RedisKey Key)
+		{
+			this.Database = Database;
+			this.Key = Key;
+		}
 
-	/// <summary>
-	/// Extension methods for typed lists
-	/// </summary>
-	public static class RedisListExtensions
-	{
 		/// <inheritdoc cref="IDatabaseAsync.ListLengthAsync(RedisKey, CommandFlags)"/>
-		public static Task<long> ListLengthAsync<T>(this IDatabase Database, RedisList<T> List)
+		public Task<long> LengthAsync()
 		{
-			return Database.ListLengthAsync(List.Key);
+			return Database.ListLengthAsync(Key);
 		}
 
 		/// <inheritdoc cref="IDatabaseAsync.ListLeftPushAsync(RedisKey, RedisValue, When, CommandFlags)"/>
-		public static Task<long> ListLeftPushAsync<T>(this IDatabase Database, RedisList<T> List, T Item, When When = When.Always, CommandFlags Flags = CommandFlags.None)
+		public Task<long> LeftPushAsync(TElement Item, When When = When.Always, CommandFlags Flags = CommandFlags.None)
 		{
-			return Database.ListLeftPushAsync(List.Key, RedisSerializer.Serialize(Item), When, Flags);
+			return Database.ListLeftPushAsync(Key, RedisSerializer.Serialize(Item), When, Flags);
 		}
 
 		/// <inheritdoc cref="IDatabaseAsync.ListLeftPushAsync(RedisKey, RedisValue[], When, CommandFlags)"/>
-		public static Task<long> ListLeftPushAsync<T>(this IDatabase Database, RedisList<T> List, IEnumerable<T> Items, When When = When.Always, CommandFlags Flags = CommandFlags.None)
+		public Task<long> LeftPushAsync(IEnumerable<TElement> Items, When When = When.Always, CommandFlags Flags = CommandFlags.None)
 		{
 			RedisValue[] Values = Items.Select(x => RedisSerializer.Serialize(x)).ToArray();
-			return Database.ListLeftPushAsync(List.Key, Values, When, Flags);
+			return Database.ListLeftPushAsync(Key, Values, When, Flags);
 		}
 
 		/// <inheritdoc cref="IDatabaseAsync.ListLeftPopAsync(RedisKey, CommandFlags)">
-		public static async Task<T?> ListLeftPopAsync<T>(this IDatabase Database, RedisList<T> List, CommandFlags Flags = CommandFlags.None) where T : class
+		public async Task<TElement> LeftPopAsync(CommandFlags Flags = CommandFlags.None)
 		{
-			RedisValue Value = await Database.ListLeftPopAsync(List.Key, Flags);
-			return Value.IsNull ? null : RedisSerializer.Deserialize<T>(Value);
+			RedisValue Value = await Database.ListLeftPopAsync(Key, Flags);
+			return Value.IsNull ? default(TElement)! : RedisSerializer.Deserialize<TElement>(Value);
 		}
 
 		/// <inheritdoc cref="IDatabaseAsync.ListRangeAsync(RedisKey, CommandFlags)">
-		public static async Task<T[]> ListRangeAsync<T>(this IDatabase Database, RedisList<T> List, long Start = 0, long Stop = -1, CommandFlags Flags = CommandFlags.None) where T : class
+		public async Task<TElement[]> RangeAsync(long Start = 0, long Stop = -1, CommandFlags Flags = CommandFlags.None)
 		{
-			RedisValue[] Values = await Database.ListRangeAsync(List.Key, Start, Stop, Flags);
-			return Array.ConvertAll(Values, x => RedisSerializer.Deserialize<T>(x));
+			RedisValue[] Values = await Database.ListRangeAsync(Key, Start, Stop, Flags);
+			return Array.ConvertAll(Values, (Converter<RedisValue, T>)(x => (T)RedisSerializer.Deserialize<TElement>(x)));
 		}
 
 		/// <inheritdoc cref="IDatabaseAsync.ListRightPushAsync(RedisKey, RedisValue, When, CommandFlags)"/>
-		public static Task<long> ListRightPushAsync<T>(this IDatabase Database, RedisList<T> List, T Item, When When = When.Always, CommandFlags Flags = CommandFlags.None)
+		public Task<long> RightPushAsync(TElement Item, When When = When.Always, CommandFlags Flags = CommandFlags.None)
 		{
-			return Database.ListRightPushAsync(List.Key, RedisSerializer.Serialize(Item), When, Flags);
+			return Database.ListRightPushAsync(Key, RedisSerializer.Serialize(Item), When, Flags);
 		}
 
 		/// <inheritdoc cref="IDatabaseAsync.ListRightPushAsync(RedisKey, RedisValue[], When, CommandFlags)"/>
-		public static Task<long> ListRightPushAsync<T>(this IDatabase Database, RedisList<T> List, IEnumerable<T> Items, When When = When.Always, CommandFlags Flags = CommandFlags.None)
+		public Task<long> RightPushAsync(IEnumerable<TElement> Items, When When = When.Always, CommandFlags Flags = CommandFlags.None)
 		{
 			RedisValue[] Values = Items.Select(x => RedisSerializer.Serialize(x)).ToArray();
-			return Database.ListRightPushAsync(List.Key, Values, When, Flags);
+			return Database.ListRightPushAsync(Key, Values, When, Flags);
 		}
 
 		/// <inheritdoc cref="IDatabaseAsync.ListTrimAsync(RedisKey, long, long, CommandFlags)">
-		public static async Task ListTrimAsync<T>(this IDatabase Database, RedisList<T> List, long Start, long Stop, CommandFlags Flags = CommandFlags.None) where T : class
+		public Task TrimAsync(long Start, long Stop, CommandFlags Flags = CommandFlags.None)
+		{
+			return Database.ListTrimAsync(Key, Start, Stop, Flags);
+		}
+	}
+
+	/// <summary>
+	/// Extension methods for sets
+	/// </summary>
+	public static class RedisListExtensions
+	{
+		/// <summary>
+		/// Creates a version of this list which modifies a transaction rather than the direct DB
+		/// </summary>
+		/// <param name="Transaction"></param>
+		/// <returns></returns>
+		public static RedisList<TElement> With<TElement>(this ITransaction Transaction, RedisList<TElement> Set)
 		{
-			await Database.ListTrimAsync(List.Key, Start, Stop, Flags);
+			return new RedisList<TElement>(Transaction, Set.Key);
 		}
 	}
 }
diff --git a/Engine/Source/Programs/Shared/EpicGames.Redis/RedisSerializer.cs b/Engine/Source/Programs/Shared/EpicGames.Redis/RedisSerializer.cs
--- a/Engine/Source/Programs/Shared/EpicGames.Redis/RedisSerializer.cs
+++ b/Engine/Source/Programs/Shared/EpicGames.Redis/RedisSerializer.cs
@@ -89,6 +89,12 @@ namespace EpicGames.Redis
 			public T FromRedisValue(RedisValue Value) => (T)TypeConverter.ConvertFrom((string)Value);
 		}
 
+		class RedisIntConverter : IRedisConverter<int>
+		{
+			public RedisValue ToRedisValue(int Value) => Value;
+			public int FromRedisValue(RedisValue Value) => (int?)Value ?? 0;
+		}
+
 		/// <summary>
 		/// Static class for caching converter lookups
 		/// </summary>
@@ -113,6 +119,12 @@ namespace EpicGames.Redis
 					return (IRedisConverter<T>)Activator.CreateInstance(ConverterType)!;
 				}
 
+				// Check for known basic types
+				if (Type == typeof(int))
+				{
+					return (IRedisConverter<T>)new RedisIntConverter();
+				}
+
 				// Check if there's a regular converter we can use to convert to/from a string
 				TypeConverter? Converter = TypeDescriptor.GetConverter(Type);
 				if (Converter != null && Converter.CanConvertFrom(typeof(string)) && Converter.CanConvertTo(typeof(string)))
diff --git a/Engine/Source/Programs/Shared/EpicGames.Redis/RedisSet.cs b/Engine/Source/Programs/Shared/EpicGames.Redis/RedisSet.cs
--- a/Engine/Source/Programs/Shared/EpicGames.Redis/RedisSet.cs
+++ b/Engine/Source/Programs/Shared/EpicGames.Redis/RedisSet.cs
@@ -12,8 +12,10 @@ namespace EpicGames.Redis
 	/// Represents a typed Redis list with a given key
 	/// </summary>
 	/// <typeparam name="TElement">The type of element stored in the list</typeparam>
-	public readonly struct RedisSet<TElement> : IEquatable<RedisSet<TElement>>
+	public readonly struct RedisSet<TElement>
 	{
+		readonly IDatabaseAsync Database;
+
 		/// <summary>
 		/// The key for the list
 		/// </summary>
@@ -23,46 +25,45 @@ namespace EpicGames.Redis
 		/// Constructor
 		/// </summary>
 		/// <param name="Key"></param>
-		public RedisSet(RedisKey Key) => this.Key = Key;
-
-		/// <inheritdoc/>
-		public override bool Equals(object? Obj) => Obj is RedisSet<TElement> Set && Key == Set.Key;
-
-		/// <inheritdoc/>
-		public bool Equals(RedisSet<TElement> Other) => Key == Other.Key;
-
-		/// <inheritdoc/>
-		public override int GetHashCode() => Key.GetHashCode();
-
-		/// <summary>Compares two instances for equality</summary>
-		public static bool operator ==(RedisSet<TElement> Left, RedisSet<TElement> Right) => Left.Key == Right.Key;
-
-		/// <summary>Compares two instances for equality</summary>
-		public static bool operator !=(RedisSet<TElement> Left, RedisSet<TElement> Right) => Left.Key != Right.Key;
-	}
+		public RedisSet(IDatabaseAsync Database, RedisKey Key)
+		{
+			this.Database = Database;
+			this.Key = Key;
+		}
 
-	/// <summary>
-	/// Extension methods for <see cref="IDatabaseAsync"/>
-	/// </summary>
-	public static class RedisSetExtensions
-	{
 		/// <inheritdoc cref="IDatabaseAsync.SetAddAsync(RedisKey, RedisValue, CommandFlags)"/>
-		public static Task SetAddAsync<T>(this IDatabaseAsync Database, RedisSet<T> Set, T Value, CommandFlags Flags = CommandFlags.None)
+		public Task AddAsync(TElement Value, CommandFlags Flags = CommandFlags.None)
 		{
-			return Database.SetAddAsync(Set.Key, RedisSerializer.Serialize(Value), Flags);
+			return Database.SetAddAsync(Key, RedisSerializer.Serialize(Value), Flags);
 		}
 
 		/// <inheritdoc cref="IDatabaseAsync.SetRemoveAsync(RedisKey, RedisValue, CommandFlags)"/>
-		public static Task SetRemoveAsync<T>(this IDatabaseAsync Database, RedisSet<T> Set, T Value, CommandFlags Flags = CommandFlags.None)
+		public Task RemoveAsync(TElement Value, CommandFlags Flags = CommandFlags.None)
 		{
-			return Database.SetRemoveAsync(Set.Key, RedisSerializer.Serialize(Value), Flags);
+			return Database.SetRemoveAsync(Key, RedisSerializer.Serialize(Value), Flags);
 		}
 
 		/// <inheritdoc cref="IDatabaseAsync.SetRemoveAsync(RedisKey, RedisValue, CommandFlags)"/>
-		public static async Task<T[]> SetMembersAsync<T>(this IDatabaseAsync Database, RedisSet<T> Set, CommandFlags Flags = CommandFlags.None)
+		public async Task<TElement[]> MembersAsync(CommandFlags Flags = CommandFlags.None)
+		{
+			RedisValue[] Values = await Database.SetMembersAsync(Key, Flags);
+			return Array.ConvertAll(Values, (Converter<RedisValue, TElement>)(x => RedisSerializer.Deserialize<TElement>(x)));
+		}
+	}
+
+	/// <summary>
+	/// Extension methods for sets
+	/// </summary>
+	public static class RedisSetExtensions
+	{
+		/// <summary>
+		/// Creates a version of this set which modifies a transaction rather than the direct DB
+		/// </summary>
+		/// <param name="Transaction"></param>
+		/// <returns></returns>
+		public static RedisSet<TElement> With<TElement>(this ITransaction Transaction, RedisSet<TElement> Set)
 		{
-			RedisValue[] Values = await Database.SetMembersAsync(Set.Key, Flags);
-			return Array.ConvertAll(Values, x => RedisSerializer.Deserialize<T>(x));
+			return new RedisSet<TElement>(Transaction, Set.Key);
 		}
 	}
 }
diff --git a/Engine/Source/Programs/Shared/EpicGames.Redis/RedisSortedSet.cs b/Engine/Source/Programs/Shared/EpicGames.Redis/RedisSortedSet.cs
--- a/Engine/Source/Programs/Shared/EpicGames.Redis/RedisSortedSet.cs
+++ b/Engine/Source/Programs/Shared/EpicGames.Redis/RedisSortedSet.cs
@@ -51,8 +51,10 @@ namespace EpicGames.Redis
 	/// Represents a typed Redis list with a given key
 	/// </summary>
 	/// <typeparam name="TElement">The type of element stored in the list</typeparam>
-	public readonly struct RedisSortedSet<TElement> : IEquatable<RedisSortedSet<TElement>>
+	public readonly struct RedisSortedSet<TElement>
 	{
+		readonly IDatabaseAsync Database;
+
 		/// <summary>
 		/// The key for the list
 		/// </summary>
@@ -62,70 +64,81 @@ namespace EpicGames.Redis
 		/// Constructor
 		/// </summary>
 		/// <param name="Key"></param>
-		public RedisSortedSet(RedisKey Key) => this.Key = Key;
-
-		/// <inheritdoc/>
-		public override bool Equals(object? Obj) => Obj is RedisSortedSet<TElement> Set && Key == Set.Key;
-
-		/// <inheritdoc/>
-		public bool Equals(RedisSortedSet<TElement> Other) => Key == Other.Key;
-
-		/// <inheritdoc/>
-		public override int GetHashCode() => Key.GetHashCode();
-
-		/// <summary>Compares two instances for equality</summary>
-		public static bool operator ==(RedisSortedSet<TElement> Left, RedisSortedSet<TElement> Right) => Left.Key == Right.Key;
-
-		/// <summary>Compares two instances for equality</summary>
-		public static bool operator !=(RedisSortedSet<TElement> Left, RedisSortedSet<TElement> Right) => Left.Key != Right.Key;
-	}
+		public RedisSortedSet(IDatabaseAsync Database, RedisKey Key)
+		{
+			this.Database = Database;
+			this.Key = Key;
+		}
 
-	/// <summary>
-	/// Extension methods for <see cref="IDatabaseAsync"/>
-	/// </summary>
-	public static class RedisSortedSetExtensions
-	{
 		/// <inheritdoc cref="IDatabaseAsync.SortedSetAddAsync(RedisKey, RedisValue, double, CommandFlags)"/>
-		public static Task<bool> SortedSetAddAsync<T>(this IDatabaseAsync Database, RedisSortedSet<T> Set, T Value, double Score, CommandFlags Flags = CommandFlags.None)
+		public Task<bool> AddAsync<T>(T Value, double Score, When When = When.Always, CommandFlags Flags = CommandFlags.None)
 		{
-			return Database.SortedSetAddAsync(Set.Key, RedisSerializer.Serialize(Value), Score, Flags);
+			return Database.SortedSetAddAsync(Key, RedisSerializer.Serialize(Value), Score, When, Flags);
 		}
 
 		/// <inheritdoc cref="IDatabaseAsync.SortedSetAddAsync(RedisKey, RedisValue, double, CommandFlags)"/>
-		public static Task<long> SortedSetAddAsync<T>(this IDatabaseAsync Database, RedisSortedSet<T> Set, SortedSetEntry<T>[] Values, CommandFlags Flags = CommandFlags.None)
+		public Task<long> AddAsync<T>(SortedSetEntry<T>[] Values, CommandFlags Flags = CommandFlags.None)
 		{
 			SortedSetEntry[] Untyped = Array.ConvertAll(Values, x => new SortedSetEntry(RedisSerializer.Serialize(x.Element), x.Score));
-			return Database.SortedSetAddAsync(Set.Key, Untyped, Flags);
+			return Database.SortedSetAddAsync(Key, Untyped, Flags);
+		}
+
+		/// <inheritdoc cref="IDatabaseAsync.SortedSetRemoveAsync(RedisKey, RedisValue, CommandFlags)"/>
+		public Task<bool> RemoveAsync<T>(T Value, CommandFlags Flags = CommandFlags.None)
+		{
+			return Database.SortedSetRemoveAsync(Key, RedisSerializer.Serialize(Value), Flags);
 		}
 
 		/// <inheritdoc cref="IDatabaseAsync.SortedSetRemoveRangeByScoreAsync(RedisKey, double, double, Exclude, Order, long, long, CommandFlags)"/>
-		public static async Task<long> SortedSetRemoveRangeByScoreAsync<T>(this IDatabaseAsync Database, RedisSortedSet<T> Set, double Start = double.NegativeInfinity, double Stop = double.PositiveInfinity, Exclude Exclude = Exclude.None, CommandFlags Flags = CommandFlags.None)
+		public async Task<long> RemoveRangeByScoreAsync<T>(double Start = double.NegativeInfinity, double Stop = double.PositiveInfinity, Exclude Exclude = Exclude.None, CommandFlags Flags = CommandFlags.None)
 		{
-			return await Database.SortedSetRemoveRangeByScoreAsync(Set.Key, Start, Stop, Exclude, Flags);
+			return await Database.SortedSetRemoveRangeByScoreAsync(Key, Start, Stop, Exclude, Flags);
 		}
 
 		/// <inheritdoc cref="IDatabaseAsync.SortedSetAddAsync(RedisKey, RedisValue, double, CommandFlags)"/>
-		public static async IAsyncEnumerable<SortedSetEntry<T>> SortedSetScanAsync<T>(this IDatabaseAsync Database, RedisSortedSet<T> Set, RedisValue Pattern = default, int PageSize = 250, long Cursor = 0, int PageOffset = 0, CommandFlags Flags = CommandFlags.None)
+		public async IAsyncEnumerable<SortedSetEntry<T>> ScanAsync<T>(RedisValue Pattern = default, int PageSize = 250, long Cursor = 0, int PageOffset = 0, CommandFlags Flags = CommandFlags.None)
 		{
-			await foreach (SortedSetEntry Entry in Database.SortedSetScanAsync(Set.Key, Pattern, PageSize, Cursor, PageOffset, Flags))
+			await foreach (SortedSetEntry Entry in Database.SortedSetScanAsync(Key, Pattern, PageSize, Cursor, PageOffset, Flags))
 			{
 				yield return new SortedSetEntry<T>(Entry);
 			}
 		}
 
+		/// <inheritdoc cref="IDatabaseAsync.SortedSetRangeByRankAsync(RedisKey, long, long, Order, CommandFlags)"/>
+		public async Task<T[]> RangeByRankAsync<T>(long Start, long Stop = -1, Order Order = Order.Ascending, CommandFlags Flags = CommandFlags.None)
+		{
+			RedisValue[] Values = await Database.SortedSetRangeByRankAsync(Key, Start, Stop, Order, Flags);
+			return Array.ConvertAll(Values, x => RedisSerializer.Deserialize<T>(x));
+		}
+
 		/// <inheritdoc cref="IDatabaseAsync.SortedSetRangeByScoreAsync(RedisKey, double, double, Exclude, Order, long, long, CommandFlags)"/>
-		public static async Task<T[]> SortedSetRangeByScoreAsync<T>(this IDatabaseAsync Database, RedisSortedSet<T> Set, double Start = double.NegativeInfinity, double Stop = double.PositiveInfinity, Exclude Exclude = Exclude.None, Order Order = Order.Ascending, long Skip = 0L, long Take = -1L, CommandFlags Flags = CommandFlags.None)
+		public async Task<T[]> RangeByScoreAsync<T>(double Start = double.NegativeInfinity, double Stop = double.PositiveInfinity, Exclude Exclude = Exclude.None, Order Order = Order.Ascending, long Skip = 0L, long Take = -1L, CommandFlags Flags = CommandFlags.None)
 		{
-			RedisValue[] Values = await Database.SortedSetRangeByScoreAsync(Set.Key, Start, Stop, Exclude, Order, Skip, Take, Flags);
+			RedisValue[] Values = await Database.SortedSetRangeByScoreAsync(Key, Start, Stop, Exclude, Order, Skip, Take, Flags);
 			return Array.ConvertAll(Values, x => RedisSerializer.Deserialize<T>(x));
 		}
 
 		/// <inheritdoc cref="IDatabaseAsync.SortedSetRangeByScoreWithScoresAsync(RedisKey, double, double, Exclude, Order, long, long, CommandFlags)"/>
-		public static async Task<SortedSetEntry<T>[]> SortedSetRangeByScoreWithScoresAsync<T>(this IDatabaseAsync Database, RedisSortedSet<T> Set, double Start = double.NegativeInfinity, double Stop = double.PositiveInfinity, Exclude Exclude = Exclude.None, Order Order = Order.Ascending, long Skip = 0L, long Take = -1L, CommandFlags Flags = CommandFlags.None)
+		public async Task<SortedSetEntry<T>[]> RangeByScoreWithScoresAsync<T>(double Start = double.NegativeInfinity, double Stop = double.PositiveInfinity, Exclude Exclude = Exclude.None, Order Order = Order.Ascending, long Skip = 0L, long Take = -1L, CommandFlags Flags = CommandFlags.None)
 		{
-			SortedSetEntry[] Values = await Database.SortedSetRangeByScoreWithScoresAsync(Set.Key, Start, Stop, Exclude, Order, Skip, Take, Flags);
+			SortedSetEntry[] Values = await Database.SortedSetRangeByScoreWithScoresAsync(Key, Start, Stop, Exclude, Order, Skip, Take, Flags);
 			return Array.ConvertAll(Values, x => new SortedSetEntry<T>(x));
 		}
+	}
 
+	/// <summary>
+	/// Extension methods for sets
+	/// </summary>
+	public static class RedisSortedSetExtensions
+	{
+		/// <summary>
+		/// Creates a version of this set which modifies a transaction rather than the direct DB
+		/// </summary>
+		/// <param name="Transaction"></param>
+		/// <returns></returns>
+		public static RedisSortedSet<TElement> With<TElement>(this ITransaction Transaction, RedisSortedSet<TElement> Set)
+		{
+			return new RedisSortedSet<TElement>(Transaction, Set.Key);
+		}
 	}
 }
diff --git a/Engine/Source/Programs/Shared/EpicGames.Redis/RedisString.cs b/Engine/Source/Programs/Shared/EpicGames.Redis/RedisString.cs
--- a/Engine/Source/Programs/Shared/EpicGames.Redis/RedisString.cs
+++ b/Engine/Source/Programs/Shared/EpicGames.Redis/RedisString.cs
@@ -12,8 +12,10 @@ namespace EpicGames.Redis
 	/// Represents a typed Redis list with a given key
 	/// </summary>
 	/// <typeparam name="TElement">The type of element stored in the list</typeparam>
-	public readonly struct RedisString<TElement> : IEquatable<RedisString<TElement>>
+	public readonly struct RedisString<TElement>
 	{
+		readonly IDatabaseAsync Database;
+
 		/// <summary>
 		/// The key for the list
 		/// </summary>
@@ -23,40 +25,39 @@ namespace EpicGames.Redis
 		/// Constructor
 		/// </summary>
 		/// <param name="Key"></param>
-		public RedisString(RedisKey Key) => this.Key = Key;
-
-		/// <inheritdoc/>
-		public override bool Equals(object? Obj) => Obj is RedisString<TElement> String && Key == String.Key;
-
-		/// <inheritdoc/>
-		public bool Equals(RedisString<TElement> Other) => Key == Other.Key;
-
-		/// <inheritdoc/>
-		public override int GetHashCode() => Key.GetHashCode();
+		public RedisString(IDatabaseAsync Database, RedisKey Key)
+		{
+			this.Database = Database;
+			this.Key = Key;
+		}
 
-		/// <summary>Compares two instances for equality</summary>
-		public static bool operator ==(RedisString<TElement> Left, RedisString<TElement> Right) => Left.Key == Right.Key;
+		/// <inheritdoc cref="IDatabaseAsync.StringGetAsync(RedisKey, CommandFlags)"/>
+		public async Task<TElement> GetAsync(CommandFlags Flags = CommandFlags.None)
+		{
+			RedisValue Value = await Database.StringGetAsync(Key, Flags);
+			return RedisSerializer.Deserialize<TElement>(Value);
+		}
 
-		/// <summary>Compares two instances for equality</summary>
-		public static bool operator !=(RedisString<TElement> Left, RedisString<TElement> Right) => Left.Key != Right.Key;
+		/// <inheritdoc cref="IDatabaseAsync.StringSetAsync(RedisKey, RedisValue, TimeSpan?, When, CommandFlags)"/>
+		public Task SetAsync(TElement Value, TimeSpan? Expiry = null, When When = When.Always, CommandFlags Flags = CommandFlags.None)
+		{
+			return Database.StringSetAsync(Key, RedisSerializer.Serialize(Value), Expiry, When, Flags);
+		}
 	}
 
 	/// <summary>
-	/// Extension methods for <see cref="IDatabaseAsync"/>
+	/// Extension methods for sets
 	/// </summary>
 	public static class RedisStringExtensions
 	{
-		/// <inheritdoc cref="IDatabaseAsync.StringSetAsync(RedisKey, RedisValue, TimeSpan?, When, CommandFlags)"/>
-		public static Task StringSetAsync<T>(this IDatabaseAsync Database, RedisString<T> Str, T Value, TimeSpan? Expiry = null, When When = When.Always, CommandFlags Flags = CommandFlags.None)
-		{
-			return Database.StringSetAsync(Str.Key, RedisSerializer.Serialize(Value), Expiry, When, Flags);
-		}
-
-		/// <inheritdoc cref="IDatabaseAsync.StringGetAsync(RedisKey, CommandFlags)"/>
-		public static async Task<T> StringSetAsync<T>(this IDatabaseAsync Database, RedisString<T> Str, CommandFlags Flags = CommandFlags.None)
+		/// <summary>
+		/// Creates a version of this string which modifies a transaction rather than the direct DB
+		/// </summary>
+		/// <param name="Transaction"></param>
+		/// <returns></returns>
+		public static RedisString<TElement> With<TElement>(this ITransaction Transaction, RedisString<TElement> Set)
 		{
-			RedisValue Value = await Database.StringGetAsync(Str.Key, Flags);
-			return RedisSerializer.Deserialize<T>(Value);
+			return new RedisString<TElement>(Transaction, Set.Key);
 		}
 	}
 }
