Horde: Missing file.

diff --git a/Engine/Source/Programs/Shared/EpicGames.Redis/RedisList.cs b/Engine/Source/Programs/Shared/EpicGames.Redis/RedisList.cs
--- a/Engine/Source/Programs/Shared/EpicGames.Redis/RedisList.cs
+++ b/Engine/Source/Programs/Shared/EpicGames.Redis/RedisList.cs
@@ -66,7 +66,7 @@ namespace EpicGames.Redis
 		public async Task<TElement[]> RangeAsync(long Start = 0, long Stop = -1, CommandFlags Flags = CommandFlags.None)
 		{
 			RedisValue[] Values = await Database.ListRangeAsync(Key, Start, Stop, Flags);
-			return Array.ConvertAll(Values, (Converter<RedisValue, T>)(x => (T)RedisSerializer.Deserialize<TElement>(x)));
+			return Array.ConvertAll(Values, (Converter<RedisValue, TElement>)(x => (TElement)RedisSerializer.Deserialize<TElement>(x)));
 		}
 
 		/// <inheritdoc cref="IDatabaseAsync.ListRightPushAsync(RedisKey, RedisValue, When, CommandFlags)"/>
