Horde: Add missing methods to strongly typed Redis wrappers, and add support for serializing all primitive types.

diff --git a/Engine/Source/Programs/Shared/EpicGames.Redis/EpicGames.Redis.csproj b/Engine/Source/Programs/Shared/EpicGames.Redis/EpicGames.Redis.csproj
--- a/Engine/Source/Programs/Shared/EpicGames.Redis/EpicGames.Redis.csproj
+++ b/Engine/Source/Programs/Shared/EpicGames.Redis/EpicGames.Redis.csproj
@@ -9,6 +9,7 @@
     <GenerateTargetFrameworkAttribute>false</GenerateTargetFrameworkAttribute>
     <Configurations>Debug;Release;Development</Configurations>
     <Nullable>enable</Nullable>
+    <GenerateDocumentationFile>True</GenerateDocumentationFile>
   </PropertyGroup>
   <PropertyGroup Condition=" '$(Configuration)|$(Platform)' == 'Debug|AnyCPU' ">
     <DocumentationFile></DocumentationFile>
diff --git a/Engine/Source/Programs/Shared/EpicGames.Redis/RedisChannel.cs b/Engine/Source/Programs/Shared/EpicGames.Redis/RedisChannel.cs
--- a/Engine/Source/Programs/Shared/EpicGames.Redis/RedisChannel.cs
+++ b/Engine/Source/Programs/Shared/EpicGames.Redis/RedisChannel.cs
@@ -80,8 +80,6 @@ namespace EpicGames.Redis
 		/// <summary>
 		/// Start the subscription
 		/// </summary>
-		/// <param name="Flags"></param>
-		/// <returns></returns>
 		internal static async Task<RedisChannelSubscription<T>> CreateAsync(ISubscriber Subscriber, RedisChannel<T> Channel, Action<RedisChannel<T>, T> Handler, CommandFlags Flags = CommandFlags.None)
 		{
 			RedisChannelSubscription<T> Subscription = new RedisChannelSubscription<T>(Subscriber, Channel, Handler);
@@ -92,8 +90,7 @@ namespace EpicGames.Redis
 		/// <summary>
 		/// Unsubscribe from the channel
 		/// </summary>
-		/// <param name="Flags">Flags for the oper</param>
-		/// <returns></returns>
+		/// <param name="Flags">Flags for the operation</param>
 		public async Task UnsubscribeAsync(CommandFlags Flags = CommandFlags.None)
 		{
 			if (Subscriber != null)
diff --git a/Engine/Source/Programs/Shared/EpicGames.Redis/RedisHash.cs b/Engine/Source/Programs/Shared/EpicGames.Redis/RedisHash.cs
--- a/Engine/Source/Programs/Shared/EpicGames.Redis/RedisHash.cs
+++ b/Engine/Source/Programs/Shared/EpicGames.Redis/RedisHash.cs
@@ -16,7 +16,7 @@ namespace EpicGames.Redis
 	/// <typeparam name="TValue">The value type for the hash</typeparam>
 	public readonly struct RedisHash<TName, TValue>
 	{
-		readonly IDatabaseAsync Database;
+		internal readonly IDatabaseAsync Database;
 
 		/// <summary>
 		/// The key for the hash
@@ -26,7 +26,6 @@ namespace EpicGames.Redis
 		/// <summary>
 		/// Constructor
 		/// </summary>
-		/// <param name="Key"></param>
 		public RedisHash(IDatabaseAsync Database, RedisKey Key)
 		{
 			this.Database = Database;
@@ -40,9 +39,31 @@ namespace EpicGames.Redis
 		}
 
 		/// <inheritdoc cref="IDatabaseAsync.HashDeleteAsync(RedisKey, RedisValue[], CommandFlags)"/>
-		public Task<long> DeleteAsync(TName[] Names, CommandFlags Flags = CommandFlags.None)
+		public Task<long> DeleteAsync(IEnumerable<TName> Names, CommandFlags Flags = CommandFlags.None)
 		{
-			return Database.HashDeleteAsync(Key, Array.ConvertAll(Names, x => RedisSerializer.Serialize(x)), Flags);
+			RedisValue[] NameArray = Names.Select(x => RedisSerializer.Serialize(x)).ToArray();
+			return Database.HashDeleteAsync(Key, NameArray, Flags);
+		}
+
+		/// <inheritdoc cref="IDatabaseAsync.HashExistsAsync(RedisKey, RedisValue, CommandFlags)"/>
+		public Task<bool> ExistsAsync(TName Name, CommandFlags Flags = CommandFlags.None)
+		{
+			return Database.HashExistsAsync(Key, RedisSerializer.Serialize(Name), Flags);
+		}
+
+		/// <inheritdoc cref="IDatabaseAsync.HashGetAsync(RedisKey, RedisValue, CommandFlags)"/>
+		public async Task<TValue> GetAsync(TName Name, CommandFlags Flags = CommandFlags.None)
+		{
+			RedisValue Value = await Database.HashGetAsync(Key, RedisSerializer.Serialize(Name), Flags);
+			return Value.IsNull ? default(TValue)! : RedisSerializer.Deserialize<TValue>(Value);
+		}
+
+		/// <inheritdoc cref="IDatabaseAsync.HashGetAsync(RedisKey, RedisValue[], CommandFlags)"/>
+		public async Task<TValue[]> GetAsync(IEnumerable<TName> Names, CommandFlags Flags = CommandFlags.None)
+		{
+			RedisValue[] NameArray = Names.Select(x => RedisSerializer.Serialize(x)).ToArray();
+			RedisValue[] ValueArray = await Database.HashGetAsync(Key, NameArray, Flags);
+			return Array.ConvertAll(ValueArray, x => RedisSerializer.Deserialize<TValue>(x));
 		}
 
 		/// <inheritdoc cref="IDatabaseAsync.HashGetAllAsync(RedisKey, CommandFlags)"/>
@@ -52,6 +73,13 @@ namespace EpicGames.Redis
 			return Array.ConvertAll(Entries, x => new HashEntry<TName, TValue>(RedisSerializer.Deserialize<TName>(x.Name), RedisSerializer.Deserialize<TValue>(x.Value)));
 		}
 
+		/// <inheritdoc cref="IDatabaseAsync.HashKeysAsync(RedisKey, CommandFlags)"/>
+		public async Task<TName[]> KeysAsync(CommandFlags Flags = CommandFlags.None)
+		{
+			RedisValue[] NameArray = await Database.HashKeysAsync(Key, Flags);
+			return Array.ConvertAll(NameArray, x => RedisSerializer.Deserialize<TName>(x));
+		}
+
 		/// <inheritdoc cref="IDatabaseAsync.HashLengthAsync(RedisKey, CommandFlags)"/>
 		public Task<long> LengthAsync(CommandFlags Flags = CommandFlags.None)
 		{
@@ -65,19 +93,26 @@ namespace EpicGames.Redis
 		}
 
 		/// <inheritdoc cref="IDatabaseAsync.HashSetAsync(RedisKey, HashEntry[], CommandFlags)"/>
-		public Task SetAsync(HashEntry[] Entries, CommandFlags Flags = CommandFlags.None)
+		public Task SetAsync(IEnumerable<HashEntry<TName, TValue>> Entries, CommandFlags Flags = CommandFlags.None)
 		{
-			return Database.HashSetAsync(Key, Array.ConvertAll(Entries, x => (HashEntry)x), Flags);
+			return Database.HashSetAsync(Key, Entries.Select(x => (HashEntry)x).ToArray(), Flags);
+		}
+
+		/// <inheritdoc cref="IDatabaseAsync.HashValuesAsync(RedisKey, CommandFlags)"/>
+		public async Task<TValue[]> ValuesAsync(CommandFlags Flags = CommandFlags.None)
+		{
+			RedisValue[] Values = await Database.HashValuesAsync(Key);
+			return Array.ConvertAll(Values, x => RedisSerializer.Deserialize<TValue>(x));
 		}
 	}
 
 	/// <inheritdoc cref="HashEntry"/>
 	public readonly struct HashEntry<TName, TValue>
 	{
-		/// <inheritdoc cref="HashEntry.Name">
+		/// <inheritdoc cref="HashEntry.Name"/>
 		public readonly TName Name { get; }
 
-		/// <inheritdoc cref="HashEntry.Value">
+		/// <inheritdoc cref="HashEntry.Value"/>
 		public readonly TValue Value { get; }
 
 		/// <summary>
@@ -115,11 +150,33 @@ namespace EpicGames.Redis
 	/// </summary>
 	public static class RedisHashExtensions
 	{
+		/// <inheritdoc cref="IDatabaseAsync.HashDecrementAsync(RedisKey, RedisValue, long, CommandFlags)"/>
+		public static Task<long> DecrementAsync<TName>(this RedisHash<TName, long> Hash, TName Name, long Value = 1L, CommandFlags Flags = CommandFlags.None)
+		{
+			return Hash.Database.HashDecrementAsync(Hash.Key, RedisSerializer.Serialize<TName>(Name), Value, Flags);
+		}
+
+		/// <inheritdoc cref="IDatabaseAsync.HashDecrementAsync(RedisKey, RedisValue, double, CommandFlags)"/>
+		public static Task<double> DecrementAsync<TName>(this RedisHash<TName, long> Hash, TName Name, double Value = 1.0, CommandFlags Flags = CommandFlags.None)
+		{
+			return Hash.Database.HashDecrementAsync(Hash.Key, RedisSerializer.Serialize<TName>(Name), Value, Flags);
+		}
+
+		/// <inheritdoc cref="IDatabaseAsync.HashIncrementAsync(RedisKey, RedisValue, long, CommandFlags)"/>
+		public static Task<long> IncrementAsync<TName>(this RedisHash<TName, long> Hash, TName Name, long Value = 1L, CommandFlags Flags = CommandFlags.None)
+		{
+			return Hash.Database.HashIncrementAsync(Hash.Key, RedisSerializer.Serialize<TName>(Name), Value, Flags);
+		}
+
+		/// <inheritdoc cref="IDatabaseAsync.HashIncrementAsync(RedisKey, RedisValue, double, CommandFlags)"/>
+		public static Task<double> IncrementAsync<TName>(this RedisHash<TName, long> Hash, TName Name, double Value = 1.0, CommandFlags Flags = CommandFlags.None)
+		{
+			return Hash.Database.HashIncrementAsync(Hash.Key, RedisSerializer.Serialize<TName>(Name), Value, Flags);
+		}
+
 		/// <summary>
 		/// Creates a version of this set which modifies a transaction rather than the direct DB
 		/// </summary>
-		/// <param name="Transaction"></param>
-		/// <returns></returns>
 		public static RedisHash<TName, TValue> With<TName, TValue>(this ITransaction Transaction, RedisHash<TName, TValue> Set)
 		{
 			return new RedisHash<TName, TValue>(Transaction, Set.Key);
diff --git a/Engine/Source/Programs/Shared/EpicGames.Redis/RedisList.cs b/Engine/Source/Programs/Shared/EpicGames.Redis/RedisList.cs
--- a/Engine/Source/Programs/Shared/EpicGames.Redis/RedisList.cs
+++ b/Engine/Source/Programs/Shared/EpicGames.Redis/RedisList.cs
@@ -29,17 +29,40 @@ namespace EpicGames.Redis
 		/// <summary>
 		/// Constructor
 		/// </summary>
-		/// <param name="Key"></param>
 		public RedisList(IDatabaseAsync Database, RedisKey Key)
 		{
 			this.Database = Database;
 			this.Key = Key;
 		}
 
-		/// <inheritdoc cref="IDatabaseAsync.ListLengthAsync(RedisKey, CommandFlags)"/>
-		public Task<long> LengthAsync()
+		/// <inheritdoc cref="IDatabaseAsync.ListGetByIndexAsync(RedisKey, long, CommandFlags)"/>
+		public async Task<TElement> GetByIndexAsync(long Index, CommandFlags Flags = CommandFlags.None)
 		{
-			return Database.ListLengthAsync(Key);
+			RedisValue Value = await Database.ListGetByIndexAsync(Key, Index, Flags);
+			return Value.IsNull? default(TElement)! : RedisSerializer.Deserialize<TElement>(Value);
+		}
+
+		/// <inheritdoc cref="IDatabaseAsync.ListInsertAfterAsync(RedisKey, RedisValue, RedisValue, CommandFlags)"/>
+		public Task<long> InsertAfterAsync(TElement Pivot, TElement Item, CommandFlags Flags = CommandFlags.None)
+		{
+			RedisValue PivotValue = RedisSerializer.Serialize(Pivot);
+			RedisValue ItemValue = RedisSerializer.Serialize(Item);
+			return Database.ListInsertAfterAsync(Key, PivotValue, ItemValue, Flags);
+		}
+
+		/// <inheritdoc cref="IDatabaseAsync.ListInsertBeforeAsync(RedisKey, RedisValue, RedisValue, CommandFlags)"/>
+		public Task<long> InsertBeforeAsync(TElement Pivot, TElement Item, CommandFlags Flags = CommandFlags.None)
+		{
+			RedisValue PivotValue = RedisSerializer.Serialize(Pivot);
+			RedisValue ItemValue = RedisSerializer.Serialize(Item);
+			return Database.ListInsertBeforeAsync(Key, PivotValue, ItemValue, Flags);
+		}
+
+		/// <inheritdoc cref="IDatabaseAsync.ListLeftPopAsync(RedisKey, CommandFlags)"/>
+		public async Task<TElement> LeftPopAsync(CommandFlags Flags = CommandFlags.None)
+		{
+			RedisValue Value = await Database.ListLeftPopAsync(Key, Flags);
+			return Value.IsNull ? default(TElement)! : RedisSerializer.Deserialize<TElement>(Value);
 		}
 
 		/// <inheritdoc cref="IDatabaseAsync.ListLeftPushAsync(RedisKey, RedisValue, When, CommandFlags)"/>
@@ -55,20 +78,33 @@ namespace EpicGames.Redis
 			return Database.ListLeftPushAsync(Key, Values, When, Flags);
 		}
 
-		/// <inheritdoc cref="IDatabaseAsync.ListLeftPopAsync(RedisKey, CommandFlags)">
-		public async Task<TElement> LeftPopAsync(CommandFlags Flags = CommandFlags.None)
+		/// <inheritdoc cref="IDatabaseAsync.ListLengthAsync(RedisKey, CommandFlags)"/>
+		public Task<long> LengthAsync()
 		{
-			RedisValue Value = await Database.ListLeftPopAsync(Key, Flags);
-			return Value.IsNull ? default(TElement)! : RedisSerializer.Deserialize<TElement>(Value);
+			return Database.ListLengthAsync(Key);
 		}
 
-		/// <inheritdoc cref="IDatabaseAsync.ListRangeAsync(RedisKey, CommandFlags)">
+		/// <inheritdoc cref="IDatabaseAsync.ListRangeAsync(RedisKey, long, long, CommandFlags)"/>
 		public async Task<TElement[]> RangeAsync(long Start = 0, long Stop = -1, CommandFlags Flags = CommandFlags.None)
 		{
 			RedisValue[] Values = await Database.ListRangeAsync(Key, Start, Stop, Flags);
 			return Array.ConvertAll(Values, (Converter<RedisValue, TElement>)(x => (TElement)RedisSerializer.Deserialize<TElement>(x)));
 		}
 
+		/// <inheritdoc cref="IDatabaseAsync.ListRemoveAsync(RedisKey, RedisValue, long, CommandFlags)"/>
+		public Task<long> RemoveAsync(TElement Item, long Count = 0L, CommandFlags Flags = CommandFlags.None)
+		{
+			RedisValue Value = RedisSerializer.Serialize(Item);
+			return Database.ListRemoveAsync(Key, Value, Count, Flags);
+		}
+
+		/// <inheritdoc cref="IDatabaseAsync.ListRightPopAsync(RedisKey, CommandFlags)"/>
+		public async Task<TElement> RightPopAsync(CommandFlags Flags = CommandFlags.None)
+		{
+			RedisValue Value = await Database.ListRightPopAsync(Key, Flags);
+			return Value.IsNull? default(TElement)! : RedisSerializer.Deserialize<TElement>(Value);
+		}
+
 		/// <inheritdoc cref="IDatabaseAsync.ListRightPushAsync(RedisKey, RedisValue, When, CommandFlags)"/>
 		public Task<long> RightPushAsync(TElement Item, When When = When.Always, CommandFlags Flags = CommandFlags.None)
 		{
@@ -82,7 +118,14 @@ namespace EpicGames.Redis
 			return Database.ListRightPushAsync(Key, Values, When, Flags);
 		}
 
-		/// <inheritdoc cref="IDatabaseAsync.ListTrimAsync(RedisKey, long, long, CommandFlags)">
+		/// <inheritdoc cref="IDatabaseAsync.ListSetByIndexAsync(RedisKey, long, RedisValue, CommandFlags)"/>
+		public Task SetByIndexAsync(long Index, TElement Item, CommandFlags Flags = CommandFlags.None)
+		{
+			RedisValue Value = RedisSerializer.Serialize(Item);
+			return Database.ListSetByIndexAsync(Key, Index, Value, Flags);
+		}
+
+		/// <inheritdoc cref="IDatabaseAsync.ListTrimAsync(RedisKey, long, long, CommandFlags)"/>
 		public Task TrimAsync(long Start, long Stop, CommandFlags Flags = CommandFlags.None)
 		{
 			return Database.ListTrimAsync(Key, Start, Stop, Flags);
@@ -97,8 +140,6 @@ namespace EpicGames.Redis
 		/// <summary>
 		/// Creates a version of this list which modifies a transaction rather than the direct DB
 		/// </summary>
-		/// <param name="Transaction"></param>
-		/// <returns></returns>
 		public static RedisList<TElement> With<TElement>(this ITransaction Transaction, RedisList<TElement> Set)
 		{
 			return new RedisList<TElement>(Transaction, Set.Key);
diff --git a/Engine/Source/Programs/Shared/EpicGames.Redis/RedisSerializer.cs b/Engine/Source/Programs/Shared/EpicGames.Redis/RedisSerializer.cs
--- a/Engine/Source/Programs/Shared/EpicGames.Redis/RedisSerializer.cs
+++ b/Engine/Source/Programs/Shared/EpicGames.Redis/RedisSerializer.cs
@@ -89,10 +89,48 @@ namespace EpicGames.Redis
 			public T FromRedisValue(RedisValue Value) => (T)TypeConverter.ConvertFrom((string)Value);
 		}
 
-		class RedisIntConverter : IRedisConverter<int>
+		class RedisNativeConverter<T> : IRedisConverter<T>
 		{
-			public RedisValue ToRedisValue(int Value) => Value;
-			public int FromRedisValue(RedisValue Value) => (int?)Value ?? 0;
+			Func<RedisValue, T> FromRedisValueFunc;
+			Func<T, RedisValue> ToRedisValueFunc;
+
+			public RedisNativeConverter(Func<RedisValue, T> FromRedisValueFunc, Func<T, RedisValue> ToRedisValueFunc)
+			{
+				this.FromRedisValueFunc = FromRedisValueFunc;
+				this.ToRedisValueFunc = ToRedisValueFunc;
+			}
+
+			public T FromRedisValue(RedisValue Value) => FromRedisValueFunc(Value);
+			public RedisValue ToRedisValue(T Value) => ToRedisValueFunc(Value);
+		}
+
+		static Dictionary<Type, object> NativeConverters = CreateNativeConverterLookup();
+
+		static Dictionary<Type, object> CreateNativeConverterLookup()
+		{
+			KeyValuePair<Type, object>[] Converters =
+			{
+				CreateNativeConverter(x => (bool)x, x => x),
+				CreateNativeConverter(x => (int)x, x => x),
+				CreateNativeConverter(x => (int?)x, x => x),
+				CreateNativeConverter(x => (uint)x, x => x),
+				CreateNativeConverter(x => (uint?)x, x => x),
+				CreateNativeConverter(x => (long)x, x => x),
+				CreateNativeConverter(x => (long?)x, x => x),
+				CreateNativeConverter(x => (ulong)x, x => x),
+				CreateNativeConverter(x => (ulong?)x, x => x),
+				CreateNativeConverter(x => (double)x, x => x),
+				CreateNativeConverter(x => (double?)x, x => x),
+				CreateNativeConverter(x => (ReadOnlyMemory<byte>)x, x => x),
+				CreateNativeConverter(x => (byte[])x, x => x),
+				CreateNativeConverter(x => (string)x, x => x),
+			};
+			return new Dictionary<Type, object>(Converters);
+		}
+
+		static KeyValuePair<Type, object> CreateNativeConverter<T>(Func<RedisValue, T> FromRedisValueFunc, Func<T, RedisValue> ToRedisValueFunc)
+		{
+			return new KeyValuePair<Type, object>(typeof(T), new RedisNativeConverter<T>(FromRedisValueFunc, ToRedisValueFunc));
 		}
 
 		/// <summary>
@@ -120,10 +158,11 @@ namespace EpicGames.Redis
 				}
 
 				// Check for known basic types
-				if (Type == typeof(int))
-				{
-					return (IRedisConverter<T>)new RedisIntConverter();
-				}
+				object? NativeConverter;
+				if (NativeConverters.TryGetValue(typeof(T), out NativeConverter))
+                {
+					return (IRedisConverter<T>)NativeConverter;
+                }
 
 				// Check if there's a regular converter we can use to convert to/from a string
 				TypeConverter? Converter = TypeDescriptor.GetConverter(Type);
diff --git a/Engine/Source/Programs/Shared/EpicGames.Redis/RedisSet.cs b/Engine/Source/Programs/Shared/EpicGames.Redis/RedisSet.cs
--- a/Engine/Source/Programs/Shared/EpicGames.Redis/RedisSet.cs
+++ b/Engine/Source/Programs/Shared/EpicGames.Redis/RedisSet.cs
@@ -24,7 +24,6 @@ namespace EpicGames.Redis
 		/// <summary>
 		/// Constructor
 		/// </summary>
-		/// <param name="Key"></param>
 		public RedisSet(IDatabaseAsync Database, RedisKey Key)
 		{
 			this.Database = Database;
@@ -32,15 +31,29 @@ namespace EpicGames.Redis
 		}
 
 		/// <inheritdoc cref="IDatabaseAsync.SetAddAsync(RedisKey, RedisValue, CommandFlags)"/>
-		public Task AddAsync(TElement Value, CommandFlags Flags = CommandFlags.None)
+		public Task AddAsync(TElement Item, CommandFlags Flags = CommandFlags.None)
 		{
-			return Database.SetAddAsync(Key, RedisSerializer.Serialize(Value), Flags);
+			return Database.SetAddAsync(Key, RedisSerializer.Serialize(Item), Flags);
 		}
 
-		/// <inheritdoc cref="IDatabaseAsync.SetRemoveAsync(RedisKey, RedisValue, CommandFlags)"/>
-		public Task RemoveAsync(TElement Value, CommandFlags Flags = CommandFlags.None)
+		/// <inheritdoc cref="IDatabaseAsync.SetAddAsync(RedisKey, RedisValue[], CommandFlags)"/>
+		public Task AddAsync(IEnumerable<TElement> Items, CommandFlags Flags = CommandFlags.None)
+		{
+			RedisValue[] Values = Items.Select(x => RedisSerializer.Serialize(x)).ToArray();
+			return Database.SetAddAsync(Key, Values, Flags);
+		}
+
+		/// <inheritdoc cref="IDatabaseAsync.SetContainsAsync(RedisKey, RedisValue, CommandFlags)"/>
+		public Task<bool> ContainsAsync(TElement Item, CommandFlags Flags = CommandFlags.None)
+		{
+			RedisValue Value = RedisSerializer.Serialize(Item);
+			return Database.SetContainsAsync(Key, Value, Flags);
+		}
+
+		/// <inheritdoc cref="IDatabaseAsync.SetLengthAsync(RedisKey, CommandFlags)"/>
+		public Task<long> LengthAsync(CommandFlags Flags = CommandFlags.None)
 		{
-			return Database.SetRemoveAsync(Key, RedisSerializer.Serialize(Value), Flags);
+			return Database.SetLengthAsync(Key, Flags);
 		}
 
 		/// <inheritdoc cref="IDatabaseAsync.SetRemoveAsync(RedisKey, RedisValue, CommandFlags)"/>
@@ -49,6 +62,33 @@ namespace EpicGames.Redis
 			RedisValue[] Values = await Database.SetMembersAsync(Key, Flags);
 			return Array.ConvertAll(Values, (Converter<RedisValue, TElement>)(x => RedisSerializer.Deserialize<TElement>(x)));
 		}
+
+		/// <inheritdoc cref="IDatabaseAsync.SetPopAsync(RedisKey, CommandFlags)"/>
+		public async Task<TElement> PopAsync(CommandFlags Flags = CommandFlags.None)
+		{
+			RedisValue Value = await Database.SetPopAsync(Key, Flags);
+			return Value.IsNull ? default(TElement)! : RedisSerializer.Deserialize<TElement>(Value);
+		}
+
+		/// <inheritdoc cref="IDatabaseAsync.SetPopAsync(RedisKey, long, CommandFlags)"/>
+		public async Task<TElement[]> PopAsync(long Count, CommandFlags Flags = CommandFlags.None)
+		{
+			RedisValue[] Values = await Database.SetPopAsync(Key, Count, Flags);
+			return Array.ConvertAll<RedisValue, TElement>(Values, x => RedisSerializer.Deserialize<TElement>(x));
+		}
+
+		/// <inheritdoc cref="IDatabaseAsync.SetRemoveAsync(RedisKey, RedisValue, CommandFlags)"/>
+		public Task RemoveAsync(TElement Item, CommandFlags Flags = CommandFlags.None)
+		{
+			return Database.SetRemoveAsync(Key, RedisSerializer.Serialize(Item), Flags);
+		}
+
+		/// <inheritdoc cref="IDatabaseAsync.SetRemoveAsync(RedisKey, RedisValue[], CommandFlags)"/>
+		public Task RemoveAsync(IEnumerable<TElement> Items, CommandFlags Flags = CommandFlags.None)
+		{
+			RedisValue[] Values = Items.Select(x => RedisSerializer.Serialize(x)).ToArray();
+			return Database.SetRemoveAsync(Key, Values, Flags);
+		}
 	}
 
 	/// <summary>
@@ -59,8 +99,6 @@ namespace EpicGames.Redis
 		/// <summary>
 		/// Creates a version of this set which modifies a transaction rather than the direct DB
 		/// </summary>
-		/// <param name="Transaction"></param>
-		/// <returns></returns>
 		public static RedisSet<TElement> With<TElement>(this ITransaction Transaction, RedisSet<TElement> Set)
 		{
 			return new RedisSet<TElement>(Transaction, Set.Key);
diff --git a/Engine/Source/Programs/Shared/EpicGames.Redis/RedisSortedSet.cs b/Engine/Source/Programs/Shared/EpicGames.Redis/RedisSortedSet.cs
--- a/Engine/Source/Programs/Shared/EpicGames.Redis/RedisSortedSet.cs
+++ b/Engine/Source/Programs/Shared/EpicGames.Redis/RedisSortedSet.cs
@@ -63,7 +63,6 @@ namespace EpicGames.Redis
 		/// <summary>
 		/// Constructor
 		/// </summary>
-		/// <param name="Key"></param>
 		public RedisSortedSet(IDatabaseAsync Database, RedisKey Key)
 		{
 			this.Database = Database;
@@ -71,58 +70,116 @@ namespace EpicGames.Redis
 		}
 
 		/// <inheritdoc cref="IDatabaseAsync.SortedSetAddAsync(RedisKey, RedisValue, double, CommandFlags)"/>
-		public Task<bool> AddAsync<T>(T Value, double Score, When When = When.Always, CommandFlags Flags = CommandFlags.None)
+		public Task<bool> AddAsync(TElement Item, double Score, When When = When.Always, CommandFlags Flags = CommandFlags.None)
 		{
-			return Database.SortedSetAddAsync(Key, RedisSerializer.Serialize(Value), Score, When, Flags);
+			RedisValue Value = RedisSerializer.Serialize(Item);
+			return Database.SortedSetAddAsync(Key, Value, Score, When, Flags);
 		}
 
-		/// <inheritdoc cref="IDatabaseAsync.SortedSetAddAsync(RedisKey, RedisValue, double, CommandFlags)"/>
-		public Task<long> AddAsync<T>(SortedSetEntry<T>[] Values, CommandFlags Flags = CommandFlags.None)
+		/// <inheritdoc cref="IDatabaseAsync.SortedSetAddAsync(RedisKey, RedisValue, double, When, CommandFlags)"/>
+		public Task<long> AddAsync(SortedSetEntry<TElement>[] Values, When When = When.Always, CommandFlags Flags = CommandFlags.None)
 		{
 			SortedSetEntry[] Untyped = Array.ConvertAll(Values, x => new SortedSetEntry(RedisSerializer.Serialize(x.Element), x.Score));
-			return Database.SortedSetAddAsync(Key, Untyped, Flags);
+			return Database.SortedSetAddAsync(Key, Untyped, When, Flags);
 		}
 
-		/// <inheritdoc cref="IDatabaseAsync.SortedSetRemoveAsync(RedisKey, RedisValue, CommandFlags)"/>
-		public Task<bool> RemoveAsync<T>(T Value, CommandFlags Flags = CommandFlags.None)
+		/// <inheritdoc cref="IDatabaseAsync.SortedSetLengthAsync(RedisKey, double, double, Exclude, CommandFlags)"/>
+		public Task<long> LengthAsync(double Min = double.NegativeInfinity, double Max = double.PositiveInfinity, Exclude Exclude = Exclude.None, CommandFlags Flags = CommandFlags.None)
 		{
-			return Database.SortedSetRemoveAsync(Key, RedisSerializer.Serialize(Value), Flags);
+			return Database.SortedSetLengthAsync(Key, Min, Max, Exclude, Flags);
 		}
 
-		/// <inheritdoc cref="IDatabaseAsync.SortedSetRemoveRangeByScoreAsync(RedisKey, double, double, Exclude, Order, long, long, CommandFlags)"/>
-		public async Task<long> RemoveRangeByScoreAsync<T>(double Start = double.NegativeInfinity, double Stop = double.PositiveInfinity, Exclude Exclude = Exclude.None, CommandFlags Flags = CommandFlags.None)
+		/// <inheritdoc cref="IDatabaseAsync.SortedSetLengthByValueAsync(RedisKey, RedisValue, RedisValue, Exclude, CommandFlags)"/>
+		public Task<long> LengthByValueAsync(TElement Min, TElement Max, Exclude Exclude = Exclude.None, CommandFlags Flags = CommandFlags.None)
 		{
-			return await Database.SortedSetRemoveRangeByScoreAsync(Key, Start, Stop, Exclude, Flags);
+			RedisValue MinValue = RedisSerializer.Serialize(Min);
+			RedisValue MaxValue = RedisSerializer.Serialize(Max);
+			return Database.SortedSetLengthByValueAsync(Key, MinValue, MaxValue, Exclude, Flags);
 		}
 
-		/// <inheritdoc cref="IDatabaseAsync.SortedSetAddAsync(RedisKey, RedisValue, double, CommandFlags)"/>
-		public async IAsyncEnumerable<SortedSetEntry<T>> ScanAsync<T>(RedisValue Pattern = default, int PageSize = 250, long Cursor = 0, int PageOffset = 0, CommandFlags Flags = CommandFlags.None)
+		/// <inheritdoc cref="IDatabaseAsync.SortedSetRangeByRankAsync(RedisKey, long, long, Order, CommandFlags)"/>
+		public async Task<TElement[]> RangeByRankAsync(long Start, long Stop = -1, Order Order = Order.Ascending, CommandFlags Flags = CommandFlags.None)
 		{
-			await foreach (SortedSetEntry Entry in Database.SortedSetScanAsync(Key, Pattern, PageSize, Cursor, PageOffset, Flags))
-			{
-				yield return new SortedSetEntry<T>(Entry);
-			}
+			RedisValue[] Values = await Database.SortedSetRangeByRankAsync(Key, Start, Stop, Order, Flags);
+			return Array.ConvertAll(Values, x => RedisSerializer.Deserialize<TElement>(x));
 		}
 
-		/// <inheritdoc cref="IDatabaseAsync.SortedSetRangeByRankAsync(RedisKey, long, long, Order, CommandFlags)"/>
-		public async Task<T[]> RangeByRankAsync<T>(long Start, long Stop = -1, Order Order = Order.Ascending, CommandFlags Flags = CommandFlags.None)
+		/// <inheritdoc cref="IDatabaseAsync.SortedSetRangeByRankWithScoresAsync(RedisKey, long, long, Order, CommandFlags)"/>
+		public async Task<SortedSetEntry<TElement>[]> RangeByRankWithScoresAsync(long Start, long Stop = -1, Order Order = Order.Ascending, CommandFlags Flags = CommandFlags.None)
 		{
-			RedisValue[] Values = await Database.SortedSetRangeByRankAsync(Key, Start, Stop, Order, Flags);
-			return Array.ConvertAll(Values, x => RedisSerializer.Deserialize<T>(x));
+			SortedSetEntry[] Values = await Database.SortedSetRangeByRankWithScoresAsync(Key, Start, Stop, Order, Flags);
+			return Array.ConvertAll(Values, x => new SortedSetEntry<TElement>(x));
 		}
 
 		/// <inheritdoc cref="IDatabaseAsync.SortedSetRangeByScoreAsync(RedisKey, double, double, Exclude, Order, long, long, CommandFlags)"/>
-		public async Task<T[]> RangeByScoreAsync<T>(double Start = double.NegativeInfinity, double Stop = double.PositiveInfinity, Exclude Exclude = Exclude.None, Order Order = Order.Ascending, long Skip = 0L, long Take = -1L, CommandFlags Flags = CommandFlags.None)
+		public async Task<TElement[]> RangeByScoreAsync(double Start = double.NegativeInfinity, double Stop = double.PositiveInfinity, Exclude Exclude = Exclude.None, Order Order = Order.Ascending, long Skip = 0L, long Take = -1L, CommandFlags Flags = CommandFlags.None)
 		{
 			RedisValue[] Values = await Database.SortedSetRangeByScoreAsync(Key, Start, Stop, Exclude, Order, Skip, Take, Flags);
-			return Array.ConvertAll(Values, x => RedisSerializer.Deserialize<T>(x));
+			return Array.ConvertAll(Values, x => RedisSerializer.Deserialize<TElement>(x));
 		}
 
 		/// <inheritdoc cref="IDatabaseAsync.SortedSetRangeByScoreWithScoresAsync(RedisKey, double, double, Exclude, Order, long, long, CommandFlags)"/>
-		public async Task<SortedSetEntry<T>[]> RangeByScoreWithScoresAsync<T>(double Start = double.NegativeInfinity, double Stop = double.PositiveInfinity, Exclude Exclude = Exclude.None, Order Order = Order.Ascending, long Skip = 0L, long Take = -1L, CommandFlags Flags = CommandFlags.None)
+		public async Task<SortedSetEntry<TElement>[]> RangeByScoreWithScoresAsync(double Start = double.NegativeInfinity, double Stop = double.PositiveInfinity, Exclude Exclude = Exclude.None, Order Order = Order.Ascending, long Skip = 0L, long Take = -1L, CommandFlags Flags = CommandFlags.None)
 		{
 			SortedSetEntry[] Values = await Database.SortedSetRangeByScoreWithScoresAsync(Key, Start, Stop, Exclude, Order, Skip, Take, Flags);
-			return Array.ConvertAll(Values, x => new SortedSetEntry<T>(x));
+			return Array.ConvertAll(Values, x => new SortedSetEntry<TElement>(x));
+		}
+
+		/// <inheritdoc cref="IDatabaseAsync.SortedSetRangeByValueAsync(RedisKey, RedisValue, RedisValue, Exclude, Order, long, long, CommandFlags)"/>
+		public async Task<TElement[]> RangeByValueAsync(TElement Min, TElement Max, Exclude Exclude = Exclude.None, Order Order = Order.Ascending, long Skip = 0L, long Take = -1L, CommandFlags Flags = CommandFlags.None)
+		{
+			RedisValue MinValue = RedisSerializer.Serialize(Min);
+			RedisValue MaxValue = RedisSerializer.Serialize(Max);
+			RedisValue[] Values = await Database.SortedSetRangeByValueAsync(Key, MinValue, MaxValue, Exclude, Order, Skip, Take, Flags);
+			return Array.ConvertAll(Values, x => RedisSerializer.Deserialize<TElement>(x));
+		}
+
+		/// <inheritdoc cref="IDatabaseAsync.SortedSetRankAsync(RedisKey, RedisValue, Order, CommandFlags)"/>
+		public Task<long?> RankAsync(TElement Item, Order Order = Order.Ascending, CommandFlags Flags = CommandFlags.None)
+		{
+			return Database.SortedSetRankAsync(Key, RedisSerializer.Serialize(Item), Order, Flags);
+		}
+
+		/// <inheritdoc cref="IDatabaseAsync.SortedSetRemoveAsync(RedisKey, RedisValue, CommandFlags)"/>
+		public Task<bool> RemoveAsync(TElement Value, CommandFlags Flags = CommandFlags.None)
+		{
+			return Database.SortedSetRemoveAsync(Key, RedisSerializer.Serialize(Value), Flags);
+		}
+
+		/// <inheritdoc cref="IDatabaseAsync.SortedSetRemoveAsync(RedisKey, RedisValue[], CommandFlags)"/>
+		public Task<long> RemoveAsync(IEnumerable<TElement> Items, CommandFlags Flags = CommandFlags.None)
+		{
+			RedisValue[] Values = Items.Select(x => RedisSerializer.Serialize(x)).ToArray();
+			return Database.SortedSetRemoveAsync(Key, Values, Flags);
+		}
+
+		/// <inheritdoc cref="IDatabaseAsync.SortedSetRemoveRangeByRankAsync(RedisKey, long, long, CommandFlags)"/>
+		public Task<long> RemoveRangeByRankAsync(long Start, long Stop, CommandFlags Flags = CommandFlags.None)
+		{
+			return Database.SortedSetRemoveRangeByRankAsync(Key, Start, Stop, Flags);
+		}
+
+		/// <inheritdoc cref="IDatabaseAsync.SortedSetRemoveRangeByScoreAsync(RedisKey, double, double, Exclude, CommandFlags)"/>
+		public async Task<long> RemoveRangeByScoreAsync(double Start, double Stop, Exclude Exclude = Exclude.None, CommandFlags Flags = CommandFlags.None)
+		{
+			return await Database.SortedSetRemoveRangeByScoreAsync(Key, Start, Stop, Exclude, Flags);
+		}
+
+		/// <inheritdoc cref="IDatabaseAsync.SortedSetRemoveRangeByValueAsync(RedisKey, RedisValue, RedisValue, Exclude, CommandFlags)"/>
+		public async Task<long> RemoveRangeByValueAsync(TElement Min, TElement Max, Exclude Exclude = Exclude.None, CommandFlags Flags = CommandFlags.None)
+		{
+			RedisValue MinValue = RedisSerializer.Serialize(Min);
+			RedisValue MaxValue = RedisSerializer.Serialize(Max);
+			return await Database.SortedSetRemoveRangeByValueAsync(Key, MinValue, MaxValue, Exclude, Flags);
+		}
+
+		/// <inheritdoc cref="IDatabaseAsync.SortedSetAddAsync(RedisKey, RedisValue, double, CommandFlags)"/>
+		public async IAsyncEnumerable<SortedSetEntry<TElement>> ScanAsync(RedisValue Pattern = default, int PageSize = 250, long Cursor = 0, int PageOffset = 0, CommandFlags Flags = CommandFlags.None)
+		{
+			await foreach (SortedSetEntry Entry in Database.SortedSetScanAsync(Key, Pattern, PageSize, Cursor, PageOffset, Flags))
+			{
+				yield return new SortedSetEntry<TElement>(Entry);
+			}
 		}
 	}
 
@@ -134,8 +191,6 @@ namespace EpicGames.Redis
 		/// <summary>
 		/// Creates a version of this set which modifies a transaction rather than the direct DB
 		/// </summary>
-		/// <param name="Transaction"></param>
-		/// <returns></returns>
 		public static RedisSortedSet<TElement> With<TElement>(this ITransaction Transaction, RedisSortedSet<TElement> Set)
 		{
 			return new RedisSortedSet<TElement>(Transaction, Set.Key);
diff --git a/Engine/Source/Programs/Shared/EpicGames.Redis/RedisString.cs b/Engine/Source/Programs/Shared/EpicGames.Redis/RedisString.cs
--- a/Engine/Source/Programs/Shared/EpicGames.Redis/RedisString.cs
+++ b/Engine/Source/Programs/Shared/EpicGames.Redis/RedisString.cs
@@ -14,7 +14,7 @@ namespace EpicGames.Redis
 	/// <typeparam name="TElement">The type of element stored in the list</typeparam>
 	public readonly struct RedisString<TElement>
 	{
-		readonly IDatabaseAsync Database;
+		internal readonly IDatabaseAsync Database;
 
 		/// <summary>
 		/// The key for the list
@@ -24,6 +24,7 @@ namespace EpicGames.Redis
 		/// <summary>
 		/// Constructor
 		/// </summary>
+		/// <param name="Database"></param>
 		/// <param name="Key"></param>
 		public RedisString(IDatabaseAsync Database, RedisKey Key)
 		{
@@ -38,6 +39,12 @@ namespace EpicGames.Redis
 			return RedisSerializer.Deserialize<TElement>(Value);
 		}
 
+		/// <inheritdoc cref="IDatabaseAsync.StringLengthAsync(RedisKey, CommandFlags)"/>
+		public Task<long> LengthAsync(CommandFlags Flags = CommandFlags.None)
+		{
+			return Database.StringLengthAsync(Key, Flags);
+		}
+
 		/// <inheritdoc cref="IDatabaseAsync.StringSetAsync(RedisKey, RedisValue, TimeSpan?, When, CommandFlags)"/>
 		public Task SetAsync(TElement Value, TimeSpan? Expiry = null, When When = When.Always, CommandFlags Flags = CommandFlags.None)
 		{
@@ -50,11 +57,33 @@ namespace EpicGames.Redis
 	/// </summary>
 	public static class RedisStringExtensions
 	{
+		/// <inheritdoc cref="IDatabaseAsync.StringDecrementAsync(RedisKey, long, CommandFlags)"/>
+		public static Task<long> DecrementAsync(this RedisString<long> String, long Value = 1L, CommandFlags Flags = CommandFlags.None)
+		{
+			return String.Database.StringDecrementAsync(String.Key, Value, Flags);
+		}
+
+		/// <inheritdoc cref="IDatabaseAsync.StringDecrementAsync(RedisKey, double, CommandFlags)"/>
+		public static Task<double> DecrementAsync(this RedisString<double> String, double Value = 1.0, CommandFlags Flags = CommandFlags.None)
+		{
+			return String.Database.StringDecrementAsync(String.Key, Value, Flags);
+		}
+
+		/// <inheritdoc cref="IDatabaseAsync.StringDecrementAsync(RedisKey, double, CommandFlags)"/>
+		public static Task<long> IncrementAsync(this RedisString<long> String, long Value = 1L, CommandFlags Flags = CommandFlags.None)
+		{
+			return String.Database.StringIncrementAsync(String.Key, Value, Flags);
+		}
+
+		/// <inheritdoc cref="IDatabaseAsync.StringDecrementAsync(RedisKey, double, CommandFlags)"/>
+		public static Task<double> IncrementAsync(this RedisString<double> String, double Value = 1.0, CommandFlags Flags = CommandFlags.None)
+		{
+			return String.Database.StringIncrementAsync(String.Key, Value, Flags);
+		}
+
 		/// <summary>
 		/// Creates a version of this string which modifies a transaction rather than the direct DB
 		/// </summary>
-		/// <param name="Transaction"></param>
-		/// <returns></returns>
 		public static RedisString<TElement> With<TElement>(this ITransaction Transaction, RedisString<TElement> Set)
 		{
 			return new RedisString<TElement>(Transaction, Set.Key);
