Horde: Fix tagged output of some Perforce commands (eg. p4 dirs).

diff --git a/Engine/Source/Programs/Shared/EpicGames.Perforce.Native/Client.cpp b/Engine/Source/Programs/Shared/EpicGames.Perforce.Native/Client.cpp
--- a/Engine/Source/Programs/Shared/EpicGames.Perforce.Native/Client.cpp
+++ b/Engine/Source/Programs/Shared/EpicGames.Perforce.Native/Client.cpp
@@ -112,6 +112,14 @@ public:
 		}
 	}
 
+	virtual void Message(Error* Err) override
+	{
+		while (!TryOutputError(Err))
+		{
+			Flush();
+		}
+	}
+
 	bool TryOutputError(Error* Err)
 	{
 		static const char CodeKey[] = "code";
@@ -119,61 +127,50 @@ public:
 		static const char GenericKey[] = "generic";
 		static const char DataKey[] = "data";
 
-		static const char ErrorCode[] = "error";
-
-		StrBuf MessageBuf;
-		StrPtr& Message = Err->StrError(MessageBuf);
-
+		StrBuf Message;
+		Err->Fmt(&Message, 0);
 		unsigned int MessageLen = Message.Length();
-		unsigned int RecordLen = 1 + MeasureStringField(CodeKey, ErrorCode) + MeasureIntField(SeverityKey) + MeasureIntField(GenericKey) + MeasureStringField(DataKey, MessageLen) + 1;
-		if (Length + RecordLen > MaxLength)
-		{
-			return false;
-		}
 
-		unsigned char* Pos = Data + Length;
-		*(Pos++) = '{';
-		Pos = WriteStringField(Pos, CodeKey, ErrorCode);
-		Pos = WriteIntField(Pos, SeverityKey, Err->GetSeverity());
-		Pos = WriteIntField(Pos, GenericKey, Err->GetGeneric());
-		Pos = WriteStringField(Pos, DataKey, Message.Text(), MessageLen);
-		*(Pos++) = '0';
-		assert(Pos == Data + Length + RecordLen);
+		unsigned int RecordLen;
+		if (Err->GetSeverity() == E_INFO)
+		{
+			static const char InfoCode[] = "info";
 
-		Length += RecordLen;
-		return true;
-	}
+			RecordLen = 1 + MeasureStringField(CodeKey, InfoCode) + MeasureStringField(DataKey, MessageLen) + 1;
+			if (Length + RecordLen > MaxLength)
+			{
+				return false;
+			}
 
-	virtual void Message(Error* Err) override
-	{
-		while (!TryOutputMessage(Err))
-		{
-			Flush();
+			unsigned char* Pos = Data + Length;
+			*(Pos++) = '{';
+			Pos = WriteStringField(Pos, CodeKey, InfoCode);
+			Pos = WriteStringField(Pos, DataKey, Message.Text(), MessageLen);
+			*(Pos++) = '0';
+			assert(Pos == Data + Length + RecordLen);
 		}
-	}
-
-	bool TryOutputMessage(Error* Err)
-	{
-		// Try to translate these into stat messages depending on the command
-		if (Func != nullptr)
+		else
 		{
-			if (Err->GetSeverity() == E_INFO && strcasecmp(Func, "dirs") == 0)
+			static const char ErrorCode[] = "error";
+
+			RecordLen = 1 + MeasureStringField(CodeKey, ErrorCode) + MeasureIntField(SeverityKey) + MeasureIntField(GenericKey) + MeasureStringField(DataKey, MessageLen) + 1;
+			if (Length + RecordLen > MaxLength)
 			{
-				StrDict* Dict = Err->GetDict();
-				if (Dict != nullptr)
-				{
-					StrPtr* DirName = Dict->GetVar("dirName");
-					if (DirName != nullptr)
-					{
-						StrBufDict StatDict;
-						StatDict.VSetVar(StrBuf("dir"), *DirName);
-						return TryWriteRecord(&StatDict);
-					}
-				}
+				return false;
 			}
+
+			unsigned char* Pos = Data + Length;
+			*(Pos++) = '{';
+			Pos = WriteStringField(Pos, CodeKey, ErrorCode);
+			Pos = WriteIntField(Pos, SeverityKey, Err->GetSeverity());
+			Pos = WriteIntField(Pos, GenericKey, Err->GetGeneric());
+			Pos = WriteStringField(Pos, DataKey, Message.Text(), MessageLen);
+			*(Pos++) = '0';
+			assert(Pos == Data + Length + RecordLen);
 		}
 
-		return TryOutputError(Err);
+		Length += RecordLen;
+		return true;
 	}
 
 	virtual void OutputError(const char* errBuf) override
@@ -307,14 +304,19 @@ public:
 	}
 
 	bool TryWriteRecord(StrDict* VarList)
+	{
+		return TryWriteRecord("stat", 4, VarList);
+	}
+
+	bool TryWriteRecord(const char* Code, unsigned int CodeLen, StrDict* VarList)
 	{
 		unsigned int Pos = Length;
 
 		// Record prefix: '{' (open record), string, 4 bytes, 'code', string ...
-		unsigned char Prefix[] = { '{', 's', 4, 0, 0, 0, 'c', 'o', 'd', 'e', 's', 4, 0, 0, 0, 's', 't', 'a', 't' };
+		unsigned char Prefix[] = { '{', 's', 4, 0, 0, 0, 'c', 'o', 'd', 'e', 's' };
 		unsigned int PrefixLen = sizeof(Prefix) / sizeof(Prefix[0]);
 
-		if (Pos + PrefixLen > MaxLength - 1)
+		if (Pos + PrefixLen + (unsigned int)sizeof(unsigned int) + CodeLen > MaxLength - 1)
 		{
 			return false;
 		}
@@ -322,6 +324,12 @@ public:
 		memcpy(Data + Pos, Prefix, PrefixLen);
 		Pos += PrefixLen;
 
+		memcpy(Data + Pos, &CodeLen, sizeof(CodeLen));
+		Pos += sizeof(unsigned int);
+
+		memcpy(Data + Pos, Code, CodeLen);
+		Pos += CodeLen;
+
 		for (int Idx = 0;; Idx++)
 		{
 			StrRef Field;
@@ -404,7 +412,7 @@ extern "C" NATIVE_API FClient* Client_Create(const FSettings* Settings, FWriteBu
 			Client->ClientApi.SetVersion(Settings->AppVersion);
 		}
 	}
-	Client->ClientApi.SetProtocol("tags", "");
+	Client->ClientApi.SetProtocol("tag", "");
 
 	Error Err;
 	Client->ClientApi.Init(&Err);
