Horde: Improve error messages from Perforce commands.

diff --git a/Engine/Source/Programs/Shared/EpicGames.Perforce/PerforceConnection.cs b/Engine/Source/Programs/Shared/EpicGames.Perforce/PerforceConnection.cs
--- a/Engine/Source/Programs/Shared/EpicGames.Perforce/PerforceConnection.cs
+++ b/Engine/Source/Programs/Shared/EpicGames.Perforce/PerforceConnection.cs
@@ -474,7 +474,7 @@ namespace EpicGames.Perforce
 				{
 					Connection.Logger.LogDebug("Unexpected response {Idx}: {Text}", Idx, Responses[Idx].ToString());
 				}
-				throw new PerforceException("Expected one result from 'p4 {0}', got {1}", Arguments, Responses.Count);
+				throw new PerforceException("Expected one result from 'p4 {0}', got {1}", Command, Responses.Count);
 			}
 			return Responses[0];
 		}
diff --git a/Engine/Source/Programs/Shared/EpicGames.Perforce/PerforceResponse.cs b/Engine/Source/Programs/Shared/EpicGames.Perforce/PerforceResponse.cs
--- a/Engine/Source/Programs/Shared/EpicGames.Perforce/PerforceResponse.cs
+++ b/Engine/Source/Programs/Shared/EpicGames.Perforce/PerforceResponse.cs
@@ -153,7 +153,7 @@ namespace EpicGames.Perforce
 					}
 					else if(InternalData is PerforceError)
 					{
-						throw new PerforceException($"Expected record of type '{typeof(T).Name}', got error: {InternalData}");
+						throw new PerforceException($"{InternalData}");
 					}
 					else
 					{
