EpicGames.Core: Log.cs

Add annotations for format-style arguments

diff --git a/Engine/Source/Programs/Shared/EpicGames.Core/EpicGames.Core.csproj b/Engine/Source/Programs/Shared/EpicGames.Core/EpicGames.Core.csproj
--- a/Engine/Source/Programs/Shared/EpicGames.Core/EpicGames.Core.csproj
+++ b/Engine/Source/Programs/Shared/EpicGames.Core/EpicGames.Core.csproj
@@ -27,6 +27,7 @@
     </Compile>
   </ItemGroup>
   <ItemGroup>
+    <PackageReference Include="JetBrains.Annotations" Version="2021.2.0" />
     <PackageReference Include="Microsoft.Extensions.Logging" Version="5.0.0" />
     <PackageReference Include="Microsoft.Extensions.Logging.Abstractions" Version="5.0.0" />
     <PackageReference Include="System.Memory" Version="4.5.4" />
diff --git a/Engine/Source/Programs/Shared/EpicGames.Core/Log.cs b/Engine/Source/Programs/Shared/EpicGames.Core/Log.cs
--- a/Engine/Source/Programs/Shared/EpicGames.Core/Log.cs
+++ b/Engine/Source/Programs/Shared/EpicGames.Core/Log.cs
@@ -10,6 +10,7 @@ using System.Linq;
 using System.Reflection;
 using System.Runtime.CompilerServices;
 using System.Text;
+using JetBrains.Annotations;
 
 #nullable disable
 
@@ -266,6 +267,7 @@ namespace EpicGames.Core
 		/// <param name="Format">Message format string.</param>
 		/// <param name="Args">Optional arguments</param>
 		[MethodImplAttribute(MethodImplOptions.NoInlining)]
+		[StringFormatMethod("Format")]
 		private static void WriteLinePrivate(int StackFramesToSkip, bool bWriteOnce, LogEventType Verbosity, LogFormatOptions FormatOptions, string Format, params object[] Args)
 		{
 			LogMessage Message = new LogMessage();
@@ -285,6 +287,7 @@ namespace EpicGames.Core
 		/// <param name="Format"></param>
 		/// <param name="Args"></param>
 		[MethodImplAttribute(MethodImplOptions.NoInlining)]
+		[StringFormatMethod("Format")]
 		public static void WriteLineIf(bool Condition, LogEventType Verbosity, string Format, params object[] Args)
 		{
 			if (Condition)
@@ -301,6 +304,7 @@ namespace EpicGames.Core
 		/// <param name="Format"></param>
 		/// <param name="Args"></param>
 		[MethodImplAttribute(MethodImplOptions.NoInlining)]
+		[StringFormatMethod("Format")]
 		public static void WriteLine(int StackFramesToSkip, LogEventType Verbosity, string Format, params object[] Args)
 		{
 			WriteLinePrivate(StackFramesToSkip + 1, false, Verbosity, LogFormatOptions.None, Format, Args);
@@ -315,6 +319,7 @@ namespace EpicGames.Core
 		/// <param name="Format"></param>
 		/// <param name="Args"></param>
 		[MethodImplAttribute(MethodImplOptions.NoInlining)]
+		[StringFormatMethod("Format")]
 		public static void WriteLine(int StackFramesToSkip, LogEventType Verbosity, LogFormatOptions FormatOptions, string Format, params object[] Args)
 		{
 			WriteLinePrivate(StackFramesToSkip + 1, false, Verbosity, FormatOptions, Format, Args);
@@ -327,6 +332,7 @@ namespace EpicGames.Core
 		/// <param name="Format"></param>
 		/// <param name="Args"></param>
 		[MethodImplAttribute(MethodImplOptions.NoInlining)]
+		[StringFormatMethod("Format")]
 		public static void WriteLine(LogEventType Verbosity, string Format, params object[] Args)
 		{
 			WriteLinePrivate(1, false, Verbosity, LogFormatOptions.None, Format, Args);
@@ -340,6 +346,7 @@ namespace EpicGames.Core
 		/// <param name="Format"></param>
 		/// <param name="Args"></param>
 		[MethodImplAttribute(MethodImplOptions.NoInlining)]
+		[StringFormatMethod("Format")]
 		public static void WriteLine(LogEventType Verbosity, LogFormatOptions FormatOptions, string Format, params object[] Args)
 		{
 			WriteLinePrivate(1, false, Verbosity, FormatOptions, Format, Args);
@@ -365,6 +372,7 @@ namespace EpicGames.Core
 		/// <param name="Format">Message format string</param>
 		/// <param name="Args">Optional arguments</param>
 		[MethodImplAttribute(MethodImplOptions.NoInlining)]
+		[StringFormatMethod("Format")]
 		public static void TraceError(string Format, params object[] Args)
 		{
 			WriteLinePrivate(1, false, LogEventType.Error, LogFormatOptions.None, Format, Args);
@@ -377,6 +385,7 @@ namespace EpicGames.Core
 		/// <param name="Format">Message format string</param>
 		/// <param name="Args">Optional arguments</param>
 		[MethodImplAttribute(MethodImplOptions.NoInlining)]
+		[StringFormatMethod("Format")]
 		public static void TraceError(FileReference File, string Format, params object[] Args)
 		{
 			WriteLinePrivate(1, false, LogEventType.Error, LogFormatOptions.NoSeverityPrefix, "{0}: error: {1}", File, String.Format(Format, Args));
@@ -390,6 +399,7 @@ namespace EpicGames.Core
 		/// <param name="Format">Message format string</param>
 		/// <param name="Args">Optional arguments</param>
 		[MethodImplAttribute(MethodImplOptions.NoInlining)]
+		[StringFormatMethod("Format")]
 		public static void TraceError(FileReference File, int Line, string Format, params object[] Args)
 		{
 			WriteLinePrivate(1, false, LogEventType.Error, LogFormatOptions.NoSeverityPrefix, "{0}({1}): error: {2}", File, Line, String.Format(Format, Args));
@@ -402,6 +412,7 @@ namespace EpicGames.Core
 		/// <param name="Args">Optional arguments</param>
 		[Conditional("TRACE")]
 		[MethodImplAttribute(MethodImplOptions.NoInlining)]
+		[StringFormatMethod("Format")]
 		public static void TraceVerbose(string Format, params object[] Args)
 		{
 			WriteLinePrivate(1, false, LogEventType.Verbose, LogFormatOptions.None, Format, Args);
@@ -413,6 +424,7 @@ namespace EpicGames.Core
 		/// <param name="Format">Message format string</param>
 		/// <param name="Args">Optional arguments</param>
 		[MethodImplAttribute(MethodImplOptions.NoInlining)]
+		[StringFormatMethod("Format")]
 		public static void TraceInformation(string Format, params object[] Args)
 		{
 			WriteLinePrivate(1, false, LogEventType.Console, LogFormatOptions.None, Format, Args);
@@ -424,6 +436,7 @@ namespace EpicGames.Core
 		/// <param name="Format">Message format string</param>
 		/// <param name="Args">Optional arguments</param>
 		[MethodImplAttribute(MethodImplOptions.NoInlining)]
+		[StringFormatMethod("Format")]
 		public static void TraceWarning(string Format, params object[] Args)
 		{
 			WriteLinePrivate(1, false, LogEventType.Warning, LogFormatOptions.None, Format, Args);
@@ -436,6 +449,7 @@ namespace EpicGames.Core
 		/// <param name="Format">Message format string</param>
 		/// <param name="Args">Optional arguments</param>
 		[MethodImplAttribute(MethodImplOptions.NoInlining)]
+		[StringFormatMethod("Format")]
 		public static void TraceWarning(FileReference File, string Format, params object[] Args)
 		{
 			WriteLinePrivate(1, false, LogEventType.Warning, LogFormatOptions.NoSeverityPrefix, "{0}: warning: {1}", File, String.Format(Format, Args));
@@ -449,6 +463,7 @@ namespace EpicGames.Core
 		/// <param name="Format">Message format string</param>
 		/// <param name="Args">Optional arguments</param>
 		[MethodImplAttribute(MethodImplOptions.NoInlining)]
+		[StringFormatMethod("Format")]
 		public static void TraceWarning(FileReference File, int Line, string Format, params object[] Args)
 		{
 			WriteLinePrivate(1, false, LogEventType.Warning, LogFormatOptions.NoSeverityPrefix, "{0}({1}): warning: {2}", File, Line, String.Format(Format, Args));
@@ -461,6 +476,7 @@ namespace EpicGames.Core
 		/// <param name="Args">Optional arguments</param>
 		[Conditional("TRACE")]
 		[MethodImplAttribute(MethodImplOptions.NoInlining)]
+		[StringFormatMethod("Format")]
 		public static void TraceVeryVerbose(string Format, params object[] Args)
 		{
 			WriteLinePrivate(1, false, LogEventType.VeryVerbose, LogFormatOptions.None, Format, Args);
@@ -473,6 +489,7 @@ namespace EpicGames.Core
 		/// <param name="Args">Optional arguments</param>
 		[Conditional("TRACE")]
 		[MethodImplAttribute(MethodImplOptions.NoInlining)]
+		[StringFormatMethod("Format")]
 		public static void TraceLog(string Format, params object[] Args)
 		{
 			WriteLinePrivate(1, false, LogEventType.Log, LogFormatOptions.None, Format, Args);
@@ -485,6 +502,7 @@ namespace EpicGames.Core
 		/// <param name="Format"></param>
 		/// <param name="Args"></param>
 		[MethodImplAttribute(MethodImplOptions.NoInlining)]
+		[StringFormatMethod("Format")]
 		public static void WriteLineOnce(LogEventType Verbosity, string Format, params object[] Args)
 		{
 			WriteLinePrivate(1, true, Verbosity, LogFormatOptions.None, Format, Args);
@@ -498,6 +516,7 @@ namespace EpicGames.Core
 		/// <param name="Format"></param>
 		/// <param name="Args"></param>
 		[MethodImplAttribute(MethodImplOptions.NoInlining)]
+		[StringFormatMethod("Format")]
 		public static void WriteLineOnce(LogEventType Verbosity, LogFormatOptions Options, string Format, params object[] Args)
 		{
 			WriteLinePrivate(1, true, Verbosity, Options, Format, Args);
@@ -509,6 +528,7 @@ namespace EpicGames.Core
 		/// <param name="Format">Message format string</param>
 		/// <param name="Args">Optional arguments</param>
 		[MethodImplAttribute(MethodImplOptions.NoInlining)]
+		[StringFormatMethod("Format")]
 		public static void TraceErrorOnce(string Format, params object[] Args)
 		{
 			WriteLinePrivate(1, true, LogEventType.Error, LogFormatOptions.None, Format, Args);
@@ -521,6 +541,7 @@ namespace EpicGames.Core
 		/// <param name="Args">Optional arguments</param>
 		[Conditional("TRACE")]
 		[MethodImplAttribute(MethodImplOptions.NoInlining)]
+		[StringFormatMethod("Format")]
 		public static void TraceVerboseOnce(string Format, params object[] Args)
 		{
 			WriteLinePrivate(1, true, LogEventType.Verbose, LogFormatOptions.None, Format, Args);
@@ -532,6 +553,7 @@ namespace EpicGames.Core
 		/// <param name="Format">Message format string</param>
 		/// <param name="Args">Optional arguments</param>
 		[MethodImplAttribute(MethodImplOptions.NoInlining)]
+		[StringFormatMethod("Format")]
 		public static void TraceInformationOnce(string Format, params object[] Args)
 		{
 			WriteLinePrivate(1, true, LogEventType.Console, LogFormatOptions.None, Format, Args);
@@ -543,6 +565,7 @@ namespace EpicGames.Core
 		/// <param name="Format">Message format string</param>
 		/// <param name="Args">Optional arguments</param>
 		[MethodImplAttribute(MethodImplOptions.NoInlining)]
+		[StringFormatMethod("Format")]
 		public static void TraceWarningOnce(string Format, params object[] Args)
 		{
 			WriteLinePrivate(1, true, LogEventType.Warning, LogFormatOptions.None, Format, Args);
@@ -555,6 +578,7 @@ namespace EpicGames.Core
 		/// <param name="Format">Message format string</param>
 		/// <param name="Args">Optional arguments</param>
 		[MethodImplAttribute(MethodImplOptions.NoInlining)]
+		[StringFormatMethod("Format")]
 		public static void TraceWarningOnce(FileReference File, string Format, params object[] Args)
 		{
 			WriteLinePrivate(1, true, LogEventType.Warning, LogFormatOptions.NoSeverityPrefix, "{0}: warning: {1}", File, String.Format(Format, Args));
@@ -568,6 +592,7 @@ namespace EpicGames.Core
 		/// <param name="Format">Message format string</param>
 		/// <param name="Args">Optional arguments</param>
 		[MethodImplAttribute(MethodImplOptions.NoInlining)]
+		[StringFormatMethod("Format")]
 		public static void TraceWarningOnce(FileReference File, int Line, string Format, params object[] Args)
 		{
 			WriteLinePrivate(1, true, LogEventType.Warning, LogFormatOptions.NoSeverityPrefix, "{0}({1}): warning: {2}", File, Line, String.Format(Format, Args));
@@ -580,6 +605,7 @@ namespace EpicGames.Core
 		/// <param name="Args">Optional arguments</param>
 		[Conditional("TRACE")]
 		[MethodImplAttribute(MethodImplOptions.NoInlining)]
+		[StringFormatMethod("Format")]
 		public static void TraceVeryVerboseOnce(string Format, params object[] Args)
 		{
 			WriteLinePrivate(1, true, LogEventType.VeryVerbose, LogFormatOptions.None, Format, Args);
@@ -592,6 +618,7 @@ namespace EpicGames.Core
 		/// <param name="Args">Optional arguments</param>
 		[Conditional("TRACE")]
 		[MethodImplAttribute(MethodImplOptions.NoInlining)]
+		[StringFormatMethod("Format")]
 		public static void TraceLogOnce(string Format, params object[] Args)
 		{
 			WriteLinePrivate(1, true, LogEventType.Log, LogFormatOptions.None, Format, Args);
@@ -921,6 +948,7 @@ namespace EpicGames.Core
 		/// <param name="Args">Message text parameters</param>
 		/// <returns>Formatted message</returns>
 		[MethodImplAttribute(MethodImplOptions.NoInlining)]
+		[StringFormatMethod("Format")]
 		private List<string> FormatMessage(int StackFramesToSkip, LogEventType Verbosity, LogFormatOptions Options, bool bForConsole, string Format, params object[] Args)
 		{
 			string TimePrefix = (!bForConsole || IncludeTimestamps) ? String.Format("[{0:hh\\:mm\\:ss\\.fff}] ", Timer.Elapsed) : "";
@@ -992,6 +1020,7 @@ namespace EpicGames.Core
 		/// <param name="Format">Message format string.</param>
 		/// <param name="Args">Optional arguments</param>
 		[MethodImplAttribute(MethodImplOptions.NoInlining)]
+		[StringFormatMethod("Format")]
 		private void WriteLinePrivate(int StackFramesToSkip, bool bWriteOnce, LogEventType Verbosity, LogFormatOptions FormatOptions, string Format, params object[] Args)
 		{
 			// if we want this message only written one time, check if it was already written out
