Horde: During a conform, clean all workspaces (from any root directory) before trying to sync anything. Would previously handle each root directory separately.

diff --git a/Engine/Source/Programs/Horde/HordeAgent/Execution/PerforceExecutor.cs b/Engine/Source/Programs/Horde/HordeAgent/Execution/PerforceExecutor.cs
--- a/Engine/Source/Programs/Horde/HordeAgent/Execution/PerforceExecutor.cs
+++ b/Engine/Source/Programs/Horde/HordeAgent/Execution/PerforceExecutor.cs
@@ -332,7 +332,8 @@ namespace HordeAgent.Execution
 				}
 			}
 
-			// Sync all the branches
+			// Sync all the branches.
+			List<Func<Task>> SyncFuncs = new List<Func<Task>>();
 			foreach (IGrouping<DirectoryReference, WorkspaceInfo> WorkspaceGroup in Workspaces.GroupBy(x => x.MetadataDir).OrderBy(x => x.Key.FullName))
 			{
 				List<PopulateRequest> PopulateRequests = new List<PopulateRequest>();
@@ -341,7 +342,14 @@ namespace HordeAgent.Execution
 					PerforceClientConnection PerforceClient = new PerforceClientConnection(Workspace.PerforceClient, Workspace.ClientName);
 					PopulateRequests.Add(new PopulateRequest(PerforceClient, Workspace.StreamName, Workspace.View));
 				}
-				await WorkspaceGroup.First().Repository.PopulateAsync(PopulateRequests, false, CancellationToken);
+
+				ManagedWorkspace Repository = WorkspaceGroup.First().Repository;
+				Tuple<int, StreamSnapshot>[] StreamStates = await Repository.PopulateCleanAsync(PopulateRequests, false, CancellationToken);
+				SyncFuncs.Add(() => Repository.PopulateSyncAsync(PopulateRequests, StreamStates, false, CancellationToken));
+			}
+			foreach (Func<Task> SyncFunc in SyncFuncs)
+			{
+				await SyncFunc();
 			}
 		}
 	}
diff --git a/Engine/Source/Programs/Shared/EpicGames.Perforce.Managed/ManagedWorkspace.cs b/Engine/Source/Programs/Shared/EpicGames.Perforce.Managed/ManagedWorkspace.cs
--- a/Engine/Source/Programs/Shared/EpicGames.Perforce.Managed/ManagedWorkspace.cs
+++ b/Engine/Source/Programs/Shared/EpicGames.Perforce.Managed/ManagedWorkspace.cs
@@ -993,119 +993,136 @@ namespace EpicGames.Perforce.Managed
 		public async Task PopulateAsync(List<PopulateRequest> Requests, bool bFakeSync, CancellationToken CancellationToken)
 		{
 			Logger.LogInformation("Populating with {NumStreams} streams", Requests.Count);
-			using(Logger.BeginIndentScope("  "))
+			using (Logger.BeginIndentScope("  "))
 			{
-				// Revert all changes in each of the unique clients
-				foreach (PopulateRequest Request in Requests)
-				{
-					PerforceConnection Perforce = Request.PerforceClient.WithoutClient();
+				Tuple<int, StreamSnapshot>[] StreamState = await PopulateCleanAsync(Requests, bFakeSync, CancellationToken);
+				await PopulateSyncAsync(Requests, StreamState, bFakeSync, CancellationToken);
+			}
+		}
 
-					PerforceResponse<ClientRecord> Response = await Perforce.TryGetClientAsync(Request.PerforceClient.ClientName, CancellationToken);
-					if (Response.Succeeded)
-					{
-						await RevertInternalAsync(Request.PerforceClient, CancellationToken);
-					}
+		/// <summary>
+		/// Perform the clean part of a populate command
+		/// </summary>
+		public async Task<Tuple<int, StreamSnapshot>[]> PopulateCleanAsync(List<PopulateRequest> Requests, bool bFakeSync, CancellationToken CancellationToken)
+		{
+			// Revert all changes in each of the unique clients
+			foreach (PopulateRequest Request in Requests)
+			{
+				PerforceConnection Perforce = Request.PerforceClient.WithoutClient();
+
+				PerforceResponse<ClientRecord> Response = await Perforce.TryGetClientAsync(Request.PerforceClient.ClientName, CancellationToken);
+				if (Response.Succeeded)
+				{
+					await RevertInternalAsync(Request.PerforceClient, CancellationToken);
 				}
+			}
 
-				// Clean the current workspace
-				await CleanAsync(true, CancellationToken);
+			// Clean the current workspace
+			await CleanAsync(true, CancellationToken);
 
-				// Update the list of files in each stream
-				Tuple<int, StreamSnapshot>[] StreamState = new Tuple<int, StreamSnapshot>[Requests.Count];
-				for(int Idx = 0; Idx < Requests.Count; Idx++)
-				{
-					PopulateRequest Request = Requests[Idx];
-					string StreamName = Request.StreamName;
-					Logger.LogInformation("Finding contents of {StreamName}:", StreamName);
+			// Update the list of files in each stream
+			Tuple<int, StreamSnapshot>[] StreamState = new Tuple<int, StreamSnapshot>[Requests.Count];
+			for (int Idx = 0; Idx < Requests.Count; Idx++)
+			{
+				PopulateRequest Request = Requests[Idx];
+				string StreamName = Request.StreamName;
+				Logger.LogInformation("Finding contents of {StreamName}:", StreamName);
 
-					using(Logger.BeginIndentScope("  "))
-					{
-						await DeleteClientAsync(Request.PerforceClient, CancellationToken);
-						await UpdateClientAsync(Request.PerforceClient, StreamName, CancellationToken);
+				using (Logger.BeginIndentScope("  "))
+				{
+					await DeleteClientAsync(Request.PerforceClient, CancellationToken);
+					await UpdateClientAsync(Request.PerforceClient, StreamName, CancellationToken);
 
-						int ChangeNumber = await GetLatestClientChangeAsync(Request.PerforceClient, CancellationToken);
-						Logger.LogInformation("Latest change is CL {CL}", ChangeNumber);
+					int ChangeNumber = await GetLatestClientChangeAsync(Request.PerforceClient, CancellationToken);
+					Logger.LogInformation("Latest change is CL {CL}", ChangeNumber);
 
-						await UpdateClientHaveTableAsync(Request.PerforceClient, ChangeNumber, Request.View, CancellationToken);
+					await UpdateClientHaveTableAsync(Request.PerforceClient, ChangeNumber, Request.View, CancellationToken);
 
-						StreamSnapshot Contents = await FindClientContentsAsync(Request.PerforceClient, ChangeNumber, bFakeSync, CancellationToken);
-						StreamState[Idx] = Tuple.Create(ChangeNumber, Contents);
+					StreamSnapshot Contents = await FindClientContentsAsync(Request.PerforceClient, ChangeNumber, bFakeSync, CancellationToken);
+					StreamState[Idx] = Tuple.Create(ChangeNumber, Contents);
 
-						GC.Collect();
-					}
+					GC.Collect();
 				}
+			}
 
-				// Remove any files from the workspace not referenced by the first stream. This ensures we can purge things from the cache that we no longer need.
-				if(Requests.Count > 0)
-				{
-					await RemoveFilesFromWorkspaceAsync(StreamState[0].Item2, CancellationToken);
-				}
+			// Remove any files from the workspace not referenced by the first stream. This ensures we can purge things from the cache that we no longer need.
+			if (Requests.Count > 0)
+			{
+				await RemoveFilesFromWorkspaceAsync(StreamState[0].Item2, CancellationToken);
+			}
 
-				// Shrink the contents of the cache
-				using (Trace("UpdateCache"))
-				using (ILoggerProgress Status = Logger.BeginProgressScope("Updating cache..."))
-				{
-					Stopwatch Timer = Stopwatch.StartNew();
+			// Shrink the contents of the cache
+			using (Trace("UpdateCache"))
+			using (ILoggerProgress Status = Logger.BeginProgressScope("Updating cache..."))
+			{
+				Stopwatch Timer = Stopwatch.StartNew();
 
-					HashSet<FileContentId> CommonContentIds = new HashSet<FileContentId>();
-					Dictionary<FileContentId, long> ContentIdToLength = new Dictionary<FileContentId, long>();
-					for(int Idx = 0; Idx < Requests.Count; Idx++)
+				HashSet<FileContentId> CommonContentIds = new HashSet<FileContentId>();
+				Dictionary<FileContentId, long> ContentIdToLength = new Dictionary<FileContentId, long>();
+				for (int Idx = 0; Idx < Requests.Count; Idx++)
+				{
+					List<StreamFile> Files = StreamState[Idx].Item2.GetFiles();
+					foreach (StreamFile File in Files)
 					{
-						List<StreamFile> Files = StreamState[Idx].Item2.GetFiles();
-						foreach(StreamFile File in Files)
-						{
-							ContentIdToLength[File.ContentId] = File.Length;
-						}
+						ContentIdToLength[File.ContentId] = File.Length;
+					}
 
-						if(Idx == 0)
-						{
-							CommonContentIds.UnionWith(Files.Select(x => x.ContentId));
-						}
-						else
-						{
-							CommonContentIds.IntersectWith(Files.Select(x => x.ContentId));
-						}
+					if (Idx == 0)
+					{
+						CommonContentIds.UnionWith(Files.Select(x => x.ContentId));
+					}
+					else
+					{
+						CommonContentIds.IntersectWith(Files.Select(x => x.ContentId));
 					}
+				}
 
-					List<CachedFileInfo> TrackedFiles = ContentIdToTrackedFile.Values.ToList();
-					foreach(CachedFileInfo TrackedFile in TrackedFiles)
+				List<CachedFileInfo> TrackedFiles = ContentIdToTrackedFile.Values.ToList();
+				foreach (CachedFileInfo TrackedFile in TrackedFiles)
+				{
+					if (!ContentIdToLength.ContainsKey(TrackedFile.ContentId))
 					{
-						if(!ContentIdToLength.ContainsKey(TrackedFile.ContentId))
-						{
-							RemoveTrackedFile(TrackedFile);
-						}
+						RemoveTrackedFile(TrackedFile);
 					}
+				}
 
-					GC.Collect();
+				GC.Collect();
 
-					double TotalSize = ContentIdToLength.Sum(x => x.Value) / (1024.0 * 1024.0);
-					Status.Progress = String.Format("{0:n1}mb total, {1:n1}mb differences ({2:0.0}s)", TotalSize, TotalSize - CommonContentIds.Sum(x => ContentIdToLength[x]) / (1024.0 * 1024.0), Timer.Elapsed.TotalSeconds);
-				}
+				double TotalSize = ContentIdToLength.Sum(x => x.Value) / (1024.0 * 1024.0);
+				Status.Progress = String.Format("{0:n1}mb total, {1:n1}mb differences ({2:0.0}s)", TotalSize, TotalSize - CommonContentIds.Sum(x => ContentIdToLength[x]) / (1024.0 * 1024.0), Timer.Elapsed.TotalSeconds);
+			}
 
-				// Sync all the new files
-				for(int Idx = 0; Idx < Requests.Count; Idx++)
-				{
-					PopulateRequest Request = Requests[Idx];
-					string StreamName = Request.StreamName;
-					Logger.LogInformation("Syncing files for {StreamName}:", StreamName);
+			return StreamState;
+		}
 
-					using(Logger.BeginIndentScope("  "))
-					{
-						await DeleteClientAsync(Request.PerforceClient, CancellationToken);
-						await UpdateClientAsync(Request.PerforceClient, StreamName, CancellationToken);
+		/// <summary>
+		/// Perform the sync part of a populate command
+		/// </summary>
+		public async Task PopulateSyncAsync(List<PopulateRequest> Requests, Tuple<int, StreamSnapshot>[] StreamState, bool bFakeSync, CancellationToken CancellationToken)
+		{
+			// Sync all the new files
+			for (int Idx = 0; Idx < Requests.Count; Idx++)
+			{
+				PopulateRequest Request = Requests[Idx];
+				string StreamName = Request.StreamName;
+				Logger.LogInformation("Syncing files for {StreamName}:", StreamName);
 
-						int ChangeNumber = StreamState[Idx].Item1;
-						await UpdateClientHaveTableAsync(Request.PerforceClient, ChangeNumber, Requests[Idx].View, CancellationToken);
+				using(Logger.BeginIndentScope("  "))
+				{
+					await DeleteClientAsync(Request.PerforceClient, CancellationToken);
+					await UpdateClientAsync(Request.PerforceClient, StreamName, CancellationToken);
 
-						StreamSnapshot Contents = StreamState[Idx].Item2;
-						await RemoveFilesFromWorkspaceAsync(Contents, CancellationToken);
-						await AddFilesToWorkspaceAsync(Request.PerforceClient, Contents, bFakeSync, CancellationToken);
-					}
-				}
+					int ChangeNumber = StreamState[Idx].Item1;
+					await UpdateClientHaveTableAsync(Request.PerforceClient, ChangeNumber, Requests[Idx].View, CancellationToken);
 
-				// Save the new repo state
-				await SaveAsync(TransactionState.Clean, CancellationToken);
+					StreamSnapshot Contents = StreamState[Idx].Item2;
+					await RemoveFilesFromWorkspaceAsync(Contents, CancellationToken);
+					await AddFilesToWorkspaceAsync(Request.PerforceClient, Contents, bFakeSync, CancellationToken);
+				}
 			}
+
+			// Save the new repo state
+			await SaveAsync(TransactionState.Clean, CancellationToken);
 		}
 
 		#endregion
