EpicGames.Core:

Add a stub for high precision scoped profiler.

diff --git a/Engine/Source/Programs/Shared/EpicGames.Core/EpicGames.Core.csproj b/Engine/Source/Programs/Shared/EpicGames.Core/EpicGames.Core.csproj
--- a/Engine/Source/Programs/Shared/EpicGames.Core/EpicGames.Core.csproj
+++ b/Engine/Source/Programs/Shared/EpicGames.Core/EpicGames.Core.csproj
@@ -9,6 +9,7 @@
     <GenerateTargetFrameworkAttribute>false</GenerateTargetFrameworkAttribute>
     <Configurations>Debug;Release;Development</Configurations>
     <Nullable>enable</Nullable>
+    <AllowUnsafeBlocks>true</AllowUnsafeBlocks>
   </PropertyGroup>
   <PropertyGroup Condition=" '$(Configuration)|$(Platform)' == 'Debug|AnyCPU' ">
     <DocumentationFile></DocumentationFile>
@@ -26,6 +27,18 @@
       <Link>Properties\MetaData.cs</Link>
     </Compile>
   </ItemGroup>
+  <ItemGroup>
+    <Compile Include="..\..\..\..\Restricted\*\Source\Programs\Shared\EpicGames.Core\ScopeProfiler\*.cs">
+      <Link>$([System.Text.RegularExpressions.Regex]::Replace(%(FullPath), '^.+?[\\/]EpicGames.Core[\\/]', ''))</Link>
+    </Compile>
+  </ItemGroup>
+  <Choose>
+    <When Condition="Exists('..\..\..\..\Restricted\NoRedist\Source\Programs\Shared\EpicGames.Core\ScopeProfiler\')">
+      <PropertyGroup>
+        <DefineConstants>$(DefineConstants);__SCOPEPROFILER_AVAILABLE__</DefineConstants>
+      </PropertyGroup>
+    </When>
+  </Choose>
   <ItemGroup>
     <PackageReference Include="JetBrains.Annotations" Version="2021.2.0" />
     <PackageReference Include="Microsoft.Extensions.Logging" Version="5.0.0" />
diff --git a/Engine/Source/Programs/Shared/EpicGames.Core/ScopeProfiler.cs b/Engine/Source/Programs/Shared/EpicGames.Core/ScopeProfiler.cs
new file mode 100644
--- /dev/null
+++ b/Engine/Source/Programs/Shared/EpicGames.Core/ScopeProfiler.cs
@@ -0,0 +1,73 @@
+// Copyright Epic Games, Inc. All Rights Reserved.
+
+using System;
+using System.Reflection;
+
+namespace EpicGames.Core
+{
+#if !__SCOPEPROFILER_AVAILABLE__
+	
+	/// <summary>
+	/// A stub/no-op scope-profiler API that can be replaced by another implementation to record execution of
+	/// instrumented scopes.
+	/// </summary>
+	public class ScopeProfiler 
+	{
+		public static ScopeProfiler Instance = new ScopeProfiler();
+		
+		public void InitializeAndStart(string ProgramName, string HostAddress, int MaxThreadCount) { }
+
+		public void StopAndShutdown() { }
+
+		/// <summary>
+		/// Start a stacked scope. Must start and end on the same thread
+		/// </summary>
+		/// <param name="Name"></param>
+		/// <param name="bIsStall"></param>
+		public void StartScope(string Name, bool bIsStall) { }
+		
+		/// <summary>
+		/// End a stacked scope. Must start and end on the same thread.
+		/// </summary>
+		public void EndScope() { }
+
+		/// <summary>
+		/// Record an un-stacked time span at a specified time.
+		/// </summary>
+		/// <param name="Name"></param>
+		/// <param name="Id"></param>
+		/// <param name="bStall"></param>
+		/// <param name="BeginThreadId"></param>
+		/// <param name="EndThreadId"></param>
+		/// <param name="StartTime">start time for the span (use FastTime())</param>
+		/// <param name="EndTime">end time for the span (use FastTime())</param>
+		public void AddSpanAtTime(string Name, ulong Id, bool bStall, uint BeginThreadId, uint EndThreadId, ulong StartTime, ulong EndTime) { }
+
+		/// <summary>
+		/// Record a value against a particular name, to be plotted over time
+		/// </summary>
+		/// <param name="Name"></param>
+		/// <param name="Value"></param>
+		public void Plot(string Name, double Value) { }
+
+		/// <summary>
+		/// Generate a timestamp value for the current moment.
+		/// </summary>
+		/// <returns></returns>
+		public ulong FastTime()
+		{
+			return 0;
+		}
+
+		/// <summary>
+		/// Retrieve the profiler's identifier for the current thread
+		/// </summary>
+		/// <returns></returns>
+		public uint GetCurrentThreadId()
+		{
+			return 0;
+		}
+	}
+#endif
+}
+
