Horde: Fix sub-optimal encoding of workspace snapshots, when the path to an item is the same as the path to its parent folder plus the item name.

diff --git a/Engine/Source/Programs/Shared/EpicGames.Perforce.Managed/StreamSnapshotFromMemory.cs b/Engine/Source/Programs/Shared/EpicGames.Perforce.Managed/StreamSnapshotFromMemory.cs
--- a/Engine/Source/Programs/Shared/EpicGames.Perforce.Managed/StreamSnapshotFromMemory.cs
+++ b/Engine/Source/Programs/Shared/EpicGames.Perforce.Managed/StreamSnapshotFromMemory.cs
@@ -77,7 +77,8 @@ namespace EpicGames.Perforce.Managed
 			CbObject RootObj = new CbObject(Data.AsMemory(CurrentSignature.Length));
 
 			CbObject RootObj2 = RootObj["root"].AsObject();
-			StreamTreeRef Root = new StreamTreeRef(RootObj2, BasePath);
+			Utf8String RootPath = RootObj2["path"].AsString(BasePath);
+			StreamTreeRef Root = new StreamTreeRef(RootPath, RootObj2);
 
 			CbArray Array = RootObj["items"].AsArray();
 
@@ -103,7 +104,11 @@ namespace EpicGames.Perforce.Managed
 			Writer.BeginObject();
 						
 			Writer.BeginObject("root");
-			Root.Write(Writer, BasePath);
+			if (Root.Path != BasePath)
+			{
+				Writer.WriteString("path", Root.Path);
+			}
+			Root.Write(Writer);
 			Writer.EndObject();
 			
 			Writer.BeginArray("items");
diff --git a/Engine/Source/Programs/Shared/EpicGames.Perforce.Managed/StreamTree.cs b/Engine/Source/Programs/Shared/EpicGames.Perforce.Managed/StreamTree.cs
--- a/Engine/Source/Programs/Shared/EpicGames.Perforce.Managed/StreamTree.cs
+++ b/Engine/Source/Programs/Shared/EpicGames.Perforce.Managed/StreamTree.cs
@@ -40,7 +40,6 @@ namespace EpicGames.Perforce.Managed
 		public int Revision { get; }
 
 		#region Field names
-		static Utf8String PathField = "path";
 		static Utf8String LengthField = "len";
 		static Utf8String DigestField = "dig";
 		static Utf8String TypeField = "type";
@@ -61,12 +60,12 @@ namespace EpicGames.Perforce.Managed
 		/// <summary>
 		/// Parse from a compact binary object
 		/// </summary>
+		/// <param name="Path">Path to the file</param>
 		/// <param name="Field"></param>
-		/// <param name="DefaultPath">Default path for this file</param>
 		/// <returns></returns>
-		public StreamFile(CbObject Field, Utf8String DefaultPath)
+		public StreamFile(Utf8String Path, CbObject Field)
 		{
-			Path = Field[PathField].AsString(DefaultPath);
+			this.Path = Path;
 			Length = Field[LengthField].AsInt64();
 
 			Md5Hash Digest = new Md5Hash(Field[DigestField].AsBinary());
@@ -80,13 +79,8 @@ namespace EpicGames.Perforce.Managed
 		/// Write this object to compact binary
 		/// </summary>
 		/// <param name="Writer"></param>
-		/// <param name="DefaultPath"></param>
-		public void Write(CbWriter Writer, Utf8String DefaultPath)
+		public void Write(CbWriter Writer)
 		{
-			if(Path != DefaultPath)
-			{
-				Writer.WriteString(PathField, Path);
-			}
 			Writer.WriteInteger(LengthField, Length);
 			Writer.WriteBinary(DigestField, ContentId.Digest.Span);
 			Writer.WriteString(TypeField, ContentId.Type);
@@ -110,7 +104,6 @@ namespace EpicGames.Perforce.Managed
 		public IoHash Hash { get; set; }
 
 		#region Field names
-		static Utf8String PathField = "path";
 		static Utf8String HashField = "hash";
 		#endregion
 
@@ -128,11 +121,11 @@ namespace EpicGames.Perforce.Managed
 		/// <summary>
 		/// Constructor
 		/// </summary>
+		/// <param name="Path"></param>
 		/// <param name="Field"></param>
-		/// <param name="DefaultPath"></param>
-		public StreamTreeRef(CbObject Field, Utf8String DefaultPath)
+		public StreamTreeRef(Utf8String Path, CbObject Field)
 		{
-			Path = Field[PathField].AsString(DefaultPath);
+			this.Path = Path;
 			Hash = Field[HashField].AsObjectAttachment();
 		}
 
@@ -144,7 +137,7 @@ namespace EpicGames.Perforce.Managed
 		{
 			CbWriter Writer = new CbWriter();
 			Writer.BeginObject();
-			Write(Writer, Utf8String.Empty);
+			Write(Writer);
 			Writer.EndObject();
 			return Writer.ToObject().GetHash();
 		}
@@ -153,13 +146,8 @@ namespace EpicGames.Perforce.Managed
 		/// Serialize to a compact binary object
 		/// </summary>
 		/// <param name="Writer"></param>
-		/// <param name="DefaultPath"></param>
-		public void Write(CbWriter Writer, Utf8String DefaultPath)
+		public void Write(CbWriter Writer)
 		{
-			if (Path != DefaultPath)
-			{
-				Writer.WriteString(PathField, Path);
-			}
 			Writer.WriteObjectAttachment(HashField, Hash);
 		}
 	}
@@ -181,6 +169,7 @@ namespace EpicGames.Perforce.Managed
 
 		#region Field names
 		static Utf8String NameField = "name";
+		static Utf8String PathField = "path";
 		static Utf8String FilesField = "files";
 		static Utf8String TreesField = "trees";
 		#endregion
@@ -214,7 +203,8 @@ namespace EpicGames.Perforce.Managed
 				CbObject FileObject = FileField.AsObject();
 
 				Utf8String Name = FileObject[NameField].AsString();
-				StreamFile File = new StreamFile(FileObject, GetDefaultPath(BasePath, Name));
+				Utf8String Path = ReadPath(FileObject, BasePath, Name);
+				StreamFile File = new StreamFile(Path, FileObject);
 
 				NameToFile.Add(Name, File);
 			}
@@ -225,7 +215,8 @@ namespace EpicGames.Perforce.Managed
 				CbObject TreeObject = TreeField.AsObject();
 
 				Utf8String Name = TreeObject[NameField].AsString();
-				StreamTreeRef Tree = new StreamTreeRef(TreeObject, GetDefaultPath(BasePath, Name));
+				Utf8String Path = ReadPath(TreeObject, BasePath, Name);
+				StreamTreeRef Tree = new StreamTreeRef(Path, TreeObject);
 
 				NameToTree.Add(Name, Tree);
 			}
@@ -287,7 +278,7 @@ namespace EpicGames.Perforce.Managed
 		/// <param name="BasePath"></param>
 		public void Write(CbWriter Writer, Utf8String BasePath)
 		{
-			if (BasePath.EndsWith("/"))
+			if (BasePath.Length > 0 && BasePath[BasePath.Length - 1] == '/')
 			{
 				throw new ArgumentException("BasePath must not end in a slash", nameof(BasePath));
 			}
@@ -299,7 +290,8 @@ namespace EpicGames.Perforce.Managed
 				{
 					Writer.BeginObject();
 					Writer.WriteString(NameField, Name);
-					File.Write(Writer, BasePath);
+					WritePath(Writer, File.Path, BasePath, Name);
+					File.Write(Writer);
 					Writer.EndObject();
 				}
 				Writer.EndArray();
@@ -312,13 +304,50 @@ namespace EpicGames.Perforce.Managed
 				{
 					Writer.BeginObject();
 					Writer.WriteString(NameField, Name);
-					Tree.Write(Writer, BasePath);
+					WritePath(Writer, Tree.Path, BasePath, Name);
+					Tree.Write(Writer);
 					Writer.EndObject();
 				}
 				Writer.EndArray();
 			}
 		}
 
+		/// <summary>
+		/// Reads a path from an object, defaulting it to the parent path plus the child name
+		/// </summary>
+		/// <param name="Object"></param>
+		/// <param name="BasePath"></param>
+		/// <param name="Name"></param>
+		/// <returns></returns>
+		static Utf8String ReadPath(CbObject Object, Utf8String BasePath, Utf8String Name)
+		{
+			Utf8String Path = Object[PathField].AsString();
+			if (Path.IsEmpty)
+			{
+				byte[] Data = new byte[BasePath.Length + 1 + Name.Length];
+				BasePath.Memory.CopyTo(Data);
+				Data[BasePath.Length] = (byte)'/';
+				Name.Memory.CopyTo(Data.AsMemory(BasePath.Length + 1));
+				Path = new Utf8String(Data);
+			}
+			return Path;
+		}
+
+		/// <summary>
+		/// Writes a path if it's not the default (the parent path, a slash, followed by the child name)
+		/// </summary>
+		/// <param name="Writer"></param>
+		/// <param name="Path"></param>
+		/// <param name="ParentPath"></param>
+		/// <param name="Name"></param>
+		static void WritePath(CbWriter Writer, Utf8String Path, Utf8String ParentPath, Utf8String Name)
+		{
+			if (Path.Length != ParentPath.Length + Name.Length + 1 || !Path.StartsWith(ParentPath) || Path[ParentPath.Length] != '/' || !Path.EndsWith(Name))
+			{
+				Writer.WriteString(PathField, Path);
+			}
+		}
+
 		/// <summary>
 		/// Serialize to a compact binary object
 		/// </summary>
