Horde: Add utf8 comparer to P4 server info.

diff --git a/Engine/Source/Programs/Shared/EpicGames.Perforce/Records/InfoRecord.cs b/Engine/Source/Programs/Shared/EpicGames.Perforce/Records/InfoRecord.cs
--- a/Engine/Source/Programs/Shared/EpicGames.Perforce/Records/InfoRecord.cs
+++ b/Engine/Source/Programs/Shared/EpicGames.Perforce/Records/InfoRecord.cs
@@ -1,5 +1,6 @@
 // Copyright Epic Games, Inc. All Rights Reserved.
 
+using EpicGames.Core;
 using System;
 using System.Collections.Generic;
 using System.Linq;
@@ -71,6 +72,11 @@ namespace EpicGames.Perforce
 		/// </summary>
 		public StringComparison PathComparison => String.Equals(CaseHandling, "insensitive", StringComparison.OrdinalIgnoreCase) ? StringComparison.OrdinalIgnoreCase : StringComparison.Ordinal;
 
+		/// <summary>
+		/// How to compare paths on this server
+		/// </summary>
+		public Utf8StringComparer Utf8PathComparer => String.Equals(CaseHandling, "insensitive", StringComparison.OrdinalIgnoreCase) ? Utf8StringComparer.OrdinalIgnoreCase : Utf8StringComparer.Ordinal;
+
 		/// <summary>
 		/// List of services provided by this server
 		/// </summary>
