Horde: Missing file.

diff --git a/Engine/Source/Programs/Shared/EpicGames.Redis/RedisSet.cs b/Engine/Source/Programs/Shared/EpicGames.Redis/RedisSet.cs
--- a/Engine/Source/Programs/Shared/EpicGames.Redis/RedisSet.cs
+++ b/Engine/Source/Programs/Shared/EpicGames.Redis/RedisSet.cs
@@ -31,13 +31,13 @@ namespace EpicGames.Redis
 		}
 
 		/// <inheritdoc cref="IDatabaseAsync.SetAddAsync(RedisKey, RedisValue, CommandFlags)"/>
-		public Task AddAsync(TElement Item, CommandFlags Flags = CommandFlags.None)
+		public Task<bool> AddAsync(TElement Item, CommandFlags Flags = CommandFlags.None)
 		{
 			return Database.SetAddAsync(Key, RedisSerializer.Serialize(Item), Flags);
 		}
 
 		/// <inheritdoc cref="IDatabaseAsync.SetAddAsync(RedisKey, RedisValue[], CommandFlags)"/>
-		public Task AddAsync(IEnumerable<TElement> Items, CommandFlags Flags = CommandFlags.None)
+		public Task<long> AddAsync(IEnumerable<TElement> Items, CommandFlags Flags = CommandFlags.None)
 		{
 			RedisValue[] Values = Items.Select(x => RedisSerializer.Serialize(x)).ToArray();
 			return Database.SetAddAsync(Key, Values, Flags);
@@ -78,13 +78,13 @@ namespace EpicGames.Redis
 		}
 
 		/// <inheritdoc cref="IDatabaseAsync.SetRemoveAsync(RedisKey, RedisValue, CommandFlags)"/>
-		public Task RemoveAsync(TElement Item, CommandFlags Flags = CommandFlags.None)
+		public Task<bool> RemoveAsync(TElement Item, CommandFlags Flags = CommandFlags.None)
 		{
 			return Database.SetRemoveAsync(Key, RedisSerializer.Serialize(Item), Flags);
 		}
 
 		/// <inheritdoc cref="IDatabaseAsync.SetRemoveAsync(RedisKey, RedisValue[], CommandFlags)"/>
-		public Task RemoveAsync(IEnumerable<TElement> Items, CommandFlags Flags = CommandFlags.None)
+		public Task<long> RemoveAsync(IEnumerable<TElement> Items, CommandFlags Flags = CommandFlags.None)
 		{
 			RedisValue[] Values = Items.Select(x => RedisSerializer.Serialize(x)).ToArray();
 			return Database.SetRemoveAsync(Key, Values, Flags);
