Horde: Add convenience methods in CompactBinarySerializer to read/write from byte spans.

diff --git a/Engine/Source/Programs/Shared/EpicGames.Serialization/CompactBinarySerializer.cs b/Engine/Source/Programs/Shared/EpicGames.Serialization/CompactBinarySerializer.cs
--- a/Engine/Source/Programs/Shared/EpicGames.Serialization/CompactBinarySerializer.cs
+++ b/Engine/Source/Programs/Shared/EpicGames.Serialization/CompactBinarySerializer.cs
@@ -236,7 +236,7 @@ namespace EpicGames.Serialization
 		}
 
 		/// <summary>
-		/// Deserialize an object from a <see cref="CbObject"/>
+		/// Deserialize an object from a <see cref="CbField"/>
 		/// </summary>
 		/// <typeparam name="T"></typeparam>
 		/// <param name="Field"></param>
@@ -245,5 +245,21 @@ namespace EpicGames.Serialization
 		{
 			return GetConverter<T>().Read(Field);
 		}
+
+		/// <summary>
+		/// Deserialize an object from a <see cref="CbObject"/>
+		/// </summary>
+		/// <typeparam name="T"></typeparam>
+		/// <param name="Object"></param>
+		/// <returns></returns>
+		public static T Deserialize<T>(CbObject Object) => Deserialize<T>(Object.AsField());
+
+		/// <summary>
+		/// Deserialize an object from a block of memory
+		/// </summary>
+		/// <typeparam name="T"></typeparam>
+		/// <param name="Data"></param>
+		/// <returns></returns>
+		public static T Deserialize<T>(ReadOnlyMemory<byte> Data) => Deserialize<T>(new CbField(Data));
 	}
 }
