Horde: Allow hyphens in property names, since AWS tags use them a lot.

diff --git a/Engine/Source/Programs/Shared/EpicGames.Horde/Common/Condition.cs b/Engine/Source/Programs/Shared/EpicGames.Horde/Common/Condition.cs
--- a/Engine/Source/Programs/Shared/EpicGames.Horde/Common/Condition.cs
+++ b/Engine/Source/Programs/Shared/EpicGames.Horde/Common/Condition.cs
@@ -265,7 +265,7 @@ namespace EpicGames.Horde.Common
 
 		static bool IsIdentifierTail(char Character)
 		{
-			return IsIdentifier(Character) || IsNumber(Character) || Character == '.';
+			return IsIdentifier(Character) || IsNumber(Character) || Character == '-' || Character == '.';
 		}
 	}
 
