Horde: Add an AuditLog class, which can be used to trace events corresponding to a particular object. Events are standard message templates.

diff --git a/Engine/Source/Programs/Horde/HordeAgent/Parser/LogEventBuilder.cs b/Engine/Source/Programs/Horde/HordeAgent/Parser/LogEventBuilder.cs
--- a/Engine/Source/Programs/Horde/HordeAgent/Parser/LogEventBuilder.cs
+++ b/Engine/Source/Programs/Horde/HordeAgent/Parser/LogEventBuilder.cs
@@ -1,5 +1,6 @@
 // Copyright Epic Games, Inc. All Rights Reserved.
 
+using EpicGames.Core;
 using HordeAgent.Parser.Interfaces;
 using HordeAgent.Utility;
 using Microsoft.Extensions.Logging;
diff --git a/Engine/Source/Programs/Horde/HordeAgent/Utility/JsonLogger.cs b/Engine/Source/Programs/Horde/HordeAgent/Utility/JsonLogger.cs
--- a/Engine/Source/Programs/Horde/HordeAgent/Utility/JsonLogger.cs
+++ b/Engine/Source/Programs/Horde/HordeAgent/Utility/JsonLogger.cs
@@ -1,5 +1,6 @@
 // Copyright Epic Games, Inc. All Rights Reserved.
 
+using EpicGames.Core;
 using HordeAgent.Parser;
 using HordeCommon;
 using Microsoft.Extensions.Logging;
diff --git a/Engine/Source/Programs/Horde/HordeAgent/Utility/MessageTemplate.cs b/Engine/Source/Programs/Horde/HordeAgent/Utility/MessageTemplate.cs
deleted file mode 100644
diff --git a/Engine/Source/Programs/Horde/HordeServer/Collections/IAuditLog.cs b/Engine/Source/Programs/Horde/HordeServer/Collections/IAuditLog.cs
new file mode 100644
--- /dev/null
+++ b/Engine/Source/Programs/Horde/HordeServer/Collections/IAuditLog.cs
@@ -0,0 +1,168 @@
+// Copyright Epic Games, Inc. All Rights Reserved.
+
+using Google.Protobuf.WellKnownTypes;
+using HordeServer.Api;
+using HordeCommon;
+using HordeServer.Models;
+using HordeCommon.Rpc.Tasks;
+using HordeServer.Services;
+using HordeServer.Utilities;
+using MongoDB.Bson;
+using MongoDB.Bson.Serialization.Attributes;
+using MongoDB.Driver;
+using System;
+using System.Collections.Generic;
+using System.Linq;
+using System.Linq.Expressions;
+using System.Threading.Tasks;
+using Microsoft.Extensions.Logging;
+using EpicGames.Core;
+using System.Buffers;
+using System.Text.Json;
+using System.Text;
+
+namespace HordeServer.Collections
+{
+	/// <summary>
+	/// Message in an audit log
+	/// </summary>
+	/// <typeparam name="TTargetId">Type of entity that the log is for</typeparam>
+	public interface IAuditLogMessage<TTargetId>
+	{
+		/// <summary>
+		/// Unique id for the entity
+		/// </summary>
+		public TTargetId TargetId { get; }
+
+		/// <summary>
+		/// Timestamp for the event
+		/// </summary>
+		public DateTime TimeUtc { get; }
+
+		/// <summary>
+		/// Severity of the message
+		/// </summary>
+		public LogLevel Level { get; }
+
+		/// <summary>
+		/// The message payload. Should be an encoded JSON object, with format/properties fields.
+		/// </summary>
+		public string Data { get; }
+	}
+
+	/// <summary>
+	/// Concrete implementation of <see cref="IAuditLogMessage{T}"/> used to add new entries
+	/// </summary>
+	/// <typeparam name="TTargetId"></typeparam>
+	public class NewAuditLogMessage<TTargetId> : IAuditLogMessage<TTargetId>
+	{
+		/// <inheritdoc/>
+		public TTargetId TargetId { get; set; }
+
+		/// <inheritdoc/>
+		public DateTime TimeUtc { get; set; }
+
+		/// <inheritdoc/>
+		public LogLevel Level { get; }
+
+		/// <inheritdoc/>
+		public string Data { get; set; }
+
+		/// <summary>
+		/// Constructor
+		/// </summary>
+		public NewAuditLogMessage(TTargetId EntityId, DateTime TimeUtc, LogLevel Level, string Data)
+		{
+			this.TargetId = EntityId;
+			this.TimeUtc = TimeUtc;
+			this.Level = Level;
+			this.Data = Data;
+		}
+	}
+
+	/// <summary>
+	/// Interface for a collection of log messages for a particular document type
+	/// </summary>
+	public interface IAuditLog<TTargetId>
+	{
+		/// <summary>
+		/// Adds a message with the given properties
+		/// </summary>
+		/// <param name="NewMessage">The message to add</param>
+		Task AddAsync(NewAuditLogMessage<TTargetId> NewMessage);
+
+		/// <summary>
+		/// Finds messages matching certain criteria
+		/// </summary>
+		/// <param name="Id">Identifier for the </param>
+		/// <param name="MinTime"></param>
+		/// <param name="MaxTime"></param>
+		/// <param name="Index"></param>
+		/// <param name="Count"></param>
+		/// <returns></returns>
+		IAsyncEnumerable<IAuditLogMessage<TTargetId>> FindAsync(TTargetId Id, DateTime? MinTime = null, DateTime? MaxTime = null, int? Index = null, int? Count = null);
+
+		/// <summary>
+		/// Deletes messages between a given time range for a particular object
+		/// </summary>
+		/// <param name="Id">Id of the object</param>
+		/// <param name="MinTime">Minimum time to remove</param>
+		/// <param name="MaxTime">Maximum time to remove</param>
+		/// <returns>Async task</returns>
+		Task<long> DeleteAsync(TTargetId Id, DateTime? MinTime = null, DateTime? MaxTime = null);
+	}
+
+	/// <summary>
+	/// Factory for instantiating audit log instances
+	/// </summary>
+	/// <typeparam name="TTargetId"></typeparam>
+	public interface IAuditLogFactory<TTargetId>
+	{
+		/// <summary>
+		/// Create a new audit log instance, with the given database name
+		/// </summary>
+		/// <param name="Name"></param>
+		/// <returns></returns>
+		IAuditLog<TTargetId> Create(string Name);
+	}
+
+	/// <summary>
+	/// Extension methods for AuditLog instances
+	/// </summary>
+	public static class AuditLogExtensions
+	{
+		/// <summary>
+		/// Write a message to an audit log
+		/// </summary>
+		public static Task WriteAsync<TTargetId>(this IAuditLog<TTargetId> AuditLog, TTargetId Id, LogLevel Level, string Format, params object[] Args)
+		{
+			return WriteAsync(AuditLog, Id, DateTime.UtcNow, Level, new EventId(0), Format, Args);
+		}
+
+		/// <summary>
+		/// Write a message to an audit log
+		/// </summary>
+		public static Task WriteAsync<TTargetId>(this IAuditLog<TTargetId> AuditLog, TTargetId Id, LogLevel Level, EventId EventId, string Format, params object[] Args)
+		{
+			return WriteAsync(AuditLog, Id, DateTime.UtcNow, Level, EventId, Format, Args);
+		}
+
+		/// <summary>
+		/// Write a message to an audit log
+		/// </summary>
+		public static Task WriteAsync<TTargetId>(this IAuditLog<TTargetId> AuditLog, TTargetId Id, DateTime TimeUtc, LogLevel Level, EventId EventId, string Format, params object[] Args)
+		{
+			ArrayBufferWriter<byte> Buffer = new ArrayBufferWriter<byte>();
+			using (Utf8JsonWriter Writer = new Utf8JsonWriter(Buffer))
+			{
+				Dictionary<string, object> Properties = new Dictionary<string, object>();
+				MessageTemplate.ParsePropertyValues(Format, Args, Properties);
+				string Text = MessageTemplate.Render(Format, Properties);
+				MessageTemplate.SerializeJson(Writer, TimeUtc, Level, EventId, Format, Properties);
+			}
+
+			NewAuditLogMessage<TTargetId> NewMessage = new NewAuditLogMessage<TTargetId>(Id, TimeUtc, Level, Encoding.UTF8.GetString(Buffer.WrittenSpan));
+			return AuditLog.AddAsync(NewMessage);
+		}
+	}
+}
diff --git a/Engine/Source/Programs/Horde/HordeServer/Collections/Impl/AuditLog.cs b/Engine/Source/Programs/Horde/HordeServer/Collections/Impl/AuditLog.cs
new file mode 100644
--- /dev/null
+++ b/Engine/Source/Programs/Horde/HordeServer/Collections/Impl/AuditLog.cs
@@ -0,0 +1,139 @@
+// Copyright Epic Games, Inc. All Rights Reserved.
+
+using Google.Protobuf.WellKnownTypes;
+using HordeCommon;
+using HordeCommon.Rpc.Tasks;
+using HordeServer.Models;
+using HordeServer.Services;
+using HordeServer.Utilities;
+using Microsoft.Extensions.Logging;
+using MongoDB.Bson;
+using MongoDB.Bson.Serialization.Attributes;
+using MongoDB.Driver;
+using System;
+using System.Collections.Generic;
+using System.Linq;
+using System.Linq.Expressions;
+using System.Threading.Tasks;
+
+namespace HordeServer.Collections.Impl
+{
+	/// <summary>
+	/// Collection of agent documents
+	/// </summary>
+	class AuditLog<TTargetId> : IAuditLog<TTargetId>
+	{
+		class AuditLogMessage : IAuditLogMessage<TTargetId>
+		{
+			public ObjectId Id { get; set; }
+
+			[BsonElement("tid")]
+			public TTargetId TargetId { get; set; }
+
+			[BsonElement("t")]
+			public DateTime TimeUtc { get; set; }
+
+			[BsonElement("l")]
+			public LogLevel Level { get; set; }
+
+			[BsonElement("d")]
+			public string Data { get; set; }
+
+			public AuditLogMessage()
+			{
+				TargetId = default!;
+				Data = String.Empty;
+			}
+
+			public AuditLogMessage(IAuditLogMessage<TTargetId> Other)
+			{
+				this.Id = ObjectId.GenerateNewId();
+				this.TargetId = Other.TargetId;
+				this.TimeUtc = Other.TimeUtc;
+				this.Level = Other.Level;
+				this.Data = Other.Data;
+			}
+		}
+
+		IMongoCollection<AuditLogMessage> Messages;
+
+		public AuditLog(DatabaseService DatabaseService, string CollectionName)
+		{
+			Messages = DatabaseService.GetCollection<AuditLogMessage>(CollectionName);
+			if (!DatabaseService.ReadOnlyMode)
+			{
+				Messages.Indexes.CreateOne(new CreateIndexModel<AuditLogMessage>(Builders<AuditLogMessage>.IndexKeys.Ascending(x => x.Id).Descending(x => x.TimeUtc)));
+			}
+		}
+
+		public Task AddAsync(NewAuditLogMessage<TTargetId> NewMessage)
+		{
+			AuditLogMessage Message = new AuditLogMessage(NewMessage);
+			return Messages.InsertOneAsync(Message);
+		}
+
+		public async IAsyncEnumerable<IAuditLogMessage<TTargetId>> FindAsync(TTargetId Id, DateTime? MinTime = null, DateTime? MaxTime = null, int? Index = null, int? Count = null)
+		{
+			FilterDefinition<AuditLogMessage> Filter = Builders<AuditLogMessage>.Filter.Eq(x => x.TargetId, Id);
+			if (MinTime != null)
+			{
+				Filter &= Builders<AuditLogMessage>.Filter.Gte(x => x.TimeUtc, MinTime.Value);
+			}
+			if (MaxTime != null)
+			{
+				Filter &= Builders<AuditLogMessage>.Filter.Lte(x => x.TimeUtc, MaxTime.Value);
+			}
+
+			using (IAsyncCursor<AuditLogMessage> Cursor = await Messages.Find(Filter).SortByDescending(x => x.TimeUtc).ToCursorAsync())
+			{
+				while (await Cursor.MoveNextAsync())
+				{
+					foreach (AuditLogMessage Message in Cursor.Current)
+					{
+						yield return Message;
+					}
+				}
+			}
+		}
+
+		public async Task<long> DeleteAsync(TTargetId Id, DateTime? MinTime = null, DateTime? MaxTime = null)
+		{
+			FilterDefinition<AuditLogMessage> Filter = Builders<AuditLogMessage>.Filter.Eq(x => x.TargetId, Id);
+			if (MinTime != null)
+			{
+				Filter &= Builders<AuditLogMessage>.Filter.Gte(x => x.TimeUtc, MinTime.Value);
+			}
+			if (MaxTime != null)
+			{
+				Filter &= Builders<AuditLogMessage>.Filter.Lte(x => x.TimeUtc, MaxTime.Value);
+			}
+
+			DeleteResult Result = await Messages.DeleteManyAsync(Filter);
+			return Result.DeletedCount;
+		}
+	}
+
+	/// <summary>
+	/// Factory for instantiating audit log instances
+	/// </summary>
+	/// <typeparam name="TTargetId"></typeparam>
+	public class AuditLogFactory<TTargetId> : IAuditLogFactory<TTargetId>
+	{
+		DatabaseService DatabaseService;
+
+		/// <summary>
+		/// Constructor
+		/// </summary>
+		/// <param name="DatabaseService"></param>
+		public AuditLogFactory(DatabaseService DatabaseService)
+		{
+			this.DatabaseService = DatabaseService;
+		}
+
+		/// <inheritdoc/>
+		public IAuditLog<TTargetId> Create(string Name)
+		{
+			return new AuditLog<TTargetId>(DatabaseService, Name);
+		}
+	}
+}
diff --git a/Engine/Source/Programs/Horde/HordeServer/Startup.cs b/Engine/Source/Programs/Horde/HordeServer/Startup.cs
--- a/Engine/Source/Programs/Horde/HordeServer/Startup.cs
+++ b/Engine/Source/Programs/Horde/HordeServer/Startup.cs
@@ -347,6 +347,7 @@ namespace HordeServer
 			Services.AddSingleton<IObjectCollection, ObjectCollection>();
 			Services.AddSingleton<IRefCollection, RefCollection>();
 
+			Services.AddSingleton(typeof(IAuditLogFactory<>), typeof(AuditLogFactory<>));
 			Services.AddSingleton(typeof(ISingletonDocument<>), typeof(SingletonDocument<>));
 
 			Services.AddSingleton<AutoscaleService>();
diff --git a/Engine/Source/Programs/Shared/EpicGames.Core/MessageTemplate.cs b/Engine/Source/Programs/Shared/EpicGames.Core/MessageTemplate.cs
new file mode 100644
--- /dev/null
+++ b/Engine/Source/Programs/Shared/EpicGames.Core/MessageTemplate.cs
@@ -0,0 +1,469 @@
+// Copyright Epic Games, Inc. All Rights Reserved.
+
+using Microsoft.Extensions.Logging;
+using System;
+using System.Buffers;
+using System.Collections.Concurrent;
+using System.Collections.Generic;
+using System.Globalization;
+using System.Linq;
+using System.Reflection;
+using System.Text;
+using System.Text.Json;
+
+namespace EpicGames.Core
+{
+	/// <summary>
+	/// Interface for argument formatters in message templates
+	/// </summary>
+	public interface IMessageTemplateFormatter
+	{
+		/// <summary>
+		/// Serializes the given value. The state of the supplied JSON writer will be within a JSON object; named properties should typically be serialized directly here.
+		/// </summary>
+		public void Write(Utf8JsonWriter Writer, object Value);
+	}
+
+	/// <summary>
+	/// Attribute indicating the formatter to use for a message template argument
+	/// </summary>
+	public class MessageTemplateFormatterAttribute : Attribute
+	{
+		public string Name { get; }
+		public Type Type { get; }
+
+		public MessageTemplateFormatterAttribute(string Name, Type Type)
+		{
+			this.Name = Name;
+			this.Type = Type;
+		}
+	}
+
+	/// <summary>
+	/// Utility class for dealing with message templates
+	/// </summary>
+	public static class MessageTemplate
+	{
+		/// <summary>
+		/// The default property name for the message template format string in an enumerable log state parameter
+		/// </summary>
+		public const string FormatPropertyName = "{OriginalFormat}";
+
+		/// <summary>
+		/// Static cache of formatters for property values
+		/// </summary>
+		static ConcurrentDictionary<Type, IMessageTemplateFormatter> TypeToWriter = new ConcurrentDictionary<Type, IMessageTemplateFormatter>();
+
+		/// <summary>
+		/// Register a formatter for the given type
+		/// </summary>
+		/// <param name="Type"></param>
+		/// <param name="Formatter"></param>
+		public static void RegisterFormatter(Type Type, IMessageTemplateFormatter Formatter)
+		{
+			TypeToWriter.TryAdd(Type, Formatter);
+		}
+
+		/// <summary>
+		/// Renders a format string
+		/// </summary>
+		/// <param name="Format">The format string</param>
+		/// <param name="Properties">Property values to embed</param>
+		/// <returns>The rendered string</returns>
+		public static string Render(string Format, IEnumerable<KeyValuePair<string, object>>? Properties)
+		{
+			StringBuilder Result = new StringBuilder();
+			Render(Format, Properties, Result);
+			return Result.ToString();
+		}
+
+		/// <summary>
+		/// Renders a format string to the end of a string builder
+		/// </summary>
+		/// <param name="Format">The format string to render</param>
+		/// <param name="Properties">Sequence of key/value properties</param>
+		/// <param name="Result">Buffer to append the rendered string to</param>
+		public static void Render(string Format, IEnumerable<KeyValuePair<string, object>>? Properties, StringBuilder Result)
+		{
+			int NextOffset = 0;
+
+			List<(int, int)>? Names = ParsePropertyNames(Format);
+			if (Names != null)
+			{
+				foreach((int Offset, int Length) in Names)
+				{
+					object? Value;
+					if (Properties != null && TryGetPropertyValue(Format.AsSpan(Offset, Length), Properties, out Value))
+					{
+						int StartOffset = Offset - 1;
+						if (Format[StartOffset] == '@' || Format[StartOffset] == '$')
+						{
+							StartOffset--;
+						}
+
+						Unescape(Format.AsSpan(NextOffset, StartOffset - NextOffset), Result);
+						Result.Append(Value?.ToString() ?? "null");
+						NextOffset = Offset + Length + 1;
+					}
+				}
+			}
+
+			Unescape(Format.AsSpan(NextOffset, Format.Length - NextOffset), Result);
+		}
+
+		/// <summary>
+		/// Serialize a message template to JOSN
+		/// </summary>
+		public static void SerializeJson(Utf8JsonWriter Writer, DateTime Time, LogLevel LogLevel, EventId EventId, string Format, IEnumerable<KeyValuePair<string, object>>? Properties, Exception? Exception = null)
+		{
+			string Text = Render(Format, Properties);
+			SerializeJson(Writer, Time, LogLevel, Text, EventId, 0, 1, Format, Properties, Exception);
+		}
+
+		/// <summary>
+		/// Serialize a message template to JSON
+		/// </summary>
+		public static void SerializeJson(Utf8JsonWriter Writer, DateTime Time, LogLevel LogLevel, string Text, EventId EventId, int LineIndex, int LineCount, string? Format, IEnumerable<KeyValuePair<string, object>>? Properties, Exception? Exception = null)
+		{
+			Writer.WriteStartObject();
+			Writer.WriteString("time", Time.ToString("s", CultureInfo.InvariantCulture));
+			Writer.WriteString("level", LogLevel.ToString());
+			Writer.WriteString("message", Text);
+
+			if (EventId.Id != 0)
+			{
+				Writer.WriteNumber("id", EventId.Id);
+			}
+
+			if (Format != null)
+			{
+				Writer.WriteString("format", Format);
+			}
+
+			if (LineCount > 1)
+			{
+				Writer.WriteNumber("line", LineIndex);
+				Writer.WriteNumber("lineCount", LineCount);
+			}
+
+			if (Properties != null && Properties.Any())
+			{
+				Writer.WriteStartObject("properties");
+				foreach (KeyValuePair<string, object> Pair in Properties!)
+				{
+					Writer.WritePropertyName(Pair.Key);
+					WritePropertyValue(Pair.Value, Writer);
+				}
+				Writer.WriteEndObject();
+			}
+
+			if (Exception != null)
+			{
+				Writer.WriteStartObject("exception");
+				WriteException(Writer, Exception);
+				Writer.WriteEndObject();
+			}
+			Writer.WriteEndObject();
+		}
+
+		/// <summary>
+		/// Writes an exception to a json object
+		/// </summary>
+		/// <param name="Writer">Writer to receive the exception data</param>
+		/// <param name="Exception">The exception</param>
+		static void WriteException(Utf8JsonWriter Writer, Exception Exception)
+		{
+			Writer.WriteString("message", Exception.Message);
+			Writer.WriteString("trace", Exception.StackTrace);
+
+			if (Exception.InnerException != null)
+			{
+				Writer.WriteStartObject("innerException");
+				WriteException(Writer, Exception.InnerException);
+				Writer.WriteEndObject();
+			}
+
+			AggregateException? AggregateException = Exception as AggregateException;
+			if (AggregateException != null)
+			{
+				Writer.WriteStartArray("innerExceptions");
+				for (int Idx = 0; Idx < 16 && Idx < AggregateException.InnerExceptions.Count; Idx++) // Cap number of exceptions returned to avoid huge messages
+				{
+					Exception InnerException = AggregateException.InnerExceptions[Idx];
+					Writer.WriteStartObject();
+					WriteException(Writer, InnerException);
+					Writer.WriteEndObject();
+				}
+				Writer.WriteEndArray();
+			}
+		}
+
+		/// <summary>
+		/// Converts a property value to log markup
+		/// </summary>
+		/// <param name="Value">The value to write</param>
+		/// <param name="Writer">The json writer</param>
+		static void WritePropertyValue(object Value, Utf8JsonWriter Writer)
+		{
+			if (Value == null)
+			{
+				Writer.WriteNullValue();
+			}
+			else
+			{
+				Type ValueType = Value.GetType();
+				if (ValueType.IsEnum)
+				{
+					Writer.WriteStringValue(Enum.GetName(ValueType, Value));
+				}
+				else
+				{
+					switch (Type.GetTypeCode(ValueType))
+					{
+						case TypeCode.Boolean:
+							Writer.WriteBooleanValue((bool)Value);
+							break;
+						case TypeCode.Byte:
+							Writer.WriteNumberValue((Byte)Value);
+							break;
+						case TypeCode.SByte:
+							Writer.WriteNumberValue((SByte)Value);
+							break;
+						case TypeCode.UInt16:
+							Writer.WriteNumberValue((UInt16)Value);
+							break;
+						case TypeCode.UInt32:
+							Writer.WriteNumberValue((UInt32)Value);
+							break;
+						case TypeCode.UInt64:
+							Writer.WriteNumberValue((UInt64)Value);
+							break;
+						case TypeCode.Int16:
+							Writer.WriteNumberValue((Int16)Value);
+							break;
+						case TypeCode.Int32:
+							Writer.WriteNumberValue((Int32)Value);
+							break;
+						case TypeCode.Int64:
+							Writer.WriteNumberValue((Int64)Value);
+							break;
+						case TypeCode.Decimal:
+							Writer.WriteNumberValue((Decimal)Value);
+							break;
+						case TypeCode.Double:
+							Writer.WriteNumberValue((Double)Value);
+							break;
+						case TypeCode.Single:
+							Writer.WriteNumberValue((Single)Value);
+							break;
+						case TypeCode.String:
+							Writer.WriteStringValue((String)Value);
+							break;
+						default:
+							WriteComplexPropertyValue(Value, ValueType, Writer);
+							break;
+					}
+				}
+			}
+		}
+
+		/// <summary>
+		/// Writes a complex object to the given json writer
+		/// </summary>
+		/// <param name="Value">The value to write</param>
+		/// <param name="ValueType">Type of the value</param>
+		/// <param name="JsonWriter">The writer to output to</param>
+		static void WriteComplexPropertyValue(object Value, Type ValueType, Utf8JsonWriter JsonWriter)
+		{
+			MessageTemplateFormatterAttribute? Attribute = ValueType.GetCustomAttribute<MessageTemplateFormatterAttribute>();
+			if (Attribute != null)
+			{
+				IMessageTemplateFormatter? Formatter;
+				if (!TypeToWriter.TryGetValue(Attribute.Type, out Formatter))
+				{
+					Formatter = (IMessageTemplateFormatter)Activator.CreateInstance(Attribute.Type)!;
+					TypeToWriter.TryAdd(Attribute.Type, Formatter);
+				}
+
+				JsonWriter.WriteStartObject();
+				JsonWriter.WriteString("$type", Attribute.Name);
+				JsonWriter.WriteString("$text", Value.ToString());
+				Formatter.Write(JsonWriter, Value);
+				JsonWriter.WriteEndObject();
+			}
+			else
+			{
+				JsonWriter.WriteStringValue(Value.ToString());
+			}
+		}
+
+		/// <summary>
+		/// Escapes a string for use in a message template
+		/// </summary>
+		/// <param name="Text">Text to escape</param>
+		/// <returns>The escaped string</returns>
+		public static string Escape(string Text)
+		{
+			StringBuilder Result = new StringBuilder();
+			Escape(Text, Result);
+			return Result.ToString();
+		}
+
+		/// <summary>
+		/// Escapes a span of characters and appends the result to a string
+		/// </summary>
+		/// <param name="Text">Span of characters to escape</param>
+		/// <param name="Result">Buffer to receive the escaped string</param>
+		public static void Escape(ReadOnlySpan<char> Text, StringBuilder Result)
+		{
+			foreach(char Char in Text)
+			{
+				Result.Append(Char);
+				if (Char == '{' || Char == '}')
+				{
+					Result.Append(Char);
+				}
+			}
+		}
+
+		/// <summary>
+		/// Unescapes a string from a message template
+		/// </summary>
+		/// <param name="Text">The text to unescape</param>
+		/// <returns>The unescaped text</returns>
+		public static string Unescape(string Text)
+		{
+			StringBuilder Result = new StringBuilder();
+			Unescape(Text.AsSpan(), Result);
+			return Result.ToString();
+		}
+
+		/// <summary>
+		/// Unescape a string and append the result to a string builder
+		/// </summary>
+		/// <param name="Text">Text to unescape</param>
+		/// <param name="Result">Receives the unescaped text</param>
+		public static void Unescape(ReadOnlySpan<char> Text, StringBuilder Result)
+		{
+			char LastChar = '\0';
+			foreach (char Char in Text)
+			{
+				if ((Char != '{' || Char != '}') || Char != LastChar)
+				{
+					Result.Append(Char);
+				}
+				LastChar = Char;
+			}
+		}
+
+		/// <summary>
+		/// Finds locations of property names from the given format string
+		/// </summary>
+		/// <param name="Format">The format string to parse</param>
+		/// <returns>List of offset, length pairs for property names. Null if the string does not contain any property references.</returns>
+		public static List<(int, int)>? ParsePropertyNames(string Format)
+		{
+			List<(int, int)>? Names = null;
+			for (int Idx = 0; Idx < Format.Length - 1; Idx++)
+			{
+				if (Format[Idx] == '{')
+				{
+					if (Format[Idx + 1] == '{')
+					{
+						Idx++;
+					}
+					else
+					{
+						int StartIdx = Idx + 1;
+
+						Idx = Format.IndexOf('}', StartIdx);
+						if (Idx == -1)
+						{
+							break;
+						}
+						if (Names == null)
+						{
+							Names = new List<(int, int)>();
+						}
+
+						Names.Add((StartIdx, Idx - StartIdx));
+					}
+				}
+			}
+			return Names;
+		}
+
+		/// <summary>
+		/// Parse the ordered arguments into a dictionary of named properties
+		/// </summary>
+		/// <param name="Format">Format string</param>
+		/// <param name="Args">Argument list to parse</param>
+		/// <returns></returns>
+		public static void ParsePropertyValues(string Format, object[] Args, Dictionary<string, object> Properties)
+		{
+			List<(int, int)>? Offsets = ParsePropertyNames(Format);
+			if (Offsets != null)
+			{
+				for (int Idx = 0; Idx < Offsets.Count; Idx++)
+				{
+					string Name = Format.Substring(Offsets[Idx].Item1, Offsets[Idx].Item2);
+
+					int Number;
+					if (int.TryParse(Name, out Number))
+					{
+						if (Number >= 0 && Number < Args.Length)
+						{
+							Properties[Name] = Args[Number];
+						}
+					}
+					else
+					{
+						if (Idx < Args.Length)
+						{
+							Properties[Name] = Args[Idx];
+						}
+					}
+				}
+			}
+		}
+
+		/// <summary>
+		/// Attempts to get a named property value from the given dictionary
+		/// </summary>
+		/// <param name="Name">Name of the property</param>
+		/// <param name="Properties">Sequence of property name/value pairs</param>
+		/// <param name="Value">On success, receives the property value</param>
+		/// <returns>True if the property was found, false otherwise</returns>
+		public static bool TryGetPropertyValue(ReadOnlySpan<char> Name, IEnumerable<KeyValuePair<string, object>> Properties, out object? Value)
+		{
+			int Number;
+			if (int.TryParse(Name, System.Globalization.NumberStyles.Integer, null, out Number))
+			{
+				foreach (KeyValuePair<string, object> Property in Properties)
+				{
+					if (Number == 0)
+					{
+						Value = Property.Value;
+						return true;
+					}
+					Number--;
+				}
+			}
+			else
+			{
+				foreach (KeyValuePair<string, object> Property in Properties)
+				{
+					ReadOnlySpan<char> ParameterName = Property.Key.AsSpan();
+					if (Name.Equals(ParameterName, StringComparison.Ordinal))
+					{
+						Value = Property.Value;
+						return true;
+					}
+				}
+			}
+
+			Value = null;
+			return false;
+		}
+	}
+}
