Horde: Add support for directly serializing raw compact binary objects. Fixes the /api/v1/objects endpoint when given a compact binary content type.

diff --git a/Engine/Source/Programs/Horde/HordeServer/Storage/Controllers/ObjectsController.cs b/Engine/Source/Programs/Horde/HordeServer/Storage/Controllers/ObjectsController.cs
--- a/Engine/Source/Programs/Horde/HordeServer/Storage/Controllers/ObjectsController.cs
+++ b/Engine/Source/Programs/Horde/HordeServer/Storage/Controllers/ObjectsController.cs
@@ -128,14 +128,14 @@ namespace HordeServer.Storage.Controllers
 		[HttpPut]
 		[Route("/api/v1/objects/{NamespaceId}/{Hash}")]
 		[ProducesResponseType(typeof(PutObjectResponse), 200)]
-		public async Task<ActionResult> PutObjectAsync(NamespaceId NamespaceId, IoHash Hash, [FromBody] byte[] Data)
+		public async Task<ActionResult> PutObjectAsync(NamespaceId NamespaceId, IoHash Hash, [FromBody] CbObject Data)
 		{
 			if (!await NamespaceCollection.AuthorizeAsync(NamespaceId, User, AclAction.WriteBlobs))
 			{
 				return Forbid();
 			}
 
-			await ObjectCollection.AddAsync(NamespaceId, Hash, new CbObject(Data));
+			await ObjectCollection.AddAsync(NamespaceId, Hash, Data);
 			return Ok(new PutObjectResponse { Id = Hash });
 		}
 
diff --git a/Engine/Source/Programs/Shared/EpicGames.Serialization/CompactBinarySerializer.cs b/Engine/Source/Programs/Shared/EpicGames.Serialization/CompactBinarySerializer.cs
--- a/Engine/Source/Programs/Shared/EpicGames.Serialization/CompactBinarySerializer.cs
+++ b/Engine/Source/Programs/Shared/EpicGames.Serialization/CompactBinarySerializer.cs
@@ -71,7 +71,11 @@ namespace EpicGames.Serialization
 		/// <summary>
 		/// Cache of type to converter
 		/// </summary>
-		public static Dictionary<Type, CbConverter> TypeToConverter = new Dictionary<Type, CbConverter>();
+		public static Dictionary<Type, CbConverter> TypeToConverter = new Dictionary<Type, CbConverter>()
+		{
+			[typeof(CbField)] = new CbFieldConverter(),
+			[typeof(CbObject)] = new CbObjectConverter()
+		};
 
 		/// <summary>
 		/// List of converter factories. Must be 
diff --git a/Engine/Source/Programs/Shared/EpicGames.Serialization/CompactBinaryWriter.cs b/Engine/Source/Programs/Shared/EpicGames.Serialization/CompactBinaryWriter.cs
--- a/Engine/Source/Programs/Shared/EpicGames.Serialization/CompactBinaryWriter.cs
+++ b/Engine/Source/Programs/Shared/EpicGames.Serialization/CompactBinaryWriter.cs
@@ -253,7 +253,7 @@ namespace EpicGames.Serialization
 		/// Writes the header for an unnamed field
 		/// </summary>
 		/// <param name="Type"></param>
-		void WriteField(CbFieldType Type)
+		void WriteFieldHeader(CbFieldType Type)
 		{
 			Scope Scope = CurrentScope;
 			if (!CbFieldUtils.IsArray(Scope.FieldType))
@@ -277,7 +277,7 @@ namespace EpicGames.Serialization
 		/// </summary>
 		/// <param name="Type"></param>
 		/// <param name="Name"></param>
-		void WriteField(CbFieldType Type, Utf8String Name)
+		void WriteFieldHeader(CbFieldType Type, Utf8String Name)
 		{
 			Scope Scope = CurrentScope;
 			if (!CbFieldUtils.IsObject(Scope.FieldType))
@@ -303,12 +303,40 @@ namespace EpicGames.Serialization
 			Scope.Count++;
 		}
 
+		/// <summary>
+		/// Copies an entire field value to the output
+		/// </summary>
+		/// <param name="Field"></param>
+		public void WriteFieldValue(CbField Field)
+		{
+			WriteFieldHeader(Field.GetType());
+			int Size = (int)Field.GetPayloadSize();
+			Field.GetPayloadView().CopyTo(Allocate(Size));
+		}
+
+		/// <summary>
+		/// Copies an entire field value to the output, using the name from the field
+		/// </summary>
+		/// <param name="Field"></param>
+		public void WriteField(CbField Field) => WriteField(Field.GetName(), Field);
+
+		/// <summary>
+		/// Copies an entire field value to the output
+		/// </summary>
+		/// <param name="Field"></param>
+		public void WriteField(Utf8String Name, CbField Field)
+		{
+			WriteFieldHeader(Field.GetType(), Name);
+			int Size = (int)Field.GetPayloadSize();
+			Field.GetPayloadView().CopyTo(Allocate(Size));
+		}
+
 		/// <summary>
 		/// Begin writing an object field
 		/// </summary>
 		public void BeginObject()
 		{
-			WriteField(CbFieldType.Object);
+			WriteFieldHeader(CbFieldType.Object);
 			PushScope(CbFieldType.Object, CbFieldType.None);
 		}
 
@@ -318,7 +346,7 @@ namespace EpicGames.Serialization
 		/// <param name="Name">Name of the field</param>
 		public void BeginObject(Utf8String Name)
 		{
-			WriteField(CbFieldType.Object, Name);
+			WriteFieldHeader(CbFieldType.Object, Name);
 			PushScope(CbFieldType.Object, CbFieldType.None);
 		}
 
@@ -335,7 +363,7 @@ namespace EpicGames.Serialization
 		/// </summary>
 		public void BeginArray()
 		{
-			WriteField(CbFieldType.Array);
+			WriteFieldHeader(CbFieldType.Array);
 			PushScope(CbFieldType.Array, CbFieldType.None);
 		}
 
@@ -345,7 +373,7 @@ namespace EpicGames.Serialization
 		/// <param name="Name"></param>
 		public void BeginArray(Utf8String Name)
 		{
-			WriteField(CbFieldType.Array, Name);
+			WriteFieldHeader(CbFieldType.Array, Name);
 			PushScope(CbFieldType.Array, CbFieldType.None);
 		}
 
@@ -363,7 +391,7 @@ namespace EpicGames.Serialization
 		/// <param name="FieldType">The field type for elements in the array</param>
 		public void BeginUniformArray(CbFieldType FieldType)
 		{
-			WriteField(CbFieldType.UniformArray);
+			WriteFieldHeader(CbFieldType.UniformArray);
 			PushScope(CbFieldType.UniformArray, FieldType);
 			Allocate(1).Span[0] = (byte)FieldType;
 		}
@@ -375,7 +403,7 @@ namespace EpicGames.Serialization
 		/// <param name="FieldType">The field type for elements in the array</param>
 		public void BeginUniformArray(Utf8String Name, CbFieldType FieldType)
 		{
-			WriteField(CbFieldType.UniformArray, Name);
+			WriteFieldHeader(CbFieldType.UniformArray, Name);
 			PushScope(CbFieldType.UniformArray, FieldType);
 			Allocate(1).Span[0] = (byte)FieldType;
 		}
@@ -393,7 +421,7 @@ namespace EpicGames.Serialization
 		/// </summary>
 		public void WriteNullValue()
 		{
-			WriteField(CbFieldType.Null);
+			WriteFieldHeader(CbFieldType.Null);
 		}
 
 		/// <summary>
@@ -402,7 +430,7 @@ namespace EpicGames.Serialization
 		/// <param name="Name">Name of the field</param>
 		public void WriteNull(Utf8String Name)
 		{
-			WriteField(CbFieldType.Null, Name);
+			WriteFieldHeader(CbFieldType.Null, Name);
 		}
 
 		/// <summary>
@@ -411,7 +439,7 @@ namespace EpicGames.Serialization
 		/// <param name="Value"></param>
 		public void WriteBoolValue(bool Value)
 		{
-			WriteField(Value? CbFieldType.BoolTrue : CbFieldType.BoolFalse);
+			WriteFieldHeader(Value? CbFieldType.BoolTrue : CbFieldType.BoolFalse);
 		}
 
 		/// <summary>
@@ -421,7 +449,7 @@ namespace EpicGames.Serialization
 		/// <param name="Value"></param>
 		public void WriteBool(Utf8String Name, bool Value)
 		{
-			WriteField(Value ? CbFieldType.BoolTrue : CbFieldType.BoolFalse, Name);
+			WriteFieldHeader(Value ? CbFieldType.BoolTrue : CbFieldType.BoolFalse, Name);
 		}
 
 		/// <summary>
@@ -443,12 +471,12 @@ namespace EpicGames.Serialization
 		{
 			if (Value >= 0)
 			{
-				WriteField(CbFieldType.IntegerPositive);
+				WriteFieldHeader(CbFieldType.IntegerPositive);
 				WriteIntegerPayload((ulong)Value);
 			}
 			else
 			{
-				WriteField(CbFieldType.IntegerNegative);
+				WriteFieldHeader(CbFieldType.IntegerNegative);
 				WriteIntegerPayload((ulong)-Value);
 			}
 		}
@@ -462,12 +490,12 @@ namespace EpicGames.Serialization
 		{
 			if (Value >= 0)
 			{
-				WriteField(CbFieldType.IntegerPositive, Name);
+				WriteFieldHeader(CbFieldType.IntegerPositive, Name);
 				WriteIntegerPayload((ulong)Value);
 			}
 			else
 			{
-				WriteField(CbFieldType.IntegerNegative, Name);
+				WriteFieldHeader(CbFieldType.IntegerNegative, Name);
 				WriteIntegerPayload((ulong)-Value);
 			}
 		}
@@ -478,7 +506,7 @@ namespace EpicGames.Serialization
 		/// <param name="Value">Value to be written</param>
 		public void WriteIntegerValue(ulong Value)
 		{
-			WriteField(CbFieldType.IntegerPositive);
+			WriteFieldHeader(CbFieldType.IntegerPositive);
 			WriteIntegerPayload(Value);
 		}
 
@@ -489,7 +517,7 @@ namespace EpicGames.Serialization
 		/// <param name="Value">Value to be written</param>
 		public void WriteInteger(Utf8String Name, ulong Value)
 		{
-			WriteField(CbFieldType.IntegerPositive, Name);
+			WriteFieldHeader(CbFieldType.IntegerPositive, Name);
 			WriteIntegerPayload(Value);
 		}
 
@@ -509,7 +537,7 @@ namespace EpicGames.Serialization
 		/// <param name="Value">Value to be written</param>
 		public void WriteDateTimeValue(DateTime Value)
 		{
-			WriteField(CbFieldType.DateTime);
+			WriteFieldHeader(CbFieldType.DateTime);
 			WriteDateTimePayload(Value);
 		}
 
@@ -520,7 +548,7 @@ namespace EpicGames.Serialization
 		/// <param name="Value">Value to be written</param>
 		public void WriteDateTime(Utf8String Name, DateTime Value)
 		{
-			WriteField(CbFieldType.DateTime, Name);
+			WriteFieldHeader(CbFieldType.DateTime, Name);
 			WriteDateTimePayload(Value);
 		}
 
@@ -540,7 +568,7 @@ namespace EpicGames.Serialization
 		/// <param name="Hash"></param>
 		public void WriteHashValue(IoHash Hash)
 		{
-			WriteField(CbFieldType.Hash);
+			WriteFieldHeader(CbFieldType.Hash);
 			WriteHashPayload(Hash);
 		}
 
@@ -551,7 +579,7 @@ namespace EpicGames.Serialization
 		/// <param name="Value">Value to be written</param>
 		public void WriteHash(Utf8String Name, IoHash Value)
 		{
-			WriteField(CbFieldType.Hash, Name);
+			WriteFieldHeader(CbFieldType.Hash, Name);
 			WriteHashPayload(Value);
 		}
 
@@ -561,7 +589,7 @@ namespace EpicGames.Serialization
 		/// <param name="Hash">Hash of the attachment</param>
 		public void WriteBinaryAttachmentValue(IoHash Hash)
 		{
-			WriteField(CbFieldType.BinaryAttachment);
+			WriteFieldHeader(CbFieldType.BinaryAttachment);
 			WriteHashPayload(Hash);
 		}
 
@@ -572,7 +600,7 @@ namespace EpicGames.Serialization
 		/// <param name="Hash">Hash of the attachment</param>
 		public void WriteBinaryAttachment(Utf8String Name, IoHash Hash)
 		{
-			WriteField(CbFieldType.BinaryAttachment, Name);
+			WriteFieldHeader(CbFieldType.BinaryAttachment, Name);
 			WriteHashPayload(Hash);
 		}
 
@@ -593,7 +621,7 @@ namespace EpicGames.Serialization
 		/// <param name="Object">Object to write</param>
 		public void WriteObject(CbObject Object)
 		{
-			WriteField(CbFieldType.Object);
+			WriteFieldHeader(CbFieldType.Object);
 			WriteObjectPayload(Object);
 		}
 
@@ -604,7 +632,7 @@ namespace EpicGames.Serialization
 		/// <param name="Object">Object to write</param>
 		public void WriteObject(Utf8String Name, CbObject Object)
 		{
-			WriteField(CbFieldType.Object, Name);
+			WriteFieldHeader(CbFieldType.Object, Name);
 			WriteObjectPayload(Object);
 		}
 
@@ -614,7 +642,7 @@ namespace EpicGames.Serialization
 		/// <param name="Hash">Hash of the attachment</param>
 		public void WriteObjectAttachmentValue(IoHash Hash)
 		{
-			WriteField(CbFieldType.ObjectAttachment);
+			WriteFieldHeader(CbFieldType.ObjectAttachment);
 			WriteHashPayload(Hash);
 		}
 
@@ -625,7 +653,7 @@ namespace EpicGames.Serialization
 		/// <param name="Hash">Hash of the attachment</param>
 		public void WriteObjectAttachment(Utf8String Name, IoHash Hash)
 		{
-			WriteField(CbFieldType.ObjectAttachment, Name);
+			WriteFieldHeader(CbFieldType.ObjectAttachment, Name);
 			WriteHashPayload(Hash);
 		}
 
@@ -660,7 +688,7 @@ namespace EpicGames.Serialization
 		/// <param name="Value">Value to be written</param>
 		public void WriteStringValue(Utf8String Value)
 		{
-			WriteField(CbFieldType.String);
+			WriteFieldHeader(CbFieldType.String);
 			WriteBinaryPayload(Value.Span);
 		}
 
@@ -671,7 +699,7 @@ namespace EpicGames.Serialization
 		/// <param name="Value">Value to be written</param>
 		public void WriteString(Utf8String Name, Utf8String Value)
 		{
-			WriteField(CbFieldType.String, Name);
+			WriteFieldHeader(CbFieldType.String, Name);
 			WriteBinaryPayload(Value.Span);
 		}
 
@@ -681,7 +709,7 @@ namespace EpicGames.Serialization
 		/// <param name="Value">Value to be written</param>
 		public void WriteBinaryValue(ReadOnlySpan<byte> Value)
 		{
-			WriteField(CbFieldType.Binary);
+			WriteFieldHeader(CbFieldType.Binary);
 			WriteBinaryPayload(Value);
 		}
 
@@ -692,7 +720,7 @@ namespace EpicGames.Serialization
 		/// <param name="Value">Value to be written</param>
 		public void WriteBinary(Utf8String Name, ReadOnlySpan<byte> Value)
 		{
-			WriteField(CbFieldType.Binary, Name);
+			WriteFieldHeader(CbFieldType.Binary, Name);
 			WriteBinaryPayload(Value);
 		}
 
diff --git a/Engine/Source/Programs/Shared/EpicGames.Serialization/Converters/CbFieldConverter.cs b/Engine/Source/Programs/Shared/EpicGames.Serialization/Converters/CbFieldConverter.cs
new file mode 100644
--- /dev/null
+++ b/Engine/Source/Programs/Shared/EpicGames.Serialization/Converters/CbFieldConverter.cs
@@ -0,0 +1,59 @@
+// Copyright Epic Games, Inc. All Rights Reserved.
+
+using EpicGames.Core;
+using System;
+using System.Collections.Generic;
+using System.Text;
+
+namespace EpicGames.Serialization.Converters
+{
+	/// <summary>
+	/// Converter for raw CbField types
+	/// </summary>
+	/// <typeparam name="T"></typeparam>
+	class CbFieldConverter : CbConverter<CbField>
+	{
+		/// <inheritdoc/>
+		public override CbField Read(CbField Field)
+		{
+			return Field;
+		}
+
+		/// <inheritdoc/>
+		public override void Write(CbWriter Writer, CbField Field)
+		{
+			Writer.WriteFieldValue(Field);
+		}
+
+		/// <inheritdoc/>
+		public override void WriteNamed(CbWriter Writer, Utf8String Name, CbField Value)
+		{
+			Writer.WriteField(Name, Value);
+		}
+	}
+
+	/// <summary>
+	/// Converter for raw CbObject types
+	/// </summary>
+	/// <typeparam name="T"></typeparam>
+	class CbObjectConverter : CbConverter<CbObject>
+	{
+		/// <inheritdoc/>
+		public override CbObject Read(CbField Field)
+		{
+			return Field.AsObject();
+		}
+
+		/// <inheritdoc/>
+		public override void Write(CbWriter Writer, CbObject Object)
+		{
+			Writer.WriteFieldValue(Object.AsField());
+		}
+
+		/// <inheritdoc/>
+		public override void WriteNamed(CbWriter Writer, Utf8String Name, CbObject Object)
+		{
+			Writer.WriteField(Name, Object.AsField());
+		}
+	}
+}
