Missing copyright notice.

diff --git a/Engine/Source/Programs/Shared/EpicGames.Redis/Utility/RedisLock.cs b/Engine/Source/Programs/Shared/EpicGames.Redis/Utility/RedisLock.cs
--- a/Engine/Source/Programs/Shared/EpicGames.Redis/Utility/RedisLock.cs
+++ b/Engine/Source/Programs/Shared/EpicGames.Redis/Utility/RedisLock.cs
@@ -1,3 +1,5 @@
+// Copyright Epic Games, Inc. All Rights Reserved.
+
 using StackExchange.Redis;
 using System;
 using System.Collections.Generic;
