Horde: Fix exception outputting error message when unshelving files in incorrect stream.

diff --git a/Engine/Source/Programs/Shared/EpicGames.Perforce.Managed/ManagedWorkspace.cs b/Engine/Source/Programs/Shared/EpicGames.Perforce.Managed/ManagedWorkspace.cs
--- a/Engine/Source/Programs/Shared/EpicGames.Perforce.Managed/ManagedWorkspace.cs
+++ b/Engine/Source/Programs/Shared/EpicGames.Perforce.Managed/ManagedWorkspace.cs
@@ -917,7 +917,7 @@ namespace EpicGames.Perforce.Managed
 			{
 				if (FileIdx >= WhereRecords.Count)
 				{
-					throw new PerforceException("Unable to get location of {File} within {Stream}. Check the correct stream is specified.", LastRecord.Files[FileIdx].DepotFile, StreamName);
+					throw new PerforceException($"Unable to get location of {LastRecord.Files[FileIdx].DepotFile} within {StreamName}. Check the correct stream is specified.");
 				}
 
 				PerforceResponse<WhereRecord> Response = WhereRecords[FileIdx];
@@ -964,7 +964,7 @@ namespace EpicGames.Perforce.Managed
 				string LocalPath = DeleteFile.Path;
 				if (File.Exists(LocalPath))
 				{
-					Logger.LogInformation($"  Deleting {LocalPath}");
+					Logger.LogInformation("  Deleting {LocalPath}", LocalPath);
 					FileUtils.ForceDeleteFile(LocalPath);
 				}
 			}
@@ -973,7 +973,7 @@ namespace EpicGames.Perforce.Managed
 			foreach (WhereRecord WriteFile in WriteFiles)
 			{
 				string LocalPath = WriteFile.Path;
-				Logger.LogInformation($"  Writing {LocalPath}");
+				Logger.LogInformation("  Writing {LocalPath}", LocalPath);
 
 				Directory.CreateDirectory(Path.GetDirectoryName(LocalPath));
 
