Horde: Move RPC protos under the Horde.Build project.

diff --git a/Engine/Source/Programs/Horde/Horde.Agent/Utilities/HordeRpc.cs b/Engine/Source/Programs/Horde/Horde.Agent/Utilities/HordeRpc.cs
new file mode 100644
--- /dev/null
+++ b/Engine/Source/Programs/Horde/Horde.Agent/Utilities/HordeRpc.cs
@@ -0,0 +1,189 @@
+// Copyright Epic Games, Inc. All Rights Reserved.
+
+using Google.Protobuf;
+using Google.Protobuf.Collections;
+using Google.Protobuf.WellKnownTypes;
+using HordeCommon;
+using System;
+using System.Collections.Generic;
+using System.Diagnostics.CodeAnalysis;
+using System.Linq;
+using System.Text;
+
+namespace HordeCommon.Rpc
+{
+#pragma warning disable CS1591
+	partial class Property
+	{
+		public Property(string Name, string Value)
+		{
+			this.Name = Name;
+			this.Value = Value;
+		}
+
+		public Property(KeyValuePair<string, string> Pair)
+		{
+			this.Name = Pair.Key;
+			this.Value = Pair.Value;
+		}
+	}
+
+	partial class PropertyUpdate
+	{
+		public PropertyUpdate(string Name, string? Value)
+		{
+			this.Name = Name;
+			this.Value = Value;
+		}
+	}
+
+	static class PropertyExtensions
+	{
+		public static string GetValue(this RepeatedField<Property> Properties, string Name)
+		{
+			return Properties.First(x => x.Name == Name).Value;
+		}
+
+		public static bool TryGetValue(this RepeatedField<Property> Properties, string Name, [MaybeNullWhen(false)] out string Result)
+		{
+			Property? Property = Properties.FirstOrDefault(x => x.Name == Name);
+			if (Property == null)
+			{
+				Result = null!;
+				return false;
+			}
+			else
+			{
+				Result = Property.Value;
+				return true;
+			}
+		}
+	}
+
+	partial class GetStreamRequest
+	{
+		public GetStreamRequest(string StreamId)
+		{
+			this.StreamId = StreamId;
+		}
+	}
+
+	partial class UpdateStreamRequest
+	{
+		public UpdateStreamRequest(string StreamId, Dictionary<string, string?> Properties)
+		{
+			this.StreamId = StreamId;
+			this.Properties.AddRange(Properties.Select(x => new PropertyUpdate(x.Key, x.Value)));
+		}
+	}
+
+	partial class GetJobRequest
+	{
+		public GetJobRequest(string JobId)
+		{
+			this.JobId = JobId;
+		}
+	}
+
+	partial class BeginBatchRequest
+	{
+		public BeginBatchRequest(string JobId, string BatchId, string LeaseId)
+		{
+			this.JobId = JobId;
+			this.BatchId = BatchId;
+			this.LeaseId = LeaseId;
+		}
+	}
+
+	partial class FinishBatchRequest
+	{
+		public FinishBatchRequest(string JobId, string BatchId, string LeaseId)
+		{
+			this.JobId = JobId;
+			this.BatchId = BatchId;
+			this.LeaseId = LeaseId;
+		}
+	}
+
+	partial class BeginStepRequest
+	{
+		public BeginStepRequest(string JobId, string BatchId, string LeaseId)
+		{
+			this.JobId = JobId;
+			this.BatchId = BatchId;
+			this.LeaseId = LeaseId;
+		}
+	}
+
+	partial class UpdateStepRequest
+	{
+		public UpdateStepRequest(string JobId, string BatchId, string StepId, JobStepState State, JobStepOutcome Outcome)
+		{
+			this.JobId = JobId;
+			this.BatchId = BatchId;
+			this.StepId = StepId;
+			this.State = State;
+			this.Outcome = Outcome;
+		}
+	}
+	
+	partial class GetStepRequest
+	{
+		public GetStepRequest(string JobId, string BatchId, string StepId)
+		{
+			this.JobId = JobId;
+			this.BatchId = BatchId;
+			this.StepId = StepId;
+		}
+	}
+	
+	partial class GetStepResponse
+	{
+		public GetStepResponse(JobStepOutcome Outcome, JobStepState State, bool AbortRequested)
+		{
+			this.Outcome = Outcome;
+			this.State = State;
+			this.AbortRequested = AbortRequested;
+		}
+	}
+
+	partial class CreateEventRequest
+	{
+		public CreateEventRequest(EventSeverity Severity, string LogId, int LineIndex, int LineCount)
+		{
+			this.Severity = Severity;
+			this.LogId = LogId;
+			this.LineIndex = LineIndex;
+			this.LineCount = LineCount;
+		}
+	}
+
+	partial class CreateEventsRequest
+	{
+		public CreateEventsRequest(IEnumerable<CreateEventRequest> Events)
+		{
+			this.Events.AddRange(Events);
+		}
+	}
+
+	partial class WriteOutputRequest
+	{
+		public WriteOutputRequest(string LogId, long Offset, int LineIndex, byte[] Data, bool Flush)
+		{
+			this.LogId = LogId;
+			this.Offset = Offset;
+			this.LineIndex = LineIndex;
+			this.Data = ByteString.CopyFrom(Data);
+			this.Flush = Flush;
+		}
+	}
+
+	partial class DownloadSoftwareRequest
+	{
+		public DownloadSoftwareRequest(string Version)
+		{
+			this.Version = Version;
+		}
+	}
+#pragma warning restore CS1591
+}
diff --git a/Engine/Source/Programs/Horde/Horde.Agent/Utilities/ProtoMethods.cs b/Engine/Source/Programs/Horde/Horde.Agent/Utilities/ProtoMethods.cs
new file mode 100644
--- /dev/null
+++ b/Engine/Source/Programs/Horde/Horde.Agent/Utilities/ProtoMethods.cs
@@ -0,0 +1,33 @@
+// Copyright Epic Games, Inc. All Rights Reserved.
+
+using EpicGames.Horde.Compute;
+using EpicGames.Serialization;
+using System;
+using System.Collections.Generic;
+using System.Text;
+
+namespace HordeCommon.Rpc.Tasks
+{
+	/// <summary>
+	/// 
+	/// </summary>
+	partial class ComputeTaskResultMessage
+	{
+		/// <summary>
+		/// Constructor
+		/// </summary>
+		public ComputeTaskResultMessage(CbObjectAttachment Result)
+		{
+			this.Result = Result;
+		}
+
+		/// <summary>
+		/// Constructor
+		/// </summary>
+		public ComputeTaskResultMessage(ComputeTaskOutcome Outcome, string? Detail = null)
+		{
+			this.Outcome = Outcome;
+			this.Detail = Detail;
+		}
+	}
+}
diff --git a/Engine/Source/Programs/Horde/Horde.Build/Utilities/Clock.cs b/Engine/Source/Programs/Horde/Horde.Build/Utilities/Clock.cs
new file mode 100644
--- /dev/null
+++ b/Engine/Source/Programs/Horde/Horde.Build/Utilities/Clock.cs
@@ -0,0 +1,56 @@
+// Copyright Epic Games, Inc. All Rights Reserved.
+
+using System;
+
+namespace HordeCommon
+{
+	/// <summary>
+	/// Interface representing time to make it pluggable during testing
+	/// In normal use, the Clock implementation below is used. 
+	/// </summary>
+	public interface IClock
+	{
+		/// <summary>
+		/// Return time expressed as the Coordinated Universal Time (UTC)
+		/// </summary>
+		/// <returns></returns>
+		DateTime UtcNow { get; }
+	}
+	
+	/// <summary>
+	/// Implementation of <see cref="IClock"/> which returns the current time
+	/// </summary>
+	public class Clock : IClock
+	{
+		/// <inheritdoc/>
+		public DateTime UtcNow => DateTime.UtcNow;
+	}
+	
+	/// <summary>
+	/// Fake clock that doesn't advance by wall block time
+	/// Requires manual ticking to progress. Used in tests.
+	/// </summary>
+	public class FakeClock : IClock
+	{
+		/// <summary>
+		/// Constructor
+		/// </summary>
+		public FakeClock()
+		{
+			UtcNow = DateTime.UtcNow;
+		}
+
+		/// <summary>
+		/// Advance time by given amount
+		/// Useful for letting time progress during tests
+		/// </summary>
+		/// <param name="Period">Time span to advance</param>
+		public void Advance(TimeSpan Period)
+		{
+			UtcNow = UtcNow.Add(Period);
+		}
+
+		/// <inheritdoc/>
+		public DateTime UtcNow { get; set; }
+	}
+}
\ No newline at end of file
diff --git a/Engine/Source/Programs/Horde/HordeCommon/Clock.cs b/Engine/Source/Programs/Horde/HordeCommon/Clock.cs
deleted file mode 100644
diff --git a/Engine/Source/Programs/Horde/HordeCommon/HordeRpc.cs b/Engine/Source/Programs/Horde/HordeCommon/HordeRpc.cs
deleted file mode 100644
diff --git a/Engine/Source/Programs/Horde/HordeCommon/KnownLogEvents.cs b/Engine/Source/Programs/Horde/HordeCommon/KnownLogEvents.cs
deleted file mode 100644
diff --git a/Engine/Source/Programs/Horde/HordeCommon/ProtoMethods.cs b/Engine/Source/Programs/Horde/HordeCommon/ProtoMethods.cs
deleted file mode 100644
diff --git a/Engine/Source/Programs/Shared/EpicGames.Core/KnownLogEvents.cs b/Engine/Source/Programs/Shared/EpicGames.Core/KnownLogEvents.cs
new file mode 100644
--- /dev/null
+++ b/Engine/Source/Programs/Shared/EpicGames.Core/KnownLogEvents.cs
@@ -0,0 +1,168 @@
+// Copyright Epic Games, Inc. All Rights Reserved.
+
+using Microsoft.Extensions.Logging;
+using System;
+using System.Collections.Generic;
+using System.Diagnostics.CodeAnalysis;
+using System.Text;
+
+namespace HordeCommon
+{
+#pragma warning disable CA1707 // Identifiers should not contain underscores
+	/// <summary>
+	/// Well known log events
+	/// </summary>
+	public static class KnownLogEvents
+	{
+		/// <summary>
+		/// Unset
+		/// </summary>
+		public static EventId None { get; } = new EventId(0);
+
+		/// <summary>
+		/// Generic error
+		/// </summary>
+		public static EventId Generic { get; } = new EventId(1);
+
+		/// <summary>
+		/// Generic exception
+		/// </summary>
+		public static EventId Exception { get; } = new EventId(2);
+
+		/// <summary>
+		/// Compiler error
+		/// </summary>
+		public static EventId Compiler { get; } = new EventId(100);
+
+		/// <summary>
+		/// Linker error
+		/// </summary>
+		public static EventId Linker { get; } = new EventId(200);
+
+		/// <summary>
+		/// Linker: Undefined symbol
+		/// </summary>
+		public static EventId Linker_UndefinedSymbol { get; } = new EventId(201);
+
+		/// <summary>
+		/// Linker: Multiply defined symbol
+		/// </summary>
+		public static EventId Linker_DuplicateSymbol { get; } = new EventId(202);
+
+		/// <summary>
+		/// Engine error
+		/// </summary>
+		public static EventId Engine { get; } = new EventId(300);
+
+		/// <summary>
+		/// Engine log channel
+		/// </summary>
+		public static EventId Engine_LogChannel { get; } = new EventId(301);
+
+		/// <summary>
+		/// Engine: Crash dump
+		/// </summary>
+		public static EventId Engine_Crash { get; } = new EventId(302);
+
+		/// <summary>
+		/// Engine: UE_ASSET_LOG output
+		/// </summary>
+		public static EventId Engine_AssetLog { get; } = new EventId(303);
+
+		/// <summary>
+		/// UAT error
+		/// </summary>
+		public static EventId AutomationTool { get; } = new EventId(400);
+
+		/// <summary>
+		/// UAT: Crash dump
+		/// </summary>
+		public static EventId AutomationTool_Crash { get; } = new EventId(401);
+
+		/// <summary>
+		/// UAT: Exit code indicating a crash
+		/// </summary>
+		public static EventId AutomationTool_CrashExitCode { get; } = new EventId(402);
+
+		/// <summary>
+		/// UAT: Source file with line number
+		/// </summary>
+		public static EventId AutomationTool_SourceFileLine { get; } = new EventId(403);
+
+		/// <summary>
+		/// UAT: Missing copyright notice
+		/// </summary>
+		public static EventId AutomationTool_MissingCopyright { get; } = new EventId(404);
+
+		/// <summary>
+		/// MSBuild: Generic error
+		/// </summary>
+		public static EventId MSBuild { get; } = new EventId(500);
+
+		/// <summary>
+		/// Microsoft: Generic Visual Studio error
+		/// </summary>
+		public static EventId Microsoft { get; } = new EventId(550);
+
+		/// <summary>
+		/// Error message from Gauntlet
+		/// </summary>
+		public static EventId Gauntlet { get; } = new EventId(600);
+
+		/// <summary>
+		/// Error message from Gauntlet engine tests
+		/// </summary>
+		public static EventId Gauntlet_UnitTest { get; } = new EventId(601);
+
+		/// <summary>
+		/// Error message from Gauntlet screenshot tests
+		/// </summary>
+		public static EventId Gauntlet_ScreenshotTest { get; } = new EventId(602);
+
+		/// <summary>
+		/// A systemic event, relating to the health of the farm
+		/// </summary>
+		public static EventId Systemic { get; } = new EventId(700);
+
+		/// <summary>
+		/// A systemic event from XGE
+		/// </summary>
+		public static EventId Systemic_Xge { get; } = new EventId(710);
+
+		/// <summary>
+		/// Builds will run in standalone mode
+		/// </summary>
+		public static EventId Systemic_Xge_Standalone { get; } = new EventId(711);
+
+		/// <summary>
+		/// BuildService.exe is not running
+		/// </summary>
+		public static EventId Systemic_Xge_ServiceNotRunning { get; } = new EventId(712);
+
+		/// <summary>
+		/// DDC is slow
+		/// </summary>
+		public static EventId Systemic_SlowDDC { get; } = new EventId(720);
+
+		/// <summary>
+		/// Internal Horde error
+		/// </summary>
+		public static EventId Systemic_Horde { get; } = new EventId(730);
+
+		/// <summary>
+		/// Artifact upload failed
+		/// </summary>
+		public static EventId Systemic_Horde_ArtifactUpload { get; } = new EventId(731);
+
+		/// <summary>
+		/// Harmless pdbutil error
+		/// </summary>
+		public static EventId Systemic_PdbUtil { get; } = new EventId(740);
+
+		/// <summary>
+		/// A systemic event from MSBuild
+		/// </summary>
+		public static EventId Systemic_MSBuild { get; } = new EventId(750);
+	}
+#pragma warning restore CA1707 // Identifiers should not contain underscores
+}
