Fix for missing copyright notice

diff --git a/Engine/Source/Programs/Shared/EpicGames.Core/CertificateUtils.cs b/Engine/Source/Programs/Shared/EpicGames.Core/CertificateUtils.cs
--- a/Engine/Source/Programs/Shared/EpicGames.Core/CertificateUtils.cs
+++ b/Engine/Source/Programs/Shared/EpicGames.Core/CertificateUtils.cs
@@ -1,3 +1,5 @@
+// Copyright Epic Games, Inc. All Rights Reserved.
+
 using System;
 using System.Collections.Generic;
 using System.Security.Cryptography;
