Add <DebugType>pdbonly</DebugType> to several projects, making them consistent with others.

diff --git a/Engine/Source/Programs/Shared/EpicGames.Build/EpicGames.Build.csproj b/Engine/Source/Programs/Shared/EpicGames.Build/EpicGames.Build.csproj
--- a/Engine/Source/Programs/Shared/EpicGames.Build/EpicGames.Build.csproj
+++ b/Engine/Source/Programs/Shared/EpicGames.Build/EpicGames.Build.csproj
@@ -8,6 +8,7 @@
     <GenerateTargetFrameworkAttribute>false</GenerateTargetFrameworkAttribute>
     <Configurations>Debug;Release;Development</Configurations>
     <Nullable>enable</Nullable>
+    <DebugType>pdbonly</DebugType>
   </PropertyGroup>
   <PropertyGroup Condition=" '$(Configuration)|$(Platform)' == 'Debug|AnyCPU' ">
     <DocumentationFile></DocumentationFile>
diff --git a/Engine/Source/Programs/Shared/EpicGames.Core/EpicGames.Core.csproj b/Engine/Source/Programs/Shared/EpicGames.Core/EpicGames.Core.csproj
--- a/Engine/Source/Programs/Shared/EpicGames.Core/EpicGames.Core.csproj
+++ b/Engine/Source/Programs/Shared/EpicGames.Core/EpicGames.Core.csproj
@@ -10,6 +10,7 @@
     <Configurations>Debug;Release;Development</Configurations>
     <Nullable>enable</Nullable>
     <AllowUnsafeBlocks>true</AllowUnsafeBlocks>
+    <DebugType>pdbonly</DebugType>
   </PropertyGroup>
   <PropertyGroup Condition=" '$(Configuration)|$(Platform)' == 'Debug|AnyCPU' ">
     <DocumentationFile></DocumentationFile>
diff --git a/Engine/Source/Programs/Shared/EpicGames.Jupiter/EpicGames.Jupiter.csproj b/Engine/Source/Programs/Shared/EpicGames.Jupiter/EpicGames.Jupiter.csproj
--- a/Engine/Source/Programs/Shared/EpicGames.Jupiter/EpicGames.Jupiter.csproj
+++ b/Engine/Source/Programs/Shared/EpicGames.Jupiter/EpicGames.Jupiter.csproj
@@ -8,6 +8,7 @@
     <GenerateAssemblyInfo>false</GenerateAssemblyInfo>
     <GenerateTargetFrameworkAttribute>false</GenerateTargetFrameworkAttribute>
     <Configurations>Debug;Release;Development</Configurations>
+    <DebugType>pdbonly</DebugType>
   </PropertyGroup>
   <PropertyGroup Condition=" '$(Configuration)|$(Platform)' == 'Debug|AnyCPU' ">
     <DocumentationFile></DocumentationFile>
