Missing file.

diff --git a/Engine/Source/Programs/Shared/EpicGames.Core/JsonExtensions.cs b/Engine/Source/Programs/Shared/EpicGames.Core/JsonExtensions.cs
new file mode 100644
--- /dev/null
+++ b/Engine/Source/Programs/Shared/EpicGames.Core/JsonExtensions.cs
@@ -0,0 +1,45 @@
+// Copyright Epic Games, Inc. All Rights Reserved.
+
+using System;
+using System.Collections.Generic;
+using System.Text;
+using System.Text.Json;
+
+namespace EpicGames.Core
+{
+	/// <summary>
+	/// Utility methods for utf8 json reader/writer
+	/// </summary>
+	public static class JsonExtensions
+	{
+		/// <summary>
+		/// Tries to read a property name then move to the value token
+		/// </summary>
+		/// <param name="Reader">Token reader</param>
+		/// <param name="PropertyName">Receives the property name on success</param>
+		/// <returns>True if the read succeeded</returns>
+		public static bool TryReadNextPropertyName(ref Utf8JsonReader Reader, out ReadOnlySpan<byte> PropertyName)
+		{
+			if (Reader.Read() && Reader.TokenType == JsonTokenType.PropertyName)
+			{
+				PropertyName = Reader.ValueSpan;
+				return Reader.Read();
+			}
+			else
+			{
+				PropertyName = ReadOnlySpan<byte>.Empty;
+				return false;
+			}
+		}
+
+		/// <summary>
+		/// Reads a Utf8 string from the current json token
+		/// </summary>
+		/// <param name="Reader"></param>
+		/// <returns></returns>
+		public static Utf8String GetUtf8String(this Utf8JsonReader Reader)
+		{
+			throw new NotImplementedException();
+		}
+	}
+}
