Added a method to compute a hash from a Stream object

#rnx

diff --git a/Engine/Source/Programs/Shared/EpicGames.Core/Md5Hash.cs b/Engine/Source/Programs/Shared/EpicGames.Core/Md5Hash.cs
--- a/Engine/Source/Programs/Shared/EpicGames.Core/Md5Hash.cs
+++ b/Engine/Source/Programs/Shared/EpicGames.Core/Md5Hash.cs
@@ -5,6 +5,7 @@ using System.Buffers.Binary;
 using System.Collections.Generic;
 using System.ComponentModel;
 using System.Globalization;
+using System.IO;
 using System.Security.Cryptography;
 using System.Text;
 
@@ -70,6 +71,19 @@ namespace EpicGames.Core
 			return new Md5Hash(Output);
 		}
 
+		/// <summary>
+		/// Creates a content hash for the input Stream object
+		/// </summary>
+		/// <param name="Stream">The Stream object to compoute the has for</param>
+		/// <returns>New content hash instance containing the hash of the data</returns>
+		public static Md5Hash Compute(Stream Stream)
+		{
+			using (MD5 Hasher = MD5.Create())
+			{
+				return new Md5Hash(Hasher.ComputeHash(Stream));
+			}
+		}
+
 		/// <summary>
 		/// Parses a digest from the given hex string
 		/// </summary>
