BuildGraph: Merging BuildGraph library refactor from //UE5/Release-5.0.

diff --git a/Engine/Source/Programs/AutomationTool/BuildGraph/Agent.cs b/Engine/Source/Programs/AutomationTool/BuildGraph/Agent.cs
deleted file mode 100644
diff --git a/Engine/Source/Programs/AutomationTool/BuildGraph/Aggregate.cs b/Engine/Source/Programs/AutomationTool/BuildGraph/Aggregate.cs
deleted file mode 100644
diff --git a/Engine/Source/Programs/AutomationTool/BuildGraph/Badge.cs b/Engine/Source/Programs/AutomationTool/BuildGraph/Badge.cs
deleted file mode 100644
diff --git a/Engine/Source/Programs/AutomationTool/BuildGraph/BuildGraph.Automation.csproj b/Engine/Source/Programs/AutomationTool/BuildGraph/BuildGraph.Automation.csproj
--- a/Engine/Source/Programs/AutomationTool/BuildGraph/BuildGraph.Automation.csproj
+++ b/Engine/Source/Programs/AutomationTool/BuildGraph/BuildGraph.Automation.csproj
@@ -42,6 +42,7 @@
     </Compile>
   </ItemGroup>
   <ItemGroup>
+    <ProjectReference Include="..\..\Shared\EpicGames.BuildGraph\EpicGames.BuildGraph.csproj" />
     <ProjectReference Include="..\..\Shared\EpicGames.Core\EpicGames.Core.csproj" PrivateAssets="All"><Private>false</Private></ProjectReference>
     <ProjectReference Include="..\..\Shared\EpicGames.Jupiter\EpicGames.Jupiter.csproj" PrivateAssets="All" />
     <ProjectReference Include="..\..\UnrealBuildTool\UnrealBuildTool.csproj" PrivateAssets="All"><Private>false</Private></ProjectReference>
diff --git a/Engine/Source/Programs/AutomationTool/BuildGraph/BuildGraph.cs b/Engine/Source/Programs/AutomationTool/BuildGraph/BuildGraph.cs
--- a/Engine/Source/Programs/AutomationTool/BuildGraph/BuildGraph.cs
+++ b/Engine/Source/Programs/AutomationTool/BuildGraph/BuildGraph.cs
@@ -1,7 +1,9 @@
 // Copyright Epic Games, Inc. All Rights Reserved.
 
+using EpicGames.BuildGraph;
 using System;
 using System.Collections.Generic;
+using System.ComponentModel;
 using System.Diagnostics;
 using System.IO;
 using System.Linq;
@@ -9,6 +11,8 @@ using System.Reflection;
 using System.Text;
 using System.Xml;
 using EpicGames.Core;
+using EpicGames.MCP.Automation;
+using Microsoft.Extensions.Logging;
 using OpenTracing;
 using OpenTracing.Util;
 using UnrealBuildBase;
@@ -16,6 +20,130 @@ using UnrealBuildTool;
 
 namespace AutomationTool
 {
+	/// <summary>
+	/// Implementation of ScriptTaskParameter corresponding to a field in a parameter class
+	/// </summary>
+	class ScriptTaskParameterBinding : BgScriptTaskParameter
+	{
+		/// <summary>
+		/// Field for this parameter
+		/// </summary>
+		public FieldInfo FieldInfo { get; }
+
+		/// <summary>
+		/// Constructor
+		/// </summary>
+		public ScriptTaskParameterBinding(string Name, FieldInfo FieldInfo, TaskParameterValidationType ValidationType, bool bOptional)
+			: base(Name, FieldInfo.FieldType, ValidationType, bOptional)
+		{
+			this.FieldInfo = FieldInfo;
+		}
+	}
+
+	/// <summary>
+	/// Binding of a ScriptTask to a Script
+	/// </summary>
+	class ScriptTaskBinding : BgScriptTask
+	{
+		/// <summary>
+		/// Type of the task to construct with this info
+		/// </summary>
+		public Type TaskClass;
+
+		/// <summary>
+		/// Type to construct with the parsed parameters
+		/// </summary>
+		public Type ParametersClass;
+
+		/// <summary>
+		/// Map from name to parameter
+		/// </summary>
+		public IReadOnlyDictionary<string, ScriptTaskParameterBinding> NameToParameter { get; }
+
+		/// <summary>
+		/// Constructor
+		/// </summary>
+		/// <param name="Name">Name of the task</param>
+		/// <param name="TaskClass">Task class to create</param>
+		/// <param name="ParametersClass">Class type of an object to be constructed and passed as an argument to the task class constructor</param>
+		public ScriptTaskBinding(string Name, Type TaskClass, Type ParametersClass)
+			: this(Name, TaskClass, ParametersClass, CreateParameters(ParametersClass))
+		{
+		}
+
+		/// <summary>
+		/// Private constructor
+		/// </summary>
+		private ScriptTaskBinding(string Name, Type TaskClass, Type ParametersClass, List<ScriptTaskParameterBinding> Parameters)
+			: base(Name, Parameters.ConvertAll<BgScriptTaskParameter>(x => x))
+		{
+			this.TaskClass = TaskClass;
+			this.ParametersClass = ParametersClass;
+			this.NameToParameter = Parameters.ToDictionary(x => x.Name, x => x);
+		}
+
+		static List<ScriptTaskParameterBinding> CreateParameters(Type ParametersClass)
+		{
+			List<ScriptTaskParameterBinding> ScriptTaskParameters = new List<ScriptTaskParameterBinding>();
+			foreach (FieldInfo Field in ParametersClass.GetFields())
+			{
+				if (Field.MemberType == MemberTypes.Field)
+				{
+					TaskParameterAttribute ParameterAttribute = Field.GetCustomAttribute<TaskParameterAttribute>();
+					if (ParameterAttribute != null)
+					{
+						ScriptTaskParameters.Add(new ScriptTaskParameterBinding(Field.Name, Field, ParameterAttribute.ValidationType, ParameterAttribute.Optional));
+					}
+				}
+			}
+			return ScriptTaskParameters;
+		}
+	}
+
+	/// <summary>
+	/// Implementation of <see cref="IBgScriptReaderContext"/> which reads files from disk
+	/// </summary>
+	class ScriptReaderFileContext : IBgScriptReaderContext
+	{
+		/// <inheritdoc/>
+		public object GetNativePath(string Path)
+		{
+			return FileReference.Combine(Unreal.RootDirectory, Path).FullName;
+		}
+
+		/// <inheritdoc/>
+		public bool Exists(string Path)
+		{
+			try
+			{
+				return FileReference.Exists(FileReference.Combine(Unreal.RootDirectory, Path));
+			}
+			catch
+			{
+				return false;
+			}
+		}
+
+		/// <inheritdoc/>
+		public bool TryRead(string Path, out byte[] Data)
+		{
+			try
+			{
+				FileReference File = FileReference.Combine(Unreal.RootDirectory, Path);
+				if (FileReference.Exists(File))
+				{
+					Data = FileReference.ReadAllBytes(File);
+					return true;
+				}
+			}
+			catch
+			{
+			}
+			Data = null;
+			return false;
+		}
+	}
+
 	/// <summary>
 	/// Tool to execute build automation scripts for UE projects, which can be run locally or in parallel across a build farm (assuming synchronization and resource allocation implemented by a separate system).
 	///
@@ -75,6 +203,11 @@ namespace AutomationTool
 	[Help("WriteToSharedStorage", "Allow writing to shared storage. If not set, but -SharedStorageDir is specified, build products will read but not written")]
 	public class BuildGraph : BuildCommand
 	{
+		/// <summary>
+		/// Context object for reading scripts and evaluating conditions
+		/// </summary>
+		ScriptReaderFileContext Context { get; } = new ScriptReaderFileContext();
+
 		/// <summary>
 		/// Main entry point for the BuildGraph command
 		/// </summary>
@@ -92,8 +225,6 @@ namespace AutomationTool
 			string SharedStorageDir = ParseParamValue("SharedStorageDir", null);
 			string SingleNodeName = ParseParamValue("SingleNode", null);
 			string TriggerName = ParseParamValue("Trigger", null);
-			string[] SkipTriggerNames = ParseParamValue("SkipTrigger", "").Split(new char[]{ '+', ';' }, StringSplitOptions.RemoveEmptyEntries).ToArray();
-			bool bSkipTriggers = ParseParam("SkipTriggers");
 			string TokenSignature = ParseParamValue("TokenSignature", null);
 			bool bSkipTargetsWithoutTokens = ParseParam("SkipTargetsWithoutTokens");
 			bool bResume = SingleNodeName != null || ParseParam("Resume");
@@ -101,6 +232,7 @@ namespace AutomationTool
 			bool bShowDiagnostics = ParseParam("ShowDiagnostics");
 			bool bWriteToSharedStorage = ParseParam("WriteToSharedStorage") || CommandUtils.IsBuildMachine;
 			bool bPublicTasksOnly = ParseParam("PublicTasksOnly");
+			bool bSkipValidation = ParseParam("SkipValidation");
 			string ReportName = ParseParamValue("ReportName", null);
 			string BranchOverride = ParseParamValue("Branch", null);
 
@@ -214,7 +346,7 @@ namespace AutomationTool
 			}
 
 			// Find all the tasks from the loaded assemblies
-			Dictionary<string, ScriptTask> NameToTask = new Dictionary<string,ScriptTask>();
+			Dictionary<string, ScriptTaskBinding> NameToTask = new Dictionary<string, ScriptTaskBinding>();
 			if(!FindAvailableTasks(NameToTask, bPublicTasksOnly))
 			{
 				return ExitCode.Error_Unknown;
@@ -228,15 +360,22 @@ namespace AutomationTool
 			}
 
 			// Import schema if one is passed in
-			ScriptSchema Schema;
+			BgScriptSchema Schema;
 			if (ImportSchemaFileName != null)
 			{
-				Schema = ScriptSchema.Import(FileReference.FromString(ImportSchemaFileName), NameToTask);
+				Schema = BgScriptSchema.Import(FileReference.FromString(ImportSchemaFileName));
 			}
 			else
 			{
+				// Add any primitive types
+				List<(Type, ScriptSchemaStandardType)> PrimitiveTypes = new List<(Type, ScriptSchemaStandardType)>();
+				PrimitiveTypes.Add((typeof(FileReference), ScriptSchemaStandardType.BalancedString));
+				PrimitiveTypes.Add((typeof(DirectoryReference), ScriptSchemaStandardType.BalancedString));
+				PrimitiveTypes.Add((typeof(UnrealTargetPlatform), ScriptSchemaStandardType.BalancedString));
+				PrimitiveTypes.Add((typeof(MCPPlatform), ScriptSchemaStandardType.BalancedString));
+
 				// Create a schema for the given tasks
-				Schema = new ScriptSchema(NameToTask);
+				Schema = new BgScriptSchema(NameToTask.Values, PrimitiveTypes);
 				if (SchemaFileName != null)
 				{
 					FileReference FullSchemaFileName = new FileReference(SchemaFileName);
@@ -256,9 +395,18 @@ namespace AutomationTool
 				return ExitCode.Error_Unknown;
 			}
 
+			// Normalize the script filename
+			FileReference FullScriptFile = FileReference.Combine(Unreal.RootDirectory, ScriptFileName);
+			if (!FullScriptFile.IsUnderDirectory(Unreal.RootDirectory))
+			{
+				LogError("BuildGraph scripts must be under the UE root directory");
+				return ExitCode.Error_Unknown;
+			}
+			ScriptFileName = FullScriptFile.MakeRelativeTo(Unreal.RootDirectory).Replace('\\', '/');
+
 			// Read the script from disk
-			Graph Graph;
-			if(!ScriptReader.TryRead(new FileReference(ScriptFileName), Arguments, DefaultProperties, PreprocessedFileName != null, Schema, out Graph, SingleNodeName))
+			BgScript Graph;
+			if(!BgScriptReader.TryRead(Context, ScriptFileName, Arguments, DefaultProperties, PreprocessedFileName != null, Schema, Logger, out Graph, SingleNodeName))
 			{
 				return ExitCode.Error_Unknown;
 			}
@@ -282,7 +430,7 @@ namespace AutomationTool
 			}
 
 			// Convert the supplied target references into nodes 
-			HashSet<Node> TargetNodes = new HashSet<Node>();
+			HashSet<BgNode> TargetNodes = new HashSet<BgNode>();
 			if(TargetNames.Length == 0)
 			{
 				if (!bListOnly && SingleNodeName == null)
@@ -309,7 +457,7 @@ namespace AutomationTool
 
 				foreach (string TargetName in NodesToResolve)
 				{
-					Node[] Nodes;
+					BgNode[] Nodes;
 					if(!Graph.TryResolveReference(TargetName, out Nodes))
 					{
 						LogError("Target '{0}' is not in graph", TargetName);
@@ -326,10 +474,10 @@ namespace AutomationTool
 				HashSet<FileReference> RequiredTokens = new HashSet<FileReference>(TargetNodes.SelectMany(x => x.RequiredTokens));
 
 				// List out all the required tokens
-				if(SingleNodeName == null)
+				if (SingleNodeName == null)
 				{
 					CommandUtils.LogInformation("Required tokens:");
-					foreach(Node Node in TargetNodes)
+					foreach(BgNode Node in TargetNodes)
 					{
 						foreach(FileReference RequiredToken in Node.RequiredTokens)
 						{
@@ -362,13 +510,13 @@ namespace AutomationTool
 					if(bSkipTargetsWithoutTokens)
 					{
 						// Find all the nodes we're going to skip
-						HashSet<Node> SkipNodes = new HashSet<Node>();
+						HashSet<BgNode> SkipNodes = new HashSet<BgNode>();
 						foreach(IGrouping<string, FileReference> MissingTokensForBuild in MissingTokens.GroupBy(x => x.Value, x => x.Key))
 						{
 							LogInformation("Skipping the following nodes due to {0}:", MissingTokensForBuild.Key);
 							foreach(FileReference MissingToken in MissingTokensForBuild)
 							{
-								foreach(Node SkipNode in TargetNodes.Where(x => x.RequiredTokens.Contains(MissingToken) && SkipNodes.Add(x)))
+								foreach(BgNode SkipNode in TargetNodes.Where(x => x.RequiredTokens.Contains(MissingToken) && SkipNodes.Add(x)))
 								{
 									LogInformation("    {0}", SkipNode);
 								}
@@ -380,7 +528,7 @@ namespace AutomationTool
 						{
 							TargetNodes.ExceptWith(SkipNodes);
 							LogInformation("Remaining target nodes:");
-							foreach(Node TargetNode in TargetNodes)
+							foreach(BgNode TargetNode in TargetNodes)
 							{
 								LogInformation("    {0}", TargetNode);
 							}
@@ -394,7 +542,7 @@ namespace AutomationTool
 					{
 						foreach(KeyValuePair<FileReference, string> Pair in MissingTokens)
 						{
-							List<Node> SkipNodes = TargetNodes.Where(x => x.RequiredTokens.Contains(Pair.Key)).ToList();
+							List<BgNode> SkipNodes = TargetNodes.Where(x => x.RequiredTokens.Contains(Pair.Key)).ToList();
 							LogError("Cannot run {0} due to previous build: {1}", String.Join(", ", SkipNodes), Pair.Value);
 						}
 						foreach(FileReference CreatedToken in CreatedTokens)
@@ -409,31 +557,10 @@ namespace AutomationTool
 			// Cull the graph to include only those nodes
 			Graph.Select(TargetNodes);
 
-			// Collapse any triggers in the graph which are marked to be skipped
-			HashSet<ManualTrigger> SkipTriggers = new HashSet<ManualTrigger>();
-			if(bSkipTriggers)
-			{
-				SkipTriggers.UnionWith(Graph.NameToTrigger.Values);
-			}
-			else
-			{
-				foreach(string SkipTriggerName in SkipTriggerNames)
-				{
-					ManualTrigger SkipTrigger;
-					if(!Graph.NameToTrigger.TryGetValue(TriggerName, out SkipTrigger))
-					{
-						LogError("Couldn't find trigger '{0}'", TriggerName);
-						return ExitCode.Error_Unknown;
-					}
-					SkipTriggers.Add(SkipTrigger);
-				}
-			}
-			Graph.SkipTriggers(SkipTriggers);
-
 			// If a report for the whole build was requested, insert it into the graph
 			if (ReportName != null)
 			{
-				Report NewReport = new Report(ReportName);
+				BgReport NewReport = new BgReport(ReportName);
 				NewReport.Nodes.UnionWith(Graph.Agents.SelectMany(x => x.Nodes));
 				Graph.NameToReport.Add(ReportName, NewReport);
 			}
@@ -447,16 +574,8 @@ namespace AutomationTool
 				bListOnly = true;
 			}
 
-			// Find the triggers which we are explicitly running.
-			ManualTrigger Trigger = null;
-			if(TriggerName != null && !Graph.NameToTrigger.TryGetValue(TriggerName, out Trigger))
-			{
-				LogError("Couldn't find trigger '{0}'", TriggerName);
-				return ExitCode.Error_Unknown;
-			}
-
 			// If we're just building a single node, find it 
-			Node SingleNode = null;
+			BgNode SingleNode = null;
 			if(SingleNodeName != null && !Graph.NameToNode.TryGetValue(SingleNodeName, out SingleNode))
 			{
 				LogError("Node '{0}' is not in the trimmed graph", SingleNodeName);
@@ -466,15 +585,14 @@ namespace AutomationTool
 			// If we just want to show the contents of the graph, do so and exit.
 			if(bListOnly)
 			{ 
-				HashSet<Node> CompletedNodes = FindCompletedNodes(Graph, Storage);
-				Graph.Print(CompletedNodes, PrintOptions);
+				HashSet<BgNode> CompletedNodes = FindCompletedNodes(Graph, Storage);
+				Graph.Print(CompletedNodes, PrintOptions, Log.Logger);
 			}
 
 			// Print out all the diagnostic messages which still apply, unless we're running a step as part of a build system or just listing the contents of the file. 
 			if(SingleNode == null && (!bListOnly || bShowDiagnostics))
 			{
-				IEnumerable<GraphDiagnostic> Diagnostics = Graph.Diagnostics.Where(x => x.EnclosingTrigger == Trigger);
-				foreach(GraphDiagnostic Diagnostic in Diagnostics)
+				foreach(BgScriptDiagnostic Diagnostic in Graph.Diagnostics)
 				{
 					if(Diagnostic.EventType == LogEventType.Console)
 					{
@@ -489,7 +607,7 @@ namespace AutomationTool
 						CommandUtils.LogError(Diagnostic.Message);
 					}
 				}
-				if(Diagnostics.Any(x => x.EventType == LogEventType.Error))
+				if(Graph.Diagnostics.Any(x => x.EventType == LogEventType.Error))
 				{
 					return ExitCode.Error_Unknown;
 				}
@@ -498,25 +616,42 @@ namespace AutomationTool
 			// Export the graph to a file
 			if(ExportFileName != null)
 			{
-				HashSet<Node> CompletedNodes = FindCompletedNodes(Graph, Storage);
-				Graph.Print(CompletedNodes, PrintOptions);
-				Graph.Export(new FileReference(ExportFileName), Trigger, CompletedNodes);
+				HashSet<BgNode> CompletedNodes = FindCompletedNodes(Graph, Storage);
+				Graph.Print(CompletedNodes, PrintOptions, Log.Logger);
+				Graph.Export(new FileReference(ExportFileName), CompletedNodes);
 				return ExitCode.Success;
 			}
 
+			// Create tasks for the entire graph
+			Dictionary<BgTask, CustomTask> TaskInfoToTask = new Dictionary<BgTask, CustomTask>();
+			if (bSkipValidation && SingleNode != null)
+			{
+				if (!CreateTaskInstances(Graph, SingleNode, NameToTask, TaskInfoToTask))
+				{
+					return ExitCode.Error_Unknown;
+				}
+			}
+			else
+			{
+				if (!CreateTaskInstances(Graph, NameToTask, TaskInfoToTask))
+				{
+					return ExitCode.Error_Unknown;
+				}
+			}
+
 			// Execute the command
-			if(!bListOnly)
+			if (!bListOnly)
 			{
 				if(SingleNode != null)
 				{
-					if(!BuildNode(new JobContext(this), Graph, SingleNode, Storage, bWithBanner: true))
+					if(!BuildNode(new JobContext(this), Graph, SingleNode, TaskInfoToTask, Storage, bWithBanner: true))
 					{
 						return ExitCode.Error_Unknown;
 					}
 				}
 				else
 				{
-					if(!BuildAllNodes(new JobContext(this), Graph, Storage))
+					if(!BuildAllNodes(new JobContext(this), Graph, TaskInfoToTask, Storage))
 					{
 						return ExitCode.Error_Unknown;
 					}
@@ -525,12 +660,185 @@ namespace AutomationTool
 			return ExitCode.Success;
 		}
 
+		bool CreateTaskInstances(BgScript Graph, Dictionary<string, ScriptTaskBinding> NameToTask, Dictionary<BgTask, CustomTask> TaskInfoToTask)
+		{
+			bool bResult = true;
+			foreach (BgAgent Agent in Graph.Agents)
+			{
+				foreach (BgNode Node in Agent.Nodes)
+				{
+					bResult &= CreateTaskInstances(Graph, Node, NameToTask, TaskInfoToTask);
+				}
+			}
+			return bResult;
+		}
+
+		bool CreateTaskInstances(BgScript Graph, BgNode Node, Dictionary<string, ScriptTaskBinding> NameToTask, Dictionary<BgTask, CustomTask> TaskInfoToTask)
+		{
+			bool bResult = true;
+			foreach (BgTask TaskInfo in Node.TaskInfos)
+			{
+				CustomTask Task = BindTask(Node, TaskInfo, NameToTask, Graph.TagNameToNodeOutput);
+				if (Task == null)
+				{
+					bResult = false;
+				}
+				else
+				{
+					TaskInfoToTask.Add(TaskInfo, Task);
+				}
+			}
+			return bResult;
+		}
+
+		CustomTask BindTask(BgNode Node, BgTask TaskInfo, Dictionary<string, ScriptTaskBinding> NameToTask, IReadOnlyDictionary<string, BgNodeOutput> TagNameToNodeOutput)
+		{
+			// Get the reflection info for this element
+			ScriptTaskBinding Task;
+			if (!NameToTask.TryGetValue(TaskInfo.Name, out Task))
+			{
+				OutputBindingError(TaskInfo, "Unknown task '{TaskName}'", TaskInfo.Name);
+				return null;
+			}
+
+			// Check all the required parameters are present
+			bool bHasRequiredAttributes = true;
+			foreach (ScriptTaskParameterBinding Parameter in Task.NameToParameter.Values)
+			{
+				if (!Parameter.bOptional && !TaskInfo.Arguments.ContainsKey(Parameter.Name))
+				{
+					OutputBindingError(TaskInfo, "Missing required attribute - {AttrName}", Parameter.Name);
+					bHasRequiredAttributes = false;
+				}
+			}
+
+			// Read all the attributes into a parameters object for this task
+			object ParametersObject = Activator.CreateInstance(Task.ParametersClass);
+			foreach ((string Name, string Value) in TaskInfo.Arguments)
+			{
+				// Get the field that this attribute should be written to in the parameters object
+				ScriptTaskParameterBinding Parameter;
+				if (!Task.NameToParameter.TryGetValue(Name, out Parameter))
+				{
+					OutputBindingError(TaskInfo, "Unknown attribute '{AttrName}'", Name);
+					continue;
+				}
+
+				// If it's a collection type, split it into separate values
+				if (Parameter.CollectionType == null)
+				{
+					// Parse it and assign it to the parameters object
+					object FieldValue = ParseValue(Value, Parameter.ValueType);
+					Parameter.FieldInfo.SetValue(ParametersObject, FieldValue);
+				}
+				else
+				{
+					// Get the collection, or create one if necessary
+					object CollectionValue = Parameter.FieldInfo.GetValue(ParametersObject);
+					if (CollectionValue == null)
+					{
+						CollectionValue = Activator.CreateInstance(Parameter.FieldInfo.FieldType);
+						Parameter.FieldInfo.SetValue(ParametersObject, CollectionValue);
+					}
+
+					// Parse the values and add them to the collection
+					List<string> ValueStrings = CustomTask.SplitDelimitedList(Value);
+					foreach (string ValueString in ValueStrings)
+					{
+						object ElementValue = ParseValue(ValueString, Parameter.ValueType);
+						Parameter.CollectionType.InvokeMember("Add", BindingFlags.InvokeMethod | BindingFlags.Instance | BindingFlags.Public, null, CollectionValue, new object[] { ElementValue });
+					}
+				}
+			}
+
+			// Construct the task
+			if (!bHasRequiredAttributes)
+			{
+				return null;
+			}
+
+			// Add it to the list
+			CustomTask NewTask = (CustomTask)Activator.CreateInstance(Task.TaskClass, ParametersObject);
+
+			// Set up the source location for diagnostics
+			NewTask.SourceLocation = TaskInfo.SourceLocation;
+
+			// Make sure all the read tags are local or listed as a dependency
+			foreach (string ReadTagName in NewTask.FindConsumedTagNames())
+			{
+				BgNodeOutput Output;
+				if (TagNameToNodeOutput.TryGetValue(ReadTagName, out Output))
+				{
+					if (Output != null && Output.ProducingNode != Node && !Node.Inputs.Contains(Output))
+					{
+						OutputBindingError(TaskInfo, "The tag '{TagName}' is not a dependency of node '{Node}'", ReadTagName, Node.Name);
+					}
+				}
+			}
+
+			// Make sure all the written tags are local or listed as an output
+			foreach (string ModifiedTagName in NewTask.FindProducedTagNames())
+			{
+				BgNodeOutput Output;
+				if (TagNameToNodeOutput.TryGetValue(ModifiedTagName, out Output))
+				{
+					if (Output != null && !Node.Outputs.Contains(Output))
+					{
+						OutputBindingError(TaskInfo, "The tag '{0}' is created by '{1}', and cannot be modified downstream", Output.TagName, Output.ProducingNode.Name);
+					}
+				}
+			}
+			return NewTask;
+		}
+
+		/// <summary>
+		/// Parse a value of the given type
+		/// </summary>
+		/// <param name="ValueText">The text to parse</param>
+		/// <param name="ValueType">Type of the value to parse</param>
+		/// <returns>Value that was parsed</returns>
+		object ParseValue(string ValueText, Type ValueType)
+		{
+			// Parse it and assign it to the parameters object
+			if (ValueType.IsEnum)
+			{
+				return Enum.Parse(ValueType, ValueText);
+			}
+			else if (ValueType == typeof(Boolean))
+			{
+				return BgCondition.Evaluate(ValueText, Context);
+			}
+			else if (ValueType == typeof(FileReference))
+			{
+				return CustomTask.ResolveFile(ValueText);
+			}
+			else if (ValueType == typeof(DirectoryReference))
+			{
+				return CustomTask.ResolveDirectory(ValueText);
+			}
+
+			TypeConverter Converter = TypeDescriptor.GetConverter(ValueType);
+			if (Converter.CanConvertFrom(typeof(string)))
+			{
+				return Converter.ConvertFromString(ValueText);
+			}
+			else
+			{
+				return Convert.ChangeType(ValueText, ValueType);
+			}
+		}
+
+		void OutputBindingError(BgTask Task, string Format, params object[] Args)
+		{
+			Logger.LogScriptError(Context.GetNativePath(Task.SourceLocation.Item1), Task.SourceLocation.Item2, Format, Args);
+		}
+
 		/// <summary>
 		/// Find all the tasks which are available from the loaded assemblies
 		/// </summary>
 		/// <param name="NameToTask">Mapping from task name to information about how to serialize it</param>
 		/// <param name="bPublicTasksOnly">Whether to include just public tasks, or all the tasks in any loaded assemblies</param>
-		static bool FindAvailableTasks(Dictionary<string, ScriptTask> NameToTask, bool bPublicTasksOnly)
+		static bool FindAvailableTasks(Dictionary<string, ScriptTaskBinding> NameToTask, bool bPublicTasksOnly)
 		{
 			IEnumerable<Assembly> LoadedScriptAssemblies = ScriptManager.AllScriptAssemblies;
 
@@ -565,7 +873,7 @@ namespace AutomationTool
 							CommandUtils.LogError("Found multiple handlers for task elements called '{0}'", ElementAttribute.Name);
 							return false;
 						}
-						NameToTask.Add(ElementAttribute.Name, new ScriptTask(ElementAttribute.Name, Type, ElementAttribute.ParametersType));
+						NameToTask.Add(ElementAttribute.Name, new ScriptTaskBinding(ElementAttribute.Name, Type, ElementAttribute.ParametersType));
 					}
 				}
 			}
@@ -666,10 +974,10 @@ namespace AutomationTool
 		/// </summary>
 		/// <param name="Graph">The graph instance</param>
 		/// <param name="Storage">The temp storage backend which stores the shared state</param>
-		HashSet<Node> FindCompletedNodes(Graph Graph, TempStorage Storage)
+		HashSet<BgNode> FindCompletedNodes(BgScript Graph, TempStorage Storage)
 		{
-			HashSet<Node> CompletedNodes = new HashSet<Node>();
-			foreach(Node Node in Graph.Agents.SelectMany(x => x.Nodes))
+			HashSet<BgNode> CompletedNodes = new HashSet<BgNode>();
+			foreach(BgNode Node in Graph.Agents.SelectMany(x => x.Nodes))
 			{
 				if(Storage.IsComplete(Node.Name))
 				{
@@ -686,15 +994,15 @@ namespace AutomationTool
 		/// <param name="Graph">The graph instance</param>
 		/// <param name="Storage">The temp storage backend which stores the shared state</param>
 		/// <returns>True if everything built successfully</returns>
-		bool BuildAllNodes(JobContext Job, Graph Graph, TempStorage Storage)
+		bool BuildAllNodes(JobContext Job, BgScript Graph, Dictionary<BgTask, CustomTask> TaskInfoToTask, TempStorage Storage)
 		{
 			// Build a flat list of nodes to execute, in order
-			Node[] NodesToExecute = Graph.Agents.SelectMany(x => x.Nodes).ToArray();
+			BgNode[] NodesToExecute = Graph.Agents.SelectMany(x => x.Nodes).ToArray();
 
 			// Check the integrity of any local nodes that have been completed. It's common to run formal builds locally between regular development builds, so we may have 
 			// stale local state. Rather than failing later, detect and clean them up now.
-			HashSet<Node> CleanedNodes = new HashSet<Node>();
-			foreach(Node NodeToExecute in NodesToExecute)
+			HashSet<BgNode> CleanedNodes = new HashSet<BgNode>();
+			foreach(BgNode NodeToExecute in NodesToExecute)
 			{
 				if(NodeToExecute.InputDependencies.Any(x => CleanedNodes.Contains(x)) || !Storage.CheckLocalIntegrity(NodeToExecute.Name, NodeToExecute.Outputs.Select(x => x.TagName)))
 				{
@@ -705,13 +1013,13 @@ namespace AutomationTool
 
 			// Execute them in order
 			int NodeIdx = 0;
-			foreach(Node NodeToExecute in NodesToExecute)
+			foreach(BgNode NodeToExecute in NodesToExecute)
 			{
 				LogInformation("****** [{0}/{1}] {2}", ++NodeIdx, NodesToExecute.Length, NodeToExecute.Name);
 				if(!Storage.IsComplete(NodeToExecute.Name))
 				{
 					LogInformation("");
-					if(!BuildNode(Job, Graph, NodeToExecute, Storage, bWithBanner: false))
+					if(!BuildNode(Job, Graph, NodeToExecute, TaskInfoToTask, Storage, bWithBanner: false))
 					{
 						return false;
 					} 
@@ -727,10 +1035,11 @@ namespace AutomationTool
 		/// <param name="Job">Information about the current job</param>
 		/// <param name="Graph">The graph to which the node belongs. Used to determine which outputs need to be transferred to temp storage.</param>
 		/// <param name="Node">The node to build</param>
+		/// <param name="TaskInfoToTask">Map from task info instance to task instance</param>
 		/// <param name="Storage">The temp storage backend which stores the shared state</param>
 		/// <param name="bWithBanner">Whether to write a banner before and after this node's log output</param>
 		/// <returns>True if the node built successfully, false otherwise.</returns>
-		bool BuildNode(JobContext Job, Graph Graph, Node Node, TempStorage Storage, bool bWithBanner)
+		bool BuildNode(JobContext Job, BgScript Graph, BgNode Node, Dictionary<BgTask, CustomTask> TaskInfoToTask, TempStorage Storage, bool bWithBanner)
 		{
 			DirectoryReference RootDir = new DirectoryReference(CommandUtils.CmdEnv.LocalRoot);
 
@@ -739,7 +1048,7 @@ namespace AutomationTool
 
 			// Read all the input tags for this node, and build a list of referenced input storage blocks
 			HashSet<TempStorageBlock> InputStorageBlocks = new HashSet<TempStorageBlock>();
-			foreach(NodeOutput Input in Node.Inputs)
+			foreach(BgNodeOutput Input in Node.Inputs)
 			{
 				TempStorageFileList FileList = Storage.ReadFileList(Input.ProducingNode.Name, Input.TagName);
 				TagNameToFileSet[Input.TagName] = FileList.ToFileSet(RootDir);
@@ -776,7 +1085,7 @@ namespace AutomationTool
 			}
 
 			// Add placeholder outputs for the current node
-			foreach(NodeOutput Output in Node.Outputs)
+			foreach(BgNodeOutput Output in Node.Outputs)
 			{
 				TagNameToFileSet.Add(Output.TagName, new HashSet<FileReference>());
 			}
@@ -787,7 +1096,7 @@ namespace AutomationTool
 				Console.WriteLine();
 				CommandUtils.LogInformation("========== Starting: {0} ==========", Node.Name);
 			}
-			if(!Node.Build(Job, TagNameToFileSet))
+			if(!ExecuteTasks(Node, Job, TaskInfoToTask, TagNameToFileSet))
 			{
 				return false;
 			}
@@ -814,14 +1123,14 @@ namespace AutomationTool
 
 			// Determine all the output files which are required to be copied to temp storage (because they're referenced by nodes in another agent)
 			HashSet<FileReference> ReferencedOutputFiles = new HashSet<FileReference>();
-			foreach(Agent Agent in Graph.Agents)
+			foreach(BgAgent Agent in Graph.Agents)
 			{
 				bool bSameAgent = Agent.Nodes.Contains(Node);
-				foreach(Node OtherNode in Agent.Nodes)
+				foreach(BgNode OtherNode in Agent.Nodes)
 				{
-					if(!bSameAgent || Node.ControllingTrigger != OtherNode.ControllingTrigger)
+					if(!bSameAgent)
 					{
-						foreach(NodeOutput Input in OtherNode.Inputs.Where(x => x.ProducingNode == Node))
+						foreach(BgNodeOutput Input in OtherNode.Inputs.Where(x => x.ProducingNode == Node))
 						{
 							ReferencedOutputFiles.UnionWith(TagNameToFileSet[Input.TagName]);
 						}
@@ -831,7 +1140,7 @@ namespace AutomationTool
 
 			// Find a block name for all new outputs
 			Dictionary<FileReference, string> FileToOutputName = new Dictionary<FileReference, string>();
-			foreach(NodeOutput Output in Node.Outputs)
+			foreach(BgNodeOutput Output in Node.Outputs)
 			{
 				HashSet<FileReference> Files = TagNameToFileSet[Output.TagName]; 
 				foreach(FileReference File in Files)
@@ -888,7 +1197,7 @@ namespace AutomationTool
 				}
 
 				// Publish all the output tags
-				foreach (NodeOutput Output in Node.Outputs)
+				foreach (BgNodeOutput Output in Node.Outputs)
 				{
 					HashSet<FileReference> Files = TagNameToFileSet[Output.TagName];
 
@@ -911,12 +1220,84 @@ namespace AutomationTool
 			return true;
 		}
 
+		/// <summary>
+		/// Build all the tasks for this node
+		/// </summary>
+		/// <param name="Job">Information about the current job</param>
+		/// <param name="TaskInfoToTask">Map from TaskInfo to Task object</param>
+		/// <param name="TagNameToFileSet">Mapping from tag names to the set of files they include. Should be set to contain the node inputs on entry.</param>
+		/// <returns>Whether the task succeeded or not. Exiting with an exception will be caught and treated as a failure.</returns>
+		bool ExecuteTasks(BgNode Node, JobContext Job, Dictionary<BgTask, CustomTask> TaskInfoToTask, Dictionary<string, HashSet<FileReference>> TagNameToFileSet)
+		{
+			List<CustomTask> Tasks = Node.TaskInfos.ConvertAll(x => TaskInfoToTask[x]);
+
+			// Run each of the tasks in order
+			HashSet<FileReference> BuildProducts = TagNameToFileSet[Node.DefaultOutput.TagName];
+			for (int Idx = 0; Idx < Tasks.Count; Idx++)
+			{
+				using (IScope Scope = GlobalTracer.Instance.BuildSpan("Task").WithTag("resource", Tasks[Idx].GetTraceName()).StartActive())
+				{
+					ITaskExecutor Executor = Tasks[Idx].GetExecutor();
+					if (Executor == null)
+					{
+						// Execute this task directly
+						try
+						{
+							Tasks[Idx].GetTraceMetadata(Scope.Span, "");
+							Tasks[Idx].Execute(Job, BuildProducts, TagNameToFileSet);
+						}
+						catch (Exception Ex)
+						{
+							ExceptionUtils.AddContext(Ex, "while executing task {0}", Tasks[Idx].GetTraceString());
+							if (Tasks[Idx].SourceLocation != null)
+							{
+								ExceptionUtils.AddContext(Ex, "at {0}({1})", Tasks[Idx].SourceLocation.Item1, Tasks[Idx].SourceLocation.Item2);
+							}
+							throw;
+						}
+					}
+					else
+					{
+						Tasks[Idx].GetTraceMetadata(Scope.Span, "1.");
+
+						// The task has a custom executor, which may be able to execute several tasks simultaneously. Try to add the following tasks.
+						int FirstIdx = Idx;
+						while (Idx + 1 < Tasks.Count && Executor.Add(Tasks[Idx + 1]))
+						{
+							Idx++;
+							Tasks[Idx].GetTraceMetadata(Scope.Span, string.Format("{0}.", 1 + Idx - FirstIdx));
+						}
+						try
+						{
+							Executor.Execute(Job, BuildProducts, TagNameToFileSet);
+						}
+						catch (Exception Ex)
+						{
+							for (int TaskIdx = FirstIdx; TaskIdx <= Idx; TaskIdx++)
+							{
+								ExceptionUtils.AddContext(Ex, "while executing {0}", Tasks[TaskIdx].GetTraceString());
+							}
+							if (Tasks[FirstIdx].SourceLocation != null)
+							{
+								ExceptionUtils.AddContext(Ex, "at {0}({1})", Tasks[FirstIdx].SourceLocation.Item1, Tasks[FirstIdx].SourceLocation.Item2);
+							}
+							throw;
+						}
+					}
+				}
+			}
+
+			// Remove anything that doesn't exist, since these files weren't explicitly tagged
+			BuildProducts.RemoveWhere(x => !FileReference.Exists(x));
+			return true;
+		}
+
 		/// <summary>
 		/// Generate HTML documentation for all the tasks
 		/// </summary>
 		/// <param name="NameToTask">Map of task name to implementation</param>
 		/// <param name="OutputFile">Output file</param>
-		static void WriteDocumentation(Dictionary<string, ScriptTask> NameToTask, FileReference OutputFile)
+		static void WriteDocumentation(Dictionary<string, ScriptTaskBinding> NameToTask, FileReference OutputFile)
 		{
 			// Find all the assemblies containing tasks
 			Assembly[] TaskAssemblies = NameToTask.Values.Select(x => x.ParametersClass.Assembly).Distinct().ToArray();
@@ -972,7 +1353,7 @@ namespace AutomationTool
 		/// <param name="NameToTask">Map of name to script task</param>
 		/// <param name="MemberNameToElement">Map of field name to XML documenation element</param>
 		/// <param name="OutputFile">The output file to write to</param>
-		static void WriteDocumentationUDN(Dictionary<string, ScriptTask> NameToTask, Dictionary<string, XmlElement> MemberNameToElement, FileReference OutputFile)
+		static void WriteDocumentationUDN(Dictionary<string, ScriptTaskBinding> NameToTask, Dictionary<string, XmlElement> MemberNameToElement, FileReference OutputFile)
 		{
 			using (StreamWriter Writer = new StreamWriter(OutputFile.FullName))
 			{
@@ -986,7 +1367,7 @@ namespace AutomationTool
 				foreach (string TaskName in NameToTask.Keys.OrderBy(x => x))
 				{
 					// Get the task object
-					ScriptTask Task = NameToTask[TaskName];
+					ScriptTaskBinding Task = NameToTask[TaskName];
 
 					// Get the documentation for this task
 					XmlElement TaskElement;
@@ -1003,7 +1384,7 @@ namespace AutomationTool
 						foreach (string ParameterName in Task.NameToParameter.Keys)
 						{
 							// Get the parameter data
-							ScriptTaskParameter Parameter = Task.NameToParameter[ParameterName];
+							ScriptTaskParameterBinding Parameter = Task.NameToParameter[ParameterName];
 
 							// Get the documentation for this parameter
 							XmlElement ParameterElement;
@@ -1069,7 +1450,7 @@ namespace AutomationTool
 		/// <param name="NameToTask">Map of name to script task</param>
 		/// <param name="MemberNameToElement">Map of field name to XML documenation element</param>
 		/// <param name="OutputFile">The output file to write to</param>
-		static void WriteDocumentationHTML(Dictionary<string, ScriptTask> NameToTask, Dictionary<string, XmlElement> MemberNameToElement, FileReference OutputFile)
+		static void WriteDocumentationHTML(Dictionary<string, ScriptTaskBinding> NameToTask, Dictionary<string, XmlElement> MemberNameToElement, FileReference OutputFile)
 		{
 			LogInformation("Writing {0}...", OutputFile);
 			using (StreamWriter Writer = new StreamWriter(OutputFile.FullName))
@@ -1086,7 +1467,7 @@ namespace AutomationTool
 				foreach(string TaskName in NameToTask.Keys.OrderBy(x => x))
 				{
 					// Get the task object
-					ScriptTask Task = NameToTask[TaskName];
+					ScriptTaskBinding Task = NameToTask[TaskName];
 
 					// Get the documentation for this task
 					XmlElement TaskElement;
@@ -1109,7 +1490,7 @@ namespace AutomationTool
 						foreach(string ParameterName in Task.NameToParameter.Keys)
 						{
 							// Get the parameter data
-							ScriptTaskParameter Parameter = Task.NameToParameter[ParameterName];
+							ScriptTaskParameterBinding Parameter = Task.NameToParameter[ParameterName];
 
 							// Get the documentation for this parameter
 							XmlElement ParameterElement;
diff --git a/Engine/Source/Programs/AutomationTool/BuildGraph/Condition.cs b/Engine/Source/Programs/AutomationTool/BuildGraph/Condition.cs
deleted file mode 100644
diff --git a/Engine/Source/Programs/AutomationTool/BuildGraph/Graph.cs b/Engine/Source/Programs/AutomationTool/BuildGraph/Graph.cs
deleted file mode 100644
diff --git a/Engine/Source/Programs/AutomationTool/BuildGraph/Label.cs b/Engine/Source/Programs/AutomationTool/BuildGraph/Label.cs
deleted file mode 100644
diff --git a/Engine/Source/Programs/AutomationTool/BuildGraph/ManualTrigger.cs b/Engine/Source/Programs/AutomationTool/BuildGraph/ManualTrigger.cs
deleted file mode 100644
diff --git a/Engine/Source/Programs/AutomationTool/BuildGraph/Node.cs b/Engine/Source/Programs/AutomationTool/BuildGraph/Node.cs
deleted file mode 100644
diff --git a/Engine/Source/Programs/AutomationTool/BuildGraph/Report.cs b/Engine/Source/Programs/AutomationTool/BuildGraph/Report.cs
deleted file mode 100644
diff --git a/Engine/Source/Programs/AutomationTool/BuildGraph/Schema.cs b/Engine/Source/Programs/AutomationTool/BuildGraph/Schema.cs
deleted file mode 100644
diff --git a/Engine/Source/Programs/AutomationTool/BuildGraph/Script.cs b/Engine/Source/Programs/AutomationTool/BuildGraph/Script.cs
deleted file mode 100644
diff --git a/Engine/Source/Programs/AutomationTool/BuildGraph/Task.cs b/Engine/Source/Programs/AutomationTool/BuildGraph/Task.cs
--- a/Engine/Source/Programs/AutomationTool/BuildGraph/Task.cs
+++ b/Engine/Source/Programs/AutomationTool/BuildGraph/Task.cs
@@ -1,5 +1,6 @@
 // Copyright Epic Games, Inc. All Rights Reserved.
 
+using EpicGames.BuildGraph;
 using System;
 using System.Collections.Generic;
 using System.IO;
@@ -15,27 +16,6 @@ using UnrealBuildTool;
 
 namespace AutomationTool
 {
-	/// <summary>
-	/// Specifies validation that should be performed on a task parameter.
-	/// </summary>
-	public enum TaskParameterValidationType
-	{
-		/// <summary>
-		/// Allow any valid values for the field type.
-		/// </summary>
-		Default,
-
-		/// <summary>
-		/// A list of tag names separated by semicolons
-		/// </summary>
-		TagList,
-
-		/// <summary>
-		/// A file specification, which may contain tags and wildcards.
-		/// </summary>
-		FileSpec,
-	}
-
 	/// <summary>
 	/// Attribute to mark parameters to a task, which should be read as XML attributes from the script file.
 	/// </summary>
@@ -120,7 +100,7 @@ namespace AutomationTool
 		/// <summary>
 		/// Line number in a source file that this task was declared. Optional; used for log messages.
 		/// </summary>
-		public Tuple<FileReference, int> SourceLocation
+		public Tuple<string, int> SourceLocation
 		{
 			get;
 			set;
@@ -208,7 +188,7 @@ namespace AutomationTool
 		/// <param name="Prefix">Prefix for metadata entries</param>
 		public virtual void GetTraceMetadata(ITraceSpan Span, string Prefix)
 		{
-			Span.AddMetadata(Prefix + "source.file", SourceLocation.Item1.MakeRelativeTo(Unreal.RootDirectory));
+			Span.AddMetadata(Prefix + "source.file", SourceLocation.Item1);
 			Span.AddMetadata(Prefix + "source.line", SourceLocation.Item2.ToString());
 		}
 		
@@ -219,7 +199,7 @@ namespace AutomationTool
 		/// <param name="Prefix">Prefix for metadata entries</param>
 		public virtual void GetTraceMetadata(ISpan Span, string Prefix)
 		{
-			Span.SetTag(Prefix + "source.file", SourceLocation.Item1.MakeRelativeTo(Unreal.RootDirectory));
+			Span.SetTag(Prefix + "source.file", SourceLocation.Item1);
 			Span.SetTag(Prefix + "source.line", SourceLocation.Item2);
 		}
 
diff --git a/Engine/Source/Programs/AutomationTool/BuildGraph/Tasks/CommandletTask.cs b/Engine/Source/Programs/AutomationTool/BuildGraph/Tasks/CommandletTask.cs
--- a/Engine/Source/Programs/AutomationTool/BuildGraph/Tasks/CommandletTask.cs
+++ b/Engine/Source/Programs/AutomationTool/BuildGraph/Tasks/CommandletTask.cs
@@ -1,5 +1,6 @@
 // Copyright Epic Games, Inc. All Rights Reserved.
 
+using EpicGames.BuildGraph;
 using System;
 using System.Collections.Generic;
 using System.IO;
diff --git a/Engine/Source/Programs/AutomationTool/BuildGraph/Tasks/CompileTask.cs b/Engine/Source/Programs/AutomationTool/BuildGraph/Tasks/CompileTask.cs
--- a/Engine/Source/Programs/AutomationTool/BuildGraph/Tasks/CompileTask.cs
+++ b/Engine/Source/Programs/AutomationTool/BuildGraph/Tasks/CompileTask.cs
@@ -1,5 +1,6 @@
 // Copyright Epic Games, Inc. All Rights Reserved.
 
+using EpicGames.BuildGraph;
 using System;
 using System.Collections.Generic;
 using System.Linq;
diff --git a/Engine/Source/Programs/AutomationTool/BuildGraph/Tasks/CookTask.cs b/Engine/Source/Programs/AutomationTool/BuildGraph/Tasks/CookTask.cs
--- a/Engine/Source/Programs/AutomationTool/BuildGraph/Tasks/CookTask.cs
+++ b/Engine/Source/Programs/AutomationTool/BuildGraph/Tasks/CookTask.cs
@@ -1,6 +1,7 @@
 // Copyright Epic Games, Inc. All Rights Reserved.
 
 using AutomationTool;
+using EpicGames.BuildGraph;
 using System;
 using System.Collections.Generic;
 using System.IO;
diff --git a/Engine/Source/Programs/AutomationTool/BuildGraph/Tasks/CopyTask.cs b/Engine/Source/Programs/AutomationTool/BuildGraph/Tasks/CopyTask.cs
--- a/Engine/Source/Programs/AutomationTool/BuildGraph/Tasks/CopyTask.cs
+++ b/Engine/Source/Programs/AutomationTool/BuildGraph/Tasks/CopyTask.cs
@@ -1,6 +1,7 @@
 // Copyright Epic Games, Inc. All Rights Reserved.
 
 using AutomationTool;
+using EpicGames.BuildGraph;
 using System;
 using System.Collections.Generic;
 using System.Linq;
diff --git a/Engine/Source/Programs/AutomationTool/BuildGraph/Tasks/CsCompileTask.cs b/Engine/Source/Programs/AutomationTool/BuildGraph/Tasks/CsCompileTask.cs
--- a/Engine/Source/Programs/AutomationTool/BuildGraph/Tasks/CsCompileTask.cs
+++ b/Engine/Source/Programs/AutomationTool/BuildGraph/Tasks/CsCompileTask.cs
@@ -1,5 +1,6 @@
 ﻿// Copyright Epic Games, Inc. All Rights Reserved.
 
+using EpicGames.BuildGraph;
 using System;
 using System.Collections.Generic;
 using System.Linq;
diff --git a/Engine/Source/Programs/AutomationTool/BuildGraph/Tasks/DeleteTask.cs b/Engine/Source/Programs/AutomationTool/BuildGraph/Tasks/DeleteTask.cs
--- a/Engine/Source/Programs/AutomationTool/BuildGraph/Tasks/DeleteTask.cs
+++ b/Engine/Source/Programs/AutomationTool/BuildGraph/Tasks/DeleteTask.cs
@@ -1,6 +1,7 @@
 // Copyright Epic Games, Inc. All Rights Reserved.
 
 using AutomationTool;
+using EpicGames.BuildGraph;
 using System;
 using System.Collections.Generic;
 using System.IO;
diff --git a/Engine/Source/Programs/AutomationTool/BuildGraph/Tasks/LogTask.cs b/Engine/Source/Programs/AutomationTool/BuildGraph/Tasks/LogTask.cs
--- a/Engine/Source/Programs/AutomationTool/BuildGraph/Tasks/LogTask.cs
+++ b/Engine/Source/Programs/AutomationTool/BuildGraph/Tasks/LogTask.cs
@@ -1,6 +1,7 @@
 // Copyright Epic Games, Inc. All Rights Reserved.
 
 using AutomationTool;
+using EpicGames.BuildGraph;
 using System;
 using System.Collections.Generic;
 using System.IO;
diff --git a/Engine/Source/Programs/AutomationTool/BuildGraph/Tasks/ModifyConfigTask.cs b/Engine/Source/Programs/AutomationTool/BuildGraph/Tasks/ModifyConfigTask.cs
--- a/Engine/Source/Programs/AutomationTool/BuildGraph/Tasks/ModifyConfigTask.cs
+++ b/Engine/Source/Programs/AutomationTool/BuildGraph/Tasks/ModifyConfigTask.cs
@@ -1,6 +1,7 @@
 // Copyright Epic Games, Inc. All Rights Reserved.
 
 using AutomationTool;
+using EpicGames.BuildGraph;
 using System;
 using System.Collections.Generic;
 using System.Linq;
diff --git a/Engine/Source/Programs/AutomationTool/BuildGraph/Tasks/ModifyJsonValue.cs b/Engine/Source/Programs/AutomationTool/BuildGraph/Tasks/ModifyJsonValue.cs
--- a/Engine/Source/Programs/AutomationTool/BuildGraph/Tasks/ModifyJsonValue.cs
+++ b/Engine/Source/Programs/AutomationTool/BuildGraph/Tasks/ModifyJsonValue.cs
@@ -1,6 +1,7 @@
 // Copyright Epic Games, Inc. All Rights Reserved.
 
 using AutomationTool;
+using EpicGames.BuildGraph;
 using System;
 using System.Collections.Generic;
 using System.Linq;
diff --git a/Engine/Source/Programs/AutomationTool/BuildGraph/Tasks/MoveTask.cs b/Engine/Source/Programs/AutomationTool/BuildGraph/Tasks/MoveTask.cs
--- a/Engine/Source/Programs/AutomationTool/BuildGraph/Tasks/MoveTask.cs
+++ b/Engine/Source/Programs/AutomationTool/BuildGraph/Tasks/MoveTask.cs
@@ -1,6 +1,7 @@
 // Copyright Epic Games, Inc. All Rights Reserved.
 
 using AutomationTool;
+using EpicGames.BuildGraph;
 using System;
 using System.Collections.Generic;
 using System.Linq;
diff --git a/Engine/Source/Programs/AutomationTool/BuildGraph/Tasks/PakFileTask.cs b/Engine/Source/Programs/AutomationTool/BuildGraph/Tasks/PakFileTask.cs
--- a/Engine/Source/Programs/AutomationTool/BuildGraph/Tasks/PakFileTask.cs
+++ b/Engine/Source/Programs/AutomationTool/BuildGraph/Tasks/PakFileTask.cs
@@ -1,6 +1,7 @@
 // Copyright Epic Games, Inc. All Rights Reserved.
 
 using AutomationTool;
+using EpicGames.BuildGraph;
 using System;
 using System.Collections.Generic;
 using System.IO;
diff --git a/Engine/Source/Programs/AutomationTool/BuildGraph/Tasks/RenameTask.cs b/Engine/Source/Programs/AutomationTool/BuildGraph/Tasks/RenameTask.cs
--- a/Engine/Source/Programs/AutomationTool/BuildGraph/Tasks/RenameTask.cs
+++ b/Engine/Source/Programs/AutomationTool/BuildGraph/Tasks/RenameTask.cs
@@ -1,6 +1,7 @@
 // Copyright Epic Games, Inc. All Rights Reserved.
 
 using AutomationTool;
+using EpicGames.BuildGraph;
 using System;
 using System.Collections.Generic;
 using System.IO;
diff --git a/Engine/Source/Programs/AutomationTool/BuildGraph/Tasks/SanitizeReceiptTask.cs b/Engine/Source/Programs/AutomationTool/BuildGraph/Tasks/SanitizeReceiptTask.cs
--- a/Engine/Source/Programs/AutomationTool/BuildGraph/Tasks/SanitizeReceiptTask.cs
+++ b/Engine/Source/Programs/AutomationTool/BuildGraph/Tasks/SanitizeReceiptTask.cs
@@ -1,5 +1,6 @@
 // Copyright Epic Games, Inc. All Rights Reserved.
 
+using EpicGames.BuildGraph;
 using System;
 using System.Collections.Generic;
 using System.Linq;
diff --git a/Engine/Source/Programs/AutomationTool/BuildGraph/Tasks/SetVersionTask.cs b/Engine/Source/Programs/AutomationTool/BuildGraph/Tasks/SetVersionTask.cs
--- a/Engine/Source/Programs/AutomationTool/BuildGraph/Tasks/SetVersionTask.cs
+++ b/Engine/Source/Programs/AutomationTool/BuildGraph/Tasks/SetVersionTask.cs
@@ -1,5 +1,6 @@
 // Copyright Epic Games, Inc. All Rights Reserved.
 
+using EpicGames.BuildGraph;
 using System;
 using System.Collections.Generic;
 using System.Linq;
diff --git a/Engine/Source/Programs/AutomationTool/BuildGraph/Tasks/SignTask.cs b/Engine/Source/Programs/AutomationTool/BuildGraph/Tasks/SignTask.cs
--- a/Engine/Source/Programs/AutomationTool/BuildGraph/Tasks/SignTask.cs
+++ b/Engine/Source/Programs/AutomationTool/BuildGraph/Tasks/SignTask.cs
@@ -1,6 +1,7 @@
 // Copyright Epic Games, Inc. All Rights Reserved.
 
 using AutomationTool;
+using EpicGames.BuildGraph;
 using System;
 using System.Collections.Generic;
 using System.Linq;
diff --git a/Engine/Source/Programs/AutomationTool/BuildGraph/Tasks/StageTask.cs b/Engine/Source/Programs/AutomationTool/BuildGraph/Tasks/StageTask.cs
--- a/Engine/Source/Programs/AutomationTool/BuildGraph/Tasks/StageTask.cs
+++ b/Engine/Source/Programs/AutomationTool/BuildGraph/Tasks/StageTask.cs
@@ -1,6 +1,7 @@
 // Copyright Epic Games, Inc. All Rights Reserved.
 
 using AutomationTool;
+using EpicGames.BuildGraph;
 using System;
 using System.Collections.Generic;
 using System.Linq;
diff --git a/Engine/Source/Programs/AutomationTool/BuildGraph/Tasks/StripTask.cs b/Engine/Source/Programs/AutomationTool/BuildGraph/Tasks/StripTask.cs
--- a/Engine/Source/Programs/AutomationTool/BuildGraph/Tasks/StripTask.cs
+++ b/Engine/Source/Programs/AutomationTool/BuildGraph/Tasks/StripTask.cs
@@ -1,6 +1,7 @@
 // Copyright Epic Games, Inc. All Rights Reserved.
 
 using AutomationTool;
+using EpicGames.BuildGraph;
 using System;
 using System.Collections.Generic;
 using System.Linq;
diff --git a/Engine/Source/Programs/AutomationTool/BuildGraph/Tasks/SubmitTask.cs b/Engine/Source/Programs/AutomationTool/BuildGraph/Tasks/SubmitTask.cs
--- a/Engine/Source/Programs/AutomationTool/BuildGraph/Tasks/SubmitTask.cs
+++ b/Engine/Source/Programs/AutomationTool/BuildGraph/Tasks/SubmitTask.cs
@@ -1,5 +1,6 @@
 // Copyright Epic Games, Inc. All Rights Reserved.
 
+using EpicGames.BuildGraph;
 using System;
 using System.Collections.Generic;
 using System.Linq;
diff --git a/Engine/Source/Programs/AutomationTool/BuildGraph/Tasks/TagReceiptTask.cs b/Engine/Source/Programs/AutomationTool/BuildGraph/Tasks/TagReceiptTask.cs
--- a/Engine/Source/Programs/AutomationTool/BuildGraph/Tasks/TagReceiptTask.cs
+++ b/Engine/Source/Programs/AutomationTool/BuildGraph/Tasks/TagReceiptTask.cs
@@ -1,5 +1,6 @@
 // Copyright Epic Games, Inc. All Rights Reserved.
 
+using EpicGames.BuildGraph;
 using System;
 using System.Collections.Generic;
 using System.IO;
diff --git a/Engine/Source/Programs/AutomationTool/BuildGraph/Tasks/TagTask.cs b/Engine/Source/Programs/AutomationTool/BuildGraph/Tasks/TagTask.cs
--- a/Engine/Source/Programs/AutomationTool/BuildGraph/Tasks/TagTask.cs
+++ b/Engine/Source/Programs/AutomationTool/BuildGraph/Tasks/TagTask.cs
@@ -1,5 +1,6 @@
 // Copyright Epic Games, Inc. All Rights Reserved.
 
+using EpicGames.BuildGraph;
 using System;
 using System.Collections.Generic;
 using System.IO;
diff --git a/Engine/Source/Programs/AutomationTool/BuildGraph/Tasks/UnzipTask.cs b/Engine/Source/Programs/AutomationTool/BuildGraph/Tasks/UnzipTask.cs
--- a/Engine/Source/Programs/AutomationTool/BuildGraph/Tasks/UnzipTask.cs
+++ b/Engine/Source/Programs/AutomationTool/BuildGraph/Tasks/UnzipTask.cs
@@ -1,6 +1,7 @@
 ﻿// Copyright Epic Games, Inc. All Rights Reserved.
 
 using AutomationTool;
+using EpicGames.BuildGraph;
 using System;
 using System.Collections.Generic;
 using System.Linq;
diff --git a/Engine/Source/Programs/AutomationTool/BuildGraph/Tasks/UploadTask.cs b/Engine/Source/Programs/AutomationTool/BuildGraph/Tasks/UploadTask.cs
--- a/Engine/Source/Programs/AutomationTool/BuildGraph/Tasks/UploadTask.cs
+++ b/Engine/Source/Programs/AutomationTool/BuildGraph/Tasks/UploadTask.cs
@@ -1,6 +1,7 @@
 // Copyright Epic Games, Inc. All Rights Reserved.
 
 using AutomationTool;
+using EpicGames.BuildGraph;
 using System;
 using System.Collections.Generic;
 using System.IO;
diff --git a/Engine/Source/Programs/AutomationTool/BuildGraph/Tasks/WriteTextFileTask.cs b/Engine/Source/Programs/AutomationTool/BuildGraph/Tasks/WriteTextFileTask.cs
--- a/Engine/Source/Programs/AutomationTool/BuildGraph/Tasks/WriteTextFileTask.cs
+++ b/Engine/Source/Programs/AutomationTool/BuildGraph/Tasks/WriteTextFileTask.cs
@@ -1,6 +1,7 @@
 // Copyright Epic Games, Inc. All Rights Reserved.
 
 using AutomationTool;
+using EpicGames.BuildGraph;
 using System;
 using System.Collections.Generic;
 using System.Linq;
diff --git a/Engine/Source/Programs/AutomationTool/BuildGraph/Tasks/ZipTask.cs b/Engine/Source/Programs/AutomationTool/BuildGraph/Tasks/ZipTask.cs
--- a/Engine/Source/Programs/AutomationTool/BuildGraph/Tasks/ZipTask.cs
+++ b/Engine/Source/Programs/AutomationTool/BuildGraph/Tasks/ZipTask.cs
@@ -1,6 +1,7 @@
 // Copyright Epic Games, Inc. All Rights Reserved.
 
 using AutomationTool;
+using EpicGames.BuildGraph;
 using System;
 using System.Collections.Generic;
 using System.Linq;
diff --git a/Engine/Source/Programs/Shared/EpicGames.BuildGraph/BgAgent.cs b/Engine/Source/Programs/Shared/EpicGames.BuildGraph/BgAgent.cs
new file mode 100644
--- /dev/null
+++ b/Engine/Source/Programs/Shared/EpicGames.BuildGraph/BgAgent.cs
@@ -0,0 +1,62 @@
+// Copyright Epic Games, Inc. All Rights Reserved.
+
+using System;
+using System.Collections.Generic;
+using System.Diagnostics;
+using System.Linq;
+using System.Text;
+using System.Threading.Tasks;
+using System.Xml;
+
+namespace EpicGames.BuildGraph
+{
+	/// <summary>
+	/// Stores a list of nodes which can be executed on a single agent
+	/// </summary>
+	[DebuggerDisplay("{Name}")]
+	public class BgAgent
+	{
+		/// <summary>
+		/// Name of this agent. Used for display purposes in a build system.
+		/// </summary>
+		public string Name;
+
+		/// <summary>
+		/// Array of valid agent types that these nodes may run on. When running in the build system, this determines the class of machine that should
+		/// be selected to run these nodes. The first defined agent type for this branch will be used.
+		/// </summary>
+		public string[] PossibleTypes;
+
+		/// <summary>
+		/// List of nodes in this agent group.
+		/// </summary>
+		public List<BgNode> Nodes = new List<BgNode>();
+
+		/// <summary>
+		/// Constructor
+		/// </summary>
+		/// <param name="InName">Name of this agent group</param>
+		/// <param name="InPossibleTypes">Array of valid agent types. See comment for AgentTypes member.</param>
+		public BgAgent(string InName, string[] InPossibleTypes)
+		{
+			Name = InName;
+			PossibleTypes = InPossibleTypes;
+		}
+
+		/// <summary>
+		/// Writes this agent group out to a file, filtering nodes by a controlling trigger
+		/// </summary>
+		/// <param name="Writer">The XML writer to output to</param>
+		public void Write(XmlWriter Writer)
+		{
+			Writer.WriteStartElement("Agent");
+			Writer.WriteAttributeString("Name", Name);
+			Writer.WriteAttributeString("Type", String.Join(";", PossibleTypes));
+			foreach (BgNode Node in Nodes)
+			{
+				Node.Write(Writer);
+			}
+			Writer.WriteEndElement();
+		}
+	}
+}
diff --git a/Engine/Source/Programs/Shared/EpicGames.BuildGraph/BgAggregate.cs b/Engine/Source/Programs/Shared/EpicGames.BuildGraph/BgAggregate.cs
new file mode 100644
--- /dev/null
+++ b/Engine/Source/Programs/Shared/EpicGames.BuildGraph/BgAggregate.cs
@@ -0,0 +1,44 @@
+// Copyright Epic Games, Inc. All Rights Reserved.
+
+using System;
+using System.Collections.Generic;
+using System.Linq;
+using System.Text;
+using System.Threading.Tasks;
+
+namespace EpicGames.BuildGraph
+{
+	/// <summary>
+	/// Defines a agggregate within a graph, which give the combined status of one or more job steps, and allow building several steps at once.
+	/// </summary>
+	public class BgAggregate
+	{
+		/// <summary>
+		/// Name of this badge
+		/// </summary>
+		public readonly string Name;
+
+		/// <summary>
+		/// Set of nodes that must be run for this label to be shown.
+		/// </summary>
+		public HashSet<BgNode> RequiredNodes = new HashSet<BgNode>();
+
+		/// <summary>
+		/// Constructor
+		/// </summary>
+		/// <param name="InName">Name of this aggregate</param>
+		public BgAggregate(string InName)
+		{
+			Name = InName;
+		}
+
+		/// <summary>
+		/// Get the name of this label
+		/// </summary>
+		/// <returns>The name of this label</returns>
+		public override string ToString()
+		{
+			return Name;
+		}
+	}
+}
diff --git a/Engine/Source/Programs/Shared/EpicGames.BuildGraph/BgBadge.cs b/Engine/Source/Programs/Shared/EpicGames.BuildGraph/BgBadge.cs
new file mode 100644
--- /dev/null
+++ b/Engine/Source/Programs/Shared/EpicGames.BuildGraph/BgBadge.cs
@@ -0,0 +1,59 @@
+// Copyright Epic Games, Inc. All Rights Reserved.
+
+using System;
+using System.Collections.Generic;
+using System.Diagnostics;
+using System.Linq;
+using System.Text;
+using System.Threading.Tasks;
+
+namespace EpicGames.BuildGraph
+{
+	/// <summary>
+	/// Defines a badge which gives an at-a-glance summary of part of the build, and can be displayed in UGS
+	/// </summary>
+	public class BgBadge
+	{
+		/// <summary>
+		/// Name of this badge
+		/// </summary>
+		public readonly string Name;
+
+		/// <summary>
+		/// Depot path to the project that this badge applies to. Used for filtering in UGS.
+		/// </summary>
+		public readonly string Project;
+
+		/// <summary>
+		/// The changelist to post the badge for
+		/// </summary>
+		public readonly int Change;
+
+		/// <summary>
+		/// Set of nodes that this badge reports the status of
+		/// </summary>
+		public HashSet<BgNode> Nodes = new HashSet<BgNode>();
+
+		/// <summary>
+		/// Constructor
+		/// </summary>
+		/// <param name="InName">Name of this report</param>
+		/// <param name="InProject">Depot path to the project that this badge applies to</param>
+		/// <param name="InChange">The changelist to post the badge for</param>
+		public BgBadge(string InName, string InProject, int InChange)
+		{
+			Name = InName;
+			Project = InProject;
+			Change = InChange;
+		}
+
+		/// <summary>
+		/// Get the name of this badge
+		/// </summary>
+		/// <returns>The name of this badge</returns>
+		public override string ToString()
+		{
+			return Name;
+		}
+	}
+}
diff --git a/Engine/Source/Programs/Shared/EpicGames.BuildGraph/BgCondition.cs b/Engine/Source/Programs/Shared/EpicGames.BuildGraph/BgCondition.cs
new file mode 100644
--- /dev/null
+++ b/Engine/Source/Programs/Shared/EpicGames.BuildGraph/BgCondition.cs
@@ -0,0 +1,506 @@
+// Copyright Epic Games, Inc. All Rights Reserved.
+
+using System;
+using System.Collections.Generic;
+using System.IO;
+using System.Linq;
+using System.Text;
+using System.Threading.Tasks;
+
+namespace EpicGames.BuildGraph
+{
+	/// <summary>
+	/// Exception class thrown due to type and syntax errors in condition expressions
+	/// </summary>
+	class BgConditionException : Exception
+	{
+		/// <summary>
+		/// Constructor; formats the exception message with the given String.Format() style parameters.
+		/// </summary>
+		/// <param name="Format">Formatting string, in String.Format syntax</param>
+		/// <param name="Args">Optional arguments for the string</param>
+		public BgConditionException(string Format, params object[] Args) : base(String.Format(Format, Args))
+		{
+		}
+	}
+
+	/// <summary>
+	/// Class to evaluate condition expressions in build scripts, following this grammar:
+	/// 
+	///		or-expression   ::= and-expression
+	///		                  | or-expression "Or" and-expression;
+	///		    
+	///		and-expression  ::= comparison
+	///		                  | and-expression "And" comparison;
+	///		                  
+	///		comparison      ::= scalar
+	///		                  | scalar "==" scalar
+	///		                  | scalar "!=" scalar
+	///		                  | scalar "&lt;" scalar
+	///		                  | scalar "&lt;=" scalar;
+	///		                  | scalar "&gt;" scalar
+	///		                  | scalar "&gt;=" scalar;
+	///		                  
+	///     scalar          ::= "(" or-expression ")"
+	///                       | "!" scalar
+	///                       | "Exists" "(" scalar ")"
+	///                       | "HasTrailingSlash" "(" scalar ")"
+	///                       | string
+	///                       | identifier;
+	///                       
+	///     string          ::= any sequence of characters terminated by single quotes (') or double quotes ("). Not escaped.
+	///     identifier      ::= any sequence of letters, digits, or underscore characters.
+	///     
+	/// The type of each subexpression is always a scalar, which are converted to expression-specific types (eg. booleans, integers) as required.
+	/// Scalar values are case-insensitive strings. The identifier 'true' and the strings "true" and "True" are all identical scalars.
+	/// </summary>
+	public static class BgCondition
+	{
+		/// <summary>
+		/// Sentinel added to the end of a sequence of tokens.
+		/// </summary>
+		const string EndToken = "<EOF>";
+
+		/// <summary>
+		/// Evaluates the given string as a condition. Throws a ConditionException on a type or syntax error.
+		/// </summary>
+		/// <param name="Text">The condition text</param>
+		/// <param name="Context">Context for evaluating the expression</param>
+		/// <returns>The result of evaluating the condition</returns>
+		public static bool Evaluate(string Text, IBgScriptReaderContext Context)
+		{
+			List<string> Tokens = new List<string>();
+			Tokenize(Text, Tokens);
+
+			bool bResult = true;
+			if(Tokens.Count > 1)
+			{
+				int Idx = 0;
+				string Result = EvaluateOr(Tokens, ref Idx, Context);
+				if(Tokens[Idx] != EndToken)
+				{
+					throw new BgConditionException("Garbage after expression: {0}", String.Join("", Tokens.Skip(Idx)));
+				}
+				bResult = CoerceToBool(Result);
+			}
+			return bResult;
+		}
+
+		/// <summary>
+		/// Evaluates an "or-expression" production.
+		/// </summary>
+		/// <param name="Tokens">List of tokens in the expression</param>
+		/// <param name="Idx">Current position in the token stream. Will be incremented as tokens are consumed.</param>
+		/// <param name="Context">Context for evaluating the expression</param>
+		/// <returns>A scalar representing the result of evaluating the expression.</returns>
+		static string EvaluateOr(List<string> Tokens, ref int Idx, IBgScriptReaderContext Context)
+		{
+			// <Condition> Or <Condition> Or...
+			string Result = EvaluateAnd(Tokens, ref Idx, Context);
+			while(String.Compare(Tokens[Idx], "Or", true) == 0)
+			{
+				// Evaluate this condition. We use a binary OR here, because we want to parse everything rather than short-circuit it.
+				Idx++;
+				string Lhs = Result;
+				string Rhs = EvaluateAnd(Tokens, ref Idx, Context);
+				Result = (CoerceToBool(Lhs) | CoerceToBool(Rhs))? "true" : "false";
+			}
+			return Result;
+		}
+
+		/// <summary>
+		/// Evaluates an "and-expression" production.
+		/// </summary>
+		/// <param name="Tokens">List of tokens in the expression</param>
+		/// <param name="Idx">Current position in the token stream. Will be incremented as tokens are consumed.</param>
+		/// <param name="Context">Context for evaluating the expression</param>
+		/// <returns>A scalar representing the result of evaluating the expression.</returns>
+		static string EvaluateAnd(List<string> Tokens, ref int Idx, IBgScriptReaderContext Context)
+		{
+			// <Condition> And <Condition> And...
+			string Result = EvaluateComparison(Tokens, ref Idx, Context);
+			while(String.Compare(Tokens[Idx], "And", true) == 0)
+			{
+				// Evaluate this condition. We use a binary AND here, because we want to parse everything rather than short-circuit it.
+				Idx++;
+				string Lhs = Result;
+				string Rhs = EvaluateComparison(Tokens, ref Idx, Context);
+				Result = (CoerceToBool(Lhs) & CoerceToBool(Rhs))? "true" : "false";
+			}
+			return Result;
+		}
+
+		/// <summary>
+		/// Evaluates a "comparison" production.
+		/// </summary>
+		/// <param name="Tokens">List of tokens in the expression</param>
+		/// <param name="Idx">Current position in the token stream. Will be incremented as tokens are consumed.</param>
+		/// <param name="Context">Context for evaluating the expression</param>
+		/// <returns>The result of evaluating the expression</returns>
+		static string EvaluateComparison(List<string> Tokens, ref int Idx, IBgScriptReaderContext Context)
+		{
+			// scalar
+			// scalar == scalar
+			// scalar != scalar
+			// scalar < scalar
+			// scalar <= scalar
+			// scalar > scalar
+			// scalar >= scalar
+
+			string Result = EvaluateScalar(Tokens, ref Idx, Context);
+			if(Tokens[Idx] == "==")
+			{
+				// Compare two scalars for equality
+				Idx++;
+				string Lhs = Result;
+				string Rhs = EvaluateScalar(Tokens, ref Idx, Context);
+				Result = (String.Compare(Lhs, Rhs, true) == 0)? "true" : "false";
+			}
+			else if(Tokens[Idx] == "!=")
+			{
+				// Compare two scalars for inequality
+				Idx++;
+				string Lhs = Result;
+				string Rhs = EvaluateScalar(Tokens, ref Idx, Context);
+				Result = (String.Compare(Lhs, Rhs, true) != 0)? "true" : "false";
+			}
+			else if(Tokens[Idx] == "<")
+			{
+				// Compares whether the first integer is less than the second
+				Idx++;
+				int Lhs = CoerceToInteger(Result);
+				int Rhs = CoerceToInteger(EvaluateScalar(Tokens, ref Idx, Context));
+				Result = (Lhs < Rhs)? "true" : "false";
+			}
+			else if(Tokens[Idx] == "<=")
+			{
+				// Compares whether the first integer is less than the second
+				Idx++;
+				int Lhs = CoerceToInteger(Result);
+				int Rhs = CoerceToInteger(EvaluateScalar(Tokens, ref Idx, Context));
+				Result = (Lhs <= Rhs)? "true" : "false";
+			}
+			else if(Tokens[Idx] == ">")
+			{
+				// Compares whether the first integer is less than the second
+				Idx++;
+				int Lhs = CoerceToInteger(Result);
+				int Rhs = CoerceToInteger(EvaluateScalar(Tokens, ref Idx, Context));
+				Result = (Lhs > Rhs)? "true" : "false";
+			}
+			else if(Tokens[Idx] == ">=")
+			{
+				// Compares whether the first integer is less than the second
+				Idx++;
+				int Lhs = CoerceToInteger(Result);
+				int Rhs = CoerceToInteger(EvaluateScalar(Tokens, ref Idx, Context));
+				Result = (Lhs >= Rhs)? "true" : "false";
+			}
+			return Result;
+		}
+
+		/// <summary>
+		/// Evaluates arguments from a token string. Arguments are all comma-separated tokens until a closing ) is encountered
+		/// </summary>
+		/// <param name="Tokens">List of tokens in the expression</param>
+		/// <param name="Idx">Current position in the token stream. Will be incremented as tokens are consumed.</param>
+		/// <returns>The result of evaluating the expression</returns>
+		static IEnumerable<string> EvaluateArguments(List<string> Tokens, ref int Idx)
+		{
+			List<string> Arguments = new List<string>();
+
+			// skip opening bracket
+			if (Tokens[Idx++] != "(")
+			{
+				throw new BgConditionException("Expected '('");
+			}
+
+			bool DidCloseBracket = false;
+
+			while (Idx < Tokens.Count)
+			{
+				string NextToken = Tokens.ElementAt(Idx++);
+
+				if (NextToken == EndToken)
+				{
+					// ran out of items
+					break;
+				}
+				else if (NextToken == ")")
+				{
+					DidCloseBracket = true;
+					break;
+				}
+				else if (NextToken != ",")
+				{
+					if (NextToken.First() == '\'' && NextToken.Last() == '\'')
+					{
+						NextToken = NextToken.Substring(1, NextToken.Length - 2);
+					}
+					Arguments.Add(NextToken);
+				}
+			}
+
+			if (!DidCloseBracket)
+			{
+				throw new BgConditionException("Expected ')'");
+			}
+
+			return Arguments;
+		}
+
+		/// <summary>
+		/// Evaluates a "scalar" production.
+		/// </summary>
+		/// <param name="Tokens">List of tokens in the expression</param>
+		/// <param name="Idx">Current position in the token stream. Will be incremented as tokens are consumed.</param>
+		/// <param name="Context">Context for evaluating the expression</param>
+		/// <returns>The result of evaluating the expression</returns>
+		static string EvaluateScalar(List<string> Tokens, ref int Idx, IBgScriptReaderContext Context)
+		{
+			string Result;
+			if(Tokens[Idx] == "(")
+			{
+				// Subexpression
+				Idx++;
+				Result = EvaluateOr(Tokens, ref Idx, Context);
+				if(Tokens[Idx] != ")")
+				{
+					throw new BgConditionException("Expected ')'");
+				}
+				Idx++;
+			}
+			else if(Tokens[Idx] == "!")
+			{
+				// Logical not
+				Idx++;
+				string Rhs = EvaluateScalar(Tokens, ref Idx, Context);
+				Result = CoerceToBool(Rhs)? "false" : "true";
+			}
+			else if(String.Compare(Tokens[Idx], "Exists", true) == 0 && Tokens[Idx + 1] == "(")
+			{
+				// Check whether file or directory exists. Evaluate the argument as a subexpression.
+				Idx++;
+				string Argument = EvaluateScalar(Tokens, ref Idx, Context);
+				Result = Context.Exists(Argument)? "true" : "false";
+			}
+			else if(String.Compare(Tokens[Idx], "HasTrailingSlash", true) == 0 && Tokens[Idx + 1] == "(")
+			{
+				// Check whether the given string ends with a slash
+				Idx++;
+				string Argument = EvaluateScalar(Tokens, ref Idx, Context);
+				Result = (Argument.Length > 0 && (Argument[Argument.Length - 1] == Path.DirectorySeparatorChar || Argument[Argument.Length - 1] == Path.AltDirectorySeparatorChar))? "true" : "false";
+			}
+			else if (String.Compare(Tokens[Idx], "Contains", true) == 0 && Tokens[Idx + 1] == "(")
+			{
+				// Check a string contains a substring. If a separator is supplied the string is first split
+				Idx++;
+				IEnumerable<string> Arguments = EvaluateArguments(Tokens, ref Idx);
+
+				if (Arguments.Count() != 2)
+				{
+					throw new BgConditionException("Invalid argument count for 'Contains'. Expected (Haystack,Needle)");
+				}
+
+				Result = Contains(Arguments.ElementAt(0), Arguments.ElementAt(1)) ? "true" : "false";
+			}
+			else if (String.Compare(Tokens[Idx], "ContainsItem", true) == 0 && Tokens[Idx + 1] == "(")
+			{
+				// Check a string contains a substring. If a separator is supplied the string is first split
+				Idx++;
+				IEnumerable<string> Arguments = EvaluateArguments(Tokens, ref Idx);
+
+				if (Arguments.Count() != 3)
+				{
+					throw new BgConditionException("Invalid argument count for 'ContainsItem'. Expected (Haystack,Needle,HaystackSeparator)");
+				}
+
+				Result = ContainsItem(Arguments.ElementAt(0), Arguments.ElementAt(1), Arguments.ElementAt(2)) ? "true" : "false";
+			}
+			else
+			{
+				// Raw scalar. Remove quotes from strings, and allow literals and simple identifiers to pass through directly.
+				string Token = Tokens[Idx];
+				if(Token.Length >= 2 && (Token[0] == '\'' || Token[0] == '\"') && Token[Token.Length - 1] == Token[0])
+				{
+					Result = Token.Substring(1, Token.Length - 2);
+					Idx++;
+				}
+				else if(Char.IsLetterOrDigit(Token[0]) || Token[0] == '_')
+				{
+					Result = Token;
+					Idx++;
+				}
+				else
+				{
+					throw new BgConditionException("Token '{0}' is not a valid scalar", Token);
+				}
+			}
+			return Result;
+		}
+
+		/// <summary>
+		/// Checks whether Haystack contains "Needle". 
+		/// </summary>
+		/// <param name="Haystack">The string to search</param>
+		/// <param name="Needle">The string to search for</param>
+		/// <returns>True if the path exists, false otherwise.</returns>
+		static bool Contains(string Haystack, string Needle)
+		{
+			try
+			{
+				return Haystack.IndexOf(Needle, StringComparison.CurrentCultureIgnoreCase) >= 0;
+			}
+			catch
+			{
+				return false;
+			}
+		}
+
+		/// <summary>
+		/// Checks whether HaystackItems contains "Needle". 
+		/// </summary>
+		/// <param name="Haystack">The separated list of items to check</param>
+		/// <param name="Needle">The item to check for</param>
+		/// <param name="HaystackSeparator">The separator used in Haystack</param>
+		/// <returns>True if the path exists, false otherwise.</returns>
+		static bool ContainsItem(string Haystack, string Needle, string HaystackSeparator)
+		{
+			try
+			{
+				IEnumerable<string> HaystackItems = Haystack.Split(new string[] { HaystackSeparator }, StringSplitOptions.RemoveEmptyEntries);
+				return HaystackItems.Any(I => I.ToLower() == Needle.ToLower());
+			}
+			catch
+			{
+				return false;
+			}
+		}
+
+		/// <summary>
+		/// Converts a scalar to a boolean value.
+		/// </summary>
+		/// <param name="Scalar">The scalar to convert</param>
+		/// <returns>The scalar converted to a boolean value.</returns>
+		static bool CoerceToBool(string Scalar)
+		{
+			bool Result;
+			if(String.Compare(Scalar, "true", true) == 0)
+			{
+				Result = true;
+			}
+			else if(String.Compare(Scalar, "false", true) == 0)
+			{
+				Result = false;
+			}
+			else
+			{
+				throw new BgConditionException("Token '{0}' cannot be coerced to a bool", Scalar);
+			}
+			return Result;
+		}
+
+		/// <summary>
+		/// Converts a scalar to a boolean value.
+		/// </summary>
+		/// <param name="Scalar">The scalar to convert</param>
+		/// <returns>The scalar converted to an integer value.</returns>
+		static int CoerceToInteger(string Scalar)
+		{
+			int Value;
+			if(!Int32.TryParse(Scalar, out Value))
+			{
+				throw new BgConditionException("Token '{0}' cannot be coerced to an integer", Scalar);
+			}
+			return Value;
+		}
+
+		/// <summary>
+		/// Splits an input string up into expression tokens.
+		/// </summary>
+		/// <param name="Text">Text to be converted into tokens</param>
+		/// <param name="Tokens">List to receive a list of tokens</param>
+		static void Tokenize(string Text, List<string> Tokens)
+		{
+			int Idx = 0;
+			while(Idx < Text.Length)
+			{
+				int EndIdx = Idx + 1;
+				if(!Char.IsWhiteSpace(Text[Idx]))
+				{
+					// Scan to the end of the current token
+					if(Char.IsNumber(Text[Idx]))
+					{
+						// Number
+						while(EndIdx < Text.Length && Char.IsNumber(Text[EndIdx]))
+						{
+							EndIdx++;
+						}
+					}
+					else if(Char.IsLetter(Text[Idx]) || Text[Idx] == '_')
+					{
+						// Identifier
+						while(EndIdx < Text.Length && (Char.IsLetterOrDigit(Text[EndIdx]) || Text[EndIdx] == '_'))
+						{
+							EndIdx++;
+						}
+					}
+					else if(Text[Idx] == '!' || Text[Idx] == '<' || Text[Idx] == '>' || Text[Idx] == '=')
+					{
+						// Operator that can be followed by an equals character
+						if(EndIdx < Text.Length && Text[EndIdx] == '=')
+						{
+							EndIdx++;
+						}
+					}
+					else if(Text[Idx] == '\'' || Text[Idx] == '\"')
+					{
+						// String
+						if(EndIdx < Text.Length)
+						{
+							EndIdx++;
+							while(EndIdx < Text.Length && Text[EndIdx - 1] != Text[Idx]) EndIdx++;
+						}
+					}
+					Tokens.Add(Text.Substring(Idx, EndIdx - Idx));
+				}
+				Idx = EndIdx;
+			}
+			Tokens.Add(EndToken);
+		}
+
+		/// <summary>
+		/// Test cases for conditions.
+		/// </summary>
+		public static void TestConditions()
+		{
+			TestCondition("1 == 2", false);
+			TestCondition("1 == 1", true);
+			TestCondition("1 != 2", true);
+			TestCondition("1 != 1", false);
+			TestCondition("'hello' == 'hello'", true);
+			TestCondition("'hello' == ('hello')", true);
+			TestCondition("'hello' == 'world'", false);
+			TestCondition("'hello' != ('world')", true);
+			TestCondition("true == ('true')", true);
+			TestCondition("true == ('True')", true);
+			TestCondition("true == ('false')", false);
+			TestCondition("true == !('False')", true);
+			TestCondition("true == 'true' and 'false' == 'False'", true);
+			TestCondition("true == 'true' and 'false' == 'true'", false);
+			TestCondition("true == 'false' or 'false' == 'false'", true);
+			TestCondition("true == 'false' or 'false' == 'true'", true);
+		}
+
+		/// <summary>
+		/// Helper method to evaluate a condition and check it's the expected result
+		/// </summary>
+		/// <param name="Condition">Condition to evaluate</param>
+		/// <param name="ExpectedResult">The expected result</param>
+		static void TestCondition(string Condition, bool ExpectedResult)
+		{
+			bool Result = Evaluate(Condition, null!);
+			Console.WriteLine("{0}: {1} = {2}", (Result == ExpectedResult)? "PASS" : "FAIL", Condition, Result);
+		}
+	}
+}
diff --git a/Engine/Source/Programs/Shared/EpicGames.BuildGraph/BgLabel.cs b/Engine/Source/Programs/Shared/EpicGames.BuildGraph/BgLabel.cs
new file mode 100644
--- /dev/null
+++ b/Engine/Source/Programs/Shared/EpicGames.BuildGraph/BgLabel.cs
@@ -0,0 +1,101 @@
+// Copyright Epic Games, Inc. All Rights Reserved.
+
+using System;
+using System.Collections.Generic;
+using System.Linq;
+using System.Text;
+using System.Threading.Tasks;
+
+namespace EpicGames.BuildGraph
+{
+	/// <summary>
+	/// Which changelist to show a UGS badge for
+	/// </summary>
+	public enum BgLabelChange
+	{
+		/// <summary>
+		/// The current changelist being built
+		/// </summary>
+		Current,
+
+		/// <summary>
+		/// The last code changelist
+		/// </summary>
+		Code,
+	}
+
+	/// <summary>
+	/// Defines a label within a graph. Labels are similar to badges, and give the combined status of one or more job steps. Unlike badges, they
+	/// separate the requirements for its status and optional nodes to be included in its status, allowing this to be handled externally.
+	/// </summary>
+	public class BgLabel
+	{
+		/// <summary>
+		/// Name of this badge
+		/// </summary>
+		public readonly string DashboardName;
+
+		/// <summary>
+		/// Category for this label
+		/// </summary>
+		public readonly string DashboardCategory;
+
+		/// <summary>
+		/// Name of the badge in UGS
+		/// </summary>
+		public readonly string UgsBadge;
+
+		/// <summary>
+		/// Path to the project folder in UGS
+		/// </summary>
+		public readonly string UgsProject;
+
+		/// <summary>
+		/// Which change to show the badge for
+		/// </summary>
+		public readonly BgLabelChange Change;
+
+		/// <summary>
+		/// Set of nodes that must be run for this label to be shown.
+		/// </summary>
+		public HashSet<BgNode> RequiredNodes = new HashSet<BgNode>();
+
+		/// <summary>
+		/// Set of nodes that will be included in this label if present.
+		/// </summary>
+		public HashSet<BgNode> IncludedNodes = new HashSet<BgNode>();
+
+		/// <summary>
+		/// Constructor
+		/// </summary>
+		/// <param name="InDashboardName">Name of this label</param>
+		/// <param name="InDashboardCategory">Type of this label</param>
+		/// <param name="InUgsBadge">The UGS badge name</param>
+		/// <param name="InUgsProject">Project to display this badge for</param>
+		/// <param name="InChange">The change to show this badge on in UGS</param>
+		public BgLabel(string InDashboardName, string InDashboardCategory, string InUgsBadge, string InUgsProject, BgLabelChange InChange)
+		{
+			DashboardName = InDashboardName;
+			DashboardCategory = InDashboardCategory;
+			UgsBadge = InUgsBadge;
+			UgsProject = InUgsProject;
+			Change = InChange;
+		}
+
+		/// <summary>
+		/// Get the name of this label
+		/// </summary>
+		/// <returns>The name of this label</returns>
+		public override string ToString()
+		{
+			if (!String.IsNullOrEmpty(DashboardName))
+			{
+				return String.Format("{0}/{1}", DashboardCategory, DashboardName);
+			}
+			else
+			{
+				return UgsBadge;
+			}
+		}
+	}
+}
diff --git a/Engine/Source/Programs/Shared/EpicGames.BuildGraph/BgNode.cs b/Engine/Source/Programs/Shared/EpicGames.BuildGraph/BgNode.cs
new file mode 100644
--- /dev/null
+++ b/Engine/Source/Programs/Shared/EpicGames.BuildGraph/BgNode.cs
@@ -0,0 +1,227 @@
+// Copyright Epic Games, Inc. All Rights Reserved.
+
+using System;
+using System.Collections.Generic;
+using System.Linq;
+using System.Text;
+using System.Threading.Tasks;
+using System.Xml;
+using EpicGames.Core;
+using OpenTracing;
+using OpenTracing.Util;
+using Microsoft.Extensions.Logging;
+using System.Reflection;
+using System.ComponentModel;
+
+namespace EpicGames.BuildGraph
+{
+	/// <summary>
+	/// Reference to an output tag from a particular node
+	/// </summary>
+	public class BgNodeOutput
+	{
+		/// <summary>
+		/// The node which produces the given output
+		/// </summary>
+		public BgNode ProducingNode;
+
+		/// <summary>
+		/// Name of the tag
+		/// </summary>
+		public string TagName;
+
+		/// <summary>
+		/// Constructor
+		/// </summary>
+		/// <param name="InProducingNode">Node which produces the given output</param>
+		/// <param name="InTagName">Name of the tag</param>
+		public BgNodeOutput(BgNode InProducingNode, string InTagName)
+		{
+			ProducingNode = InProducingNode;
+			TagName = InTagName;
+		}
+
+		/// <summary>
+		/// Returns a string representation of this output for debugging purposes
+		/// </summary>
+		/// <returns>The name of this output</returns>
+		public override string ToString()
+		{
+			return String.Format("{0}: {1}", ProducingNode.Name, TagName);
+		}
+	}
+
+	/// <summary>
+	/// Defines a node, a container for tasks and the smallest unit of execution that can be run as part of a build graph.
+	/// </summary>
+	public class BgNode
+	{
+		/// <summary>
+		/// The node's name
+		/// </summary>
+		public string Name;
+
+		/// <summary>
+		/// Array of inputs which this node requires to run
+		/// </summary>
+		public BgNodeOutput[] Inputs;
+
+		/// <summary>
+		/// Array of outputs produced by this node
+		/// </summary>
+		public BgNodeOutput[] Outputs;
+
+		/// <summary>
+		/// Nodes which this node has input dependencies on
+		/// </summary>
+		public BgNode[] InputDependencies;
+
+		/// <summary>
+		/// Nodes which this node needs to run after
+		/// </summary>
+		public BgNode[] OrderDependencies;
+
+		/// <summary>
+		/// Tokens which must be acquired for this node to run
+		/// </summary>
+		public FileReference[] RequiredTokens;
+
+		/// <summary>
+		/// List of tasks to execute
+		/// </summary>
+		public List<BgTask> TaskInfos = new List<BgTask>();
+
+		/// <summary>
+		/// List of email addresses to notify if this node fails.
+		/// </summary>
+		public HashSet<string> NotifyUsers = new HashSet<string>(StringComparer.InvariantCultureIgnoreCase);
+
+		/// <summary>
+		/// If set, anyone that has submitted to one of the given paths will be notified on failure of this node
+		/// </summary>
+		public HashSet<string> NotifySubmitters = new HashSet<string>(StringComparer.InvariantCultureIgnoreCase);
+
+		/// <summary>
+		/// Whether to start this node as soon as its dependencies are satisfied, rather than waiting for all of its agent's dependencies to be met.
+		/// </summary>
+		public bool bRunEarly = false;
+
+		/// <summary>
+		/// Whether to ignore warnings produced by this node
+		/// </summary>
+		public bool bNotifyOnWarnings = true;
+
+		/// <summary>
+		/// Constructor
+		/// </summary>
+		/// <param name="InName">The name of this node</param>
+		/// <param name="InInputs">Inputs that this node depends on</param>
+		/// <param name="InOutputNames">Names of the outputs that this node produces</param>
+		/// <param name="InInputDependencies">Nodes which this node is dependent on for its inputs</param>
+		/// <param name="InOrderDependencies">Nodes which this node needs to run after. Should include all input dependencies.</param>
+		/// <param name="InRequiredTokens">Optional tokens which must be required for this node to run</param>
+		public BgNode(string InName, BgNodeOutput[] InInputs, string[] InOutputNames, BgNode[] InInputDependencies, BgNode[] InOrderDependencies, FileReference[] InRequiredTokens)
+		{
+			Name = InName;
+			Inputs = InInputs;
+
+			List<BgNodeOutput> AllOutputs = new List<BgNodeOutput>();
+			AllOutputs.Add(new BgNodeOutput(this, "#" + Name));
+			AllOutputs.AddRange(InOutputNames.Where(x => String.Compare(x, Name, StringComparison.InvariantCultureIgnoreCase) != 0).Select(x => new BgNodeOutput(this, x)));
+			Outputs = AllOutputs.ToArray();
+
+			InputDependencies = InInputDependencies;
+			OrderDependencies = InOrderDependencies;
+			RequiredTokens = InRequiredTokens;
+		}
+
+		/// <summary>
+		/// Returns the default output for this node, which includes all build products
+		/// </summary>
+		public BgNodeOutput DefaultOutput
+		{
+			get { return Outputs[0]; }
+		}
+
+		/// <summary>
+		/// Determines the minimal set of direct input dependencies for this node to run
+		/// </summary>
+		/// <returns>Sequence of nodes that are direct inputs to this node</returns>
+		public IEnumerable<BgNode> GetDirectInputDependencies()
+		{
+			HashSet<BgNode> DirectDependencies = new HashSet<BgNode>(InputDependencies);
+			foreach(BgNode InputDependency in InputDependencies)
+			{
+				DirectDependencies.ExceptWith(InputDependency.InputDependencies);
+			}
+			return DirectDependencies;
+		}
+
+		/// <summary>
+		/// Determines the minimal set of direct order dependencies for this node to run
+		/// </summary>
+		/// <returns>Sequence of nodes that are direct order dependencies of this node</returns>
+		public IEnumerable<BgNode> GetDirectOrderDependencies()
+		{
+			HashSet<BgNode> DirectDependencies = new HashSet<BgNode>(OrderDependencies);
+			foreach(BgNode OrderDependency in OrderDependencies)
+			{
+				DirectDependencies.ExceptWith(OrderDependency.OrderDependencies);
+			}
+			return DirectDependencies;
+		}
+
+		/// <summary>
+		/// Write this node to an XML writer
+		/// </summary>
+		/// <param name="Writer">The writer to output the node to</param>
+		public void Write(XmlWriter Writer)
+		{
+			Writer.WriteStartElement("Node");
+			Writer.WriteAttributeString("Name", Name);
+
+			string[] RequireNames = Inputs.Select(x => x.TagName).ToArray();
+			if (RequireNames.Length > 0)
+			{
+				Writer.WriteAttributeString("Requires", String.Join(";", RequireNames));
+			}
+
+			string[] ProducesNames = Outputs.Where(x => x != DefaultOutput).Select(x => x.TagName).ToArray();
+			if (ProducesNames.Length > 0)
+			{
+				Writer.WriteAttributeString("Produces", String.Join(";", ProducesNames));
+			}
+
+			string[] AfterNames = GetDirectOrderDependencies().Except(InputDependencies).Select(x => x.Name).ToArray();
+			if (AfterNames.Length > 0)
+			{
+				Writer.WriteAttributeString("After", String.Join(";", AfterNames));
+			}
+
+			if (!bNotifyOnWarnings)
+			{
+				Writer.WriteAttributeString("NotifyOnWarnings", bNotifyOnWarnings.ToString());
+			}
+
+			if(bRunEarly)
+			{
+				Writer.WriteAttributeString("RunEarly", bRunEarly.ToString());
+			}
+
+			foreach (BgTask Task in TaskInfos)
+			{
+				Task.Write(Writer);
+			}
+			Writer.WriteEndElement();
+		}
+
+		/// <summary>
+		/// Returns the name of this node
+		/// </summary>
+		/// <returns>The name of this node</returns>
+		public override string ToString()
+		{
+			return Name;
+		}
+	}
+}
diff --git a/Engine/Source/Programs/Shared/EpicGames.BuildGraph/BgReport.cs b/Engine/Source/Programs/Shared/EpicGames.BuildGraph/BgReport.cs
new file mode 100644
--- /dev/null
+++ b/Engine/Source/Programs/Shared/EpicGames.BuildGraph/BgReport.cs
@@ -0,0 +1,50 @@
+// Copyright Epic Games, Inc. All Rights Reserved.
+
+using System;
+using System.Collections.Generic;
+using System.Diagnostics;
+using System.Linq;
+using System.Text;
+using System.Threading.Tasks;
+
+namespace EpicGames.BuildGraph
+{
+	/// <summary>
+	/// Defines a report to be generated as part of the build.
+	/// </summary>
+	public class BgReport
+	{
+		/// <summary>
+		/// Name of this trigger
+		/// </summary>
+		public readonly string Name;
+
+		/// <summary>
+		/// Set of nodes to include in the report
+		/// </summary>
+		public HashSet<BgNode> Nodes = new HashSet<BgNode>();
+
+		/// <summary>
+		/// List of users to notify with this report
+		/// </summary>
+		public HashSet<string> NotifyUsers = new HashSet<string>(StringComparer.InvariantCultureIgnoreCase);
+
+		/// <summary>
+		/// Constructor
+		/// </summary>
+		/// <param name="InName">Name of this report</param>
+		public BgReport(string InName)
+		{
+			Name = InName;
+		}
+
+		/// <summary>
+		/// Get the name of this report
+		/// </summary>
+		/// <returns>The name of this report</returns>
+		public override string ToString()
+		{
+			return Name;
+		}
+	}
+}
diff --git a/Engine/Source/Programs/Shared/EpicGames.BuildGraph/BgSchema.cs b/Engine/Source/Programs/Shared/EpicGames.BuildGraph/BgSchema.cs
new file mode 100644
--- /dev/null
+++ b/Engine/Source/Programs/Shared/EpicGames.BuildGraph/BgSchema.cs
@@ -0,0 +1,1109 @@
+// Copyright Epic Games, Inc. All Rights Reserved.
+
+using System;
+using System.Collections.Generic;
+using System.Diagnostics;
+using System.IO;
+using System.Linq;
+using System.Reflection;
+using System.Text;
+using System.Threading.Tasks;
+using System.Xml;
+using System.Xml.Schema;
+using System.Xml.Serialization;
+using EpicGames.Core;
+
+namespace EpicGames.BuildGraph
+{
+	/// <summary>
+	/// Specifies validation that should be performed on a task parameter.
+	/// </summary>
+	public enum TaskParameterValidationType
+	{
+		/// <summary>
+		/// Allow any valid values for the field type.
+		/// </summary>
+		Default,
+
+		/// <summary>
+		/// A list of tag names separated by semicolons
+		/// </summary>
+		TagList,
+
+		/// <summary>
+		/// A file specification, which may contain tags and wildcards.
+		/// </summary>
+		FileSpec,
+	}
+
+	/// <summary>
+	/// Information about a parameter to a task
+	/// </summary>
+	[DebuggerDisplay("{Name}")]
+	public class BgScriptTaskParameter
+	{
+		/// <summary>
+		/// Name of this parameter
+		/// </summary>
+		public string Name { get; }
+
+		/// <summary>
+		/// The type for values assigned to this field
+		/// </summary>
+		public Type ValueType { get; }
+
+		/// <summary>
+		/// The ICollection interface for this type
+		/// </summary>
+		public Type CollectionType { get; }
+
+		/// <summary>
+		/// Validation type for this field
+		/// </summary>
+		public TaskParameterValidationType ValidationType { get; }
+
+		/// <summary>
+		/// Whether this parameter is optional
+		/// </summary>
+		public bool bOptional { get; }
+
+		/// <summary>
+		/// Constructor
+		/// </summary>
+		public BgScriptTaskParameter(string InName, Type InValueType, TaskParameterValidationType InValidationType, bool bInOptional)
+		{
+			Name = InName;
+			ValueType = InValueType;
+			ValidationType = InValidationType;
+			bOptional = bInOptional;
+
+			if (ValueType.IsGenericType && ValueType.GetGenericTypeDefinition() == typeof(Nullable<>))
+			{
+				ValueType = ValueType.GetGenericArguments()[0];
+				bOptional = true;
+			}
+
+			if(ValueType.IsClass)
+			{
+				foreach(Type InterfaceType in ValueType.GetInterfaces())
+				{
+					if(InterfaceType.IsGenericType)
+					{
+						Type GenericInterfaceType = InterfaceType.GetGenericTypeDefinition();
+						if(GenericInterfaceType == typeof(ICollection<>))
+						{
+							CollectionType = InterfaceType;
+							ValueType = InterfaceType.GetGenericArguments()[0];
+						}
+					}
+				}
+			}
+		}
+	}
+
+	/// <summary>
+	/// Helper class to serialize a task from an xml element
+	/// </summary>
+	[DebuggerDisplay("{Name}")]
+	public class BgScriptTask
+	{
+		/// <summary>
+		/// Name of this task
+		/// </summary>
+		public string Name { get; }
+
+		/// <summary>
+		/// Parameters for this task
+		/// </summary>
+		public List<BgScriptTaskParameter> Parameters { get; }
+
+		/// <summary>
+		/// Constructor
+		/// </summary>
+		/// <param name="Name">Name of the task</param>
+		/// <param name="Parameters">Parameters for the task</param>
+		public BgScriptTask(string Name, List<BgScriptTaskParameter> Parameters)
+		{
+			this.Name = Name;
+			this.Parameters = new List<BgScriptTaskParameter>(Parameters);
+		}
+	}
+
+	/// <summary>
+	/// Enumeration of standard types used in the schema. Avoids hard-coding names.
+	/// </summary>
+	public enum ScriptSchemaStandardType
+	{
+		Graph,
+		Agent,
+		AgentBody,
+		Node,
+		NodeBody,
+		Aggregate,
+		Report,
+		Badge,
+		Label,
+		Notify,
+		Include,
+		Option,
+		EnvVar,
+		Property,
+		Regex,
+		Macro,
+		MacroBody,
+		Expand,
+		Trace,
+		Warning,
+		Error,
+		Name,
+		NameList,
+		Tag,
+		TagList,
+		NameOrTag,
+		NameOrTagList,
+		QualifiedName,
+		BalancedString,
+		Boolean,
+		Integer,
+		LabelChange
+	}
+
+	/// <summary>
+	/// Schema for build graph definitions. Stores information about the supported tasks, and allows validating an XML document.
+	/// </summary>
+	public class BgScriptSchema
+	{
+		/// <summary>
+		/// Name of the root element
+		/// </summary>
+		public const string RootElementName = "BuildGraph";
+
+		/// <summary>
+		/// Namespace for the schema
+		/// </summary>
+		public const string NamespaceURI = "http://www.epicgames.com/BuildGraph";
+
+		/// <summary>
+		/// Qualified name for the string type
+		/// </summary>
+		static readonly XmlQualifiedName StringTypeName = new XmlQualifiedName("string", "http://www.w3.org/2001/XMLSchema");
+
+		/// <summary>
+		/// The inner xml schema
+		/// </summary>
+		public readonly XmlSchema CompiledSchema;
+
+		/// <summary>
+		/// Characters which are not permitted in names.
+		/// </summary>
+		public const string IllegalNameCharacters = "^<>:\"/\\|?*";
+
+		/// <summary>
+		/// Pattern which matches any name; alphanumeric characters, with single embedded spaces.
+		/// </summary>
+		const string NamePattern = "[^ " + IllegalNameCharacters +"]+( [^ " + IllegalNameCharacters + "]+)*";
+
+		/// <summary>
+		/// Pattern which matches a list of names, separated by semicolons.
+		/// </summary>
+		const string NameListPattern = NamePattern + "(;" + NamePattern + ")*";
+
+		/// <summary>
+		/// Pattern which matches any tag name; a name with a leading '#' character
+		/// </summary>
+		const string TagPattern = "#" + NamePattern;
+
+		/// <summary>
+		/// Pattern which matches a list of tag names, separated by semicolons;
+		/// </summary>
+		const string TagListPattern = TagPattern + "(;" + TagPattern + ")*";
+
+		/// <summary>
+		/// Pattern which matches any name or tag name; a name with a leading '#' character
+		/// </summary>
+		const string NameOrTagPattern = "#?" + NamePattern;
+
+		/// <summary>
+		/// Pattern which matches a list of names or tag names, separated by semicolons;
+		/// </summary>
+		const string NameOrTagListPattern = NameOrTagPattern + "(;" + NameOrTagPattern + ")*";
+
+		/// <summary>
+		/// Pattern which matches a qualified name.
+		/// </summary>
+		const string QualifiedNamePattern = NamePattern + "(\\." + NamePattern + ")*";
+
+		/// <summary>
+		/// Pattern which matches a property name
+		/// </summary>
+		const string PropertyPattern = "\\$\\(" + NamePattern + "\\)";
+
+		/// <summary>
+		/// Pattern which matches balanced parentheses in a string
+		/// </summary>
+		const string StringWithPropertiesPattern = "[^\\$]*" + "(" + "(" + PropertyPattern + "|" + "\\$[^\\(]" + ")" + "[^\\$]*" + ")+" + "\\$?";
+
+		/// <summary>
+		/// Pattern which matches balanced parentheses in a string
+		/// </summary>
+		const string BalancedStringPattern = "[^\\$]*" + "(" + "(" + PropertyPattern + "|" + "\\$[^\\(]" + ")" + "[^\\$]*" + ")*" + "\\$?";
+
+		private BgScriptSchema(XmlSchema Schema)
+		{
+			CompiledSchema = Schema;
+		}
+
+		/// <summary>
+		/// Constructor
+		/// </summary>
+		/// <param name="InNameToTask">Mapping of task name to information about how to construct it</param>
+		public BgScriptSchema(IEnumerable<BgScriptTask> Tasks, List<(Type, ScriptSchemaStandardType)> PrimitiveTypes)
+		{
+			// Create a lookup from standard types to their qualified names
+			Dictionary<Type, XmlQualifiedName> TypeToSchemaTypeName = new Dictionary<Type,XmlQualifiedName>();
+			TypeToSchemaTypeName.Add(typeof(String), GetQualifiedTypeName(ScriptSchemaStandardType.BalancedString));
+			TypeToSchemaTypeName.Add(typeof(Boolean), GetQualifiedTypeName(ScriptSchemaStandardType.Boolean));
+			TypeToSchemaTypeName.Add(typeof(Int32), GetQualifiedTypeName(ScriptSchemaStandardType.Integer));
+			foreach ((Type Type, ScriptSchemaStandardType SchemaType) in PrimitiveTypes)
+			{
+				TypeToSchemaTypeName.Add(Type, GetQualifiedTypeName(SchemaType));
+			}
+
+			// If a dictionary of tasks was specified, add those to the schema. Otherwise we'll allow any syntactically valid invocation.
+			List<XmlSchemaType> UserTypes = new List<XmlSchemaType>();
+			Dictionary<string, XmlSchemaComplexType> TaskNameToType = null;
+			if (Tasks != null)
+			{
+				// Create all the custom user types, and add them to the qualified name lookup
+				foreach (Type Type in Tasks.SelectMany(x => x.Parameters).Select(x => x.ValueType))
+				{
+					if (!TypeToSchemaTypeName.ContainsKey(Type))
+					{
+						if (Type.IsClass && Type.GetInterfaces().Any(x => x.GetGenericTypeDefinition() == typeof(ICollection<>)))
+						{
+							TypeToSchemaTypeName.Add(Type, GetQualifiedTypeName(ScriptSchemaStandardType.BalancedString));
+						}
+						else
+						{
+							string Name = Type.Name + "UserType";
+							XmlSchemaType SchemaType = CreateUserType(Name, Type);
+							UserTypes.Add(SchemaType);
+							TypeToSchemaTypeName.Add(Type, new XmlQualifiedName(Name, NamespaceURI));
+						}
+					}
+				}
+
+				// Create all the task types
+				TaskNameToType = new Dictionary<string, XmlSchemaComplexType>();
+				foreach (BgScriptTask Task in Tasks)
+				{
+					XmlSchemaComplexType TaskType = new XmlSchemaComplexType();
+					TaskType.Name = Task.Name + "TaskType";
+					foreach (BgScriptTaskParameter Parameter in Task.Parameters)
+					{
+						XmlQualifiedName SchemaTypeName = GetQualifiedTypeName(Parameter.ValidationType);
+						if (SchemaTypeName == null)
+						{
+							SchemaTypeName = TypeToSchemaTypeName[Parameter.ValueType];
+						}
+						TaskType.Attributes.Add(CreateSchemaAttribute(Parameter.Name, SchemaTypeName, Parameter.bOptional ? XmlSchemaUse.Optional : XmlSchemaUse.Required));
+					}
+					TaskType.Attributes.Add(CreateSchemaAttribute("If", ScriptSchemaStandardType.BalancedString, XmlSchemaUse.Optional));
+					TaskNameToType.Add(Task.Name, TaskType);
+				}
+			}
+
+			// Create the schema object
+			XmlSchema NewSchema = new XmlSchema();
+			NewSchema.TargetNamespace = NamespaceURI;
+			NewSchema.ElementFormDefault = XmlSchemaForm.Qualified;
+			NewSchema.Items.Add(CreateSchemaElement(RootElementName, ScriptSchemaStandardType.Graph));
+			NewSchema.Items.Add(CreateGraphType());
+			NewSchema.Items.Add(CreateAgentType());
+			NewSchema.Items.Add(CreateAgentBodyType());
+			NewSchema.Items.Add(CreateNodeType());
+			NewSchema.Items.Add(CreateNodeBodyType(TaskNameToType));
+			NewSchema.Items.Add(CreateAggregateType());
+			NewSchema.Items.Add(CreateReportType());
+			NewSchema.Items.Add(CreateBadgeType());
+			NewSchema.Items.Add(CreateLabelType());
+			NewSchema.Items.Add(CreateEnumType(GetTypeName(ScriptSchemaStandardType.LabelChange), typeof(BgLabelChange)));
+			NewSchema.Items.Add(CreateNotifyType());
+			NewSchema.Items.Add(CreateIncludeType());
+			NewSchema.Items.Add(CreateOptionType());
+			NewSchema.Items.Add(CreateEnvVarType());
+			NewSchema.Items.Add(CreatePropertyType());
+			NewSchema.Items.Add(CreateRegexType());
+			NewSchema.Items.Add(CreateMacroType());
+			NewSchema.Items.Add(CreateMacroBodyType(TaskNameToType));
+			NewSchema.Items.Add(CreateExpandType());
+			NewSchema.Items.Add(CreateDiagnosticType(ScriptSchemaStandardType.Trace));
+			NewSchema.Items.Add(CreateDiagnosticType(ScriptSchemaStandardType.Warning));
+			NewSchema.Items.Add(CreateDiagnosticType(ScriptSchemaStandardType.Error));
+			NewSchema.Items.Add(CreateSimpleTypeFromRegex(GetTypeName(ScriptSchemaStandardType.Name), "(" + NamePattern + "|" + StringWithPropertiesPattern + ")"));
+			NewSchema.Items.Add(CreateSimpleTypeFromRegex(GetTypeName(ScriptSchemaStandardType.NameList), "(" + NameListPattern + "|" + StringWithPropertiesPattern + ")"));
+			NewSchema.Items.Add(CreateSimpleTypeFromRegex(GetTypeName(ScriptSchemaStandardType.Tag), "(" + TagPattern + "|" + StringWithPropertiesPattern + ")"));
+			NewSchema.Items.Add(CreateSimpleTypeFromRegex(GetTypeName(ScriptSchemaStandardType.TagList), "(" + TagListPattern + "|" + StringWithPropertiesPattern + ")"));
+			NewSchema.Items.Add(CreateSimpleTypeFromRegex(GetTypeName(ScriptSchemaStandardType.NameOrTag), "(" + NameOrTagPattern + "|" + StringWithPropertiesPattern + ")"));
+			NewSchema.Items.Add(CreateSimpleTypeFromRegex(GetTypeName(ScriptSchemaStandardType.NameOrTagList), "(" + NameOrTagListPattern + "|" + StringWithPropertiesPattern + ")"));
+			NewSchema.Items.Add(CreateSimpleTypeFromRegex(GetTypeName(ScriptSchemaStandardType.QualifiedName), "(" + QualifiedNamePattern + "|" + StringWithPropertiesPattern + ")"));
+			NewSchema.Items.Add(CreateSimpleTypeFromRegex(GetTypeName(ScriptSchemaStandardType.BalancedString), BalancedStringPattern));
+			NewSchema.Items.Add(CreateSimpleTypeFromRegex(GetTypeName(ScriptSchemaStandardType.Boolean), "(true|True|false|False|" + StringWithPropertiesPattern + ")"));
+			NewSchema.Items.Add(CreateSimpleTypeFromRegex(GetTypeName(ScriptSchemaStandardType.Integer), "(" + "(-?[1-9][0-9]*|0)" + "|" + StringWithPropertiesPattern + ")"));
+			if (TaskNameToType != null)
+			{
+				foreach (XmlSchemaComplexType Type in TaskNameToType.Values)
+				{
+					NewSchema.Items.Add(Type);
+				}
+			}
+			foreach(XmlSchemaSimpleType Type in UserTypes)
+			{
+				NewSchema.Items.Add(Type);
+			}
+
+			// Now that we've finished, compile it and store it to the class
+			XmlSchemaSet NewSchemaSet = new XmlSchemaSet();
+			NewSchemaSet.Add(NewSchema);
+			NewSchemaSet.Compile();
+			foreach(XmlSchema NewCompiledSchema in NewSchemaSet.Schemas())
+			{
+				CompiledSchema = NewCompiledSchema;
+			}
+		}
+
+		/// <summary>
+		/// Export the schema to a file
+		/// </summary>
+		/// <param name="File"></param>
+		public void Export(FileReference File)
+		{
+			DirectoryReference.CreateDirectory(File.Directory);
+
+			XmlWriterSettings Settings = new XmlWriterSettings();
+			Settings.Indent = true;
+			Settings.IndentChars = "  ";
+			Settings.NewLineChars = "\n";
+
+			DirectoryReference.CreateDirectory(File.Directory);
+			using(XmlWriter Writer = XmlWriter.Create(File.FullName, Settings))
+			{
+				CompiledSchema.Write(Writer);
+			}
+		}
+
+		/// <summary>
+		/// Imports a schema from a file
+		/// </summary>
+		/// <param name="File">The XML file to import</param>
+		/// <param name="NameToTask">Mapping of task name to information about how to construct it</param>
+		/// <returns>A <see cref="BgScriptSchema"/> deserialized from the XML file, or null if file doesn't exist</returns>
+		public static BgScriptSchema Import(FileReference File)
+		{
+			if (!FileReference.Exists(File))
+			{
+				return null;
+			}
+			
+			using (XmlReader SchemaFile = XmlReader.Create(File.FullName))
+			{
+				BgScriptSchema ImportedSchema = new BgScriptSchema(XmlSchema.Read(SchemaFile, ValidationCallback));
+				return ImportedSchema;
+			}
+		}
+
+		static void ValidationCallback(object Sender, ValidationEventArgs EventArgs)
+		{
+			if (EventArgs.Severity == XmlSeverityType.Warning)
+			{
+				Log.WriteLine(LogEventType.Warning, "WARNING: {0}", EventArgs.Message);
+			}
+			else if (EventArgs.Severity == XmlSeverityType.Error)
+			{
+				Log.WriteLine(LogEventType.Error, "ERROR: {0}", EventArgs.Message);
+			}
+		}
+
+		/// <summary>
+		/// Gets the bare name for the given script type
+		/// </summary>
+		/// <param name="Type">Script type to find the name of</param>
+		/// <returns>Name of the schema type that matches the given script type</returns>
+		static string GetTypeName(ScriptSchemaStandardType Type)
+		{
+			return Type.ToString() + "Type";
+		}
+		
+		/// <summary>
+		/// Gets the qualified name for the given script type
+		/// </summary>
+		/// <param name="Type">Script type to find the qualified name for</param>
+		/// <returns>Qualified name of the schema type that matches the given script type</returns>
+		static XmlQualifiedName GetQualifiedTypeName(ScriptSchemaStandardType Type)
+		{
+			return new XmlQualifiedName(GetTypeName(Type), NamespaceURI);
+		}
+
+		/// <summary>
+		/// Gets the qualified name of the schema type for the given type of validation
+		/// </summary>
+		/// <returns>Qualified name for the corresponding schema type</returns>
+		static XmlQualifiedName GetQualifiedTypeName(TaskParameterValidationType Type)
+		{
+			switch(Type)
+			{
+				case TaskParameterValidationType.TagList:
+					return GetQualifiedTypeName(ScriptSchemaStandardType.TagList);
+			}
+			return null;
+		}
+
+		/// <summary>
+		/// Creates the schema type representing the graph type
+		/// </summary>
+		/// <returns>Type definition for a graph</returns>
+		static XmlSchemaType CreateGraphType()
+		{
+			XmlSchemaChoice GraphChoice = new XmlSchemaChoice();
+			GraphChoice.MinOccurs = 0;
+			GraphChoice.MaxOccursString = "unbounded";
+			GraphChoice.Items.Add(CreateSchemaElement("Include", ScriptSchemaStandardType.Include));
+			GraphChoice.Items.Add(CreateSchemaElement("Option", ScriptSchemaStandardType.Option));
+			GraphChoice.Items.Add(CreateSchemaElement("EnvVar", ScriptSchemaStandardType.EnvVar));
+			GraphChoice.Items.Add(CreateSchemaElement("Property", ScriptSchemaStandardType.Property));
+			GraphChoice.Items.Add(CreateSchemaElement("Regex", ScriptSchemaStandardType.Regex));
+			GraphChoice.Items.Add(CreateSchemaElement("Macro", ScriptSchemaStandardType.Macro));
+			GraphChoice.Items.Add(CreateSchemaElement("Agent", ScriptSchemaStandardType.Agent));
+			GraphChoice.Items.Add(CreateSchemaElement("Aggregate", ScriptSchemaStandardType.Aggregate));
+			GraphChoice.Items.Add(CreateSchemaElement("Report", ScriptSchemaStandardType.Report));
+			GraphChoice.Items.Add(CreateSchemaElement("Badge", ScriptSchemaStandardType.Badge));
+			GraphChoice.Items.Add(CreateSchemaElement("Label", ScriptSchemaStandardType.Label));
+			GraphChoice.Items.Add(CreateSchemaElement("Notify", ScriptSchemaStandardType.Notify));
+			GraphChoice.Items.Add(CreateSchemaElement("Trace", ScriptSchemaStandardType.Trace));
+			GraphChoice.Items.Add(CreateSchemaElement("Warning", ScriptSchemaStandardType.Warning));
+			GraphChoice.Items.Add(CreateSchemaElement("Error", ScriptSchemaStandardType.Error));
+			GraphChoice.Items.Add(CreateSchemaElement("Expand", ScriptSchemaStandardType.Expand));
+			GraphChoice.Items.Add(CreateDoElement(ScriptSchemaStandardType.Graph));
+			GraphChoice.Items.Add(CreateSwitchElement(ScriptSchemaStandardType.Graph));
+			GraphChoice.Items.Add(CreateForEachElement(ScriptSchemaStandardType.Graph));
+
+			XmlSchemaComplexType GraphType = new XmlSchemaComplexType();
+			GraphType.Name = GetTypeName(ScriptSchemaStandardType.Graph);
+			GraphType.Particle = GraphChoice;
+			return GraphType;
+		}
+
+		/// <summary>
+		/// Creates the schema type representing the agent type
+		/// </summary>
+		/// <returns>Type definition for an agent</returns>
+		static XmlSchemaType CreateAgentType()
+		{
+			XmlSchemaComplexContentExtension Extension = new XmlSchemaComplexContentExtension();
+			Extension.BaseTypeName = GetQualifiedTypeName(ScriptSchemaStandardType.AgentBody);
+			Extension.Attributes.Add(CreateSchemaAttribute("Name", StringTypeName, XmlSchemaUse.Required));
+			Extension.Attributes.Add(CreateSchemaAttribute("Type", ScriptSchemaStandardType.NameList, XmlSchemaUse.Optional));
+			Extension.Attributes.Add(CreateSchemaAttribute("If", ScriptSchemaStandardType.BalancedString, XmlSchemaUse.Optional));
+
+			XmlSchemaComplexContent ContentModel = new XmlSchemaComplexContent();
+			ContentModel.Content = Extension;
+
+			XmlSchemaComplexType ComplexType = new XmlSchemaComplexType();
+			ComplexType.Name = GetTypeName(ScriptSchemaStandardType.Agent);
+			ComplexType.ContentModel = ContentModel;
+			return ComplexType;
+		}
+
+		/// <summary>
+		/// Creates the schema type representing the contents of agent type
+		/// </summary>
+		/// <returns>Type definition for an agent</returns>
+		static XmlSchemaType CreateAgentBodyType()
+		{
+			XmlSchemaChoice AgentChoice = new XmlSchemaChoice();
+			AgentChoice.MinOccurs = 0;
+			AgentChoice.MaxOccursString = "unbounded";
+			AgentChoice.Items.Add(CreateSchemaElement("Property", ScriptSchemaStandardType.Property));
+			AgentChoice.Items.Add(CreateSchemaElement("Regex", ScriptSchemaStandardType.Regex));
+			AgentChoice.Items.Add(CreateSchemaElement("EnvVar", ScriptSchemaStandardType.EnvVar));
+			AgentChoice.Items.Add(CreateSchemaElement("Node", ScriptSchemaStandardType.Node));
+			AgentChoice.Items.Add(CreateSchemaElement("Trace", ScriptSchemaStandardType.Trace));
+			AgentChoice.Items.Add(CreateSchemaElement("Label", ScriptSchemaStandardType.Label));
+			AgentChoice.Items.Add(CreateSchemaElement("Warning", ScriptSchemaStandardType.Warning));
+			AgentChoice.Items.Add(CreateSchemaElement("Error", ScriptSchemaStandardType.Error));
+			AgentChoice.Items.Add(CreateSchemaElement("Expand", ScriptSchemaStandardType.Expand));
+			AgentChoice.Items.Add(CreateDoElement(ScriptSchemaStandardType.AgentBody));
+			AgentChoice.Items.Add(CreateSwitchElement(ScriptSchemaStandardType.AgentBody));
+			AgentChoice.Items.Add(CreateForEachElement(ScriptSchemaStandardType.AgentBody));
+
+			XmlSchemaComplexType AgentType = new XmlSchemaComplexType();
+			AgentType.Name = GetTypeName(ScriptSchemaStandardType.AgentBody);
+			AgentType.Particle = AgentChoice;
+			return AgentType;
+		}
+
+		/// <summary>
+		/// Creates the schema type representing the node type
+		/// </summary>
+		/// <returns>Type definition for a node</returns>
+		static XmlSchemaType CreateNodeType()
+		{
+			XmlSchemaComplexContentExtension Extension = new XmlSchemaComplexContentExtension();
+			Extension.BaseTypeName = GetQualifiedTypeName(ScriptSchemaStandardType.NodeBody);
+			Extension.Attributes.Add(CreateSchemaAttribute("Name", ScriptSchemaStandardType.Name, XmlSchemaUse.Required));
+			Extension.Attributes.Add(CreateSchemaAttribute("Requires", ScriptSchemaStandardType.NameOrTagList, XmlSchemaUse.Optional));
+			Extension.Attributes.Add(CreateSchemaAttribute("Produces", ScriptSchemaStandardType.TagList, XmlSchemaUse.Optional));
+			Extension.Attributes.Add(CreateSchemaAttribute("After", ScriptSchemaStandardType.NameOrTagList, XmlSchemaUse.Optional));
+			Extension.Attributes.Add(CreateSchemaAttribute("Token", ScriptSchemaStandardType.BalancedString, XmlSchemaUse.Optional));
+			Extension.Attributes.Add(CreateSchemaAttribute("If", ScriptSchemaStandardType.BalancedString, XmlSchemaUse.Optional));
+			Extension.Attributes.Add(CreateSchemaAttribute("RunEarly", ScriptSchemaStandardType.Boolean, XmlSchemaUse.Optional));
+			Extension.Attributes.Add(CreateSchemaAttribute("NotifyOnWarnings", ScriptSchemaStandardType.Boolean, XmlSchemaUse.Optional));
+
+			XmlSchemaComplexContent ContentModel = new XmlSchemaComplexContent();
+			ContentModel.Content = Extension;
+
+			XmlSchemaComplexType ComplexType = new XmlSchemaComplexType();
+			ComplexType.Name = GetTypeName(ScriptSchemaStandardType.Node);
+			ComplexType.ContentModel = ContentModel;
+			return ComplexType;
+		}
+
+		/// <summary>
+		/// Creates the schema type representing the body of the node type
+		/// </summary>
+		/// <returns>Type definition for a node</returns>
+		static XmlSchemaType CreateNodeBodyType(Dictionary<string, XmlSchemaComplexType> TaskNameToType)
+		{
+			XmlSchemaChoice NodeChoice = new XmlSchemaChoice();
+			NodeChoice.MinOccurs = 0;
+			NodeChoice.MaxOccursString = "unbounded";
+			NodeChoice.Items.Add(CreateSchemaElement("Property", ScriptSchemaStandardType.Property));
+			NodeChoice.Items.Add(CreateSchemaElement("Regex", ScriptSchemaStandardType.Regex));
+			NodeChoice.Items.Add(CreateSchemaElement("EnvVar", ScriptSchemaStandardType.EnvVar));
+			NodeChoice.Items.Add(CreateSchemaElement("Trace", ScriptSchemaStandardType.Trace));
+			NodeChoice.Items.Add(CreateSchemaElement("Warning", ScriptSchemaStandardType.Warning));
+			NodeChoice.Items.Add(CreateSchemaElement("Error", ScriptSchemaStandardType.Error));
+			NodeChoice.Items.Add(CreateSchemaElement("Expand", ScriptSchemaStandardType.Expand));
+			NodeChoice.Items.Add(CreateDoElement(ScriptSchemaStandardType.NodeBody));
+			NodeChoice.Items.Add(CreateSwitchElement(ScriptSchemaStandardType.NodeBody));
+			NodeChoice.Items.Add(CreateForEachElement(ScriptSchemaStandardType.NodeBody));
+
+			if (TaskNameToType == null)
+			{
+				NodeChoice.Items.Add(new XmlSchemaAny());
+			}
+			else
+			{
+				foreach (KeyValuePair<string, XmlSchemaComplexType> Pair in TaskNameToType.OrderBy(x => x.Key))
+				{
+					NodeChoice.Items.Add(CreateSchemaElement(Pair.Key, new XmlQualifiedName(Pair.Value.Name, NamespaceURI)));
+				}
+			}
+
+			XmlSchemaComplexType NodeType = new XmlSchemaComplexType();
+			NodeType.Name = GetTypeName(ScriptSchemaStandardType.NodeBody);
+			NodeType.Particle = NodeChoice;
+			return NodeType;
+		}
+
+		/// <summary>
+		/// Creates the schema type representing the aggregate type
+		/// </summary>
+		/// <returns>Type definition for an aggregate</returns>
+		static XmlSchemaType CreateAggregateType()
+		{
+			XmlSchemaComplexType AggregateType = new XmlSchemaComplexType();
+			AggregateType.Name = GetTypeName(ScriptSchemaStandardType.Aggregate);
+			AggregateType.Attributes.Add(CreateSchemaAttribute("Name", ScriptSchemaStandardType.Name, XmlSchemaUse.Required));
+			AggregateType.Attributes.Add(CreateSchemaAttribute("Label", ScriptSchemaStandardType.BalancedString, XmlSchemaUse.Optional));
+			AggregateType.Attributes.Add(CreateSchemaAttribute("Requires", ScriptSchemaStandardType.NameOrTagList, XmlSchemaUse.Required));
+			AggregateType.Attributes.Add(CreateSchemaAttribute("Include", ScriptSchemaStandardType.NameOrTagList, XmlSchemaUse.Optional));
+			AggregateType.Attributes.Add(CreateSchemaAttribute("Exclude", ScriptSchemaStandardType.NameOrTagList, XmlSchemaUse.Optional));
+			AggregateType.Attributes.Add(CreateSchemaAttribute("If", ScriptSchemaStandardType.BalancedString, XmlSchemaUse.Optional));
+			return AggregateType;
+		}
+
+		/// <summary>
+		/// Creates the schema type representing the report type
+		/// </summary>
+		/// <returns>Type definition for a report</returns>
+		static XmlSchemaType CreateReportType()
+		{
+			XmlSchemaComplexType ReportType = new XmlSchemaComplexType();
+			ReportType.Name = GetTypeName(ScriptSchemaStandardType.Report);
+			ReportType.Attributes.Add(CreateSchemaAttribute("Name", ScriptSchemaStandardType.Name, XmlSchemaUse.Required));
+			ReportType.Attributes.Add(CreateSchemaAttribute("Requires", ScriptSchemaStandardType.NameOrTagList, XmlSchemaUse.Required));
+			ReportType.Attributes.Add(CreateSchemaAttribute("If", ScriptSchemaStandardType.BalancedString, XmlSchemaUse.Optional));
+			return ReportType;
+		}
+
+		/// <summary>
+		/// Creates the schema type representing the badge type
+		/// </summary>
+		/// <returns>Type definition for a badge</returns>
+		static XmlSchemaType CreateBadgeType()
+		{
+			XmlSchemaComplexType BadgeType = new XmlSchemaComplexType();
+			BadgeType.Name = GetTypeName(ScriptSchemaStandardType.Badge);
+			BadgeType.Attributes.Add(CreateSchemaAttribute("Name", ScriptSchemaStandardType.Name, XmlSchemaUse.Required));
+			BadgeType.Attributes.Add(CreateSchemaAttribute("Requires", ScriptSchemaStandardType.NameOrTagList, XmlSchemaUse.Optional));
+			BadgeType.Attributes.Add(CreateSchemaAttribute("Targets", ScriptSchemaStandardType.NameOrTagList, XmlSchemaUse.Optional));
+			BadgeType.Attributes.Add(CreateSchemaAttribute("Project", ScriptSchemaStandardType.BalancedString, XmlSchemaUse.Required));
+			BadgeType.Attributes.Add(CreateSchemaAttribute("Change", ScriptSchemaStandardType.BalancedString, XmlSchemaUse.Optional));
+			BadgeType.Attributes.Add(CreateSchemaAttribute("If", ScriptSchemaStandardType.BalancedString, XmlSchemaUse.Optional));
+			return BadgeType;
+		}
+
+		/// <summary>
+		/// Creates the schema type representing the label type
+		/// </summary>
+		/// <returns>Type definition for a label</returns>
+		static XmlSchemaType CreateLabelType()
+		{
+			XmlSchemaComplexType LabelType = new XmlSchemaComplexType();
+			LabelType.Name = GetTypeName(ScriptSchemaStandardType.Label);
+			LabelType.Attributes.Add(CreateSchemaAttribute("Name", ScriptSchemaStandardType.Name, XmlSchemaUse.Optional));
+			LabelType.Attributes.Add(CreateSchemaAttribute("Category", ScriptSchemaStandardType.Name, XmlSchemaUse.Optional));
+			LabelType.Attributes.Add(CreateSchemaAttribute("UgsBadge", ScriptSchemaStandardType.Name, XmlSchemaUse.Optional));
+			LabelType.Attributes.Add(CreateSchemaAttribute("UgsProject", ScriptSchemaStandardType.BalancedString, XmlSchemaUse.Optional));
+			LabelType.Attributes.Add(CreateSchemaAttribute("Change", ScriptSchemaStandardType.LabelChange, XmlSchemaUse.Optional));
+			LabelType.Attributes.Add(CreateSchemaAttribute("Requires", ScriptSchemaStandardType.NameOrTagList, XmlSchemaUse.Required));
+			LabelType.Attributes.Add(CreateSchemaAttribute("Include", ScriptSchemaStandardType.NameOrTagList, XmlSchemaUse.Optional));
+			LabelType.Attributes.Add(CreateSchemaAttribute("Exclude", ScriptSchemaStandardType.NameOrTagList, XmlSchemaUse.Optional));
+			LabelType.Attributes.Add(CreateSchemaAttribute("If", ScriptSchemaStandardType.BalancedString, XmlSchemaUse.Optional));
+			return LabelType;
+		}
+
+		/// <summary>
+		/// Creates the schema type representing a notifier
+		/// </summary>
+		/// <returns>Type definition for a notifier</returns>
+		static XmlSchemaType CreateNotifyType()
+		{
+			XmlSchemaComplexType AggregateType = new XmlSchemaComplexType();
+			AggregateType.Name = GetTypeName(ScriptSchemaStandardType.Notify);
+			AggregateType.Attributes.Add(CreateSchemaAttribute("Targets", ScriptSchemaStandardType.NameOrTagList, XmlSchemaUse.Optional));
+			AggregateType.Attributes.Add(CreateSchemaAttribute("Except", ScriptSchemaStandardType.NameOrTagList, XmlSchemaUse.Optional));
+			AggregateType.Attributes.Add(CreateSchemaAttribute("Nodes", ScriptSchemaStandardType.NameOrTagList, XmlSchemaUse.Optional));
+			AggregateType.Attributes.Add(CreateSchemaAttribute("Reports", ScriptSchemaStandardType.NameList, XmlSchemaUse.Optional));
+			AggregateType.Attributes.Add(CreateSchemaAttribute("Users", ScriptSchemaStandardType.BalancedString, XmlSchemaUse.Optional));
+			AggregateType.Attributes.Add(CreateSchemaAttribute("Submitters", ScriptSchemaStandardType.BalancedString, XmlSchemaUse.Optional));
+			AggregateType.Attributes.Add(CreateSchemaAttribute("Warnings", ScriptSchemaStandardType.Boolean, XmlSchemaUse.Optional));
+			AggregateType.Attributes.Add(CreateSchemaAttribute("If", ScriptSchemaStandardType.BalancedString, XmlSchemaUse.Optional));
+			AggregateType.Attributes.Add(CreateSchemaAttribute("Absolute", ScriptSchemaStandardType.BalancedString, XmlSchemaUse.Optional));
+			return AggregateType;
+		}
+
+		/// <summary>
+		/// Creates the schema type representing an include type
+		/// </summary>
+		/// <returns>Type definition for an include directive</returns>
+		static XmlSchemaType CreateIncludeType()
+		{
+			XmlSchemaComplexType PropertyType = new XmlSchemaComplexType();
+			PropertyType.Name = GetTypeName(ScriptSchemaStandardType.Include);
+			PropertyType.Attributes.Add(CreateSchemaAttribute("Script", ScriptSchemaStandardType.BalancedString, XmlSchemaUse.Required));
+			PropertyType.Attributes.Add(CreateSchemaAttribute("If", ScriptSchemaStandardType.BalancedString, XmlSchemaUse.Optional));
+			return PropertyType;
+		}
+
+		/// <summary>
+		/// Creates the schema type representing a parameter type
+		/// </summary>
+		/// <returns>Type definition for a parameter</returns>
+		static XmlSchemaType CreateOptionType()
+		{
+			XmlSchemaComplexType OptionType = new XmlSchemaComplexType();
+			OptionType.Name = GetTypeName(ScriptSchemaStandardType.Option);
+			OptionType.Attributes.Add(CreateSchemaAttribute("Name", ScriptSchemaStandardType.Name, XmlSchemaUse.Required));
+			OptionType.Attributes.Add(CreateSchemaAttribute("Restrict", StringTypeName, XmlSchemaUse.Optional));
+			OptionType.Attributes.Add(CreateSchemaAttribute("DefaultValue", StringTypeName, XmlSchemaUse.Required));
+			OptionType.Attributes.Add(CreateSchemaAttribute("Description", StringTypeName, XmlSchemaUse.Required));
+			OptionType.Attributes.Add(CreateSchemaAttribute("If", ScriptSchemaStandardType.BalancedString, XmlSchemaUse.Optional));
+			return OptionType;
+		}
+
+		/// <summary>
+		/// Creates the schema type representing a environment variable type
+		/// </summary>
+		/// <returns>Type definition for an environment variable property</returns>
+		static XmlSchemaType CreateEnvVarType()
+		{
+			XmlSchemaComplexType EnvVarType = new XmlSchemaComplexType();
+			EnvVarType.Name = GetTypeName(ScriptSchemaStandardType.EnvVar);
+			EnvVarType.Attributes.Add(CreateSchemaAttribute("Name", ScriptSchemaStandardType.Name, XmlSchemaUse.Required));
+			EnvVarType.Attributes.Add(CreateSchemaAttribute("If", ScriptSchemaStandardType.BalancedString, XmlSchemaUse.Optional));
+			return EnvVarType;
+		}
+
+		/// <summary>
+		/// Creates the schema type representing a property type
+		/// </summary>
+		/// <returns>Type definition for a property</returns>
+		static XmlSchemaType CreatePropertyType()
+		{
+			XmlSchemaSimpleContentExtension Extension = new XmlSchemaSimpleContentExtension();
+			Extension.BaseTypeName = StringTypeName;
+			Extension.Attributes.Add(CreateSchemaAttribute("Name", ScriptSchemaStandardType.Name, XmlSchemaUse.Required));
+			Extension.Attributes.Add(CreateSchemaAttribute("Value", StringTypeName, XmlSchemaUse.Optional));
+			Extension.Attributes.Add(CreateSchemaAttribute("If", ScriptSchemaStandardType.BalancedString, XmlSchemaUse.Optional));
+
+			XmlSchemaSimpleContent ContentModel = new XmlSchemaSimpleContent();
+			ContentModel.Content = Extension;
+
+			XmlSchemaComplexType PropertyType = new XmlSchemaComplexType();
+			PropertyType.Name = GetTypeName(ScriptSchemaStandardType.Property);
+			PropertyType.ContentModel = ContentModel;
+			return PropertyType;
+		}
+
+		/// <summary>
+		/// Creates the schema type representing a regex type
+		/// </summary>
+		/// <returns>Type definition for a regex</returns>
+		static XmlSchemaType CreateRegexType()
+		{
+			XmlSchemaSimpleContentExtension Extension = new XmlSchemaSimpleContentExtension();
+			Extension.BaseTypeName = StringTypeName;
+
+			Extension.Attributes.Add(CreateSchemaAttribute("Input", StringTypeName, XmlSchemaUse.Required));
+			Extension.Attributes.Add(CreateSchemaAttribute("Pattern", StringTypeName, XmlSchemaUse.Required));
+			Extension.Attributes.Add(CreateSchemaAttribute("Capture", ScriptSchemaStandardType.NameList, XmlSchemaUse.Required));
+			Extension.Attributes.Add(CreateSchemaAttribute("Optional", ScriptSchemaStandardType.Boolean, XmlSchemaUse.Optional));
+			Extension.Attributes.Add(CreateSchemaAttribute("If", ScriptSchemaStandardType.BalancedString, XmlSchemaUse.Optional));
+
+			XmlSchemaSimpleContent ContentModel = new XmlSchemaSimpleContent();
+			ContentModel.Content = Extension;
+
+			XmlSchemaComplexType RegexType  = new XmlSchemaComplexType();
+			RegexType.Name = GetTypeName(ScriptSchemaStandardType.Regex);
+			RegexType.ContentModel = ContentModel;
+			return RegexType;
+		}
+
+		/// <summary>
+		/// Creates the schema type representing the macro type
+		/// </summary>
+		/// <returns>Type definition for a node</returns>
+		static XmlSchemaType CreateMacroType()
+		{
+			XmlSchemaComplexContentExtension Extension = new XmlSchemaComplexContentExtension();
+			Extension.BaseTypeName = GetQualifiedTypeName(ScriptSchemaStandardType.MacroBody);
+			Extension.Attributes.Add(CreateSchemaAttribute("Name", ScriptSchemaStandardType.Name, XmlSchemaUse.Required));
+			Extension.Attributes.Add(CreateSchemaAttribute("Arguments", ScriptSchemaStandardType.BalancedString, XmlSchemaUse.Optional));
+			Extension.Attributes.Add(CreateSchemaAttribute("OptionalArguments", ScriptSchemaStandardType.BalancedString, XmlSchemaUse.Optional));
+			Extension.Attributes.Add(CreateSchemaAttribute("If", ScriptSchemaStandardType.BalancedString, XmlSchemaUse.Optional));
+
+			XmlSchemaComplexContent ContentModel = new XmlSchemaComplexContent();
+			ContentModel.Content = Extension;
+
+			XmlSchemaComplexType ComplexType = new XmlSchemaComplexType();
+			ComplexType.Name = GetTypeName(ScriptSchemaStandardType.Macro);
+			ComplexType.ContentModel = ContentModel;
+			return ComplexType;
+		}
+
+		/// <summary>
+		/// Creates the schema type representing the macro type
+		/// </summary>
+		/// <returns>Type definition for a node</returns>
+		static XmlSchemaType CreateMacroBodyType(Dictionary<string, XmlSchemaComplexType> TaskNameToType)
+		{
+			XmlSchemaChoice MacroChoice = new XmlSchemaChoice();
+			MacroChoice.MinOccurs = 0;
+			MacroChoice.MaxOccursString = "unbounded";
+
+			// Graph scope
+			MacroChoice.Items.Add(CreateSchemaElement("Include", ScriptSchemaStandardType.Include));
+			MacroChoice.Items.Add(CreateSchemaElement("Option", ScriptSchemaStandardType.Option));
+			MacroChoice.Items.Add(CreateSchemaElement("EnvVar", ScriptSchemaStandardType.EnvVar));
+			MacroChoice.Items.Add(CreateSchemaElement("Property", ScriptSchemaStandardType.Property));
+			MacroChoice.Items.Add(CreateSchemaElement("Regex", ScriptSchemaStandardType.Regex));
+			MacroChoice.Items.Add(CreateSchemaElement("Macro", ScriptSchemaStandardType.Macro));
+			MacroChoice.Items.Add(CreateSchemaElement("Agent", ScriptSchemaStandardType.Agent));
+			MacroChoice.Items.Add(CreateSchemaElement("Aggregate", ScriptSchemaStandardType.Aggregate));
+			MacroChoice.Items.Add(CreateSchemaElement("Report", ScriptSchemaStandardType.Report));
+			MacroChoice.Items.Add(CreateSchemaElement("Badge", ScriptSchemaStandardType.Badge));
+			MacroChoice.Items.Add(CreateSchemaElement("Notify", ScriptSchemaStandardType.Notify));
+			MacroChoice.Items.Add(CreateSchemaElement("Trace", ScriptSchemaStandardType.Trace));
+			MacroChoice.Items.Add(CreateSchemaElement("Warning", ScriptSchemaStandardType.Warning));
+			MacroChoice.Items.Add(CreateSchemaElement("Error", ScriptSchemaStandardType.Error));
+			MacroChoice.Items.Add(CreateSchemaElement("Expand", ScriptSchemaStandardType.Expand));
+
+			// Agent scope
+			MacroChoice.Items.Add(CreateSchemaElement("Node", ScriptSchemaStandardType.Node));
+
+			// Node scope
+			MacroChoice.Items.Add(CreateDoElement(ScriptSchemaStandardType.NodeBody));
+			MacroChoice.Items.Add(CreateSwitchElement(ScriptSchemaStandardType.NodeBody));
+			MacroChoice.Items.Add(CreateForEachElement(ScriptSchemaStandardType.NodeBody));
+			foreach (KeyValuePair<string, XmlSchemaComplexType> Pair in TaskNameToType.OrderBy(x => x.Key))
+			{
+				MacroChoice.Items.Add(CreateSchemaElement(Pair.Key, new XmlQualifiedName(Pair.Value.Name, NamespaceURI)));
+			}
+
+			XmlSchemaComplexType NodeType = new XmlSchemaComplexType();
+			NodeType.Name = GetTypeName(ScriptSchemaStandardType.MacroBody);
+			NodeType.Particle = MacroChoice;
+			return NodeType;
+		}
+
+		/// <summary>
+		/// Creates the schema type representing a macro expansion
+		/// </summary>
+		/// <returns>Type definition for expanding a macro</returns>
+		static XmlSchemaType CreateExpandType()
+		{
+			XmlSchemaAnyAttribute AnyAttribute = new XmlSchemaAnyAttribute();
+			AnyAttribute.ProcessContents = XmlSchemaContentProcessing.Skip;
+
+			XmlSchemaComplexType PropertyType = new XmlSchemaComplexType();
+			PropertyType.Name = GetTypeName(ScriptSchemaStandardType.Expand);
+			PropertyType.Attributes.Add(CreateSchemaAttribute("Name", ScriptSchemaStandardType.BalancedString, XmlSchemaUse.Required));
+			PropertyType.Attributes.Add(CreateSchemaAttribute("If", ScriptSchemaStandardType.BalancedString, XmlSchemaUse.Optional));
+			PropertyType.AnyAttribute = AnyAttribute;
+			return PropertyType;
+		}
+
+		/// <summary>
+		/// Creates the schema type representing a warning or error type
+		/// </summary>
+		/// <returns>Type definition for a warning</returns>
+		static XmlSchemaType CreateDiagnosticType(ScriptSchemaStandardType StandardType)
+		{
+			XmlSchemaComplexType PropertyType = new XmlSchemaComplexType();
+			PropertyType.Name = GetTypeName(StandardType);
+			PropertyType.Attributes.Add(CreateSchemaAttribute("Message", ScriptSchemaStandardType.BalancedString, XmlSchemaUse.Required));
+			PropertyType.Attributes.Add(CreateSchemaAttribute("If", ScriptSchemaStandardType.BalancedString, XmlSchemaUse.Optional));
+			return PropertyType;
+		}
+
+		/// <summary>
+		/// Creates an element representing a conditional "Do" block, which recursively contains another type
+		/// </summary>
+		/// <param name="InnerType">The base type for the do block to contain</param>
+		/// <returns>New schema element for the block</returns>
+		static XmlSchemaElement CreateDoElement(ScriptSchemaStandardType InnerType)
+		{
+			XmlSchemaComplexContentExtension Extension = new XmlSchemaComplexContentExtension();
+			Extension.BaseTypeName = GetQualifiedTypeName(InnerType);
+			Extension.Attributes.Add(CreateSchemaAttribute("If", ScriptSchemaStandardType.BalancedString, XmlSchemaUse.Optional));
+
+			XmlSchemaComplexContent ContentModel = new XmlSchemaComplexContent();
+			ContentModel.Content = Extension;
+
+			XmlSchemaComplexType SchemaType = new XmlSchemaComplexType();
+			SchemaType.ContentModel = ContentModel;
+
+			XmlSchemaElement Element = new XmlSchemaElement();
+			Element.Name = "Do";
+			Element.SchemaType = SchemaType;
+			return Element;
+		}
+
+		/// <summary>
+		/// Creates an element representing a conditional "Switch" block, which recursively contains another type
+		/// </summary>
+		/// <param name="InnerType">The base type for the do block to contain</param>
+		/// <returns>New schema element for the block</returns>
+		static XmlSchemaElement CreateSwitchElement(ScriptSchemaStandardType InnerType)
+		{
+			// Create the "Option" element
+			XmlSchemaComplexContentExtension CaseExtension = new XmlSchemaComplexContentExtension();
+			CaseExtension.BaseTypeName = GetQualifiedTypeName(InnerType);
+			CaseExtension.Attributes.Add(CreateSchemaAttribute("If", ScriptSchemaStandardType.BalancedString, XmlSchemaUse.Required));
+
+			XmlSchemaComplexContent CaseContentModel = new XmlSchemaComplexContent();
+			CaseContentModel.Content = CaseExtension;
+
+			XmlSchemaComplexType CaseSchemaType = new XmlSchemaComplexType();
+			CaseSchemaType.ContentModel = CaseContentModel;
+
+			XmlSchemaElement CaseElement = new XmlSchemaElement();
+			CaseElement.Name = "Case";
+			CaseElement.SchemaType = CaseSchemaType;
+			CaseElement.MinOccurs = 0;
+			CaseElement.MaxOccursString = "unbounded";
+
+			// Create the "Otherwise" element
+			XmlSchemaElement OtherwiseElement = new XmlSchemaElement();
+			OtherwiseElement.Name = "Default";
+			OtherwiseElement.SchemaTypeName = GetQualifiedTypeName(InnerType);
+			OtherwiseElement.MinOccurs = 0;
+			OtherwiseElement.MaxOccurs = 1;
+
+			// Create the "Switch" element
+			XmlSchemaSequence SwitchSequence = new XmlSchemaSequence();
+			SwitchSequence.Items.Add(CaseElement);
+			SwitchSequence.Items.Add(OtherwiseElement);
+
+			XmlSchemaComplexType SwitchSchemaType = new XmlSchemaComplexType();
+			SwitchSchemaType.Particle = SwitchSequence;
+
+			XmlSchemaElement SwitchElement = new XmlSchemaElement();
+			SwitchElement.Name = "Switch";
+			SwitchElement.SchemaType = SwitchSchemaType;
+			return SwitchElement;
+		}
+
+		/// <summary>
+		/// Creates an element representing a conditional "ForEach" block, which recursively contains another type
+		/// </summary>
+		/// <param name="InnerType">The base type for the foreach block to contain</param>
+		/// <returns>New schema element for the block</returns>
+		static XmlSchemaElement CreateForEachElement(ScriptSchemaStandardType InnerType)
+		{
+			XmlSchemaComplexContentExtension Extension = new XmlSchemaComplexContentExtension();
+			Extension.BaseTypeName = GetQualifiedTypeName(InnerType);
+			Extension.Attributes.Add(CreateSchemaAttribute("Name", ScriptSchemaStandardType.BalancedString, XmlSchemaUse.Required));
+			Extension.Attributes.Add(CreateSchemaAttribute("Values", ScriptSchemaStandardType.BalancedString, XmlSchemaUse.Required));
+			Extension.Attributes.Add(CreateSchemaAttribute("Separator", ScriptSchemaStandardType.BalancedString, XmlSchemaUse.Optional));
+			Extension.Attributes.Add(CreateSchemaAttribute("If", ScriptSchemaStandardType.BalancedString, XmlSchemaUse.Optional));
+
+			XmlSchemaComplexContent ContentModel = new XmlSchemaComplexContent();
+			ContentModel.Content = Extension;
+
+			XmlSchemaComplexType SchemaType = new XmlSchemaComplexType();
+			SchemaType.ContentModel = ContentModel;
+
+			XmlSchemaElement Element = new XmlSchemaElement();
+			Element.Name = "ForEach";
+			Element.SchemaType = SchemaType;
+			return Element;
+		}
+
+		/// <summary>
+		/// Constructs an XmlSchemaElement and initializes it with the given parameters
+		/// </summary>
+		/// <param name="Name">Element name</param>
+		/// <param name="SchemaType">Type enumeration for the attribute</param>
+		/// <returns>A new XmlSchemaElement object</returns>
+		static XmlSchemaElement CreateSchemaElement(string Name, ScriptSchemaStandardType SchemaType)
+		{
+			return CreateSchemaElement(Name, GetQualifiedTypeName(SchemaType));
+		}
+
+		/// <summary>
+		/// Constructs an XmlSchemaElement and initializes it with the given parameters
+		/// </summary>
+		/// <param name="Name">Element name</param>
+		/// <param name="SchemaTypeName">Qualified name of the type for this element</param>
+		/// <returns>A new XmlSchemaElement object</returns>
+		static XmlSchemaElement CreateSchemaElement(string Name, XmlQualifiedName SchemaTypeName)
+		{
+			XmlSchemaElement Element = new XmlSchemaElement();
+			Element.Name = Name;
+			Element.SchemaTypeName = SchemaTypeName;
+			return Element;
+		}
+
+		/// <summary>
+		/// Constructs an XmlSchemaAttribute and initialize it with the given parameters
+		/// </summary>
+		/// <param name="Name">The attribute name</param>
+		/// <param name="SchemaType">Type enumeration for the attribute</param>
+		/// <param name="Use">Whether the attribute is required or optional</param>
+		/// <returns>A new XmlSchemaAttribute object</returns>
+		static XmlSchemaAttribute CreateSchemaAttribute(string Name, ScriptSchemaStandardType SchemaType, XmlSchemaUse Use)
+		{
+			return CreateSchemaAttribute(Name, GetQualifiedTypeName(SchemaType), Use);
+		}
+
+		/// <summary>
+		/// Constructs an XmlSchemaAttribute and initialize it with the given parameters
+		/// </summary>
+		/// <param name="Name">The attribute name</param>
+		/// <param name="SchemaTypeName">Qualified name of the type for this attribute</param>
+		/// <param name="Use">Whether the attribute is required or optional</param>
+		/// <returns>The new attribute</returns>
+		static XmlSchemaAttribute CreateSchemaAttribute(string Name, XmlQualifiedName SchemaTypeName, XmlSchemaUse Use)
+		{
+			XmlSchemaAttribute Attribute = new XmlSchemaAttribute();
+			Attribute.Name = Name;
+			Attribute.SchemaTypeName = SchemaTypeName;
+			Attribute.Use = Use;
+			return Attribute;
+		}
+
+		/// <summary>
+		/// Creates a simple type that is the union of two other types
+		/// </summary>
+		/// <param name="Name">The name of the type</param>
+		/// <param name="ValidTypes">List of valid types for the union</param>
+		/// <returns>A simple type which will match the given pattern</returns>
+		static XmlSchemaSimpleType CreateSimpleTypeFromUnion(string Name, params XmlSchemaType[] ValidTypes)
+		{
+			XmlSchemaSimpleTypeUnion Union = new XmlSchemaSimpleTypeUnion();
+			foreach (XmlSchemaType ValidType in ValidTypes)
+			{
+				Union.BaseTypes.Add(ValidType);
+			}
+
+			XmlSchemaSimpleType UnionType = new XmlSchemaSimpleType();
+			UnionType.Name = Name;
+			UnionType.Content = Union;
+			return UnionType;
+		}
+
+		/// <summary>
+		/// Creates a simple type that matches a regex
+		/// </summary>
+		/// <param name="Name">Name of the new type</param>
+		/// <param name="Pattern">Regex pattern to match</param>
+		/// <returns>A simple type which will match the given pattern</returns>
+		static XmlSchemaSimpleType CreateSimpleTypeFromRegex(string Name, string Pattern)
+		{
+			XmlSchemaPatternFacet PatternFacet = new XmlSchemaPatternFacet();
+			PatternFacet.Value = Pattern;
+
+			XmlSchemaSimpleTypeRestriction Restriction = new XmlSchemaSimpleTypeRestriction();
+			Restriction.BaseTypeName = StringTypeName;
+			Restriction.Facets.Add(PatternFacet);
+
+			XmlSchemaSimpleType SimpleType = new XmlSchemaSimpleType();
+			SimpleType.Name = Name;
+			SimpleType.Content = Restriction;
+			return SimpleType;
+		}
+
+		/// <summary>
+		/// Create a schema type for the given user type. Currently only handles enumerations.
+		/// </summary>
+		/// <param name="Name">Name for the new type</param>
+		/// <param name="Type">CLR type information to create a schema type for</param>
+		static XmlSchemaType CreateUserType(string Name, Type Type)
+		{
+			if(Type.IsEnum)
+			{
+				return CreateSimpleTypeFromUnion(Name, CreateEnumType(null, Type), CreateSimpleTypeFromRegex(null, StringWithPropertiesPattern));
+			}
+			else
+			{
+				throw new Exception($"Cannot create custom type in schema for '{Type.Name}'");
+			}
+		}
+
+		/// <summary>
+		/// Create a schema type for the given enum.
+		/// </summary>
+		/// <param name="Name">Name for the new type</param>
+		/// <param name="Type">CLR type information to create a schema type for</param>
+		static XmlSchemaType CreateEnumType(string Name, Type Type)
+		{
+			XmlSchemaSimpleTypeRestriction Restriction = new XmlSchemaSimpleTypeRestriction();
+			Restriction.BaseTypeName = StringTypeName;
+
+			foreach(string EnumName in Enum.GetNames(Type))
+			{
+				XmlSchemaEnumerationFacet Facet = new XmlSchemaEnumerationFacet();
+				Facet.Value = EnumName;
+				Restriction.Facets.Add(Facet);
+			}
+
+			XmlSchemaSimpleType SchemaType = new XmlSchemaSimpleType();
+			SchemaType.Name = Name;
+			SchemaType.Content = Restriction;
+			return SchemaType;
+		}
+	}
+}
diff --git a/Engine/Source/Programs/Shared/EpicGames.BuildGraph/BgScript.cs b/Engine/Source/Programs/Shared/EpicGames.BuildGraph/BgScript.cs
new file mode 100644
--- /dev/null
+++ b/Engine/Source/Programs/Shared/EpicGames.BuildGraph/BgScript.cs
@@ -0,0 +1,774 @@
+// Copyright Epic Games, Inc. All Rights Reserved.
+
+using System;
+using System.Collections.Generic;
+using System.Text;
+using System.IO;
+using System.Reflection;
+using System.Xml;
+using System.Linq;
+using System.Diagnostics;
+using EpicGames.Core;
+using Microsoft.Extensions.Logging;
+
+namespace EpicGames.BuildGraph
+{
+	/// <summary>
+	/// Options for how the graph should be printed
+	/// </summary>
+	public enum GraphPrintOptions
+	{
+		/// <summary>
+		/// Includes a list of the graph options
+		/// </summary>
+		ShowCommandLineOptions = 0x1,
+
+		/// <summary>
+		/// Includes the list of dependencies for each node
+		/// </summary>
+		ShowDependencies = 0x2,
+
+		/// <summary>
+		/// Includes the list of notifiers for each node
+		/// </summary>
+		ShowNotifications = 0x4,
+	}
+
+	/// <summary>
+	/// Diagnostic message from the graph script. These messages are parsed at startup, then culled along with the rest of the graph nodes before output. Doing so
+	/// allows errors and warnings which are only output if a node is part of the graph being executed.
+	/// </summary>
+	public class BgScriptDiagnostic
+	{
+		/// <summary>
+		/// The diagnostic event type
+		/// </summary>
+		public LogEventType EventType;
+
+		/// <summary>
+		/// The message to display
+		/// </summary>
+		public string Message;
+
+		/// <summary>
+		/// The node which this diagnostic is declared in. If the node is culled from the graph, the message will not be displayed.
+		/// </summary>
+		public BgNode EnclosingNode;
+
+		/// <summary>
+		/// The agent that this diagnostic is declared in. If the entire agent is culled from the graph, the message will not be displayed.
+		/// </summary>
+		public BgAgent EnclosingAgent;
+	}
+
+	/// <summary>
+	/// Represents a graph option. These are expanded during preprocessing, but are retained in order to display help messages.
+	/// </summary>
+	public class BgScriptOption
+	{
+		/// <summary>
+		/// Name of this option
+		/// </summary>
+		public string Name;
+
+		/// <summary>
+		/// Description for this option
+		/// </summary>
+		public string Description;
+
+		/// <summary>
+		/// Default value for this option
+		/// </summary>
+		public string DefaultValue;
+
+		/// <summary>
+		/// Constructor
+		/// </summary>
+		/// <param name="Name">The name of this option</param>
+		/// <param name="Description">Description of the option, for display on help pages</param>
+		/// <param name="DefaultValue">Default value for the option</param>
+		public BgScriptOption(string Name, string Description, string DefaultValue)
+		{
+			this.Name = Name;
+			this.Description = Description;
+			this.DefaultValue = DefaultValue;
+		}
+
+		/// <summary>
+		/// Returns a name of this option for debugging
+		/// </summary>
+		/// <returns>Name of the option</returns>
+		public override string ToString()
+		{
+			return Name;
+		}
+	}
+
+	/// <summary>
+	/// Definition of a graph.
+	/// </summary>
+	public class BgScript
+	{
+		/// <summary>
+		/// List of options, in the order they were specified
+		/// </summary>
+		public List<BgScriptOption> Options = new List<BgScriptOption>();
+
+		/// <summary>
+		/// List of agents containing nodes to execute
+		/// </summary>
+		public List<BgAgent> Agents = new List<BgAgent>();
+
+		/// <summary>
+		/// Mapping from name to agent
+		/// </summary>
+		public Dictionary<string, BgAgent> NameToAgent = new Dictionary<string, BgAgent>(StringComparer.OrdinalIgnoreCase);
+
+		/// <summary>
+		/// Mapping of names to the corresponding node.
+		/// </summary>
+		public Dictionary<string, BgNode> NameToNode = new Dictionary<string,BgNode>(StringComparer.OrdinalIgnoreCase);
+
+		/// <summary>
+		/// Mapping of names to the corresponding report.
+		/// </summary>
+		public Dictionary<string, BgReport> NameToReport = new Dictionary<string, BgReport>(StringComparer.OrdinalIgnoreCase);
+
+		/// <summary>
+		/// Mapping of names to their corresponding node output.
+		/// </summary>
+		public Dictionary<string, BgNodeOutput> TagNameToNodeOutput = new Dictionary<string,BgNodeOutput>(StringComparer.OrdinalIgnoreCase);
+
+		/// <summary>
+		/// Mapping of aggregate names to their respective nodes
+		/// </summary>
+		public Dictionary<string, BgAggregate> NameToAggregate = new Dictionary<string, BgAggregate>(StringComparer.OrdinalIgnoreCase);
+
+		/// <summary>
+		/// List of badges that can be displayed for this build
+		/// </summary>
+		public List<BgBadge> Badges = new List<BgBadge>();
+
+		/// <summary>
+		/// List of labels that can be displayed for this build
+		/// </summary>
+		public List<BgLabel> Labels = new List<BgLabel>();
+
+		/// <summary>
+		/// Diagnostic messages for this graph
+		/// </summary>
+		public List<BgScriptDiagnostic> Diagnostics = new List<BgScriptDiagnostic>();
+
+		/// <summary>
+		/// Default constructor
+		/// </summary>
+		public BgScript()
+		{
+		}
+
+		/// <summary>
+		/// Checks whether a given name already exists
+		/// </summary>
+		/// <param name="Name">The name to check.</param>
+		/// <returns>True if the name exists, false otherwise.</returns>
+		public bool ContainsName(string Name)
+		{
+			return NameToNode.ContainsKey(Name) || NameToReport.ContainsKey(Name) || NameToAggregate.ContainsKey(Name);
+		}
+
+		/// <summary>
+		/// Tries to resolve the given name to one or more nodes. Checks for aggregates, and actual nodes.
+		/// </summary>
+		/// <param name="Name">The name to search for</param>
+		/// <param name="OutNodes">If the name is a match, receives an array of nodes and their output names</param>
+		/// <returns>True if the name was found, false otherwise.</returns>
+		public bool TryResolveReference(string Name, out BgNode[] OutNodes)
+		{
+			// Check if it's a tag reference or node reference
+			if(Name.StartsWith("#"))
+			{
+				// Check if it's a regular node or output name
+				BgNodeOutput Output;
+				if(TagNameToNodeOutput.TryGetValue(Name, out Output))
+				{
+					OutNodes = new BgNode[]{ Output.ProducingNode };
+					return true;
+				}
+			}
+			else
+			{
+				// Check if it's a regular node or output name
+				BgNode Node;
+				if(NameToNode.TryGetValue(Name, out Node))
+				{
+					OutNodes = new BgNode[]{ Node };
+					return true;
+				}
+
+				// Check if it's an aggregate name
+				BgAggregate Aggregate;
+				if(NameToAggregate.TryGetValue(Name, out Aggregate))
+				{
+					OutNodes = Aggregate.RequiredNodes.ToArray();
+					return true;
+				}
+
+				// Check if it's a group name
+				BgAgent Agent;
+				if(NameToAgent.TryGetValue(Name, out Agent))
+				{
+					OutNodes = Agent.Nodes.ToArray();
+					return true;
+				}
+			}
+
+			// Otherwise fail
+			OutNodes = null;
+			return false;
+		}
+
+		/// <summary>
+		/// Tries to resolve the given name to one or more node outputs. Checks for aggregates, and actual nodes.
+		/// </summary>
+		/// <param name="Name">The name to search for</param>
+		/// <param name="OutOutputs">If the name is a match, receives an array of nodes and their output names</param>
+		/// <returns>True if the name was found, false otherwise.</returns>
+		public bool TryResolveInputReference(string Name, out BgNodeOutput[] OutOutputs)
+		{
+			// Check if it's a tag reference or node reference
+			if(Name.StartsWith("#"))
+			{
+				// Check if it's a regular node or output name
+				BgNodeOutput Output;
+				if(TagNameToNodeOutput.TryGetValue(Name, out Output))
+				{
+					OutOutputs = new BgNodeOutput[]{ Output };
+					return true;
+				}
+			}
+			else
+			{
+				// Check if it's a regular node or output name
+				BgNode Node;
+				if(NameToNode.TryGetValue(Name, out Node))
+				{
+					OutOutputs = Node.Outputs.Union(Node.Inputs).ToArray();
+					return true;
+				}
+
+				// Check if it's an aggregate name
+				BgAggregate Aggregate;
+				if(NameToAggregate.TryGetValue(Name, out Aggregate))
+				{
+					OutOutputs = Aggregate.RequiredNodes.SelectMany(x => x.Outputs.Union(x.Inputs)).Distinct().ToArray();
+					return true;
+				}
+			}
+
+			// Otherwise fail
+			OutOutputs = null;
+			return false;
+		}
+
+		/// <summary>
+		/// Cull the graph to only include the given nodes and their dependencies
+		/// </summary>
+		/// <param name="TargetNodes">A set of target nodes to build</param>
+		public void Select(IEnumerable<BgNode> TargetNodes)
+		{
+			// Find this node and all its dependencies
+			HashSet<BgNode> RetainNodes = new HashSet<BgNode>(TargetNodes);
+			foreach(BgNode TargetNode in TargetNodes)
+			{
+				RetainNodes.UnionWith(TargetNode.InputDependencies);
+			}
+
+			// Remove all the nodes which are not marked to be kept
+			foreach(BgAgent Agent in Agents)
+			{
+				Agent.Nodes = Agent.Nodes.Where(x => RetainNodes.Contains(x)).ToList();
+			}
+
+			// Remove all the empty agents
+			Agents.RemoveAll(x => x.Nodes.Count == 0);
+
+			// Trim down the list of nodes for each report to the ones that are being built
+			foreach (BgReport Report in NameToReport.Values)
+			{
+				Report.Nodes.RemoveWhere(x => !RetainNodes.Contains(x));
+			}
+
+			// Remove all the empty reports
+			NameToReport = NameToReport.Where(x => x.Value.Nodes.Count > 0).ToDictionary(Pair => Pair.Key, Pair => Pair.Value, StringComparer.InvariantCultureIgnoreCase);
+
+			// Remove all the order dependencies which are no longer part of the graph. Since we don't need to build them, we don't need to wait for them
+			foreach(BgNode Node in RetainNodes)
+			{
+				Node.OrderDependencies = Node.OrderDependencies.Where(x => RetainNodes.Contains(x)).ToArray();
+			}
+
+			// Create a new list of aggregates for everything that's left
+			Dictionary<string, BgAggregate> NewNameToAggregate = new Dictionary<string, BgAggregate>(NameToAggregate.Comparer);
+			foreach(BgAggregate Aggregate in NameToAggregate.Values)
+			{
+				if (Aggregate.RequiredNodes.All(x => RetainNodes.Contains(x)))
+				{
+					NewNameToAggregate[Aggregate.Name] = Aggregate;
+				}
+			}
+			NameToAggregate = NewNameToAggregate;
+
+			// Remove any labels that are no longer value
+			foreach (BgLabel Label in Labels)
+			{
+				Label.RequiredNodes.RemoveWhere(x => !RetainNodes.Contains(x));
+				Label.IncludedNodes.RemoveWhere(x => !RetainNodes.Contains(x));
+			}
+			Labels.RemoveAll(x => x.RequiredNodes.Count == 0);
+
+			// Remove any badges which do not have all their dependencies
+			Badges.RemoveAll(x => x.Nodes.Any(y => !RetainNodes.Contains(y)));
+
+			// Remove any diagnostics which are no longer part of the graph
+			Diagnostics.RemoveAll(x => (x.EnclosingNode != null && !RetainNodes.Contains(x.EnclosingNode)) || (x.EnclosingAgent != null && !Agents.Contains(x.EnclosingAgent)));
+		}
+
+		/// <summary>
+		/// Writes a preprocessed build graph to a script file
+		/// </summary>
+		/// <param name="File">The file to load</param>
+		/// <param name="SchemaFile">Schema file for validation</param>
+		public void Write(FileReference File, FileReference SchemaFile)
+		{
+			XmlWriterSettings Settings = new XmlWriterSettings();
+			Settings.Indent = true;
+			Settings.IndentChars = "\t";
+
+			using (XmlWriter Writer = XmlWriter.Create(File.FullName, Settings))
+			{
+				Writer.WriteStartElement("BuildGraph", "http://www.epicgames.com/BuildGraph");
+
+				if (SchemaFile != null)
+				{
+					Writer.WriteAttributeString("schemaLocation", "http://www.w3.org/2001/XMLSchema-instance", "http://www.epicgames.com/BuildGraph " + SchemaFile.MakeRelativeTo(File.Directory));
+				}
+
+				foreach (BgAgent Agent in Agents)
+				{
+					Agent.Write(Writer);
+				}
+
+				foreach (BgAggregate Aggregate in NameToAggregate.Values)
+				{
+					// If the aggregate has no required elements, skip it.
+					if (Aggregate.RequiredNodes.Count == 0)
+					{
+						continue;
+					}
+
+					Writer.WriteStartElement("Aggregate");
+					Writer.WriteAttributeString("Name", Aggregate.Name);
+					Writer.WriteAttributeString("Requires", String.Join(";", Aggregate.RequiredNodes.Select(x => x.Name)));
+					Writer.WriteEndElement();
+				}
+
+				foreach(BgLabel Label in Labels)
+				{
+					Writer.WriteStartElement("Label");
+					if (Label.DashboardCategory != null)
+					{
+						Writer.WriteAttributeString("Category", Label.DashboardCategory);
+					}
+					Writer.WriteAttributeString("Name", Label.DashboardName);
+					Writer.WriteAttributeString("Requires", String.Join(";", Label.RequiredNodes.Select(x => x.Name)));
+
+					HashSet<BgNode> IncludedNodes = new HashSet<BgNode>(Label.IncludedNodes);
+					IncludedNodes.ExceptWith(Label.IncludedNodes.SelectMany(x => x.InputDependencies));
+					IncludedNodes.ExceptWith(Label.RequiredNodes);
+					if (IncludedNodes.Count > 0)
+					{
+						Writer.WriteAttributeString("Include", String.Join(";", IncludedNodes.Select(x => x.Name)));
+					}
+
+					HashSet<BgNode> ExcludedNodes = new HashSet<BgNode>(Label.IncludedNodes);
+					ExcludedNodes.UnionWith(Label.IncludedNodes.SelectMany(x => x.InputDependencies));
+					ExcludedNodes.ExceptWith(Label.IncludedNodes);
+					ExcludedNodes.ExceptWith(ExcludedNodes.ToArray().SelectMany(x => x.InputDependencies));
+					if (ExcludedNodes.Count > 0)
+					{
+						Writer.WriteAttributeString("Exclude", String.Join(";", ExcludedNodes.Select(x => x.Name)));
+					}
+					Writer.WriteEndElement();
+				}
+
+				foreach (BgReport Report in NameToReport.Values)
+				{
+					Writer.WriteStartElement("Report");
+					Writer.WriteAttributeString("Name", Report.Name);
+					Writer.WriteAttributeString("Requires", String.Join(";", Report.Nodes.Select(x => x.Name)));
+					Writer.WriteEndElement();
+				}
+
+				foreach (BgBadge Badge in Badges)
+				{
+					Writer.WriteStartElement("Badge");
+					Writer.WriteAttributeString("Name", Badge.Name);
+					if (Badge.Project != null)
+					{
+						Writer.WriteAttributeString("Project", Badge.Project);
+					}
+					if (Badge.Change != 0)
+					{
+						Writer.WriteAttributeString("Change", Badge.Change.ToString());
+					}
+					Writer.WriteAttributeString("Requires", String.Join(";", Badge.Nodes.Select(x => x.Name)));
+					Writer.WriteEndElement();
+				}
+
+				Writer.WriteEndElement();
+			}
+		}
+
+		/// <summary>
+		/// Export the build graph to a Json file, for parallel execution by the build system
+		/// </summary>
+		/// <param name="File">Output file to write</param>
+		/// <param name="CompletedNodes">Set of nodes which have been completed</param>
+		public void Export(FileReference File, HashSet<BgNode> CompletedNodes)
+		{
+			// Find all the nodes which we're actually going to execute. We'll use this to filter the graph.
+			HashSet<BgNode> NodesToExecute = new HashSet<BgNode>();
+			foreach(BgNode Node in Agents.SelectMany(x => x.Nodes))
+			{
+				if(!CompletedNodes.Contains(Node))
+				{
+					NodesToExecute.Add(Node);
+				}
+			}
+
+			// Open the output file
+			using(JsonWriter JsonWriter = new JsonWriter(File.FullName))
+			{
+				JsonWriter.WriteObjectStart();
+
+				// Write all the agents
+				JsonWriter.WriteArrayStart("Groups");
+				foreach(BgAgent Agent in Agents)
+				{
+					BgNode[] Nodes = Agent.Nodes.Where(x => NodesToExecute.Contains(x)).ToArray();
+					if(Nodes.Length > 0)
+					{
+						JsonWriter.WriteObjectStart();
+						JsonWriter.WriteValue("Name", Agent.Name);
+						JsonWriter.WriteArrayStart("Agent Types");
+						foreach(string AgentType in Agent.PossibleTypes)
+						{
+							JsonWriter.WriteValue(AgentType);
+						}
+						JsonWriter.WriteArrayEnd();
+						JsonWriter.WriteArrayStart("Nodes");
+						foreach(BgNode Node in Nodes)
+						{
+							JsonWriter.WriteObjectStart();
+							JsonWriter.WriteValue("Name", Node.Name);
+							JsonWriter.WriteValue("DependsOn", String.Join(";", Node.GetDirectOrderDependencies().Where(x => NodesToExecute.Contains(x))));
+							JsonWriter.WriteValue("RunEarly", Node.bRunEarly);
+							JsonWriter.WriteObjectStart("Notify");
+							JsonWriter.WriteValue("Default", String.Join(";", Node.NotifyUsers));
+							JsonWriter.WriteValue("Submitters", String.Join(";", Node.NotifySubmitters));
+							JsonWriter.WriteValue("Warnings", Node.bNotifyOnWarnings);
+							JsonWriter.WriteObjectEnd();
+							JsonWriter.WriteObjectEnd();
+						}
+						JsonWriter.WriteArrayEnd();
+						JsonWriter.WriteObjectEnd();
+					}
+				}
+				JsonWriter.WriteArrayEnd();
+
+				// Write all the badges
+				JsonWriter.WriteArrayStart("Badges");
+				foreach (BgBadge Badge in Badges)
+				{
+					BgNode[] Dependencies = Badge.Nodes.Where(x => NodesToExecute.Contains(x)).ToArray();
+					if (Dependencies.Length > 0)
+					{
+						// Reduce that list to the smallest subset of direct dependencies
+						HashSet<BgNode> DirectDependencies = new HashSet<BgNode>(Dependencies);
+						foreach (BgNode Dependency in Dependencies)
+						{
+							DirectDependencies.ExceptWith(Dependency.OrderDependencies);
+						}
+
+						JsonWriter.WriteObjectStart();
+						JsonWriter.WriteValue("Name", Badge.Name);
+						if (!String.IsNullOrEmpty(Badge.Project))
+						{
+							JsonWriter.WriteValue("Project", Badge.Project);
+						}
+						if(Badge.Change != 0)
+						{
+							JsonWriter.WriteValue("Change", Badge.Change);
+						}
+						JsonWriter.WriteValue("AllDependencies", String.Join(";", Agents.SelectMany(x => x.Nodes).Where(x => Dependencies.Contains(x)).Select(x => x.Name)));
+						JsonWriter.WriteValue("DirectDependencies", String.Join(";", DirectDependencies.Select(x => x.Name)));
+						JsonWriter.WriteObjectEnd();
+					}
+				}
+				JsonWriter.WriteArrayEnd();
+
+				// Write all the triggers and reports. 
+				JsonWriter.WriteArrayStart("Reports");
+				foreach (BgReport Report in NameToReport.Values)
+				{
+					BgNode[] Dependencies = Report.Nodes.Where(x => NodesToExecute.Contains(x)).ToArray();
+					if (Dependencies.Length > 0)
+					{
+						// Reduce that list to the smallest subset of direct dependencies
+						HashSet<BgNode> DirectDependencies = new HashSet<BgNode>(Dependencies);
+						foreach (BgNode Dependency in Dependencies)
+						{
+							DirectDependencies.ExceptWith(Dependency.OrderDependencies);
+						}
+
+						JsonWriter.WriteObjectStart();
+						JsonWriter.WriteValue("Name", Report.Name);
+						JsonWriter.WriteValue("AllDependencies", String.Join(";", Agents.SelectMany(x => x.Nodes).Where(x => Dependencies.Contains(x)).Select(x => x.Name)));
+						JsonWriter.WriteValue("DirectDependencies", String.Join(";", DirectDependencies.Select(x => x.Name)));
+						JsonWriter.WriteValue("Notify", String.Join(";", Report.NotifyUsers));
+						JsonWriter.WriteValue("IsTrigger", false);
+						JsonWriter.WriteObjectEnd();
+					}
+				}
+				JsonWriter.WriteArrayEnd();
+
+				JsonWriter.WriteObjectEnd();
+			}
+		}
+
+		/// <summary>
+		/// Export the build graph to a Json file for parsing by Horde
+		/// </summary>
+		/// <param name="File">Output file to write</param>
+		public void ExportForHorde(FileReference File)
+		{
+			DirectoryReference.CreateDirectory(File.Directory);
+			using (JsonWriter JsonWriter = new JsonWriter(File.FullName))
+			{
+				JsonWriter.WriteObjectStart();
+				JsonWriter.WriteArrayStart("Groups");
+				foreach (BgAgent Agent in Agents)
+				{
+					JsonWriter.WriteObjectStart();
+					JsonWriter.WriteArrayStart("Types");
+					foreach(string PossibleType in Agent.PossibleTypes)
+					{
+						JsonWriter.WriteValue(PossibleType);
+					}
+					JsonWriter.WriteArrayEnd();
+					JsonWriter.WriteArrayStart("Nodes");
+					foreach (BgNode Node in Agent.Nodes)
+					{
+						JsonWriter.WriteObjectStart();
+						JsonWriter.WriteValue("Name", Node.Name);
+						JsonWriter.WriteValue("RunEarly", Node.bRunEarly);
+						JsonWriter.WriteValue("Warnings", Node.bNotifyOnWarnings);
+
+						JsonWriter.WriteArrayStart("InputDependencies");
+						foreach (string InputDependency in Node.GetDirectInputDependencies().Select(x => x.Name))
+						{
+							JsonWriter.WriteValue(InputDependency);
+						}
+						JsonWriter.WriteArrayEnd();
+
+						JsonWriter.WriteArrayStart("OrderDependencies");
+						foreach (string OrderDependency in Node.GetDirectOrderDependencies().Select(x => x.Name))
+						{
+							JsonWriter.WriteValue(OrderDependency);
+						}
+						JsonWriter.WriteArrayEnd();
+
+						JsonWriter.WriteObjectEnd();
+					}
+					JsonWriter.WriteArrayEnd();
+					JsonWriter.WriteObjectEnd();
+				}
+				JsonWriter.WriteArrayEnd();
+
+				JsonWriter.WriteArrayStart("Aggregates");
+				foreach (BgAggregate Aggregate in NameToAggregate.Values)
+				{
+					JsonWriter.WriteObjectStart();
+					JsonWriter.WriteValue("Name", Aggregate.Name);
+					JsonWriter.WriteArrayStart("Nodes");
+					foreach (BgNode RequiredNode in Aggregate.RequiredNodes.OrderBy(x => x.Name))
+					{
+						JsonWriter.WriteValue(RequiredNode.Name);
+					}
+					JsonWriter.WriteArrayEnd();
+					JsonWriter.WriteObjectEnd();
+				}
+				JsonWriter.WriteArrayEnd();
+
+				JsonWriter.WriteArrayStart("Labels");
+				foreach (BgLabel Label in Labels)
+				{
+					JsonWriter.WriteObjectStart();
+					if (!String.IsNullOrEmpty(Label.DashboardName))
+					{
+						JsonWriter.WriteValue("Name", Label.DashboardName);
+					}
+					if (!String.IsNullOrEmpty(Label.DashboardCategory))
+					{
+						JsonWriter.WriteValue("Category", Label.DashboardCategory);
+					}
+					if (!String.IsNullOrEmpty(Label.UgsBadge))
+					{
+						JsonWriter.WriteValue("UgsBadge", Label.UgsBadge);
+					}
+					if (!String.IsNullOrEmpty(Label.UgsProject))
+					{
+						JsonWriter.WriteValue("UgsProject", Label.UgsProject);
+					}
+					if (Label.Change != BgLabelChange.Current)
+					{
+						JsonWriter.WriteValue("Change", Label.Change.ToString());
+					}
+
+					JsonWriter.WriteArrayStart("RequiredNodes");
+					foreach (BgNode RequiredNode in Label.RequiredNodes.OrderBy(x => x.Name))
+					{
+						JsonWriter.WriteValue(RequiredNode.Name);
+					}
+					JsonWriter.WriteArrayEnd();
+					JsonWriter.WriteArrayStart("IncludedNodes");
+					foreach (BgNode IncludedNode in Label.IncludedNodes.OrderBy(x => x.Name))
+					{
+						JsonWriter.WriteValue(IncludedNode.Name);
+					}
+					JsonWriter.WriteArrayEnd();
+					JsonWriter.WriteObjectEnd();
+				}
+				JsonWriter.WriteArrayEnd();
+
+				JsonWriter.WriteArrayStart("Badges");
+				foreach (BgBadge Badge in Badges)
+				{
+					HashSet<BgNode> Dependencies = Badge.Nodes;
+					if (Dependencies.Count > 0)
+					{
+						// Reduce that list to the smallest subset of direct dependencies
+						HashSet<BgNode> DirectDependencies = new HashSet<BgNode>(Dependencies);
+						foreach (BgNode Dependency in Dependencies)
+						{
+							DirectDependencies.ExceptWith(Dependency.OrderDependencies);
+						}
+
+						JsonWriter.WriteObjectStart();
+						JsonWriter.WriteValue("Name", Badge.Name);
+						if (!String.IsNullOrEmpty(Badge.Project))
+						{
+							JsonWriter.WriteValue("Project", Badge.Project);
+						}
+						if (Badge.Change != 0)
+						{
+							JsonWriter.WriteValue("Change", Badge.Change);
+						}
+						JsonWriter.WriteValue("Dependencies", String.Join(";", DirectDependencies.Select(x => x.Name)));
+						JsonWriter.WriteObjectEnd();
+					}
+				}
+				JsonWriter.WriteArrayEnd();
+
+				JsonWriter.WriteObjectEnd();
+			}
+		}
+
+		/// <summary>
+		/// Print the contents of the graph
+		/// </summary>
+		/// <param name="CompletedNodes">Set of nodes which are already complete</param>
+		/// <param name="PrintOptions">Options for how to print the graph</param>
+		/// <param name="Logger"></param>
+		public void Print(HashSet<BgNode> CompletedNodes, GraphPrintOptions PrintOptions, ILogger Logger)
+		{
+			// Print the options
+			if((PrintOptions & GraphPrintOptions.ShowCommandLineOptions) != 0)
+			{
+				// Get the list of messages
+				List<KeyValuePair<string, string>> Parameters = new List<KeyValuePair<string, string>>();
+				foreach(BgScriptOption Option in Options)
+				{
+					string Name = String.Format("-set:{0}=...", Option.Name);
+
+					StringBuilder Description = new StringBuilder(Option.Description);
+					if(!String.IsNullOrEmpty(Option.DefaultValue))
+					{
+						Description.AppendFormat(" (Default: {0})", Option.DefaultValue);
+					}
+
+					Parameters.Add(new KeyValuePair<string, string>(Name, Description.ToString()));
+				}
+
+				// Format them to the log
+				if(Parameters.Count > 0)
+				{
+					Logger.LogInformation("");
+					Logger.LogInformation("Options:");
+					Logger.LogInformation("");
+					HelpUtils.PrintTable(Parameters, 4, 24, Logger);
+				}
+			}
+
+			// Output all the triggers in order
+			Logger.LogInformation("");
+			Logger.LogInformation("Graph:");
+			foreach(BgAgent Agent in Agents)
+			{
+				Logger.LogInformation("        Agent: {0} ({1})", Agent.Name, String.Join(";", Agent.PossibleTypes));
+				foreach(BgNode Node in Agent.Nodes)
+				{
+					Logger.LogInformation("            Node: {0}{1}", Node.Name, CompletedNodes.Contains(Node)? " (completed)" : Node.bRunEarly? " (early)" : "");
+					if(PrintOptions.HasFlag(GraphPrintOptions.ShowDependencies))
+					{
+						HashSet<BgNode> InputDependencies = new HashSet<BgNode>(Node.GetDirectInputDependencies());
+						foreach(BgNode InputDependency in InputDependencies)
+						{
+							Logger.LogInformation("                input> {0}", InputDependency.Name);
+						}
+						HashSet<BgNode> OrderDependencies = new HashSet<BgNode>(Node.GetDirectOrderDependencies());
+						foreach(BgNode OrderDependency in OrderDependencies.Except(InputDependencies))
+						{
+							Logger.LogInformation("                after> {0}", OrderDependency.Name);
+						}
+					}
+					if(PrintOptions.HasFlag(GraphPrintOptions.ShowNotifications))
+					{
+						string Label = Node.bNotifyOnWarnings? "warnings" : "errors";
+						foreach(string User in Node.NotifyUsers)
+						{
+							Logger.LogInformation("                {0}> {1}", Label, User);
+						}
+						foreach(string Submitter in Node.NotifySubmitters)
+						{
+							Logger.LogInformation("                {0}> submitters to {1}", Label, Submitter);
+						}
+					}
+				}
+			}
+			Logger.LogInformation("");
+
+			// Print out all the non-empty aggregates
+			BgAggregate[] Aggregates = NameToAggregate.Values.OrderBy(x => x.Name).ToArray();
+			if(Aggregates.Length > 0)
+			{
+				Logger.LogInformation("Aggregates:");
+				foreach(string AggregateName in Aggregates.Select(x => x.Name))
+				{
+					Logger.LogInformation("    {0}", AggregateName);
+				}
+				Logger.LogInformation("");
+			}
+		}
+	}
+}
diff --git a/Engine/Source/Programs/Shared/EpicGames.BuildGraph/BgScriptReader.cs b/Engine/Source/Programs/Shared/EpicGames.BuildGraph/BgScriptReader.cs
new file mode 100644
--- /dev/null
+++ b/Engine/Source/Programs/Shared/EpicGames.BuildGraph/BgScriptReader.cs
@@ -0,0 +1,1960 @@
+// Copyright Epic Games, Inc. All Rights Reserved.
+
+using System;
+using System.Collections.Generic;
+using System.ComponentModel;
+using System.IO;
+using System.Linq;
+using System.Reflection;
+using System.Text;
+using System.Text.RegularExpressions;
+using System.Xml;
+using System.Xml.Schema;
+using EpicGames.Core;
+using Microsoft.Extensions.Logging;
+
+namespace EpicGames.BuildGraph
+{
+	/// <summary>
+	/// Implementation of XmlDocument which preserves line numbers for its elements
+	/// </summary>
+	class BgScriptDocument : XmlDocument
+	{
+		/// <summary>
+		/// The file being read
+		/// </summary>
+		string File;
+
+		/// <summary>
+		/// Native file representation
+		/// </summary>
+		object NativeFile;
+
+		/// <summary>
+		/// Interface to the LineInfo on the active XmlReader
+		/// </summary>
+		IXmlLineInfo LineInfo;
+
+		/// <summary>
+		/// Set to true if the reader encounters an error
+		/// </summary>
+		bool bHasErrors;
+
+		/// <summary>
+		/// Logger for validation errors
+		/// </summary>
+		ILogger Logger;
+
+		/// <summary>
+		/// Private constructor. Use ScriptDocument.Load to read an XML document.
+		/// </summary>
+		BgScriptDocument(string InFile, object InNativeFile, ILogger InLogger)
+		{
+			File = InFile;
+			NativeFile = InNativeFile;
+			Logger = InLogger;
+		}
+
+		/// <summary>
+		/// Overrides XmlDocument.CreateElement() to construct ScriptElements rather than XmlElements
+		/// </summary>
+		public override XmlElement CreateElement(string Prefix, string LocalName, string NamespaceUri)
+		{
+			return new BgScriptElement(File, NativeFile, LineInfo.LineNumber, Prefix, LocalName, NamespaceUri, this);
+		}
+
+		/// <summary>
+		/// Loads a script document from the given file
+		/// </summary>
+		/// <param name="File">The file to load</param>
+		/// <param name="NativeFile"></param>
+		/// <param name="Data"></param>
+		/// <param name="Schema">The schema to validate against</param>
+		/// <param name="Logger">Logger for output messages</param>
+		/// <param name="OutDocument">If successful, the document that was read</param>
+		/// <returns>True if the document could be read, false otherwise</returns>
+		public static bool TryRead(string File, object NativeFile, byte[] Data, BgScriptSchema Schema, ILogger Logger, out BgScriptDocument OutDocument)
+		{
+			BgScriptDocument Document = new BgScriptDocument(File, NativeFile, Logger);
+
+			XmlReaderSettings Settings = new XmlReaderSettings();
+			Settings.Schemas.Add(Schema.CompiledSchema);
+			Settings.ValidationType = ValidationType.Schema;
+			Settings.ValidationEventHandler += Document.ValidationEvent;
+
+			using (MemoryStream Stream = new MemoryStream(Data))
+			using (XmlReader Reader = XmlReader.Create(Stream, Settings))
+			{
+				// Read the document
+				Document.LineInfo = (IXmlLineInfo)Reader;
+				try
+				{
+					Document.Load(Reader);
+				}
+				catch (XmlException Ex)
+				{
+					if (!Document.bHasErrors)
+					{
+						Logger.LogScriptError(NativeFile, Ex.LineNumber, "{Message}", Ex.Message);
+						Document.bHasErrors = true;
+					}
+				}
+
+				// If we hit any errors while parsing
+				if (Document.bHasErrors)
+				{
+					OutDocument = null;
+					return false;
+				}
+
+				// Check that the root element is valid. If not, we didn't actually validate against the schema.
+				if (Document.DocumentElement.Name != BgScriptSchema.RootElementName)
+				{
+					Logger.LogScriptError(NativeFile, 1, "Script does not have a root element called '{ElementName}'", BgScriptSchema.RootElementName);
+					OutDocument = null;
+					return false;
+				}
+				if (Document.DocumentElement.NamespaceURI != BgScriptSchema.NamespaceURI)
+				{
+					Logger.LogScriptError(NativeFile, 1, "Script root element is not in the '{Namespace}' namespace (add the xmlns=\"{NewNamespace}\" attribute)", BgScriptSchema.NamespaceURI, BgScriptSchema.NamespaceURI);
+					OutDocument = null;
+					return false;
+				}
+			}
+
+			OutDocument = Document;
+			return true;
+		}
+
+		/// <summary>
+		/// Callback for validation errors in the document
+		/// </summary>
+		/// <param name="Sender">Standard argument for ValidationEventHandler</param>
+		/// <param name="Args">Standard argument for ValidationEventHandler</param>
+		void ValidationEvent(object Sender, ValidationEventArgs Args)
+		{
+			if (Args.Severity == XmlSeverityType.Warning)
+			{
+				Logger.LogScriptWarning(NativeFile, Args.Exception.LineNumber, "{Message}", Args.Message);
+			}
+			else
+			{
+				Logger.LogScriptError(NativeFile, Args.Exception.LineNumber, "{Message}", Args.Message);
+				bHasErrors = true;
+			}
+		}
+	}
+
+	/// <summary>
+	/// Implementation of XmlElement which preserves line numbers
+	/// </summary>
+	class BgScriptElement : XmlElement
+	{
+		/// <summary>
+		/// The file containing this element
+		/// </summary>
+		public string File { get; }
+
+		/// <summary>
+		/// Native file object
+		/// </summary>
+		public object NativeFile { get; }
+
+		/// <summary>
+		/// The line number containing this element
+		/// </summary>
+		public int LineNumber { get; }
+
+		/// <summary>
+		/// Constructor
+		/// </summary>
+		public BgScriptElement(string File, object NativeFile, int LineNumber, string Prefix, string LocalName, string NamespaceUri, BgScriptDocument Document)
+			: base(Prefix, LocalName, NamespaceUri, Document)
+		{
+			this.File = File;
+			this.NativeFile = NativeFile;
+			this.LineNumber = LineNumber;
+		}
+	}
+
+	/// <summary>
+	/// Stores information about a script function that has been declared
+	/// </summary>
+	class BgScriptMacro
+	{
+		/// <summary>
+		/// Name of the function
+		/// </summary>
+		public readonly string Name;
+
+		/// <summary>
+		/// Element where the function was declared
+		/// </summary>
+		public readonly BgScriptElement Element;
+
+		/// <summary>
+		/// The total number of arguments
+		/// </summary>
+		public readonly int NumArguments;
+
+		/// <summary>
+		/// Number of arguments that are required
+		/// </summary>
+		public readonly int NumRequiredArguments;
+
+		/// <summary>
+		/// Maps an argument name to its type
+		/// </summary>
+		public readonly Dictionary<string, int> ArgumentNameToIndex;
+
+		/// <summary>
+		/// Constructor
+		/// </summary>
+		/// <param name="Name">Name of the function</param>
+		/// <param name="Element">Element containing the function definition</param>
+		/// <param name="ArgumentNameToIndex">Map of argument name to index</param>
+		/// <param name="NumRequiredArguments">Number of arguments that are required. Indices 0 to NumRequiredArguments - 1 are required.</param>
+		public BgScriptMacro(string Name, BgScriptElement Element, Dictionary<string, int> ArgumentNameToIndex, int NumRequiredArguments)
+		{
+			this.Name = Name;
+			this.Element = Element;
+			this.NumArguments = ArgumentNameToIndex.Count;
+			this.NumRequiredArguments = NumRequiredArguments;
+			this.ArgumentNameToIndex = ArgumentNameToIndex;
+		}
+	}
+
+	public interface IBgScriptReaderContext
+	{
+		/// <summary>
+		/// Tests whether the given file or directory exists
+		/// </summary>
+		/// <param name="File">Path to a file</param>
+		/// <returns>True if the file exists</returns>
+		bool Exists(string Path);
+
+		/// <summary>
+		/// Tries to read a file from the given path
+		/// </summary>
+		/// <param name="Path">Path of the file to read</param>
+		/// <param name="Data">On success, contains the read data</param>
+		/// <returns></returns>
+		bool TryRead(string Path, out byte[] Data);
+
+		/// <summary>
+		/// Converts a path to its native form, for display to the user
+		/// </summary>
+		/// <param name="Path">Path to format</param>
+		/// <returns></returns>
+		object GetNativePath(string Path);
+	}
+
+	/// <summary>
+	/// Extension methods for writing script error messages
+	/// </summary>
+	public static class BgScriptExtensions
+	{
+		public static void LogScriptError(this ILogger Logger, object File, int LineNumber, string Format, params object[] Args)
+		{
+			object[] AllArgs = new object[Args.Length + 2];
+			AllArgs[0] = File;
+			AllArgs[1] = LineNumber;
+			Args.CopyTo(AllArgs, 2);
+			Logger.LogError($"{{Script}}({{Line}}): error: {Format}", AllArgs);
+		}
+
+		public static void LogScriptWarning(this ILogger Logger, object File, int LineNumber, string Format, params object[] Args)
+		{
+			object[] AllArgs = new object[Args.Length + 2];
+			AllArgs[0] = File;
+			AllArgs[1] = LineNumber;
+			Args.CopyTo(AllArgs, 2);
+			Logger.LogWarning($"{{Script}}({{Line}}): warning: {Format}", AllArgs);
+		}
+	}
+
+	/// <summary>
+	/// Reader for build graph definitions. Instanced to contain temporary state; public interface is through ScriptReader.TryRead().
+	/// </summary>
+	public class BgScriptReader
+	{
+		/// <summary>
+		/// Interface used for reading files
+		/// </summary>
+		IBgScriptReaderContext Context;
+
+		/// <summary>
+		/// The current graph
+		/// </summary>
+		BgScript Graph = new BgScript();
+
+		/// <summary>
+		/// List of property name to value lookups. Modifications to properties are scoped to nodes and agents. EnterScope() pushes an empty dictionary onto the end of this list, and LeaveScope() removes one. 
+		/// ExpandProperties() searches from last to first lookup when trying to resolve a property name, and takes the first it finds.
+		/// </summary>
+		List<Dictionary<string, string>> ScopedProperties = new List<Dictionary<string, string>>();
+
+		/// <summary>
+		/// When declaring a property in a nested scope, we enter its name into a set for each parent scope which prevents redeclaration in an OUTER scope later. Subsequent NESTED scopes can redeclare it.
+		/// The former is likely a coding error, since it implies that the scope of the variable was meant to be further out, whereas the latter is common for temporary and loop variables.
+		/// </summary>
+		List<HashSet<string>> ShadowProperties = new List<HashSet<string>>();
+
+		/// <summary>
+		/// Maps from a function name to its definition
+		/// </summary>
+		Dictionary<string, BgScriptMacro> MacroNameToDefinition = new Dictionary<string, BgScriptMacro>();
+
+		/// <summary>
+		/// Preprocess the file, but do not expand any values that are not portable (eg. paths on the local machine)
+		/// </summary>
+		bool bPreprocessOnly;
+
+		/// <summary>
+		/// Schema for the script
+		/// </summary>
+		BgScriptSchema Schema;
+
+		/// <summary>
+		/// Logger for diagnostic messages
+		/// </summary>
+		ILogger Logger;
+
+		/// <summary>
+		/// The number of errors encountered during processing so far
+		/// </summary>
+		int NumErrors;
+
+		/// <summary>
+		/// The name of the node if only a single node is going to be built, otherwise null.
+		/// </summary>
+		string SingleNodeName;
+
+		/// <summary>
+		/// Private constructor. Use ScriptReader.TryRead() to read a script file.
+		/// </summary>
+		/// <param name="Context">Context object</param>
+		/// <param name="DefaultProperties">Default properties available to the script</param>
+		/// <param name="bPreprocessOnly">Preprocess the file, but do not expand any values that are not portable (eg. paths on the local machine)</param>
+		/// <param name="Schema">Schema for the script</param>
+		/// <param name="Logger">Logger for diagnostic messages</param>
+		private BgScriptReader(IBgScriptReaderContext Context, IDictionary<string, string> DefaultProperties, bool bPreprocessOnly, BgScriptSchema Schema, ILogger Logger)
+		{
+			this.Context = Context;
+			this.bPreprocessOnly = bPreprocessOnly;
+			this.Schema = Schema;
+			this.Logger = Logger;
+
+			EnterScope();
+
+			foreach(KeyValuePair<string, string> Pair in DefaultProperties)
+			{
+				ScopedProperties[ScopedProperties.Count - 1].Add(Pair.Key, Pair.Value);
+			}
+		}
+
+		/// <summary>
+		/// Try to read a script file from the given file.
+		/// </summary>
+		/// <param name="Context">Supplies context about the parse</param>
+		/// <param name="File">File to read from</param>
+		/// <param name="Arguments">Arguments passed in to the graph on the command line</param>
+		/// <param name="DefaultProperties">Default properties available to the script</param>
+		/// <param name="bPreprocessOnly">Preprocess the file, but do not expand any values that are not portable (eg. paths on the local machine)</param>
+		/// <param name="Schema">Schema for the script</param>
+		/// <param name="Logger">Logger for output messages</param>
+		/// <param name="Graph">If successful, the graph constructed from the given script</param>
+		/// <param name="SingleNodeName">If a single node will be processed, the name of that node.</param>
+		/// <returns>True if the graph was read, false if there were errors</returns>
+		public static bool TryRead(IBgScriptReaderContext Context, string File, Dictionary<string, string> Arguments, Dictionary<string, string> DefaultProperties, bool bPreprocessOnly, BgScriptSchema Schema, ILogger Logger, out BgScript Graph, string SingleNodeName = null)
+		{
+			// Read the file and build the graph
+			BgScriptReader Reader = new BgScriptReader(Context, DefaultProperties, bPreprocessOnly, Schema, Logger);
+			if (!Reader.TryRead(File, Arguments, Logger, SingleNodeName) || Reader.NumErrors > 0)
+			{
+				Graph = null;
+				return false;
+			}
+
+			// Make sure all the arguments were valid
+			HashSet<string> ValidArgumentNames = new HashSet<string>(Reader.Graph.Options.Select(x => x.Name), StringComparer.OrdinalIgnoreCase);
+			ValidArgumentNames.Add("PreflightChange");
+
+			foreach(string ArgumentName in Arguments.Keys)
+			{
+				if (!ValidArgumentNames.Contains(ArgumentName))
+				{
+					Logger.LogWarning("Unknown argument '{ArgumentName}' for '{Script}'", ArgumentName, Context.GetNativePath(File));
+				}
+			}
+
+			// Return the constructed graph
+			Graph = Reader.Graph;
+			return true;
+		}
+
+		/// <summary>
+		/// Read the script from the given file
+		/// </summary>
+		/// <param name="File">File to read from</param>
+		/// <param name="Arguments">Arguments passed in to the graph on the command line</param>
+		/// <param name="Logger">Logger for output messages</param>
+		/// <param name="SingleNodeName">The name of the node if only a single node is going to be built, otherwise null.</param>
+		bool TryRead(string File, Dictionary<string, string> Arguments, ILogger Logger, string SingleNodeName = null)
+		{
+			// Get the data for this file
+			byte[] Data;
+			if (!Context.TryRead(File, out Data))
+			{
+				Logger.LogError("Unable to open file {File}", File);
+				NumErrors++;
+				return false;
+			}
+
+			// Read the document and validate it against the schema
+			BgScriptDocument Document;
+			if (!BgScriptDocument.TryRead(File, Context.GetNativePath(File), Data, Schema, Logger, out Document))
+			{
+				NumErrors++;
+				return false;
+			}
+
+			// Read the root BuildGraph element
+			this.SingleNodeName = SingleNodeName;
+			ReadGraphBody(Document.DocumentElement, Arguments);
+			return true;
+		}
+
+		/// <summary>
+		/// Reads the contents of a graph
+		/// </summary>
+		/// <param name="Element">The parent element to read from</param>
+		/// <param name="Arguments">Arguments passed in to the graph on the command line</param>
+		void ReadGraphBody(XmlElement Element, Dictionary<string, string> Arguments)
+		{
+			foreach (BgScriptElement ChildElement in Element.ChildNodes.OfType<BgScriptElement>())
+			{
+				switch (ChildElement.Name)
+				{
+					case "Include":
+						ReadInclude(ChildElement, Arguments);
+						break;
+					case "Option":
+						ReadOption(ChildElement, Arguments);
+						break;
+					case "Property":
+						ReadProperty(ChildElement);
+						break;
+					case "Regex":
+						ReadRegex(ChildElement);
+						break;
+					case "EnvVar":
+						ReadEnvVar(ChildElement);
+						break;
+					case "Macro":
+						ReadMacro(ChildElement);
+						break;
+					case "Agent":
+						ReadAgent(ChildElement);
+						break;
+					case "Aggregate":
+						ReadAggregate(ChildElement);
+						break;
+					case "Report":
+						ReadReport(ChildElement);
+						break;
+					case "Badge":
+						ReadBadge(ChildElement);
+						break;
+					case "Label":
+						ReadLabel(ChildElement);
+						break;
+					case "Notify":
+						ReadNotifier(ChildElement);
+						break;
+					case "Trace":
+						ReadDiagnostic(ChildElement, LogEventType.Console, null, null);
+						break;
+					case "Warning":
+						ReadDiagnostic(ChildElement, LogEventType.Warning, null, null);
+						break;
+					case "Error":
+						ReadDiagnostic(ChildElement, LogEventType.Error, null, null);
+						break;
+					case "Do":
+						ReadBlock(ChildElement, x => ReadGraphBody(x, Arguments));
+						break;
+					case "Switch":
+						ReadSwitch(ChildElement, x => ReadGraphBody(x, Arguments));
+						break;
+					case "ForEach":
+						ReadForEach(ChildElement, x => ReadGraphBody(x, Arguments));
+						break;
+					case "Expand":
+						ReadExpand(ChildElement, x => ReadGraphBody(x, Arguments));
+						break;
+					default:
+						LogError(ChildElement, "Invalid element '{0}'", ChildElement.Name);
+						break;
+				}
+			}
+		}
+
+		/// <summary>
+		/// Handles validation messages from validating the document against its schema
+		/// </summary>
+		/// <param name="Sender">The source of the event</param>
+		/// <param name="Args">Event arguments</param>
+		void ValidationHandler(object Sender, ValidationEventArgs Args)
+		{
+			if (Args.Severity == XmlSeverityType.Warning)
+			{
+				Logger.LogWarning("Script: {Message}", Args.Message);
+			}
+			else
+			{
+				Logger.LogError("Script: {Message}", Args.Message);
+				NumErrors++;
+			}
+		}
+
+		/// <summary>
+		/// Push a new property scope onto the stack
+		/// </summary>
+		void EnterScope()
+		{
+			ScopedProperties.Add(new Dictionary<string, string>(StringComparer.InvariantCultureIgnoreCase));
+			ShadowProperties.Add(new HashSet<string>(StringComparer.InvariantCultureIgnoreCase));
+		}
+
+		/// <summary>
+		/// Pop a property scope from the stack
+		/// </summary>
+		void LeaveScope()
+		{
+			ScopedProperties.RemoveAt(ScopedProperties.Count - 1);
+			ShadowProperties.RemoveAt(ShadowProperties.Count - 1);
+		}
+
+		/// <summary>
+		/// Sets a property value in the current scope
+		/// </summary>
+		/// <param name="Element">Element containing the property assignment. Used for error messages if the property is shadowed in another scope.</param>
+		/// <param name="Name">Name of the property</param>
+		/// <param name="Value">Value for the property</param>
+		void SetPropertyValue(BgScriptElement Element, string Name, string Value)
+		{
+			// Find the scope containing this property, defaulting to the current scope
+			int ScopeIdx = 0;
+			while(ScopeIdx < ScopedProperties.Count - 1 && !ScopedProperties[ScopeIdx].ContainsKey(Name))
+			{
+				ScopeIdx++;
+			}
+
+			// Make sure this property name was not already used in a child scope; it likely indicates an error.
+			if(ShadowProperties[ScopeIdx].Contains(Name))
+			{
+				LogError(Element, "Property '{0}' was already used in a child scope. Move this definition before the previous usage if they are intended to share scope, or use a different name.", Name);
+			}
+			else
+			{
+				// Make sure it's added to the shadow property list for every parent scope
+				for(int Idx = 0; Idx < ScopeIdx; Idx++)
+				{
+					ShadowProperties[Idx].Add(Name);
+				}
+				ScopedProperties[ScopeIdx][Name] = Value;
+			}
+		}
+
+		/// <summary>
+		/// Tries to get the value of a property
+		/// </summary>
+		/// <param name="Name">Name of the property</param>
+		/// <param name="Value">On success, contains the value of the property. Set to null otherwise.</param>
+		/// <returns>True if the property was found, false otherwise</returns>
+		bool TryGetPropertyValue(string Name, out string Value)
+		{
+			int ValueLength = 0;
+			if (Name.Contains(":"))
+			{
+				string[] Tokens = Name.Split(':');
+				Name = Tokens[0];
+				ValueLength = int.Parse(Tokens[1]);
+			}
+
+			// Check each scope for the property
+			for (int ScopeIdx = ScopedProperties.Count - 1; ScopeIdx >= 0; ScopeIdx--)
+			{
+				string ScopeValue;
+				if (ScopedProperties[ScopeIdx].TryGetValue(Name, out ScopeValue))
+				{
+					Value = ScopeValue;
+
+					// It's valid for a property to exist but have a null value. It won't be expanded
+					// Handle $(PropName:-6) where PropName might be "Foo"
+					if (Value != null && Value.Length > Math.Abs(ValueLength))
+					{
+						if (ValueLength > 0)
+						{
+							Value = Value.Substring(0, ValueLength);
+						}
+						if (ValueLength < 0)
+						{
+							Value = Value.Substring(Value.Length + ValueLength, -ValueLength);
+						}
+					}
+					return true;
+				}
+			}
+
+			// If we didn't find it, return false.
+			Value = null;
+			return false;
+		}
+
+		static string CombinePaths(string BasePath, string NextPath)
+		{
+			List<string> Fragments = new List<string>(BasePath.Split('/'));
+			Fragments.RemoveAt(Fragments.Count - 1);
+
+			foreach (string AppendFragment in NextPath.Split('/'))
+			{
+				if (AppendFragment.Equals(".", StringComparison.Ordinal))
+				{
+					continue;
+				}
+				else if (AppendFragment.Equals("..", StringComparison.Ordinal))
+				{
+					if (Fragments.Count > 0)
+					{
+						Fragments.RemoveAt(Fragments.Count - 1);
+					}
+					else
+					{
+						throw new Exception($"Path '{NextPath}' cannot be combined with '{BasePath}'");
+					}
+				}
+				else
+				{
+					Fragments.Add(AppendFragment);
+				}
+			}
+			return String.Join('/', Fragments);
+		}
+
+		/// <summary>
+		/// Read an include directive, and the contents of the target file
+		/// </summary>
+		/// <param name="Element">Xml element to read the definition from</param>
+		/// <param name="Arguments">Arguments passed in to the graph on the command line</param>
+		void ReadInclude(BgScriptElement Element, Dictionary<string, string> Arguments)
+		{
+			if (EvaluateCondition(Element))
+			{
+				foreach (string Script in ReadListAttribute(Element, "Script"))
+				{
+					TryRead(CombinePaths(Element.File, Script), Arguments, Logger);
+				}
+			}
+		}
+
+		/// <summary>
+		/// Reads the definition of a graph option; a parameter which can be set by the user on the command-line or via an environment variable.
+		/// </summary>
+		/// <param name="Element">Xml element to read the definition from</param>
+		/// <param name="Arguments">Arguments passed in to the graph on the command line</param>
+		void ReadOption(BgScriptElement Element, IDictionary<string, string> Arguments)
+		{
+			if (EvaluateCondition(Element))
+			{
+				string Name = ReadAttribute(Element, "Name");
+				if (ValidateName(Element, Name))
+				{
+					// Make sure we're at global scope
+					if(ScopedProperties.Count > 1)
+					{
+						throw new Exception("Incorrect scope depth for reading option settings");
+					}
+
+					// Check if the property already exists. If it does, we don't need to register it as an option.
+					string ExistingValue;
+					if(TryGetPropertyValue(Name, out ExistingValue))
+					{
+						// If there's a restriction on this definition, check it matches
+						string Restrict = ReadAttribute(Element, "Restrict");
+						if(!String.IsNullOrEmpty(Restrict) && !Regex.IsMatch(ExistingValue, "^" + Restrict + "$", RegexOptions.IgnoreCase))
+						{
+							LogError(Element, "'{0} is already set to '{1}', which does not match the given restriction ('{2}')", Name, ExistingValue, Restrict);
+						}
+					}
+					else
+					{
+						// Create a new option object to store the settings
+						string Description = ReadAttribute(Element, "Description");
+						string DefaultValue = ReadAttribute(Element, "DefaultValue");
+						BgScriptOption Option = new BgScriptOption(Name, Description, DefaultValue);
+						Graph.Options.Add(Option);
+
+						// Get the value of this property
+						string Value;
+						if(!Arguments.TryGetValue(Name, out Value))
+						{
+							Value = Option.DefaultValue;
+						}
+						SetPropertyValue(Element, Name, Value);
+
+						// If there's a restriction on it, check it's valid
+						string Restrict = ReadAttribute(Element, "Restrict");
+						if(!String.IsNullOrEmpty(Restrict))
+						{
+							string Pattern = "^(" + Restrict + ")$";
+							if(!Regex.IsMatch(Value, Pattern, RegexOptions.IgnoreCase))
+							{
+								LogError(Element, "'{0}' is not a valid value for '{1}' (required: '{2}')", Value, Name, Restrict);
+							}
+							if(Option.DefaultValue != Value && !Regex.IsMatch(Option.DefaultValue, Pattern, RegexOptions.IgnoreCase))
+							{
+								LogError(Element, "Default value '{0}' is not valid for '{1}' (required: '{2}')", Option.DefaultValue, Name, Restrict);
+							}
+						}
+					}
+				}
+			}
+		}
+
+		/// <summary>
+		/// Reads a property assignment.
+		/// </summary>
+		/// <param name="Element">Xml element to read the definition from</param>
+		void ReadProperty(BgScriptElement Element)
+		{
+			if (EvaluateCondition(Element))
+			{
+				string Name = ReadAttribute(Element, "Name");
+				if (ValidateName(Element, Name))
+				{
+					string Value = ReadAttribute(Element, "Value");
+					if(Element.HasChildNodes)
+					{
+						// Read the element content, and append each line to the value as a semicolon delimited list
+						StringBuilder Builder = new StringBuilder(Value);
+						foreach(string Line in Element.InnerText.Split('\n'))
+						{
+							string TrimLine = ExpandProperties(Element, Line.Trim());
+							if(TrimLine.Length > 0)
+							{
+								if(Builder.Length > 0)
+								{
+									Builder.Append(";");
+								}
+								Builder.Append(TrimLine);
+							}
+						}
+						Value = Builder.ToString();
+					}
+					SetPropertyValue(Element, Name, Value);
+				}
+			}
+		}
+
+		/// <summary>
+		/// Reads a Regex assignment.
+		/// </summary>
+		/// <param name="Element">Xml element to read the definition from</param>
+		void ReadRegex(BgScriptElement Element)
+		{
+			if (EvaluateCondition(Element))
+			{
+				// Get the pattern
+				string RegexString = ReadAttribute(Element, "Pattern");
+				// Make sure its a valid regex.
+				Regex RegexValue = ParseRegex(Element, RegexString);
+				if (RegexValue != null)
+				{
+					// read the names in 
+					string[] CaptureNames = ReadListAttribute(Element, "Capture");
+
+					// get number of groups we passed in
+					int[] GroupNumbers = RegexValue.GetGroupNumbers();
+
+					// make sure the number of property names is the same as the number of match groups
+					// this includes the entire string match group as [0], so don't count that one.
+					if (CaptureNames.Length != GroupNumbers.Count() - 1)
+					{
+						LogError(Element, "MatchGroup count: {0} does not match the number of names specified: {1}", GroupNumbers.Count() - 1, CaptureNames.Length);
+					}
+					else
+					{
+						// apply the regex to the value
+						string Input = ReadAttribute(Element, "Input");
+						Match Match = RegexValue.Match(Input);
+
+						bool Optional = BgCondition.Evaluate(ReadAttribute(Element, "Optional"), Context);
+						if (!Match.Success)
+						{
+							if (!Optional)
+							{
+								LogError(Element, "Regex {0} did not find a match against input string {1}", RegexString, Input);
+							}
+						}
+						else
+						{
+							// assign each property to the group it matches, skip over [0]
+							for (int MatchIdx = 1; MatchIdx < GroupNumbers.Count(); MatchIdx++)
+							{
+								SetPropertyValue(Element, CaptureNames[MatchIdx - 1], Match.Groups[MatchIdx].Value);
+							}
+						}
+					}
+				}
+			}
+		}
+
+		/// <summary>
+		/// Reads a property assignment from an environment variable.
+		/// </summary>
+		/// <param name="Element">Xml element to read the definition from</param>
+		void ReadEnvVar(BgScriptElement Element)
+		{
+			if (EvaluateCondition(Element))
+			{
+				string Name = ReadAttribute(Element, "Name");
+				if (ValidateName(Element, Name))
+				{
+					string Value = Environment.GetEnvironmentVariable(Name) ?? "";
+					SetPropertyValue(Element, Name, Value);
+				}
+			}
+		}
+
+		/// <summary>
+		/// Reads a macro definition
+		/// </summary>
+		/// <param name="Element">Xml element to read the definition from</param>
+		void ReadMacro(BgScriptElement Element)
+		{
+			if(EvaluateCondition(Element))
+			{
+				string Name = ReadAttribute(Element, "Name");
+				if (ValidateName(Element, Name))
+				{
+					BgScriptMacro OriginalDefinition;
+					if(MacroNameToDefinition.TryGetValue(Name, out OriginalDefinition))
+					{
+						LogError(Element, "Function '{0}' has already been declared (see {1} line {2})", OriginalDefinition.Element.File, OriginalDefinition.Element.LineNumber);
+					}
+					else
+					{
+						Dictionary<string, int> ArgumentNameToIndex = new Dictionary<string, int>();
+						ReadMacroArguments(Element, "Arguments", ArgumentNameToIndex);
+
+						int NumRequiredArguments = ArgumentNameToIndex.Count;
+						ReadMacroArguments(Element, "OptionalArguments", ArgumentNameToIndex);
+
+						MacroNameToDefinition.Add(Name, new BgScriptMacro(Name, Element, ArgumentNameToIndex, NumRequiredArguments));
+					}
+				}
+			}
+		}
+
+		/// <summary>
+		/// Reads a list of macro arguments from an attribute
+		/// </summary>
+		/// <param name="Element">The element containing the attributes</param>
+		/// <param name="AttributeName">Name of the attribute containing the arguments</param>
+		/// <param name="ArgumentNameToIndex">List of arguments to add to</param>
+		void ReadMacroArguments(BgScriptElement Element, string AttributeName, Dictionary<string, int> ArgumentNameToIndex)
+		{
+			string AttributeValue = ReadAttribute(Element, AttributeName);
+			if(AttributeValue != null)
+			{
+				foreach(string ArgumentName in AttributeValue.Split(new char[]{ ';' }, StringSplitOptions.RemoveEmptyEntries))
+				{
+					if(ArgumentNameToIndex.ContainsKey(ArgumentName))
+					{
+						LogWarning(Element, "Argument '{0}' is listed multiple times", ArgumentName);
+					}
+					else
+					{
+						ArgumentNameToIndex.Add(ArgumentName, ArgumentNameToIndex.Count);
+					}
+				}
+			}
+		}
+
+		/// <summary>
+		/// Reads the definition for an agent.
+		/// </summary>
+		/// <param name="Element">Xml element to read the definition from</param>
+		/// <param name="Trigger">The controlling trigger for nodes in this agent</param>
+		void ReadAgent(BgScriptElement Element)
+		{
+			string Name;
+			if (EvaluateCondition(Element) && TryReadObjectName(Element, out Name))
+			{
+				// Read the valid agent types. This may be omitted if we're continuing an existing agent.
+				string[] Types = ReadListAttribute(Element, "Type");
+
+				// Create the agent object, or continue an existing one
+				BgAgent Agent;
+				if (Graph.NameToAgent.TryGetValue(Name, out Agent))
+				{
+					if (Types.Length > 0 && Agent.PossibleTypes.Length > 0)
+					{
+						if (Types.Length != Agent.PossibleTypes.Length || !Types.SequenceEqual(Agent.PossibleTypes, StringComparer.InvariantCultureIgnoreCase))
+						{
+							LogError(Element, "Agent types ({0}) were different than previous agent definition with types ({1}). Must either be empty or match exactly.", string.Join(",", Types), string.Join(",", Agent.PossibleTypes));
+						}
+					}
+				}
+				else
+				{
+					if (Types.Length == 0)
+					{
+						LogError(Element, "Missing type for agent '{0}'", Name);
+					}
+					Agent = new BgAgent(Name, Types);
+					Graph.NameToAgent.Add(Name, Agent);
+					Graph.Agents.Add(Agent);
+				}
+
+				// Process all the child elements.
+				ReadAgentBody(Element, Agent);
+			}
+		}
+
+		/// <summary>
+		/// Read the contents of an agent definition
+		/// </summary>
+		/// <param name="Element">Xml element to read the definition from</param>
+		/// <param name="ParentAgent">The agent to contain the definition</param>
+		void ReadAgentBody(BgScriptElement Element, BgAgent ParentAgent)
+		{
+			EnterScope();
+			foreach (BgScriptElement ChildElement in Element.ChildNodes.OfType<BgScriptElement>())
+			{
+				switch (ChildElement.Name)
+				{
+					case "Property":
+						ReadProperty(ChildElement);
+						break;
+					case "Regex":
+						ReadRegex(ChildElement);
+						break;
+					case "Node":
+						ReadNode(ChildElement, ParentAgent);
+						break;
+					case "Aggregate":
+						ReadAggregate(ChildElement);
+						break;
+					case "Trace":
+						ReadDiagnostic(ChildElement, LogEventType.Console, null, ParentAgent);
+						break;
+					case "Warning":
+						ReadDiagnostic(ChildElement, LogEventType.Warning, null, ParentAgent);
+						break;
+					case "Error":
+						ReadDiagnostic(ChildElement, LogEventType.Error, null, ParentAgent);
+						break;
+					case "Label":
+						ReadLabel(ChildElement);
+						break;
+					case "Do":
+						ReadBlock(ChildElement, x => ReadAgentBody(x, ParentAgent));
+						break;
+					case "Switch":
+						ReadSwitch(ChildElement, x => ReadAgentBody(x, ParentAgent));
+						break;
+					case "ForEach":
+						ReadForEach(ChildElement, x => ReadAgentBody(x, ParentAgent));
+						break;
+					case "Expand":
+						ReadExpand(ChildElement, x => ReadAgentBody(x, ParentAgent));
+						break;
+					default:
+						LogError(ChildElement, "Unexpected element type '{0}'", ChildElement.Name);
+						break;
+				}
+			}
+			LeaveScope();
+		}
+
+		/// <summary>
+		/// Reads the definition for an aggregate
+		/// </summary>
+		/// <param name="Element">Xml element to read the definition from</param>
+		void ReadAggregate(BgScriptElement Element)
+		{
+			string Name;
+			if (EvaluateCondition(Element) && TryReadObjectName(Element, out Name) && CheckNameIsUnique(Element, Name))
+			{
+				string[] RequiredNames = ReadListAttribute(Element, "Requires");
+
+				BgAggregate NewAggregate = new BgAggregate(Name);
+				foreach (BgNode ReferencedNode in ResolveReferences(Element, RequiredNames))
+				{
+					NewAggregate.RequiredNodes.Add(ReferencedNode);
+				}
+				Graph.NameToAggregate[Name] = NewAggregate;
+
+				string LabelCategoryName = ReadAttribute(Element, "Label");
+				if (!String.IsNullOrEmpty(LabelCategoryName))
+				{
+					BgLabel Label;
+
+					// Create the label
+					int SlashIdx = LabelCategoryName.IndexOf('/');
+					if (SlashIdx != -1)
+					{
+						Label = new BgLabel(LabelCategoryName.Substring(SlashIdx + 1), LabelCategoryName.Substring(0, SlashIdx), null, null, BgLabelChange.Current);
+					}
+					else
+					{
+						Label = new BgLabel(LabelCategoryName, "Other", null, null, BgLabelChange.Current);
+					}
+
+					// Find all the included nodes
+					foreach (BgNode RequiredNode in NewAggregate.RequiredNodes)
+					{
+						Label.RequiredNodes.Add(RequiredNode);
+						Label.IncludedNodes.Add(RequiredNode);
+						Label.IncludedNodes.UnionWith(RequiredNode.OrderDependencies);
+					}
+
+					string[] IncludedNames = ReadListAttribute(Element, "Include");
+					foreach (BgNode IncludedNode in ResolveReferences(Element, IncludedNames))
+					{
+						Label.IncludedNodes.Add(IncludedNode);
+						Label.IncludedNodes.UnionWith(IncludedNode.OrderDependencies);
+					}
+
+					string[] ExcludedNames = ReadListAttribute(Element, "Exclude");
+					foreach (BgNode ExcludedNode in ResolveReferences(Element, ExcludedNames))
+					{
+						Label.IncludedNodes.Remove(ExcludedNode);
+						Label.IncludedNodes.ExceptWith(ExcludedNode.OrderDependencies);
+					}
+
+					Graph.Labels.Add(Label);
+				}
+			}
+		}
+
+		/// <summary>
+		/// Reads the definition for a report
+		/// </summary>
+		/// <param name="Element">Xml element to read the definition from</param>
+		void ReadReport(BgScriptElement Element)
+		{
+			string Name;
+			if (EvaluateCondition(Element) && TryReadObjectName(Element, out Name) && CheckNameIsUnique(Element, Name))
+			{
+				string[] RequiredNames = ReadListAttribute(Element, "Requires");
+
+				BgReport NewReport = new BgReport(Name);
+				foreach (BgNode ReferencedNode in ResolveReferences(Element, RequiredNames))
+				{
+					NewReport.Nodes.Add(ReferencedNode);
+					NewReport.Nodes.UnionWith(ReferencedNode.OrderDependencies);
+				}
+				Graph.NameToReport.Add(Name, NewReport);
+			}
+		}
+
+		/// <summary>
+		/// Reads the definition for a badge
+		/// </summary>
+		/// <param name="Element">Xml element to read the definition from</param>
+		void ReadBadge(BgScriptElement Element)
+		{
+			string Name;
+			if (EvaluateCondition(Element) && TryReadObjectName(Element, out Name))
+			{
+				string[] RequiredNames = ReadListAttribute(Element, "Requires");
+				string[] TargetNames = ReadListAttribute(Element, "Targets");
+				string Project = ReadAttribute(Element, "Project");
+				int Change = ReadIntegerAttribute(Element, "Change", 0);
+
+				BgBadge NewBadge = new BgBadge(Name, Project, Change);
+				foreach (BgNode ReferencedNode in ResolveReferences(Element, RequiredNames))
+				{
+					NewBadge.Nodes.Add(ReferencedNode);
+				}
+				foreach (BgNode ReferencedNode in ResolveReferences(Element, TargetNames))
+				{
+					NewBadge.Nodes.Add(ReferencedNode);
+					NewBadge.Nodes.UnionWith(ReferencedNode.OrderDependencies);
+				}
+				Graph.Badges.Add(NewBadge);
+			}
+		}
+
+		/// <summary>
+		/// Reads the definition for a label
+		/// </summary>
+		/// <param name="Element">Xml element to read the definition from</param>
+		void ReadLabel(BgScriptElement Element)
+		{
+			if (EvaluateCondition(Element))
+			{
+				string Name = ReadAttribute(Element, "Name");
+				if (!String.IsNullOrEmpty(Name))
+				{
+					ValidateName(Element, Name);
+				}
+
+				string Category = ReadAttribute(Element, "Category");
+
+				string[] RequiredNames = ReadListAttribute(Element, "Requires");
+				string[] IncludedNames = ReadListAttribute(Element, "Include");
+				string[] ExcludedNames = ReadListAttribute(Element, "Exclude");
+
+				string UgsBadge = ReadAttribute(Element, "UgsBadge");
+				string UgsProject = ReadAttribute(Element, "UgsProject");
+
+				BgLabelChange Change = ReadEnumAttribute<BgLabelChange>(Element, "Change", BgLabelChange.Current);
+
+				BgLabel NewLabel = new BgLabel(Name, Category, UgsBadge, UgsProject, Change);
+				foreach (BgNode ReferencedNode in ResolveReferences(Element, RequiredNames))
+				{
+					NewLabel.RequiredNodes.Add(ReferencedNode);
+					NewLabel.IncludedNodes.Add(ReferencedNode);
+					NewLabel.IncludedNodes.UnionWith(ReferencedNode.OrderDependencies);
+				}
+				foreach (BgNode IncludedNode in ResolveReferences(Element, IncludedNames))
+				{
+					NewLabel.IncludedNodes.Add(IncludedNode);
+					NewLabel.IncludedNodes.UnionWith(IncludedNode.OrderDependencies);
+				}
+				foreach (BgNode ExcludedNode in ResolveReferences(Element, ExcludedNames))
+				{
+					NewLabel.IncludedNodes.Remove(ExcludedNode);
+					NewLabel.IncludedNodes.ExceptWith(ExcludedNode.OrderDependencies);
+				}
+				Graph.Labels.Add(NewLabel);
+			}
+		}
+
+		/// <summary>
+		/// Reads the definition for a node, and adds it to the given agent
+		/// </summary>
+		/// <param name="Element">Xml element to read the definition from</param>
+		/// <param name="ParentAgent">Agent for the node to be added to</param>
+		void ReadNode(BgScriptElement Element, BgAgent ParentAgent)
+		{
+			string Name;
+			if (EvaluateCondition(Element) && TryReadObjectName(Element, out Name))
+			{
+				string[] RequiresNames = ReadListAttribute(Element, "Requires");
+				string[] ProducesNames = ReadListAttribute(Element, "Produces");
+				string[] AfterNames = ReadListAttribute(Element, "After");
+				string[] TokenFileNames = ReadListAttribute(Element, "Token");
+				bool bRunEarly = ReadBooleanAttribute(Element, "RunEarly", false);
+				bool bNotifyOnWarnings = ReadBooleanAttribute(Element, "NotifyOnWarnings", true);
+
+				// Resolve all the inputs we depend on
+				HashSet<BgNodeOutput> Inputs = ResolveInputReferences(Element, RequiresNames);
+
+				// Gather up all the input dependencies, and check they're all upstream of the current node
+				HashSet<BgNode> InputDependencies = new HashSet<BgNode>();
+				foreach (BgNode InputDependency in Inputs.Select(x => x.ProducingNode).Distinct())
+				{
+					InputDependencies.Add(InputDependency);
+				}
+
+				// Remove all the lock names from the list of required names
+				HashSet<FileReference> RequiredTokens = new HashSet<FileReference>(TokenFileNames.Select(x => new FileReference(x)));
+
+				// Recursively include all their dependencies too
+				foreach (BgNode InputDependency in InputDependencies.ToArray())
+				{
+					RequiredTokens.UnionWith(InputDependency.RequiredTokens);
+					InputDependencies.UnionWith(InputDependency.InputDependencies);
+				}
+
+				// Validate all the outputs
+				List<string> ValidOutputNames = new List<string>();
+				foreach (string ProducesName in ProducesNames)
+				{
+					BgNodeOutput ExistingOutput;
+					if(Graph.TagNameToNodeOutput.TryGetValue(ProducesName, out ExistingOutput))
+					{
+						LogError(Element, "Output tag '{0}' is already generated by node '{1}'", ProducesName, ExistingOutput.ProducingNode.Name);
+					}
+					else if(!ProducesName.StartsWith("#"))
+					{
+						LogError(Element, "Output tag names must begin with a '#' character ('{0}')", ProducesName);
+					}
+					else
+					{
+						ValidOutputNames.Add(ProducesName);
+					}
+				}
+
+				// Gather up all the order dependencies
+				HashSet<BgNode> OrderDependencies = new HashSet<BgNode>(InputDependencies);
+				OrderDependencies.UnionWith(ResolveReferences(Element, AfterNames));
+
+				// Recursively include all their order dependencies too
+				foreach (BgNode OrderDependency in OrderDependencies.ToArray())
+				{
+					OrderDependencies.UnionWith(OrderDependency.OrderDependencies);
+				}
+
+				// Check that we're not dependent on anything completing that is declared after the initial declaration of this agent.
+				int AgentIdx = Graph.Agents.IndexOf(ParentAgent);
+				for (int Idx = AgentIdx + 1; Idx < Graph.Agents.Count; Idx++)
+				{
+					foreach (BgNode Node in Graph.Agents[Idx].Nodes.Where(x => OrderDependencies.Contains(x)))
+					{
+						LogError(Element, "Node '{0}' has a dependency on '{1}', which was declared after the initial definition of '{2}'.", Name, Node.Name, ParentAgent.Name);
+					}
+				}
+
+				// Construct and register the node
+				if (CheckNameIsUnique(Element, Name))
+				{
+					// Add it to the node lookup
+					BgNode NewNode = new BgNode(Name, Inputs.ToArray(), ValidOutputNames.ToArray(), InputDependencies.ToArray(), OrderDependencies.ToArray(), RequiredTokens.ToArray());
+					NewNode.bRunEarly = bRunEarly;
+					NewNode.bNotifyOnWarnings = bNotifyOnWarnings;
+					Graph.NameToNode.Add(Name, NewNode);
+
+					// Register all the output tags in the global name table.
+					foreach(BgNodeOutput Output in NewNode.Outputs)
+					{
+						BgNodeOutput ExistingOutput;
+						if(Graph.TagNameToNodeOutput.TryGetValue(Output.TagName, out ExistingOutput))
+						{
+							LogError(Element, "Node '{0}' already has an output called '{1}'", ExistingOutput.ProducingNode.Name, Output.TagName);
+						}
+						else
+						{
+							Graph.TagNameToNodeOutput.Add(Output.TagName, Output);
+						}
+					}
+
+					// Add all the tasks
+					ReadNodeBody(Element, NewNode, ParentAgent);
+
+					// Add it to the current agent
+					ParentAgent.Nodes.Add(NewNode);
+				}
+			}
+		}
+
+		/// <summary>
+		/// Reads the contents of a node element
+		/// </summary>
+		/// <param name="Element">Xml element to read the definition from</param>
+		/// <param name="NewNode">The new node that has been created</param>
+		/// <param name="ParentAgent">Agent for the node to be added to</param>
+		void ReadNodeBody(XmlElement Element, BgNode NewNode, BgAgent ParentAgent)
+		{
+			EnterScope();
+			foreach (BgScriptElement ChildElement in Element.ChildNodes.OfType<BgScriptElement>())
+			{
+				switch (ChildElement.Name)
+				{
+					case "Property":
+						ReadProperty(ChildElement);
+						break;
+					case "Regex":
+						ReadRegex(ChildElement);
+						break;
+					case "Trace":
+						ReadDiagnostic(ChildElement, LogEventType.Console, NewNode, ParentAgent);
+						break;
+					case "Warning":
+						ReadDiagnostic(ChildElement, LogEventType.Warning, NewNode, ParentAgent);
+						break;
+					case "Error":
+						ReadDiagnostic(ChildElement, LogEventType.Error, NewNode, ParentAgent);
+						break;
+					case "Do":
+						ReadBlock(ChildElement, x => ReadNodeBody(x, NewNode, ParentAgent));
+						break;
+					case "Switch":
+						ReadSwitch(ChildElement, x => ReadNodeBody(x, NewNode, ParentAgent));
+						break;
+					case "ForEach":
+						ReadForEach(ChildElement, x => ReadNodeBody(x, NewNode, ParentAgent));
+						break;
+					case "Expand":
+						ReadExpand(ChildElement, x => ReadNodeBody(x, NewNode, ParentAgent));
+						break;
+					default:
+						ReadTask(ChildElement, NewNode);
+						break;
+				}
+			}
+			LeaveScope();
+		}
+
+		/// <summary>
+		/// Reads a block element
+		/// </summary>
+		/// <param name="Element">Xml element to read the definition from</param>
+		/// <param name="ReadContents">Delegate to read the contents of the element, if the condition evaluates to true</param>
+		void ReadBlock(BgScriptElement Element, Action<BgScriptElement> ReadContents)
+		{
+			if (EvaluateCondition(Element))
+			{
+				ReadContents(Element);
+			}
+		}
+
+		/// <summary>
+		/// Reads a "Switch" element 
+		/// </summary>
+		/// <param name="Element">Xml element to read the definition from</param>
+		/// <param name="ReadContents">Delegate to read the contents of the element, if the condition evaluates to true</param>
+		void ReadSwitch(BgScriptElement Element, Action<BgScriptElement> ReadContents)
+		{
+			foreach (BgScriptElement ChildElement in Element.ChildNodes.OfType<BgScriptElement>())
+			{
+				if (ChildElement.Name == "Default" || EvaluateCondition(ChildElement))
+				{
+					ReadContents(ChildElement);
+					break;
+				}
+			}
+		}
+
+		/// <summary>
+		/// Reads a "ForEach" element 
+		/// </summary>
+		/// <param name="Element">Xml element to read the definition from</param>
+		/// <param name="ReadContents">Delegate to read the contents of the element, if the condition evaluates to true</param>
+		void ReadForEach(BgScriptElement Element, Action<BgScriptElement> ReadContents)
+		{
+			EnterScope();
+			if(EvaluateCondition(Element))
+			{
+				string Name = ReadAttribute(Element, "Name");
+				string Separator = ReadAttribute(Element, "Separator");
+				if(Separator.Length > 1)
+				{
+					LogWarning(Element, "Node {0}'s Separator attribute is more than one character ({1}). Defaulting to ;", Name, Separator);
+					Separator = ";";
+				}
+				if(string.IsNullOrEmpty(Separator))
+				{
+					Separator = ";";
+				}
+				if (ValidateName(Element, Name))
+				{
+					if(ScopedProperties.Any(x => x.ContainsKey(Name)))
+					{
+						LogError(Element, "Loop variable '{0}' already exists as a local property in an outer scope", Name);
+					}
+					else
+					{
+						// Loop through all the values
+						string[] Values = ReadListAttribute(Element, "Values", Convert.ToChar(Separator));
+						foreach(string Value in Values)
+						{
+							ScopedProperties[ScopedProperties.Count - 1][Name] = Value;
+							ReadContents(Element);
+						}
+					}
+				}
+			}
+			LeaveScope();
+		}
+
+		/// <summary>
+		/// Reads an "Expand" element 
+		/// </summary>
+		/// <param name="Element">Xml element to read the definition from</param>
+		/// <param name="ReadContents">Delegate to read the contents of the element, if the condition evaluates to true</param>
+		void ReadExpand(BgScriptElement Element, Action<BgScriptElement> ReadContents)
+		{
+			if(EvaluateCondition(Element))
+			{
+				string Name = ReadAttribute(Element, "Name");
+
+				BgScriptMacro Macro;
+				if(!MacroNameToDefinition.TryGetValue(Name, out Macro))
+				{
+					LogError(Element, "Macro '{0}' does not exist", Name);
+				}
+				else
+				{
+					// Parse the argument list
+					string[] Arguments = new string[Macro.ArgumentNameToIndex.Count];
+					foreach(XmlAttribute Attribute in Element.Attributes)
+					{
+						if(Attribute.Name != "Name" && Attribute.Name != "If")
+						{
+							int Index;
+							if(Macro.ArgumentNameToIndex.TryGetValue(Attribute.Name, out Index))
+							{
+								Arguments[Index] = ExpandProperties(Element, Attribute.Value);
+							}
+							else
+							{
+								LogWarning(Element, "Macro '{0}' does not take an argument '{1}'", Name, Attribute.Name);
+							}
+						}
+					}
+
+					// Make sure none of the required arguments are missing
+					bool bHasMissingArguments = false;
+					for(int Idx = 0; Idx < Macro.NumRequiredArguments; Idx++)
+					{
+						if(Arguments[Idx] == null)
+						{
+							LogWarning(Element, "Macro '{0}' is missing argument '{1}'", Macro.Name, Macro.ArgumentNameToIndex.First(x => x.Value == Idx).Key);
+							bHasMissingArguments = true;
+						}
+					}
+
+					// Expand the function
+					if(!bHasMissingArguments)
+					{
+						EnterScope();
+						foreach(KeyValuePair<string, int> Pair in Macro.ArgumentNameToIndex)
+						{
+							ScopedProperties[ScopedProperties.Count - 1][Pair.Key] = Arguments[Pair.Value] ?? "";
+						}
+						ReadContents(Macro.Element);
+						LeaveScope();
+					}
+				}
+			}
+		}
+
+		/// <summary>
+		/// Reads a task definition from the given element, and add it to the given list
+		/// </summary>
+		/// <param name="Element">Xml element to read the definition from</param>
+		/// <param name="ParentNode">The node which owns this task</param>
+		void ReadTask(BgScriptElement Element, BgNode ParentNode)
+		{
+			// If we're running a single node and this element's parent isn't the single node to run, ignore the error and return.
+			if (!string.IsNullOrWhiteSpace(SingleNodeName) && ParentNode.Name != SingleNodeName)
+			{
+				return;
+			}
+
+			if (EvaluateCondition(Element))
+			{
+				BgTask Info = new BgTask(Tuple.Create(Element.File, Element.LineNumber), Element.Name);
+				foreach (XmlAttribute Attribute in Element.Attributes)
+				{
+					if (String.Compare(Attribute.Name, "If", StringComparison.InvariantCultureIgnoreCase) != 0)
+					{
+						string ExpandedValue = ExpandProperties(Element, Attribute.Value);
+						Info.Arguments.Add(Attribute.Name, ExpandedValue);
+					}
+				}
+				ParentNode.TaskInfos.Add(Info);
+			}
+		}
+
+		/// <summary>
+		/// Reads the definition for an email notifier
+		/// </summary>
+		/// <param name="Element">Xml element to read the definition from</param>
+		void ReadNotifier(BgScriptElement Element)
+		{
+			if (EvaluateCondition(Element))
+			{
+				string[] TargetNames = ReadListAttribute(Element, "Targets");
+				string[] ExceptNames = ReadListAttribute(Element, "Except");
+				string[] IndividualNodeNames = ReadListAttribute(Element, "Nodes");
+				string[] ReportNames = ReadListAttribute(Element, "Reports");
+				string[] Users = ReadListAttribute(Element, "Users");
+				string[] Submitters = ReadListAttribute(Element, "Submitters");
+				bool? bWarnings = Element.HasAttribute("Warnings") ? (bool?)ReadBooleanAttribute(Element, "Warnings", true) : null;
+				bool bAbsolute = Element.HasAttribute("Absolute") ? ReadBooleanAttribute(Element, "Absolute", true) : false;
+
+				// Find the list of targets which are included, and recurse through all their dependencies
+				HashSet<BgNode> Nodes = new HashSet<BgNode>();
+				if (TargetNames != null)
+				{
+					HashSet<BgNode> TargetNodes = ResolveReferences(Element, TargetNames);
+					foreach (BgNode Node in TargetNodes)
+					{
+						Nodes.Add(Node);
+						Nodes.UnionWith(Node.InputDependencies);
+					}
+				}
+
+				// Add all the individually referenced nodes
+				if (IndividualNodeNames != null)
+				{
+					HashSet<BgNode> IndividualNodes = ResolveReferences(Element, IndividualNodeNames);
+					Nodes.UnionWith(IndividualNodes);
+				}
+
+				// Exclude all the exceptions
+				if (ExceptNames != null)
+				{
+					HashSet<BgNode> ExceptNodes = ResolveReferences(Element, ExceptNames);
+					Nodes.ExceptWith(ExceptNodes);
+				}
+
+				// Update all the referenced nodes with the settings
+				foreach (BgNode Node in Nodes)
+				{
+					if (Users != null)
+					{
+						if (bAbsolute)
+						{
+							Node.NotifyUsers = new HashSet<string>(Users);
+						}
+						else
+						{
+							Node.NotifyUsers.UnionWith(Users);
+						}
+					}
+					if (Submitters != null)
+					{
+						if (bAbsolute)
+						{
+							Node.NotifySubmitters = new HashSet<string>(Submitters);
+						}
+						else
+						{
+							Node.NotifySubmitters.UnionWith(Submitters);
+						}
+					}
+					if (bWarnings.HasValue)
+					{
+						Node.bNotifyOnWarnings = bWarnings.Value;
+					}
+				}
+
+				// Add the users to the list of reports
+				if (ReportNames != null)
+				{
+					foreach (string ReportName in ReportNames)
+					{
+						BgReport Report;
+						if (Graph.NameToReport.TryGetValue(ReportName, out Report))
+						{
+							Report.NotifyUsers.UnionWith(Users);
+						}
+						else
+						{
+							LogError(Element, "Report '{0}' has not been defined", ReportName);
+						}
+					}
+				}
+			}
+		}
+
+		/// <summary>
+		/// Reads a warning from the given element, evaluates the condition on it, and writes it to the log if the condition passes.
+		/// </summary>
+		/// <param name="Element">Xml element to read the definition from</param>
+		/// <param name="EventType">The diagnostic event type</param>
+		/// <param name="EnclosingNode">The node that this diagnostic is declared in, or null</param>
+		/// <param name="EnclosingAgent">The agent that this diagnostic is declared in, or null</param>
+		void ReadDiagnostic(BgScriptElement Element, LogEventType EventType, BgNode EnclosingNode, BgAgent EnclosingAgent)
+		{
+			if (EvaluateCondition(Element))
+			{
+				string Message = ReadAttribute(Element, "Message");
+
+				BgScriptDiagnostic Diagnostic = new BgScriptDiagnostic();
+				Diagnostic.EventType = EventType;
+				Diagnostic.Message = String.Format("{0}({1}): {2}", Element.File, Element.LineNumber, Message);
+				Diagnostic.EnclosingNode = EnclosingNode;
+				Diagnostic.EnclosingAgent = EnclosingAgent;
+				Graph.Diagnostics.Add(Diagnostic);
+			}
+		}
+
+		/// <summary>
+		/// Checks that the given name does not already used to refer to a node, and print an error if it is.
+		/// </summary>
+		/// <param name="Element">Xml element to read from</param>
+		/// <param name="Name">Name of the alias</param>
+		/// <returns>True if the name was registered correctly, false otherwise.</returns>
+		bool CheckNameIsUnique(BgScriptElement Element, string Name)
+		{
+			// Get the nodes that it maps to
+			if (Graph.ContainsName(Name))
+			{
+				LogError(Element, "'{0}' is already defined; cannot add a second time", Name);
+				return false;
+			}
+			return true;
+		}
+
+		/// <summary>
+		/// Resolve a list of references to a set of nodes
+		/// </summary>
+		/// <param name="Element">Element used to locate any errors</param>
+		/// <param name="ReferenceNames">Sequence of names to look up</param>
+		/// <returns>Hashset of all the nodes included by the given names</returns>
+		HashSet<BgNode> ResolveReferences(BgScriptElement Element, IEnumerable<string> ReferenceNames)
+		{
+			HashSet<BgNode> Nodes = new HashSet<BgNode>();
+			foreach (string ReferenceName in ReferenceNames)
+			{
+				BgNode[] OtherNodes;
+				if (Graph.TryResolveReference(ReferenceName, out OtherNodes))
+				{
+					Nodes.UnionWith(OtherNodes);
+				}
+				else if (!ReferenceName.StartsWith("#") && Graph.TagNameToNodeOutput.ContainsKey("#" + ReferenceName))
+				{
+					LogError(Element, "Reference to '{0}' cannot be resolved; did you mean '#{0}'?", ReferenceName);
+				}
+				else
+				{
+					LogError(Element, "Reference to '{0}' cannot be resolved; check it has been defined.", ReferenceName);
+				}
+			}
+			return Nodes;
+		}
+
+		/// <summary>
+		/// Resolve a list of references to a set of nodes
+		/// </summary>
+		/// <param name="Element">Element used to locate any errors</param>
+		/// <param name="ReferenceNames">Sequence of names to look up</param>
+		/// <returns>Set of all the nodes included by the given names</returns>
+		HashSet<BgNodeOutput> ResolveInputReferences(BgScriptElement Element, IEnumerable<string> ReferenceNames)
+		{
+			HashSet<BgNodeOutput> Inputs = new HashSet<BgNodeOutput>();
+			foreach (string ReferenceName in ReferenceNames)
+			{
+				BgNodeOutput[] ReferenceInputs;
+				if (Graph.TryResolveInputReference(ReferenceName, out ReferenceInputs))
+				{
+					Inputs.UnionWith(ReferenceInputs);
+				}
+				else if (!ReferenceName.StartsWith("#") && Graph.TagNameToNodeOutput.ContainsKey("#" + ReferenceName))
+				{
+					LogError(Element, "Reference to '{0}' cannot be resolved; did you mean '#{0}'?", ReferenceName);
+				}
+				else
+				{
+					LogError(Element, "Reference to '{0}' cannot be resolved; check it has been defined.", ReferenceName);
+				}
+			}
+			return Inputs;
+		}
+
+		/// <summary>
+		/// Reads an object name from its defining element. Outputs an error if the name is missing.
+		/// </summary>
+		/// <param name="Element">Element to read the name for</param>
+		/// <param name="Name">Output variable to receive the name of the object</param>
+		/// <returns>True if the object had a valid name (assigned to the Name variable), false if the name was invalid or missing.</returns>
+		bool TryReadObjectName(BgScriptElement Element, out string Name)
+		{
+			// Check the name attribute is present
+			if (!Element.HasAttribute("Name"))
+			{
+				LogError(Element, "Missing 'Name' attribute");
+				Name = null;
+				return false;
+			}
+
+			// Get the value of it, strip any leading or trailing whitespace, and make sure it's not empty
+			string Value = ReadAttribute(Element, "Name");
+			if (!ValidateName(Element, Value))
+			{
+				Name = null;
+				return false;
+			}
+
+			// Return it
+			Name = Value;
+			return true;
+		}
+
+		/// <summary>
+		/// Reads an qualified object name from its defining element. Outputs an error if the name is missing.
+		/// </summary>
+		/// <param name="Element">Element to read the name for</param>
+		/// <param name="QualifiedName">Output variable to receive the name of the object</param>
+		/// <returns>True if the object had a valid name (assigned to the Name variable), false if the name was invalid or missing.</returns>
+		bool TryReadQualifiedObjectName(BgScriptElement Element, out string[] QualifiedName)
+		{
+			// Check the name attribute is present
+			if (!Element.HasAttribute("Name"))
+			{
+				LogError(Element, "Missing 'Name' attribute");
+				QualifiedName = null;
+				return false;
+			}
+
+			// Get the value of it, strip any leading or trailing whitespace, and make sure it's not empty
+			string[] Values = ReadAttribute(Element, "Name").Split('.');
+			foreach (string Value in Values)
+			{
+				if (!ValidateName(Element, Value))
+				{
+					QualifiedName = null;
+					return false;
+				}
+			}
+
+			// Return it
+			QualifiedName = Values;
+			return true;
+		}
+
+		/// <summary>
+		/// Checks that the given name is valid syntax
+		/// </summary>
+		/// <param name="Element">The element that contains the name</param>
+		/// <param name="Name">The name to check</param>
+		/// <returns>True if the name is valid</returns>
+		bool ValidateName(BgScriptElement Element, string Name)
+		{
+			// Check it's not empty
+			if (Name.Length == 0)
+			{
+				LogError(Element, "Name is empty");
+				return false;
+			}
+
+			// Check there are no invalid characters
+			for (int Idx = 0; Idx < Name.Length; Idx++)
+			{
+				if (Idx > 0 && Name[Idx] == ' ' && Name[Idx - 1] == ' ')
+				{
+					LogError(Element, "Consecutive spaces in object name - '{0}'", Name);
+					return false;
+				}
+				if(Char.IsControl(Name[Idx]) || BgScriptSchema.IllegalNameCharacters.IndexOf(Name[Idx]) != -1)
+				{
+					LogError(Element, "Invalid character in object name - '{0}'", Name[Idx]);
+					return false;
+				}
+			}
+			return true;
+		}
+
+		/// <summary>
+		/// Constructs a regex from a regex string and returns it
+		/// </summary>
+		/// <param name="Element">The element that contains the regex</param>
+		/// <param name="Regex">The pattern to construct</param>
+		/// <returns>The regex if is valid, otherwise null</returns>
+		Regex ParseRegex(BgScriptElement Element, string Regex)
+		{
+			if(Regex.Length == 0)
+			{
+				LogError(Element, "Regex is empty");
+				return null;
+			}
+			try
+			{
+				return new Regex(Regex);
+			}
+			catch(ArgumentException InvalidRegex)
+			{
+				LogError(Element, "Could not construct the Regex, reason: {0}", InvalidRegex.Message);
+				return null;
+			}
+		}
+
+		/// <summary>
+		/// Expands any properties and reads an attribute.
+		/// </summary>
+		/// <param name="Element">Element to read the attribute from</param>
+		/// <param name="Name">Name of the attribute</param>
+		/// <returns>Array of names, with all leading and trailing whitespace removed</returns>
+		string ReadAttribute(BgScriptElement Element, string Name)
+		{
+			return ExpandProperties(Element, Element.GetAttribute(Name));
+		}
+
+		/// <summary>
+		/// Expands any properties and reads a list of strings from an attribute, separated by semi-colon characters
+		/// </summary>
+		/// <param name="Element"></param>
+		/// <param name="Name"></param>
+		/// <param name="Separator"></param>
+		/// <returns>Array of names, with all leading and trailing whitespace removed</returns>
+		string[] ReadListAttribute(BgScriptElement Element, string Name, char Separator = ';')
+		{
+			string Value = ReadAttribute(Element, Name);
+			return Value.Split(new char[] { Separator }).Select(x => x.Trim()).Where(x => x.Length > 0).ToArray();
+		}
+
+		/// <summary>
+		/// Reads an attribute from the given XML element, expands any properties in it, and parses it as a boolean.
+		/// </summary>
+		/// <param name="Element">Element to read the attribute from</param>
+		/// <param name="Name">Name of the attribute</param>
+		/// <param name="bDefaultValue">Default value if the attribute is missing</param>
+		/// <returns>The value of the attribute field</returns>
+		bool ReadBooleanAttribute(BgScriptElement Element, string Name, bool bDefaultValue)
+		{
+			bool bResult = bDefaultValue;
+			if (Element.HasAttribute(Name))
+			{
+				string Value = ReadAttribute(Element, Name).Trim();
+				if (Value.Equals("true", StringComparison.InvariantCultureIgnoreCase))
+				{
+					bResult = true;
+				}
+				else if (Value.Equals("false", StringComparison.InvariantCultureIgnoreCase))
+				{
+					bResult = false;
+				}
+				else
+				{
+					LogError(Element, "Invalid boolean value '{0}' - expected 'true' or 'false'", Value);
+				}
+			}
+			return bResult;
+		}
+
+		/// <summary>
+		/// Reads an attribute from the given XML element, expands any properties in it, and parses it as an integer.
+		/// </summary>
+		/// <param name="Element">Element to read the attribute from</param>
+		/// <param name="Name">Name of the attribute</param>
+		/// <param name="DefaultValue">Default value for the integer, if the attribute is missing</param>
+		/// <returns>The value of the attribute field</returns>
+		int ReadIntegerAttribute(BgScriptElement Element, string Name, int DefaultValue)
+		{
+			int Result = DefaultValue;
+			if (Element.HasAttribute(Name))
+			{
+				string Value = ReadAttribute(Element, Name).Trim();
+
+				int IntValue;
+				if (Int32.TryParse(Value, out IntValue))
+				{
+					Result = IntValue;
+				}
+				else
+				{
+					LogError(Element, "Invalid integer value '{0}'", Value);
+				}
+			}
+			return Result;
+		}
+
+		/// <summary>
+		/// Reads an attribute from the given XML element, expands any properties in it, and parses it as an enum of the given type.
+		/// </summary>
+		/// <typeparam name="T">The enum type to parse the attribute as</typeparam>
+		/// <param name="Element">Element to read the attribute from</param>
+		/// <param name="Name">Name of the attribute</param>
+		/// <param name="DefaultValue">Default value for the enum, if the attribute is missing</param>
+		/// <returns>The value of the attribute field</returns>
+		T ReadEnumAttribute<T>(BgScriptElement Element, string Name, T DefaultValue) where T : struct
+		{
+			T Result = DefaultValue;
+			if (Element.HasAttribute(Name))
+			{
+				string Value = ReadAttribute(Element, Name).Trim();
+
+				T EnumValue;
+				if (Enum.TryParse(Value, true, out EnumValue))
+				{
+					Result = EnumValue;
+				}
+				else
+				{
+					LogError(Element, "Invalid value '{0}' - expected {1}", Value, String.Join("/", Enum.GetNames(typeof(T))));
+				}
+			}
+			return Result;
+		}
+
+		/// <summary>
+		/// Outputs an error message to the log and increments the number of errors, referencing the file and line number of the element that caused it.
+		/// </summary>
+		/// <param name="Element">The script element causing the error</param>
+		/// <param name="Format">Standard String.Format()-style format string</param>
+		/// <param name="Args">Optional arguments</param>
+		void LogError(BgScriptElement Element, string Format, params object[] Args)
+		{
+			Logger.LogScriptError(Element.NativeFile, Element.LineNumber, Format, Args);
+			NumErrors++;
+		}
+
+		/// <summary>
+		/// Outputs a warning message to the log and increments the number of errors, referencing the file and line number of the element that caused it.
+		/// </summary>
+		/// <param name="Element">The script element causing the error</param>
+		/// <param name="Format">Standard String.Format()-style format string</param>
+		/// <param name="Args">Optional arguments</param>
+		void LogWarning(BgScriptElement Element, string Format, params object[] Args)
+		{
+			Logger.LogScriptWarning(Element.NativeFile, Element.LineNumber, Format, Args);
+		}
+
+		/// <summary>
+		/// Evaluates the (optional) conditional expression on a given XML element via the If="..." attribute, and returns true if the element is enabled.
+		/// </summary>
+		/// <param name="Element">The element to check</param>
+		/// <returns>True if the element's condition evaluates to true (or doesn't have a conditional expression), false otherwise</returns>
+		bool EvaluateCondition(BgScriptElement Element)
+		{
+			// Check if the element has a conditional attribute
+			const string AttributeName = "If";
+			if (!Element.HasAttribute(AttributeName))
+			{
+				return true;
+			}
+
+			// If it does, try to evaluate it.
+			try
+			{
+				string Text = ExpandProperties(Element, Element.GetAttribute("If"));
+				return BgCondition.Evaluate(Text, Context);
+			}
+			catch (BgConditionException Ex)
+			{
+				LogError(Element, "Error in condition: {0}", Ex.Message);
+				return false;
+			}
+		}
+
+		/// <summary>
+		/// Expand all the property references (of the form $(PropertyName)) in a string.
+		/// </summary>
+		/// <param name="Element">The element containing the string. Used for diagnostic messages.</param>
+		/// <param name="Text">The input string to expand properties in</param>
+		/// <returns>The expanded string</returns>
+		string ExpandProperties(BgScriptElement Element, string Text)
+		{
+			string Result = Text;
+			// Iterate in reverse order to handle cases where there are nested expansions like $(Outer$(Inner))
+			for (int Idx = Result.LastIndexOf("$("); Idx != -1; Idx = Result.LastIndexOf("$(", Idx, Idx+1))
+			{
+				// Find the end of the variable name
+				int EndIdx = Result.IndexOf(')', Idx + 2);
+				if (EndIdx == -1)
+				{
+					break;
+				}
+
+				// Extract the variable name from the string
+				string Name = Result.Substring(Idx + 2, EndIdx - (Idx + 2));
+
+				// Find the value for it, either from the dictionary or the environment block
+				string Value;
+				if (!TryGetPropertyValue(Name, out Value))
+				{
+					LogWarning(Element, "Property '{0}' is not defined", Name);
+					Value = "";
+				}
+
+				// Check if we've got a value for this variable
+				if (Value != null)
+				{
+					// Replace the variable, or skip past it
+					Result = Result.Substring(0, Idx) + Value + Result.Substring(EndIdx + 1);
+				}
+			}
+			return Result;
+		}
+	}
+}
diff --git a/Engine/Source/Programs/Shared/EpicGames.BuildGraph/BgTask.cs b/Engine/Source/Programs/Shared/EpicGames.BuildGraph/BgTask.cs
new file mode 100644
--- /dev/null
+++ b/Engine/Source/Programs/Shared/EpicGames.BuildGraph/BgTask.cs
@@ -0,0 +1,59 @@
+// Copyright Epic Games, Inc. All Rights Reserved.
+
+using EpicGames.Core;
+using Microsoft.Extensions.Logging;
+using System;
+using System.Collections.Generic;
+using System.Text;
+using System.Xml;
+
+namespace EpicGames.BuildGraph
+{
+	/// <summary>
+	/// A task invocation
+	/// </summary>
+	public class BgTask
+	{
+		/// <summary>
+		/// Line number in a source file that this task was declared. Optional; used for log messages.
+		/// </summary>
+		public Tuple<string, int> SourceLocation
+		{
+			get;
+			set;
+		}
+
+		/// <summary>
+		/// Name of the task
+		/// </summary>
+		public string Name { get; set; }
+
+		/// <summary>
+		/// Arguments for the task
+		/// </summary>
+		public Dictionary<string, string> Arguments { get; } = new Dictionary<string, string>();
+
+		/// <summary>
+		/// Constructor
+		/// </summary>
+		public BgTask(Tuple<string, int> SourceLocation, string Name)
+		{
+			this.SourceLocation = SourceLocation;
+			this.Name = Name;
+		}
+
+		/// <summary>
+		/// Write to an xml file
+		/// </summary>
+		/// <param name="Writer"></param>
+		public void Write(XmlWriter Writer)
+		{
+			Writer.WriteStartElement(Name);
+			foreach (KeyValuePair<string, string> Argument in Arguments)
+			{
+				Writer.WriteAttributeString(Argument.Key, Argument.Value);
+			}
+			Writer.WriteEndElement();
+		}
+	}
+}
diff --git a/Engine/Source/Programs/Shared/EpicGames.BuildGraph/EpicGames.BuildGraph.csproj b/Engine/Source/Programs/Shared/EpicGames.BuildGraph/EpicGames.BuildGraph.csproj
new file mode 100644
--- /dev/null
+++ b/Engine/Source/Programs/Shared/EpicGames.BuildGraph/EpicGames.BuildGraph.csproj
@@ -0,0 +1,35 @@
+﻿<Project Sdk="Microsoft.NET.Sdk">
+  <PropertyGroup>
+    <TargetFramework>netcoreapp3.1</TargetFramework>
+    <Configuration Condition=" '$(Configuration)' == '' ">Development</Configuration>
+    <OutputType>Library</OutputType>
+    <RootNamespace>EpicGames.BuildGraph</RootNamespace>
+    <GenerateAssemblyInfo>false</GenerateAssemblyInfo>
+    <GenerateTargetFrameworkAttribute>false</GenerateTargetFrameworkAttribute>
+    <Configurations>Debug;Release;Development</Configurations>
+    <Nullable>disable</Nullable>
+    <DebugType>pdbonly</DebugType>
+  </PropertyGroup>
+  <PropertyGroup Condition=" '$(Configuration)|$(Platform)' == 'Debug|AnyCPU' ">
+    <DocumentationFile></DocumentationFile>
+  </PropertyGroup>
+  <PropertyGroup Condition="'$(Configuration)|$(Platform)' == 'Development|AnyCPU'">
+    <DocumentationFile></DocumentationFile>
+    <Optimize>true</Optimize>
+  </PropertyGroup>
+  <PropertyGroup Condition="'$(Configuration)|$(Platform)' == 'Release|AnyCPU'">
+    <DocumentationFile></DocumentationFile>
+    <Optimize>true</Optimize>
+  </PropertyGroup>
+  
+  <ItemGroup>
+    <Compile Include="..\MetaData.cs">
+      <Link>Properties\MetaData.cs</Link>
+    </Compile>
+  </ItemGroup>
+  <ItemGroup>
+    <ProjectReference Include="..\..\Shared\EpicGames.Core\EpicGames.Core.csproj" />
+  </ItemGroup>
+  <ItemGroup>
+  </ItemGroup>
+</Project>
\ No newline at end of file
