Add support for setting passwords on the native client at creation time.

diff --git a/Engine/Source/Programs/Shared/EpicGames.Perforce.Native/Client.cpp b/Engine/Source/Programs/Shared/EpicGames.Perforce.Native/Client.cpp
--- a/Engine/Source/Programs/Shared/EpicGames.Perforce.Native/Client.cpp
+++ b/Engine/Source/Programs/Shared/EpicGames.Perforce.Native/Client.cpp
@@ -13,8 +13,9 @@
 
 struct FSettings
 {
-	const char* Port;
+	const char* ServerAndPort;
 	const char* User;
+	const char* Password;
 	const char* Client;
 	const char* AppName;
 	const char* AppVersion;
@@ -391,14 +392,18 @@ extern "C" NATIVE_API FClient* Client_Create(const FSettings* Settings, FWriteBu
 
 	if (Settings != nullptr)
 	{
-		if (Settings->Port != nullptr)
+		if (Settings->ServerAndPort != nullptr)
 		{
-			Client->ClientApi.SetPort(Settings->Port);
+			Client->ClientApi.SetPort(Settings->ServerAndPort);
 		}
 		if (Settings->User != nullptr)
 		{
 			Client->ClientApi.SetUser(Settings->User);
 		}
+		if (Settings->Password != nullptr)
+		{
+			Client->ClientApi.SetPassword(Settings->Password);
+		}
 		if (Settings->Client != nullptr)
 		{
 			Client->ClientApi.SetClient(Settings->Client);
diff --git a/Engine/Source/Programs/Shared/EpicGames.Perforce/NativePerforceConnection.cs b/Engine/Source/Programs/Shared/EpicGames.Perforce/NativePerforceConnection.cs
--- a/Engine/Source/Programs/Shared/EpicGames.Perforce/NativePerforceConnection.cs
+++ b/Engine/Source/Programs/Shared/EpicGames.Perforce/NativePerforceConnection.cs
@@ -25,11 +25,14 @@ namespace EpicGames.Perforce
 		class NativeSettings
 		{
 			[MarshalAs(UnmanagedType.LPStr)]
-			public string? Port;
+			public string? ServerAndPort;
 
 			[MarshalAs(UnmanagedType.LPStr)]
 			public string? User;
 
+			[MarshalAs(UnmanagedType.LPStr)]
+			public string? Password;
+
 			[MarshalAs(UnmanagedType.LPStr)]
 			public string? Client;
 
@@ -374,8 +377,9 @@ namespace EpicGames.Perforce
 			if (Settings != null)
 			{
 				NativeSettings = new NativeSettings();
-				NativeSettings.Port = Settings.Port;
+				NativeSettings.ServerAndPort = Settings.ServerAndPort;
 				NativeSettings.User = Settings.User;
+				NativeSettings.Password = Settings.Password;
 				NativeSettings.Client = Settings.Client;
 				NativeSettings.AppName = Settings.AppName;
 				NativeSettings.AppVersion = Settings.AppVersion;
diff --git a/Engine/Source/Programs/Shared/EpicGames.Perforce/PerforceSettings.cs b/Engine/Source/Programs/Shared/EpicGames.Perforce/PerforceSettings.cs
--- a/Engine/Source/Programs/Shared/EpicGames.Perforce/PerforceSettings.cs
+++ b/Engine/Source/Programs/Shared/EpicGames.Perforce/PerforceSettings.cs
@@ -14,13 +14,18 @@ namespace EpicGames.Perforce
 		/// <summary>
 		/// Server and port to connect to
 		/// </summary>
-		public string? Port { get; set; }
+		public string? ServerAndPort { get; set; }
 
 		/// <summary>
 		/// Username to log in with
 		/// </summary>
 		public string? User { get; set; }
 
+		/// <summary>
+		/// Password to use
+		/// </summary>
+		public string? Password { get; set; }
+
 		/// <summary>
 		/// Name of the client to use
 		/// </summary>
