BuildGraph: Allow passing a null schema to BgScriptReader.TryRead.

diff --git a/Engine/Source/Programs/Shared/EpicGames.BuildGraph/BgScriptReader.cs b/Engine/Source/Programs/Shared/EpicGames.BuildGraph/BgScriptReader.cs
--- a/Engine/Source/Programs/Shared/EpicGames.BuildGraph/BgScriptReader.cs
+++ b/Engine/Source/Programs/Shared/EpicGames.BuildGraph/BgScriptReader.cs
@@ -78,9 +78,12 @@ namespace EpicGames.BuildGraph
 			BgScriptDocument Document = new BgScriptDocument(File, NativeFile, Logger);
 
 			XmlReaderSettings Settings = new XmlReaderSettings();
-			Settings.Schemas.Add(Schema.CompiledSchema);
-			Settings.ValidationType = ValidationType.Schema;
-			Settings.ValidationEventHandler += Document.ValidationEvent;
+			if (Schema != null)
+			{
+				Settings.Schemas.Add(Schema.CompiledSchema);
+				Settings.ValidationType = ValidationType.Schema;
+				Settings.ValidationEventHandler += Document.ValidationEvent;
+			}
 
 			using (MemoryStream Stream = new MemoryStream(Data))
 			using (XmlReader Reader = XmlReader.Create(Stream, Settings))
