Horde: Convert EpicGames.BuildGraph project to enable #nullable.

diff --git a/Engine/Source/Programs/Shared/EpicGames.BuildGraph/BgAgent.cs b/Engine/Source/Programs/Shared/EpicGames.BuildGraph/BgAgent.cs
--- a/Engine/Source/Programs/Shared/EpicGames.BuildGraph/BgAgent.cs
+++ b/Engine/Source/Programs/Shared/EpicGames.BuildGraph/BgAgent.cs
@@ -19,18 +19,18 @@ namespace EpicGames.BuildGraph
 		/// <summary>
 		/// Name of this agent. Used for display purposes in a build system.
 		/// </summary>
-		public string Name;
+		public string Name { get; }
 
 		/// <summary>
 		/// Array of valid agent types that these nodes may run on. When running in the build system, this determines the class of machine that should
 		/// be selected to run these nodes. The first defined agent type for this branch will be used.
 		/// </summary>
-		public string[] PossibleTypes;
+		public string[] PossibleTypes { get; }
 
 		/// <summary>
 		/// List of nodes in this agent group.
 		/// </summary>
-		public List<BgNode> Nodes = new List<BgNode>();
+		public List<BgNode> Nodes { get; set; } = new List<BgNode>();
 
 		/// <summary>
 		/// Constructor
diff --git a/Engine/Source/Programs/Shared/EpicGames.BuildGraph/BgAggregate.cs b/Engine/Source/Programs/Shared/EpicGames.BuildGraph/BgAggregate.cs
--- a/Engine/Source/Programs/Shared/EpicGames.BuildGraph/BgAggregate.cs
+++ b/Engine/Source/Programs/Shared/EpicGames.BuildGraph/BgAggregate.cs
@@ -16,12 +16,12 @@ namespace EpicGames.BuildGraph
 		/// <summary>
 		/// Name of this badge
 		/// </summary>
-		public readonly string Name;
+		public string Name { get; }
 
 		/// <summary>
 		/// Set of nodes that must be run for this label to be shown.
 		/// </summary>
-		public HashSet<BgNode> RequiredNodes = new HashSet<BgNode>();
+		public HashSet<BgNode> RequiredNodes { get; } = new HashSet<BgNode>();
 
 		/// <summary>
 		/// Constructor
diff --git a/Engine/Source/Programs/Shared/EpicGames.BuildGraph/BgBadge.cs b/Engine/Source/Programs/Shared/EpicGames.BuildGraph/BgBadge.cs
--- a/Engine/Source/Programs/Shared/EpicGames.BuildGraph/BgBadge.cs
+++ b/Engine/Source/Programs/Shared/EpicGames.BuildGraph/BgBadge.cs
@@ -17,22 +17,22 @@ namespace EpicGames.BuildGraph
 		/// <summary>
 		/// Name of this badge
 		/// </summary>
-		public readonly string Name;
+		public string Name { get; }
 
 		/// <summary>
 		/// Depot path to the project that this badge applies to. Used for filtering in UGS.
 		/// </summary>
-		public readonly string Project;
+		public string Project { get; }
 
 		/// <summary>
 		/// The changelist to post the badge for
 		/// </summary>
-		public readonly int Change;
+		public int Change { get; }
 
 		/// <summary>
 		/// Set of nodes that this badge reports the status of
 		/// </summary>
-		public HashSet<BgNode> Nodes = new HashSet<BgNode>();
+		public HashSet<BgNode> Nodes { get; } = new HashSet<BgNode>();
 
 		/// <summary>
 		/// Constructor
diff --git a/Engine/Source/Programs/Shared/EpicGames.BuildGraph/BgLabel.cs b/Engine/Source/Programs/Shared/EpicGames.BuildGraph/BgLabel.cs
--- a/Engine/Source/Programs/Shared/EpicGames.BuildGraph/BgLabel.cs
+++ b/Engine/Source/Programs/Shared/EpicGames.BuildGraph/BgLabel.cs
@@ -33,22 +33,22 @@ namespace EpicGames.BuildGraph
 		/// <summary>
 		/// Name of this badge
 		/// </summary>
-		public readonly string DashboardName;
+		public string? DashboardName { get; }
 
 		/// <summary>
 		/// Category for this label
 		/// </summary>
-		public readonly string DashboardCategory;
+		public string? DashboardCategory { get; }
 
 		/// <summary>
 		/// Name of the badge in UGS
 		/// </summary>
-		public readonly string UgsBadge;
+		public string? UgsBadge { get; }
 
 		/// <summary>
 		/// Path to the project folder in UGS
 		/// </summary>
-		public readonly string UgsProject;
+		public string? UgsProject { get; }
 
 		/// <summary>
 		/// Which change to show the badge for
@@ -73,7 +73,7 @@ namespace EpicGames.BuildGraph
 		/// <param name="InUgsBadge">The UGS badge name</param>
 		/// <param name="InUgsProject">Project to display this badge for</param>
 		/// <param name="InChange">The change to show this badge on in UGS</param>
-		public BgLabel(string InDashboardName, string InDashboardCategory, string InUgsBadge, string InUgsProject, BgLabelChange InChange)
+		public BgLabel(string? InDashboardName, string? InDashboardCategory, string? InUgsBadge, string? InUgsProject, BgLabelChange InChange)
 		{
 			DashboardName = InDashboardName;
 			DashboardCategory = InDashboardCategory;
@@ -94,7 +94,7 @@ namespace EpicGames.BuildGraph
 			}
 			else
 			{
-				return UgsBadge;
+				return UgsBadge ?? "Unknown";
 			}
 		}
 	}
diff --git a/Engine/Source/Programs/Shared/EpicGames.BuildGraph/BgNode.cs b/Engine/Source/Programs/Shared/EpicGames.BuildGraph/BgNode.cs
--- a/Engine/Source/Programs/Shared/EpicGames.BuildGraph/BgNode.cs
+++ b/Engine/Source/Programs/Shared/EpicGames.BuildGraph/BgNode.cs
@@ -23,12 +23,12 @@ namespace EpicGames.BuildGraph
 		/// <summary>
 		/// The node which produces the given output
 		/// </summary>
-		public BgNode ProducingNode;
+		public BgNode ProducingNode { get; }
 
 		/// <summary>
 		/// Name of the tag
 		/// </summary>
-		public string TagName;
+		public string TagName { get; }
 
 		/// <summary>
 		/// Constructor
@@ -59,57 +59,57 @@ namespace EpicGames.BuildGraph
 		/// <summary>
 		/// The node's name
 		/// </summary>
-		public string Name;
+		public string Name { get; }
 
 		/// <summary>
 		/// Array of inputs which this node requires to run
 		/// </summary>
-		public BgNodeOutput[] Inputs;
+		public BgNodeOutput[] Inputs { get; }
 
 		/// <summary>
 		/// Array of outputs produced by this node
 		/// </summary>
-		public BgNodeOutput[] Outputs;
+		public BgNodeOutput[] Outputs { get; }
 
 		/// <summary>
 		/// Nodes which this node has input dependencies on
 		/// </summary>
-		public BgNode[] InputDependencies;
+		public BgNode[] InputDependencies { get; set; }
 
 		/// <summary>
 		/// Nodes which this node needs to run after
 		/// </summary>
-		public BgNode[] OrderDependencies;
+		public BgNode[] OrderDependencies { get; set; }
 
 		/// <summary>
 		/// Tokens which must be acquired for this node to run
 		/// </summary>
-		public FileReference[] RequiredTokens;
+		public FileReference[] RequiredTokens { get; }
 
 		/// <summary>
 		/// List of tasks to execute
 		/// </summary>
-		public List<BgTask> TaskInfos = new List<BgTask>();
+		public List<BgTask> Tasks { get; } = new List<BgTask>();
 
 		/// <summary>
 		/// List of email addresses to notify if this node fails.
 		/// </summary>
-		public HashSet<string> NotifyUsers = new HashSet<string>(StringComparer.InvariantCultureIgnoreCase);
+		public HashSet<string> NotifyUsers { get; set; } = new HashSet<string>(StringComparer.InvariantCultureIgnoreCase);
 
 		/// <summary>
 		/// If set, anyone that has submitted to one of the given paths will be notified on failure of this node
 		/// </summary>
-		public HashSet<string> NotifySubmitters = new HashSet<string>(StringComparer.InvariantCultureIgnoreCase);
+		public HashSet<string> NotifySubmitters { get; set; } = new HashSet<string>(StringComparer.InvariantCultureIgnoreCase);
 
 		/// <summary>
 		/// Whether to start this node as soon as its dependencies are satisfied, rather than waiting for all of its agent's dependencies to be met.
 		/// </summary>
-		public bool bRunEarly = false;
+		public bool bRunEarly { get; set; } = false;
 
 		/// <summary>
 		/// Whether to ignore warnings produced by this node
 		/// </summary>
-		public bool bNotifyOnWarnings = true;
+		public bool bNotifyOnWarnings { get; set; } = true;
 
 		/// <summary>
 		/// Constructor
@@ -208,7 +208,7 @@ namespace EpicGames.BuildGraph
 				Writer.WriteAttributeString("RunEarly", bRunEarly.ToString());
 			}
 
-			foreach (BgTask Task in TaskInfos)
+			foreach (BgTask Task in Tasks)
 			{
 				Task.Write(Writer);
 			}
diff --git a/Engine/Source/Programs/Shared/EpicGames.BuildGraph/BgSchema.cs b/Engine/Source/Programs/Shared/EpicGames.BuildGraph/BgSchema.cs
--- a/Engine/Source/Programs/Shared/EpicGames.BuildGraph/BgSchema.cs
+++ b/Engine/Source/Programs/Shared/EpicGames.BuildGraph/BgSchema.cs
@@ -55,7 +55,7 @@ namespace EpicGames.BuildGraph
 		/// <summary>
 		/// The ICollection interface for this type
 		/// </summary>
-		public Type CollectionType { get; }
+		public Type? CollectionType { get; }
 
 		/// <summary>
 		/// Validation type for this field
@@ -269,48 +269,43 @@ namespace EpicGames.BuildGraph
 				TypeToSchemaTypeName.Add(Type, GetQualifiedTypeName(SchemaType));
 			}
 
-			// If a dictionary of tasks was specified, add those to the schema. Otherwise we'll allow any syntactically valid invocation.
+			// Create all the custom user types, and add them to the qualified name lookup
 			List<XmlSchemaType> UserTypes = new List<XmlSchemaType>();
-			Dictionary<string, XmlSchemaComplexType> TaskNameToType = null;
-			if (Tasks != null)
+			foreach (Type Type in Tasks.SelectMany(x => x.Parameters).Select(x => x.ValueType))
 			{
-				// Create all the custom user types, and add them to the qualified name lookup
-				foreach (Type Type in Tasks.SelectMany(x => x.Parameters).Select(x => x.ValueType))
+				if (!TypeToSchemaTypeName.ContainsKey(Type))
 				{
-					if (!TypeToSchemaTypeName.ContainsKey(Type))
+					if (Type.IsClass && Type.GetInterfaces().Any(x => x.GetGenericTypeDefinition() == typeof(ICollection<>)))
 					{
-						if (Type.IsClass && Type.GetInterfaces().Any(x => x.GetGenericTypeDefinition() == typeof(ICollection<>)))
-						{
-							TypeToSchemaTypeName.Add(Type, GetQualifiedTypeName(ScriptSchemaStandardType.BalancedString));
-						}
-						else
-						{
-							string Name = Type.Name + "UserType";
-							XmlSchemaType SchemaType = CreateUserType(Name, Type);
-							UserTypes.Add(SchemaType);
-							TypeToSchemaTypeName.Add(Type, new XmlQualifiedName(Name, NamespaceURI));
-						}
+						TypeToSchemaTypeName.Add(Type, GetQualifiedTypeName(ScriptSchemaStandardType.BalancedString));
+					}
+					else
+					{
+						string Name = Type.Name + "UserType";
+						XmlSchemaType SchemaType = CreateUserType(Name, Type);
+						UserTypes.Add(SchemaType);
+						TypeToSchemaTypeName.Add(Type, new XmlQualifiedName(Name, NamespaceURI));
 					}
 				}
+			}
 
-				// Create all the task types
-				TaskNameToType = new Dictionary<string, XmlSchemaComplexType>();
-				foreach (BgScriptTask Task in Tasks)
+			// Create all the task types
+			Dictionary<string, XmlSchemaComplexType>? TaskNameToType = new Dictionary<string, XmlSchemaComplexType>();
+			foreach (BgScriptTask Task in Tasks)
+			{
+				XmlSchemaComplexType TaskType = new XmlSchemaComplexType();
+				TaskType.Name = Task.Name + "TaskType";
+				foreach (BgScriptTaskParameter Parameter in Task.Parameters)
 				{
-					XmlSchemaComplexType TaskType = new XmlSchemaComplexType();
-					TaskType.Name = Task.Name + "TaskType";
-					foreach (BgScriptTaskParameter Parameter in Task.Parameters)
+					XmlQualifiedName? SchemaTypeName = GetQualifiedTypeName(Parameter.ValidationType);
+					if (SchemaTypeName == null)
 					{
-						XmlQualifiedName SchemaTypeName = GetQualifiedTypeName(Parameter.ValidationType);
-						if (SchemaTypeName == null)
-						{
-							SchemaTypeName = TypeToSchemaTypeName[Parameter.ValueType];
-						}
-						TaskType.Attributes.Add(CreateSchemaAttribute(Parameter.Name, SchemaTypeName, Parameter.bOptional ? XmlSchemaUse.Optional : XmlSchemaUse.Required));
+						SchemaTypeName = TypeToSchemaTypeName[Parameter.ValueType];
 					}
-					TaskType.Attributes.Add(CreateSchemaAttribute("If", ScriptSchemaStandardType.BalancedString, XmlSchemaUse.Optional));
-					TaskNameToType.Add(Task.Name, TaskType);
+					TaskType.Attributes.Add(CreateSchemaAttribute(Parameter.Name, SchemaTypeName, Parameter.bOptional ? XmlSchemaUse.Optional : XmlSchemaUse.Required));
 				}
+				TaskType.Attributes.Add(CreateSchemaAttribute("If", ScriptSchemaStandardType.BalancedString, XmlSchemaUse.Optional));
+				TaskNameToType.Add(Task.Name, TaskType);
 			}
 
 			// Create the schema object
@@ -366,10 +361,13 @@ namespace EpicGames.BuildGraph
 			XmlSchemaSet NewSchemaSet = new XmlSchemaSet();
 			NewSchemaSet.Add(NewSchema);
 			NewSchemaSet.Compile();
-			foreach(XmlSchema NewCompiledSchema in NewSchemaSet.Schemas())
+
+			XmlSchema? Schema = null;
+			foreach(XmlSchema? NewCompiledSchema in NewSchemaSet.Schemas())
 			{
-				CompiledSchema = NewCompiledSchema;
+				Schema = NewCompiledSchema!;
 			}
+			CompiledSchema = Schema!;
 		}
 
 		/// <summary>
@@ -396,9 +394,8 @@ namespace EpicGames.BuildGraph
 		/// Imports a schema from a file
 		/// </summary>
 		/// <param name="File">The XML file to import</param>
-		/// <param name="NameToTask">Mapping of task name to information about how to construct it</param>
 		/// <returns>A <see cref="BgScriptSchema"/> deserialized from the XML file, or null if file doesn't exist</returns>
-		public static BgScriptSchema Import(FileReference File)
+		public static BgScriptSchema? Import(FileReference File)
 		{
 			if (!FileReference.Exists(File))
 			{
@@ -448,7 +445,7 @@ namespace EpicGames.BuildGraph
 		/// Gets the qualified name of the schema type for the given type of validation
 		/// </summary>
 		/// <returns>Qualified name for the corresponding schema type</returns>
-		static XmlQualifiedName GetQualifiedTypeName(TaskParameterValidationType Type)
+		static XmlQualifiedName? GetQualifiedTypeName(TaskParameterValidationType Type)
 		{
 			switch(Type)
 			{
@@ -1051,7 +1048,7 @@ namespace EpicGames.BuildGraph
 		/// <param name="Name">Name of the new type</param>
 		/// <param name="Pattern">Regex pattern to match</param>
 		/// <returns>A simple type which will match the given pattern</returns>
-		static XmlSchemaSimpleType CreateSimpleTypeFromRegex(string Name, string Pattern)
+		static XmlSchemaSimpleType CreateSimpleTypeFromRegex(string? Name, string Pattern)
 		{
 			XmlSchemaPatternFacet PatternFacet = new XmlSchemaPatternFacet();
 			PatternFacet.Value = Pattern;
@@ -1088,7 +1085,7 @@ namespace EpicGames.BuildGraph
 		/// </summary>
 		/// <param name="Name">Name for the new type</param>
 		/// <param name="Type">CLR type information to create a schema type for</param>
-		static XmlSchemaType CreateEnumType(string Name, Type Type)
+		static XmlSchemaType CreateEnumType(string? Name, Type Type)
 		{
 			XmlSchemaSimpleTypeRestriction Restriction = new XmlSchemaSimpleTypeRestriction();
 			Restriction.BaseTypeName = StringTypeName;
diff --git a/Engine/Source/Programs/Shared/EpicGames.BuildGraph/BgScript.cs b/Engine/Source/Programs/Shared/EpicGames.BuildGraph/BgScript.cs
--- a/Engine/Source/Programs/Shared/EpicGames.BuildGraph/BgScript.cs
+++ b/Engine/Source/Programs/Shared/EpicGames.BuildGraph/BgScript.cs
@@ -10,6 +10,7 @@ using System.Linq;
 using System.Diagnostics;
 using EpicGames.Core;
 using Microsoft.Extensions.Logging;
+using System.Diagnostics.CodeAnalysis;
 
 namespace EpicGames.BuildGraph
 {
@@ -53,12 +54,23 @@ namespace EpicGames.BuildGraph
 		/// <summary>
 		/// The node which this diagnostic is declared in. If the node is culled from the graph, the message will not be displayed.
 		/// </summary>
-		public BgNode EnclosingNode;
+		public BgNode? EnclosingNode;
 
 		/// <summary>
 		/// The agent that this diagnostic is declared in. If the entire agent is culled from the graph, the message will not be displayed.
 		/// </summary>
-		public BgAgent EnclosingAgent;
+		public BgAgent? EnclosingAgent;
+
+		/// <summary>
+		/// Constructor
+		/// </summary>
+		public BgScriptDiagnostic(LogEventType EventType, string Message, BgNode? EnclosingNode, BgAgent? EnclosingAgent)
+		{
+			this.EventType = EventType;
+			this.Message = Message;
+			this.EnclosingNode = EnclosingNode;
+			this.EnclosingAgent = EnclosingAgent;
+		}
 	}
 
 	/// <summary>
@@ -182,13 +194,13 @@ namespace EpicGames.BuildGraph
 		/// <param name="Name">The name to search for</param>
 		/// <param name="OutNodes">If the name is a match, receives an array of nodes and their output names</param>
 		/// <returns>True if the name was found, false otherwise.</returns>
-		public bool TryResolveReference(string Name, out BgNode[] OutNodes)
+		public bool TryResolveReference(string Name, [NotNullWhen(true)] out BgNode[]? OutNodes)
 		{
 			// Check if it's a tag reference or node reference
 			if(Name.StartsWith("#"))
 			{
 				// Check if it's a regular node or output name
-				BgNodeOutput Output;
+				BgNodeOutput? Output;
 				if(TagNameToNodeOutput.TryGetValue(Name, out Output))
 				{
 					OutNodes = new BgNode[]{ Output.ProducingNode };
@@ -198,7 +210,7 @@ namespace EpicGames.BuildGraph
 			else
 			{
 				// Check if it's a regular node or output name
-				BgNode Node;
+				BgNode? Node;
 				if(NameToNode.TryGetValue(Name, out Node))
 				{
 					OutNodes = new BgNode[]{ Node };
@@ -206,7 +218,7 @@ namespace EpicGames.BuildGraph
 				}
 
 				// Check if it's an aggregate name
-				BgAggregate Aggregate;
+				BgAggregate? Aggregate;
 				if(NameToAggregate.TryGetValue(Name, out Aggregate))
 				{
 					OutNodes = Aggregate.RequiredNodes.ToArray();
@@ -214,7 +226,7 @@ namespace EpicGames.BuildGraph
 				}
 
 				// Check if it's a group name
-				BgAgent Agent;
+				BgAgent? Agent;
 				if(NameToAgent.TryGetValue(Name, out Agent))
 				{
 					OutNodes = Agent.Nodes.ToArray();
@@ -233,13 +245,13 @@ namespace EpicGames.BuildGraph
 		/// <param name="Name">The name to search for</param>
 		/// <param name="OutOutputs">If the name is a match, receives an array of nodes and their output names</param>
 		/// <returns>True if the name was found, false otherwise.</returns>
-		public bool TryResolveInputReference(string Name, out BgNodeOutput[] OutOutputs)
+		public bool TryResolveInputReference(string Name, [NotNullWhen(true)] out BgNodeOutput[]? OutOutputs)
 		{
 			// Check if it's a tag reference or node reference
 			if(Name.StartsWith("#"))
 			{
 				// Check if it's a regular node or output name
-				BgNodeOutput Output;
+				BgNodeOutput? Output;
 				if(TagNameToNodeOutput.TryGetValue(Name, out Output))
 				{
 					OutOutputs = new BgNodeOutput[]{ Output };
@@ -249,7 +261,7 @@ namespace EpicGames.BuildGraph
 			else
 			{
 				// Check if it's a regular node or output name
-				BgNode Node;
+				BgNode? Node;
 				if(NameToNode.TryGetValue(Name, out Node))
 				{
 					OutOutputs = Node.Outputs.Union(Node.Inputs).ToArray();
@@ -257,7 +269,7 @@ namespace EpicGames.BuildGraph
 				}
 
 				// Check if it's an aggregate name
-				BgAggregate Aggregate;
+				BgAggregate? Aggregate;
 				if(NameToAggregate.TryGetValue(Name, out Aggregate))
 				{
 					OutOutputs = Aggregate.RequiredNodes.SelectMany(x => x.Outputs.Union(x.Inputs)).Distinct().ToArray();
diff --git a/Engine/Source/Programs/Shared/EpicGames.BuildGraph/BgScriptReader.cs b/Engine/Source/Programs/Shared/EpicGames.BuildGraph/BgScriptReader.cs
--- a/Engine/Source/Programs/Shared/EpicGames.BuildGraph/BgScriptReader.cs
+++ b/Engine/Source/Programs/Shared/EpicGames.BuildGraph/BgScriptReader.cs
@@ -3,6 +3,7 @@
 using System;
 using System.Collections.Generic;
 using System.ComponentModel;
+using System.Diagnostics.CodeAnalysis;
 using System.IO;
 using System.Linq;
 using System.Reflection;
@@ -33,7 +34,7 @@ namespace EpicGames.BuildGraph
 		/// <summary>
 		/// Interface to the LineInfo on the active XmlReader
 		/// </summary>
-		IXmlLineInfo LineInfo;
+		IXmlLineInfo? LineInfo;
 
 		/// <summary>
 		/// Set to true if the reader encounters an error
@@ -60,7 +61,7 @@ namespace EpicGames.BuildGraph
 		/// </summary>
 		public override XmlElement CreateElement(string Prefix, string LocalName, string NamespaceUri)
 		{
-			return new BgScriptElement(File, NativeFile, LineInfo.LineNumber, Prefix, LocalName, NamespaceUri, this);
+			return new BgScriptElement(File, NativeFile, LineInfo!.LineNumber, Prefix, LocalName, NamespaceUri, this);
 		}
 
 		/// <summary>
@@ -73,7 +74,7 @@ namespace EpicGames.BuildGraph
 		/// <param name="Logger">Logger for output messages</param>
 		/// <param name="OutDocument">If successful, the document that was read</param>
 		/// <returns>True if the document could be read, false otherwise</returns>
-		public static bool TryRead(string File, object NativeFile, byte[] Data, BgScriptSchema Schema, ILogger Logger, out BgScriptDocument OutDocument)
+		public static bool TryRead(string File, object NativeFile, byte[] Data, BgScriptSchema Schema, ILogger Logger, [NotNullWhen(true)] out BgScriptDocument? OutDocument)
 		{
 			BgScriptDocument Document = new BgScriptDocument(File, NativeFile, Logger);
 
@@ -331,7 +332,7 @@ namespace EpicGames.BuildGraph
 		/// <summary>
 		/// The name of the node if only a single node is going to be built, otherwise null.
 		/// </summary>
-		string SingleNodeName;
+		string? SingleNodeName;
 
 		/// <summary>
 		/// Private constructor. Use ScriptReader.TryRead() to read a script file.
@@ -369,7 +370,7 @@ namespace EpicGames.BuildGraph
 		/// <param name="Graph">If successful, the graph constructed from the given script</param>
 		/// <param name="SingleNodeName">If a single node will be processed, the name of that node.</param>
 		/// <returns>True if the graph was read, false if there were errors</returns>
-		public static bool TryRead(IBgScriptReaderContext Context, string File, Dictionary<string, string> Arguments, Dictionary<string, string> DefaultProperties, bool bPreprocessOnly, BgScriptSchema Schema, ILogger Logger, out BgScript Graph, string SingleNodeName = null)
+		public static bool TryRead(IBgScriptReaderContext Context, string File, Dictionary<string, string> Arguments, Dictionary<string, string> DefaultProperties, bool bPreprocessOnly, BgScriptSchema Schema, ILogger Logger, out BgScript? Graph, string? SingleNodeName = null)
 		{
 			// Read the file and build the graph
 			BgScriptReader Reader = new BgScriptReader(Context, DefaultProperties, bPreprocessOnly, Schema, Logger);
@@ -403,7 +404,7 @@ namespace EpicGames.BuildGraph
 		/// <param name="Arguments">Arguments passed in to the graph on the command line</param>
 		/// <param name="Logger">Logger for output messages</param>
 		/// <param name="SingleNodeName">The name of the node if only a single node is going to be built, otherwise null.</param>
-		bool TryRead(string File, Dictionary<string, string> Arguments, ILogger Logger, string SingleNodeName = null)
+		bool TryRead(string File, Dictionary<string, string> Arguments, ILogger Logger, string? SingleNodeName = null)
 		{
 			// Get the data for this file
 			byte[] Data;
@@ -415,7 +416,7 @@ namespace EpicGames.BuildGraph
 			}
 
 			// Read the document and validate it against the schema
-			BgScriptDocument Document;
+			BgScriptDocument? Document;
 			if (!BgScriptDocument.TryRead(File, Context.GetNativePath(File), Data, Schema, Logger, out Document))
 			{
 				NumErrors++;
@@ -576,7 +577,7 @@ namespace EpicGames.BuildGraph
 		/// <param name="Name">Name of the property</param>
 		/// <param name="Value">On success, contains the value of the property. Set to null otherwise.</param>
 		/// <returns>True if the property was found, false otherwise</returns>
-		bool TryGetPropertyValue(string Name, out string Value)
+		bool TryGetPropertyValue(string Name, out string? Value)
 		{
 			int ValueLength = 0;
 			if (Name.Contains(":"))
@@ -589,7 +590,7 @@ namespace EpicGames.BuildGraph
 			// Check each scope for the property
 			for (int ScopeIdx = ScopedProperties.Count - 1; ScopeIdx >= 0; ScopeIdx--)
 			{
-				string ScopeValue;
+				string? ScopeValue;
 				if (ScopedProperties[ScopeIdx].TryGetValue(Name, out ScopeValue))
 				{
 					Value = ScopeValue;
@@ -681,8 +682,8 @@ namespace EpicGames.BuildGraph
 					}
 
 					// Check if the property already exists. If it does, we don't need to register it as an option.
-					string ExistingValue;
-					if(TryGetPropertyValue(Name, out ExistingValue))
+					string? ExistingValue;
+					if(TryGetPropertyValue(Name, out ExistingValue) && ExistingValue != null)
 					{
 						// If there's a restriction on this definition, check it matches
 						string Restrict = ReadAttribute(Element, "Restrict");
@@ -700,7 +701,7 @@ namespace EpicGames.BuildGraph
 						Graph.Options.Add(Option);
 
 						// Get the value of this property
-						string Value;
+						string? Value;
 						if(!Arguments.TryGetValue(Name, out Value))
 						{
 							Value = Option.DefaultValue;
@@ -771,8 +772,9 @@ namespace EpicGames.BuildGraph
 			{
 				// Get the pattern
 				string RegexString = ReadAttribute(Element, "Pattern");
+
 				// Make sure its a valid regex.
-				Regex RegexValue = ParseRegex(Element, RegexString);
+				Regex? RegexValue = ParseRegex(Element, RegexString);
 				if (RegexValue != null)
 				{
 					// read the names in 
@@ -842,7 +844,7 @@ namespace EpicGames.BuildGraph
 				string Name = ReadAttribute(Element, "Name");
 				if (ValidateName(Element, Name))
 				{
-					BgScriptMacro OriginalDefinition;
+					BgScriptMacro? OriginalDefinition;
 					if(MacroNameToDefinition.TryGetValue(Name, out OriginalDefinition))
 					{
 						LogError(Element, "Function '{0}' has already been declared (see {1} line {2})", OriginalDefinition.Element.File, OriginalDefinition.Element.LineNumber);
@@ -893,14 +895,14 @@ namespace EpicGames.BuildGraph
 		/// <param name="Trigger">The controlling trigger for nodes in this agent</param>
 		void ReadAgent(BgScriptElement Element)
 		{
-			string Name;
+			string? Name;
 			if (EvaluateCondition(Element) && TryReadObjectName(Element, out Name))
 			{
 				// Read the valid agent types. This may be omitted if we're continuing an existing agent.
 				string[] Types = ReadListAttribute(Element, "Type");
 
 				// Create the agent object, or continue an existing one
-				BgAgent Agent;
+				BgAgent? Agent;
 				if (Graph.NameToAgent.TryGetValue(Name, out Agent))
 				{
 					if (Types.Length > 0 && Agent.PossibleTypes.Length > 0)
@@ -989,7 +991,7 @@ namespace EpicGames.BuildGraph
 		/// <param name="Element">Xml element to read the definition from</param>
 		void ReadAggregate(BgScriptElement Element)
 		{
-			string Name;
+			string? Name;
 			if (EvaluateCondition(Element) && TryReadObjectName(Element, out Name) && CheckNameIsUnique(Element, Name))
 			{
 				string[] RequiredNames = ReadListAttribute(Element, "Requires");
@@ -1050,7 +1052,7 @@ namespace EpicGames.BuildGraph
 		/// <param name="Element">Xml element to read the definition from</param>
 		void ReadReport(BgScriptElement Element)
 		{
-			string Name;
+			string? Name;
 			if (EvaluateCondition(Element) && TryReadObjectName(Element, out Name) && CheckNameIsUnique(Element, Name))
 			{
 				string[] RequiredNames = ReadListAttribute(Element, "Requires");
@@ -1071,7 +1073,7 @@ namespace EpicGames.BuildGraph
 		/// <param name="Element">Xml element to read the definition from</param>
 		void ReadBadge(BgScriptElement Element)
 		{
-			string Name;
+			string? Name;
 			if (EvaluateCondition(Element) && TryReadObjectName(Element, out Name))
 			{
 				string[] RequiredNames = ReadListAttribute(Element, "Requires");
@@ -1146,7 +1148,7 @@ namespace EpicGames.BuildGraph
 		/// <param name="ParentAgent">Agent for the node to be added to</param>
 		void ReadNode(BgScriptElement Element, BgAgent ParentAgent)
 		{
-			string Name;
+			string? Name;
 			if (EvaluateCondition(Element) && TryReadObjectName(Element, out Name))
 			{
 				string[] RequiresNames = ReadListAttribute(Element, "Requires");
@@ -1180,7 +1182,7 @@ namespace EpicGames.BuildGraph
 				List<string> ValidOutputNames = new List<string>();
 				foreach (string ProducesName in ProducesNames)
 				{
-					BgNodeOutput ExistingOutput;
+					BgNodeOutput? ExistingOutput;
 					if(Graph.TagNameToNodeOutput.TryGetValue(ProducesName, out ExistingOutput))
 					{
 						LogError(Element, "Output tag '{0}' is already generated by node '{1}'", ProducesName, ExistingOutput.ProducingNode.Name);
@@ -1227,7 +1229,7 @@ namespace EpicGames.BuildGraph
 					// Register all the output tags in the global name table.
 					foreach(BgNodeOutput Output in NewNode.Outputs)
 					{
-						BgNodeOutput ExistingOutput;
+						BgNodeOutput? ExistingOutput;
 						if(Graph.TagNameToNodeOutput.TryGetValue(Output.TagName, out ExistingOutput))
 						{
 							LogError(Element, "Node '{0}' already has an output called '{1}'", ExistingOutput.ProducingNode.Name, Output.TagName);
@@ -1378,7 +1380,7 @@ namespace EpicGames.BuildGraph
 			{
 				string Name = ReadAttribute(Element, "Name");
 
-				BgScriptMacro Macro;
+				BgScriptMacro? Macro;
 				if(!MacroNameToDefinition.TryGetValue(Name, out Macro))
 				{
 					LogError(Element, "Macro '{0}' does not exist", Name);
@@ -1387,9 +1389,9 @@ namespace EpicGames.BuildGraph
 				{
 					// Parse the argument list
 					string[] Arguments = new string[Macro.ArgumentNameToIndex.Count];
-					foreach(XmlAttribute Attribute in Element.Attributes)
+					foreach(XmlAttribute? Attribute in Element.Attributes)
 					{
-						if(Attribute.Name != "Name" && Attribute.Name != "If")
+						if(Attribute != null && Attribute.Name != "Name" && Attribute.Name != "If")
 						{
 							int Index;
 							if(Macro.ArgumentNameToIndex.TryGetValue(Attribute.Name, out Index))
@@ -1445,15 +1447,15 @@ namespace EpicGames.BuildGraph
 			if (EvaluateCondition(Element))
 			{
 				BgTask Info = new BgTask(Tuple.Create(Element.File, Element.LineNumber), Element.Name);
-				foreach (XmlAttribute Attribute in Element.Attributes)
+				foreach (XmlAttribute? Attribute in Element.Attributes)
 				{
-					if (String.Compare(Attribute.Name, "If", StringComparison.InvariantCultureIgnoreCase) != 0)
+					if (String.Compare(Attribute!.Name, "If", StringComparison.InvariantCultureIgnoreCase) != 0)
 					{
 						string ExpandedValue = ExpandProperties(Element, Attribute.Value);
 						Info.Arguments.Add(Attribute.Name, ExpandedValue);
 					}
 				}
-				ParentNode.TaskInfos.Add(Info);
+				ParentNode.Tasks.Add(Info);
 			}
 		}
 
@@ -1536,7 +1538,7 @@ namespace EpicGames.BuildGraph
 				{
 					foreach (string ReportName in ReportNames)
 					{
-						BgReport Report;
+						BgReport? Report;
 						if (Graph.NameToReport.TryGetValue(ReportName, out Report))
 						{
 							Report.NotifyUsers.UnionWith(Users);
@@ -1557,17 +1559,13 @@ namespace EpicGames.BuildGraph
 		/// <param name="EventType">The diagnostic event type</param>
 		/// <param name="EnclosingNode">The node that this diagnostic is declared in, or null</param>
 		/// <param name="EnclosingAgent">The agent that this diagnostic is declared in, or null</param>
-		void ReadDiagnostic(BgScriptElement Element, LogEventType EventType, BgNode EnclosingNode, BgAgent EnclosingAgent)
+		void ReadDiagnostic(BgScriptElement Element, LogEventType EventType, BgNode? EnclosingNode, BgAgent? EnclosingAgent)
 		{
 			if (EvaluateCondition(Element))
 			{
 				string Message = ReadAttribute(Element, "Message");
 
-				BgScriptDiagnostic Diagnostic = new BgScriptDiagnostic();
-				Diagnostic.EventType = EventType;
-				Diagnostic.Message = String.Format("{0}({1}): {2}", Element.File, Element.LineNumber, Message);
-				Diagnostic.EnclosingNode = EnclosingNode;
-				Diagnostic.EnclosingAgent = EnclosingAgent;
+				BgScriptDiagnostic Diagnostic = new BgScriptDiagnostic(EventType, $"{Element.File}({Element.LineNumber}): {Message}", EnclosingNode, EnclosingAgent);
 				Graph.Diagnostics.Add(Diagnostic);
 			}
 		}
@@ -1600,7 +1598,7 @@ namespace EpicGames.BuildGraph
 			HashSet<BgNode> Nodes = new HashSet<BgNode>();
 			foreach (string ReferenceName in ReferenceNames)
 			{
-				BgNode[] OtherNodes;
+				BgNode[]? OtherNodes;
 				if (Graph.TryResolveReference(ReferenceName, out OtherNodes))
 				{
 					Nodes.UnionWith(OtherNodes);
@@ -1628,7 +1626,7 @@ namespace EpicGames.BuildGraph
 			HashSet<BgNodeOutput> Inputs = new HashSet<BgNodeOutput>();
 			foreach (string ReferenceName in ReferenceNames)
 			{
-				BgNodeOutput[] ReferenceInputs;
+				BgNodeOutput[]? ReferenceInputs;
 				if (Graph.TryResolveInputReference(ReferenceName, out ReferenceInputs))
 				{
 					Inputs.UnionWith(ReferenceInputs);
@@ -1651,7 +1649,7 @@ namespace EpicGames.BuildGraph
 		/// <param name="Element">Element to read the name for</param>
 		/// <param name="Name">Output variable to receive the name of the object</param>
 		/// <returns>True if the object had a valid name (assigned to the Name variable), false if the name was invalid or missing.</returns>
-		bool TryReadObjectName(BgScriptElement Element, out string Name)
+		bool TryReadObjectName(BgScriptElement Element, [NotNullWhen(true)] out string? Name)
 		{
 			// Check the name attribute is present
 			if (!Element.HasAttribute("Name"))
@@ -1674,38 +1672,6 @@ namespace EpicGames.BuildGraph
 			return true;
 		}
 
-		/// <summary>
-		/// Reads an qualified object name from its defining element. Outputs an error if the name is missing.
-		/// </summary>
-		/// <param name="Element">Element to read the name for</param>
-		/// <param name="QualifiedName">Output variable to receive the name of the object</param>
-		/// <returns>True if the object had a valid name (assigned to the Name variable), false if the name was invalid or missing.</returns>
-		bool TryReadQualifiedObjectName(BgScriptElement Element, out string[] QualifiedName)
-		{
-			// Check the name attribute is present
-			if (!Element.HasAttribute("Name"))
-			{
-				LogError(Element, "Missing 'Name' attribute");
-				QualifiedName = null;
-				return false;
-			}
-
-			// Get the value of it, strip any leading or trailing whitespace, and make sure it's not empty
-			string[] Values = ReadAttribute(Element, "Name").Split('.');
-			foreach (string Value in Values)
-			{
-				if (!ValidateName(Element, Value))
-				{
-					QualifiedName = null;
-					return false;
-				}
-			}
-
-			// Return it
-			QualifiedName = Values;
-			return true;
-		}
-
 		/// <summary>
 		/// Checks that the given name is valid syntax
 		/// </summary>
@@ -1744,7 +1710,7 @@ namespace EpicGames.BuildGraph
 		/// <param name="Element">The element that contains the regex</param>
 		/// <param name="Regex">The pattern to construct</param>
 		/// <returns>The regex if is valid, otherwise null</returns>
-		Regex ParseRegex(BgScriptElement Element, string Regex)
+		Regex? ParseRegex(BgScriptElement Element, string Regex)
 		{
 			if(Regex.Length == 0)
 			{
@@ -1943,7 +1909,7 @@ namespace EpicGames.BuildGraph
 				string Name = Result.Substring(Idx + 2, EndIdx - (Idx + 2));
 
 				// Find the value for it, either from the dictionary or the environment block
-				string Value;
+				string? Value;
 				if (!TryGetPropertyValue(Name, out Value))
 				{
 					LogWarning(Element, "Property '{0}' is not defined", Name);
diff --git a/Engine/Source/Programs/Shared/EpicGames.BuildGraph/EpicGames.BuildGraph.csproj b/Engine/Source/Programs/Shared/EpicGames.BuildGraph/EpicGames.BuildGraph.csproj
--- a/Engine/Source/Programs/Shared/EpicGames.BuildGraph/EpicGames.BuildGraph.csproj
+++ b/Engine/Source/Programs/Shared/EpicGames.BuildGraph/EpicGames.BuildGraph.csproj
@@ -7,7 +7,7 @@
     <GenerateAssemblyInfo>false</GenerateAssemblyInfo>
     <GenerateTargetFrameworkAttribute>false</GenerateTargetFrameworkAttribute>
     <Configurations>Debug;Release;Development</Configurations>
-    <Nullable>disable</Nullable>
+    <Nullable>enable</Nullable>
     <DebugType>pdbonly</DebugType>
   </PropertyGroup>
   <PropertyGroup Condition=" '$(Configuration)|$(Platform)' == 'Debug|AnyCPU' ">
