Build fix.

diff --git a/Engine/Source/Programs/Shared/EpicGames.Core/Log.cs b/Engine/Source/Programs/Shared/EpicGames.Core/Log.cs
--- a/Engine/Source/Programs/Shared/EpicGames.Core/Log.cs
+++ b/Engine/Source/Programs/Shared/EpicGames.Core/Log.cs
@@ -1023,7 +1023,10 @@ namespace EpicGames.Core
 						{
 							JsonBufferWriter.Clear();
 							JsonWriter.Reset();
-							MessageTemplate.Serialize(JsonWriter, LogLevel, EventId, State, Exception, Formatter);
+
+							LogEvent Event = new LogEvent(DateTime.UtcNow, LogLevel, EventId, Formatter(State, Exception), null, null, LogException.FromException(Exception));
+							Event.Write(JsonWriter);
+
 							JsonWriter.Flush();
 							Console.WriteLine(Encoding.UTF8.GetString(JsonBufferWriter.WrittenSpan));
 						}
