BuildGraph: Make Exists() and Read() methods asynchronous.

diff --git a/Engine/Source/Programs/AutomationTool/BuildGraph/BuildGraph.cs b/Engine/Source/Programs/AutomationTool/BuildGraph/BuildGraph.cs
--- a/Engine/Source/Programs/AutomationTool/BuildGraph/BuildGraph.cs
+++ b/Engine/Source/Programs/AutomationTool/BuildGraph/BuildGraph.cs
@@ -17,6 +17,7 @@ using OpenTracing;
 using OpenTracing.Util;
 using UnrealBuildBase;
 using UnrealBuildTool;
+using System.Threading.Tasks;
 
 namespace AutomationTool
 {
@@ -112,35 +113,33 @@ namespace AutomationTool
 		}
 
 		/// <inheritdoc/>
-		public bool Exists(string Path)
+		public Task<bool> ExistsAsync(string Path)
 		{
 			try
 			{
-				return FileReference.Exists(FileReference.Combine(Unreal.RootDirectory, Path));
+				return Task.FromResult(FileReference.Exists(FileReference.Combine(Unreal.RootDirectory, Path)));
 			}
 			catch
 			{
-				return false;
+				return Task.FromResult(false);
 			}
 		}
 
 		/// <inheritdoc/>
-		public bool TryRead(string Path, out byte[] Data)
+		public async Task<byte[]> ReadAsync(string Path)
 		{
 			try
 			{
 				FileReference File = FileReference.Combine(Unreal.RootDirectory, Path);
 				if (FileReference.Exists(File))
 				{
-					Data = FileReference.ReadAllBytes(File);
-					return true;
+					return await FileReference.ReadAllBytesAsync(File);
 				}
 			}
 			catch
 			{
 			}
-			Data = null;
-			return false;
+			return null;
 		}
 	}
 
@@ -405,8 +404,8 @@ namespace AutomationTool
 			ScriptFileName = FullScriptFile.MakeRelativeTo(Unreal.RootDirectory).Replace('\\', '/');
 
 			// Read the script from disk
-			BgScript Graph;
-			if(!BgScriptReader.TryRead(Context, ScriptFileName, Arguments, DefaultProperties, PreprocessedFileName != null, Schema, Logger, out Graph, SingleNodeName))
+			BgScript Graph = BgScriptReader.ReadAsync(Context, ScriptFileName, Arguments, DefaultProperties, PreprocessedFileName != null, Schema, Logger, SingleNodeName).Result;
+			if(Graph == null)
 			{
 				return ExitCode.Error_Unknown;
 			}
@@ -806,7 +805,7 @@ namespace AutomationTool
 			}
 			else if (ValueType == typeof(Boolean))
 			{
-				return BgCondition.Evaluate(ValueText, Context);
+				return BgCondition.EvaluateAsync(ValueText, Context).Result;
 			}
 			else if (ValueType == typeof(FileReference))
 			{
diff --git a/Engine/Source/Programs/Horde/Horde.Build/Utilities/BuildGraphExtensions.cs b/Engine/Source/Programs/Horde/Horde.Build/Utilities/BuildGraphExtensions.cs
--- a/Engine/Source/Programs/Horde/Horde.Build/Utilities/BuildGraphExtensions.cs
+++ b/Engine/Source/Programs/Horde/Horde.Build/Utilities/BuildGraphExtensions.cs
@@ -31,9 +31,9 @@ namespace Horde.Build.Utilities
 				this.ChangeNumber = ChangeNumber;
 			}
 
-			public bool Exists(string Path)
+			public async Task<bool> ExistsAsync(string Path)
 			{
-				List<FStatRecord> Records = PerforceConnection.FStatAsync($"{Stream}/{Path}@{ChangeNumber}").Result;
+				List<FStatRecord> Records = await PerforceConnection.FStatAsync($"{Stream}/{Path}@{ChangeNumber}");
 				return Records.Count > 0;
 			}
 
@@ -42,19 +42,17 @@ namespace Horde.Build.Utilities
 				return $"{Stream}/{Path}";
 			}
 
-			public bool TryRead(string Path, out byte[]? Data)
+			public async Task<byte[]?> ReadAsync(string Path)
 			{
 				string TempFileName = System.IO.Path.GetTempFileName();
 				try
 				{
-					PerforceConnection.PrintAsync(TempFileName, $"{Stream}/{Path}").Wait();
-					Data = File.ReadAllBytes(TempFileName);
-					return true;
+					await PerforceConnection.PrintAsync(TempFileName, $"{Stream}/{Path}");
+					return await File.ReadAllBytesAsync(TempFileName);
 				}
 				catch
 				{
-					Data = null;
-					return false;
+					return null;
 				}
 				finally
 				{
@@ -128,8 +126,13 @@ namespace Horde.Build.Utilities
 				}
 			}
 
-			BgScript Script;
-			if (!BgScriptReader.TryRead(Context, ScriptFile, Options, Properties, false, Schema, EpicGames.Core.Log.Logger, out Script))
+			if(ScriptFile == null)
+			{
+				throw new NotImplementedException();
+			}
+
+			BgScript? Script = await BgScriptReader.ReadAsync(Context, ScriptFile, Options, Properties, false, Schema, EpicGames.Core.Log.Logger);
+			if (Script == null)
 			{
 				throw new NotImplementedException();
 			}
diff --git a/Engine/Source/Programs/Shared/EpicGames.BuildGraph/BgCondition.cs b/Engine/Source/Programs/Shared/EpicGames.BuildGraph/BgCondition.cs
--- a/Engine/Source/Programs/Shared/EpicGames.BuildGraph/BgCondition.cs
+++ b/Engine/Source/Programs/Shared/EpicGames.BuildGraph/BgCondition.cs
@@ -54,7 +54,7 @@ namespace EpicGames.BuildGraph
 	/// The type of each subexpression is always a scalar, which are converted to expression-specific types (eg. booleans, integers) as required.
 	/// Scalar values are case-insensitive strings. The identifier 'true' and the strings "true" and "True" are all identical scalars.
 	/// </summary>
-	public static class BgCondition
+	public class BgCondition
 	{
 		/// <summary>
 		/// Sentinel added to the end of a sequence of tokens.
@@ -62,21 +62,54 @@ namespace EpicGames.BuildGraph
 		const string EndToken = "<EOF>";
 
 		/// <summary>
-		/// Evaluates the given string as a condition. Throws a ConditionException on a type or syntax error.
+		/// Tokens for the condition
+		/// </summary>
+		List<string> Tokens = new List<string>();
+
+		/// <summary>
+		/// The current token index
+		/// </summary>
+		int Idx;
+
+		/// <summary>
+		/// Context for evaluating the expression
+		/// </summary>
+		IBgScriptReaderContext Context;
+
+		/// <summary>
+		/// Constructor
 		/// </summary>
 		/// <param name="Text">The condition text</param>
 		/// <param name="Context">Context for evaluating the expression</param>
-		/// <returns>The result of evaluating the condition</returns>
-		public static bool Evaluate(string Text, IBgScriptReaderContext Context)
+		private BgCondition(string Text, IBgScriptReaderContext Context)
 		{
-			List<string> Tokens = new List<string>();
+			this.Context = Context;
+
 			Tokenize(Text, Tokens);
+		}
 
+		/// <summary>
+		/// Evaluates the given string as a condition. Throws a ConditionException on a type or syntax error.
+		/// </summary>
+		/// <param name="Text"></param>
+		/// <param name="Context"></param>
+		/// <returns>The result of evaluating the condition</returns>
+		public static ValueTask<bool> EvaluateAsync(string Text, IBgScriptReaderContext Context)
+		{
+			return new BgCondition(Text, Context).EvaluateAsync();
+		}
+
+		/// <summary>
+		/// Evaluates the given string as a condition. Throws a ConditionException on a type or syntax error.
+		/// </summary>
+		/// <returns>The result of evaluating the condition</returns>
+		async ValueTask<bool> EvaluateAsync()
+		{
 			bool bResult = true;
 			if(Tokens.Count > 1)
 			{
-				int Idx = 0;
-				string Result = EvaluateOr(Tokens, ref Idx, Context);
+				Idx = 0;
+				string Result = await EvaluateOrAsync();
 				if(Tokens[Idx] != EndToken)
 				{
 					throw new BgConditionException("Garbage after expression: {0}", String.Join("", Tokens.Skip(Idx)));
@@ -89,20 +122,17 @@ namespace EpicGames.BuildGraph
 		/// <summary>
 		/// Evaluates an "or-expression" production.
 		/// </summary>
-		/// <param name="Tokens">List of tokens in the expression</param>
-		/// <param name="Idx">Current position in the token stream. Will be incremented as tokens are consumed.</param>
-		/// <param name="Context">Context for evaluating the expression</param>
 		/// <returns>A scalar representing the result of evaluating the expression.</returns>
-		static string EvaluateOr(List<string> Tokens, ref int Idx, IBgScriptReaderContext Context)
+		async ValueTask<string> EvaluateOrAsync()
 		{
 			// <Condition> Or <Condition> Or...
-			string Result = EvaluateAnd(Tokens, ref Idx, Context);
+			string Result = await EvaluateAndAsync();
 			while(String.Compare(Tokens[Idx], "Or", true) == 0)
 			{
 				// Evaluate this condition. We use a binary OR here, because we want to parse everything rather than short-circuit it.
 				Idx++;
 				string Lhs = Result;
-				string Rhs = EvaluateAnd(Tokens, ref Idx, Context);
+				string Rhs = await EvaluateAndAsync();
 				Result = (CoerceToBool(Lhs) | CoerceToBool(Rhs))? "true" : "false";
 			}
 			return Result;
@@ -111,20 +141,17 @@ namespace EpicGames.BuildGraph
 		/// <summary>
 		/// Evaluates an "and-expression" production.
 		/// </summary>
-		/// <param name="Tokens">List of tokens in the expression</param>
-		/// <param name="Idx">Current position in the token stream. Will be incremented as tokens are consumed.</param>
-		/// <param name="Context">Context for evaluating the expression</param>
 		/// <returns>A scalar representing the result of evaluating the expression.</returns>
-		static string EvaluateAnd(List<string> Tokens, ref int Idx, IBgScriptReaderContext Context)
+		async ValueTask<string> EvaluateAndAsync()
 		{
 			// <Condition> And <Condition> And...
-			string Result = EvaluateComparison(Tokens, ref Idx, Context);
+			string Result = await EvaluateComparisonAsync();
 			while(String.Compare(Tokens[Idx], "And", true) == 0)
 			{
 				// Evaluate this condition. We use a binary AND here, because we want to parse everything rather than short-circuit it.
 				Idx++;
 				string Lhs = Result;
-				string Rhs = EvaluateComparison(Tokens, ref Idx, Context);
+				string Rhs = await EvaluateComparisonAsync();
 				Result = (CoerceToBool(Lhs) & CoerceToBool(Rhs))? "true" : "false";
 			}
 			return Result;
@@ -133,11 +160,8 @@ namespace EpicGames.BuildGraph
 		/// <summary>
 		/// Evaluates a "comparison" production.
 		/// </summary>
-		/// <param name="Tokens">List of tokens in the expression</param>
-		/// <param name="Idx">Current position in the token stream. Will be incremented as tokens are consumed.</param>
-		/// <param name="Context">Context for evaluating the expression</param>
 		/// <returns>The result of evaluating the expression</returns>
-		static string EvaluateComparison(List<string> Tokens, ref int Idx, IBgScriptReaderContext Context)
+		async ValueTask<string> EvaluateComparisonAsync()
 		{
 			// scalar
 			// scalar == scalar
@@ -147,13 +171,13 @@ namespace EpicGames.BuildGraph
 			// scalar > scalar
 			// scalar >= scalar
 
-			string Result = EvaluateScalar(Tokens, ref Idx, Context);
+			string Result = await EvaluateScalarAsync();
 			if(Tokens[Idx] == "==")
 			{
 				// Compare two scalars for equality
 				Idx++;
 				string Lhs = Result;
-				string Rhs = EvaluateScalar(Tokens, ref Idx, Context);
+				string Rhs = await EvaluateScalarAsync();
 				Result = (String.Compare(Lhs, Rhs, true) == 0)? "true" : "false";
 			}
 			else if(Tokens[Idx] == "!=")
@@ -161,7 +185,7 @@ namespace EpicGames.BuildGraph
 				// Compare two scalars for inequality
 				Idx++;
 				string Lhs = Result;
-				string Rhs = EvaluateScalar(Tokens, ref Idx, Context);
+				string Rhs = await EvaluateScalarAsync();
 				Result = (String.Compare(Lhs, Rhs, true) != 0)? "true" : "false";
 			}
 			else if(Tokens[Idx] == "<")
@@ -169,7 +193,7 @@ namespace EpicGames.BuildGraph
 				// Compares whether the first integer is less than the second
 				Idx++;
 				int Lhs = CoerceToInteger(Result);
-				int Rhs = CoerceToInteger(EvaluateScalar(Tokens, ref Idx, Context));
+				int Rhs = CoerceToInteger(await EvaluateScalarAsync());
 				Result = (Lhs < Rhs)? "true" : "false";
 			}
 			else if(Tokens[Idx] == "<=")
@@ -177,7 +201,7 @@ namespace EpicGames.BuildGraph
 				// Compares whether the first integer is less than the second
 				Idx++;
 				int Lhs = CoerceToInteger(Result);
-				int Rhs = CoerceToInteger(EvaluateScalar(Tokens, ref Idx, Context));
+				int Rhs = CoerceToInteger(await EvaluateScalarAsync());
 				Result = (Lhs <= Rhs)? "true" : "false";
 			}
 			else if(Tokens[Idx] == ">")
@@ -185,7 +209,7 @@ namespace EpicGames.BuildGraph
 				// Compares whether the first integer is less than the second
 				Idx++;
 				int Lhs = CoerceToInteger(Result);
-				int Rhs = CoerceToInteger(EvaluateScalar(Tokens, ref Idx, Context));
+				int Rhs = CoerceToInteger(await EvaluateScalarAsync());
 				Result = (Lhs > Rhs)? "true" : "false";
 			}
 			else if(Tokens[Idx] == ">=")
@@ -193,7 +217,7 @@ namespace EpicGames.BuildGraph
 				// Compares whether the first integer is less than the second
 				Idx++;
 				int Lhs = CoerceToInteger(Result);
-				int Rhs = CoerceToInteger(EvaluateScalar(Tokens, ref Idx, Context));
+				int Rhs = CoerceToInteger(await EvaluateScalarAsync());
 				Result = (Lhs >= Rhs)? "true" : "false";
 			}
 			return Result;
@@ -202,10 +226,8 @@ namespace EpicGames.BuildGraph
 		/// <summary>
 		/// Evaluates arguments from a token string. Arguments are all comma-separated tokens until a closing ) is encountered
 		/// </summary>
-		/// <param name="Tokens">List of tokens in the expression</param>
-		/// <param name="Idx">Current position in the token stream. Will be incremented as tokens are consumed.</param>
 		/// <returns>The result of evaluating the expression</returns>
-		static IEnumerable<string> EvaluateArguments(List<string> Tokens, ref int Idx)
+		IEnumerable<string> EvaluateArguments()
 		{
 			List<string> Arguments = new List<string>();
 
@@ -252,18 +274,15 @@ namespace EpicGames.BuildGraph
 		/// <summary>
 		/// Evaluates a "scalar" production.
 		/// </summary>
-		/// <param name="Tokens">List of tokens in the expression</param>
-		/// <param name="Idx">Current position in the token stream. Will be incremented as tokens are consumed.</param>
-		/// <param name="Context">Context for evaluating the expression</param>
 		/// <returns>The result of evaluating the expression</returns>
-		static string EvaluateScalar(List<string> Tokens, ref int Idx, IBgScriptReaderContext Context)
+		async ValueTask<string> EvaluateScalarAsync()
 		{
 			string Result;
 			if(Tokens[Idx] == "(")
 			{
 				// Subexpression
 				Idx++;
-				Result = EvaluateOr(Tokens, ref Idx, Context);
+				Result = await EvaluateOrAsync();
 				if(Tokens[Idx] != ")")
 				{
 					throw new BgConditionException("Expected ')'");
@@ -274,28 +293,28 @@ namespace EpicGames.BuildGraph
 			{
 				// Logical not
 				Idx++;
-				string Rhs = EvaluateScalar(Tokens, ref Idx, Context);
+				string Rhs = await EvaluateScalarAsync();
 				Result = CoerceToBool(Rhs)? "false" : "true";
 			}
 			else if(String.Compare(Tokens[Idx], "Exists", true) == 0 && Tokens[Idx + 1] == "(")
 			{
 				// Check whether file or directory exists. Evaluate the argument as a subexpression.
 				Idx++;
-				string Argument = EvaluateScalar(Tokens, ref Idx, Context);
-				Result = Context.Exists(Argument)? "true" : "false";
+				string Argument = await EvaluateScalarAsync();
+				Result = await Context.ExistsAsync(Argument)? "true" : "false";
 			}
 			else if(String.Compare(Tokens[Idx], "HasTrailingSlash", true) == 0 && Tokens[Idx + 1] == "(")
 			{
 				// Check whether the given string ends with a slash
 				Idx++;
-				string Argument = EvaluateScalar(Tokens, ref Idx, Context);
+				string Argument = await EvaluateScalarAsync();
 				Result = (Argument.Length > 0 && (Argument[Argument.Length - 1] == Path.DirectorySeparatorChar || Argument[Argument.Length - 1] == Path.AltDirectorySeparatorChar))? "true" : "false";
 			}
 			else if (String.Compare(Tokens[Idx], "Contains", true) == 0 && Tokens[Idx + 1] == "(")
 			{
 				// Check a string contains a substring. If a separator is supplied the string is first split
 				Idx++;
-				IEnumerable<string> Arguments = EvaluateArguments(Tokens, ref Idx);
+				IEnumerable<string> Arguments = EvaluateArguments();
 
 				if (Arguments.Count() != 2)
 				{
@@ -308,7 +327,7 @@ namespace EpicGames.BuildGraph
 			{
 				// Check a string contains a substring. If a separator is supplied the string is first split
 				Idx++;
-				IEnumerable<string> Arguments = EvaluateArguments(Tokens, ref Idx);
+				IEnumerable<string> Arguments = EvaluateArguments();
 
 				if (Arguments.Count() != 3)
 				{
@@ -472,24 +491,24 @@ namespace EpicGames.BuildGraph
 		/// <summary>
 		/// Test cases for conditions.
 		/// </summary>
-		public static void TestConditions()
+		public static async Task TestConditions()
 		{
-			TestCondition("1 == 2", false);
-			TestCondition("1 == 1", true);
-			TestCondition("1 != 2", true);
-			TestCondition("1 != 1", false);
-			TestCondition("'hello' == 'hello'", true);
-			TestCondition("'hello' == ('hello')", true);
-			TestCondition("'hello' == 'world'", false);
-			TestCondition("'hello' != ('world')", true);
-			TestCondition("true == ('true')", true);
-			TestCondition("true == ('True')", true);
-			TestCondition("true == ('false')", false);
-			TestCondition("true == !('False')", true);
-			TestCondition("true == 'true' and 'false' == 'False'", true);
-			TestCondition("true == 'true' and 'false' == 'true'", false);
-			TestCondition("true == 'false' or 'false' == 'false'", true);
-			TestCondition("true == 'false' or 'false' == 'true'", true);
+			await TestConditionAsync("1 == 2", false);
+			await TestConditionAsync("1 == 1", true);
+			await TestConditionAsync("1 != 2", true);
+			await TestConditionAsync("1 != 1", false);
+			await TestConditionAsync("'hello' == 'hello'", true);
+			await TestConditionAsync("'hello' == ('hello')", true);
+			await TestConditionAsync("'hello' == 'world'", false);
+			await TestConditionAsync("'hello' != ('world')", true);
+			await TestConditionAsync("true == ('true')", true);
+			await TestConditionAsync("true == ('True')", true);
+			await TestConditionAsync("true == ('false')", false);
+			await TestConditionAsync("true == !('False')", true);
+			await TestConditionAsync("true == 'true' and 'false' == 'False'", true);
+			await TestConditionAsync("true == 'true' and 'false' == 'true'", false);
+			await TestConditionAsync("true == 'false' or 'false' == 'false'", true);
+			await TestConditionAsync("true == 'false' or 'false' == 'true'", true);
 		}
 
 		/// <summary>
@@ -497,9 +516,9 @@ namespace EpicGames.BuildGraph
 		/// </summary>
 		/// <param name="Condition">Condition to evaluate</param>
 		/// <param name="ExpectedResult">The expected result</param>
-		static void TestCondition(string Condition, bool ExpectedResult)
+		static async Task TestConditionAsync(string Condition, bool ExpectedResult)
 		{
-			bool Result = Evaluate(Condition, null!);
+			bool Result = await new BgCondition(Condition, null!).EvaluateAsync();
 			Console.WriteLine("{0}: {1} = {2}", (Result == ExpectedResult)? "PASS" : "FAIL", Condition, Result);
 		}
 	}
diff --git a/Engine/Source/Programs/Shared/EpicGames.BuildGraph/BgScriptReader.cs b/Engine/Source/Programs/Shared/EpicGames.BuildGraph/BgScriptReader.cs
--- a/Engine/Source/Programs/Shared/EpicGames.BuildGraph/BgScriptReader.cs
+++ b/Engine/Source/Programs/Shared/EpicGames.BuildGraph/BgScriptReader.cs
@@ -9,6 +9,7 @@ using System.Linq;
 using System.Reflection;
 using System.Text;
 using System.Text.RegularExpressions;
+using System.Threading.Tasks;
 using System.Xml;
 using System.Xml.Schema;
 using EpicGames.Core;
@@ -235,7 +236,7 @@ namespace EpicGames.BuildGraph
 		/// </summary>
 		/// <param name="File">Path to a file</param>
 		/// <returns>True if the file exists</returns>
-		bool Exists(string Path);
+		Task<bool> ExistsAsync(string Path);
 
 		/// <summary>
 		/// Tries to read a file from the given path
@@ -243,7 +244,7 @@ namespace EpicGames.BuildGraph
 		/// <param name="Path">Path of the file to read</param>
 		/// <param name="Data">On success, contains the read data</param>
 		/// <returns></returns>
-		bool TryRead(string Path, out byte[] Data);
+		Task<byte[]?> ReadAsync(string Path);
 
 		/// <summary>
 		/// Converts a path to its native form, for display to the user
@@ -367,17 +368,15 @@ namespace EpicGames.BuildGraph
 		/// <param name="bPreprocessOnly">Preprocess the file, but do not expand any values that are not portable (eg. paths on the local machine)</param>
 		/// <param name="Schema">Schema for the script</param>
 		/// <param name="Logger">Logger for output messages</param>
-		/// <param name="Graph">If successful, the graph constructed from the given script</param>
 		/// <param name="SingleNodeName">If a single node will be processed, the name of that node.</param>
 		/// <returns>True if the graph was read, false if there were errors</returns>
-		public static bool TryRead(IBgScriptReaderContext Context, string File, Dictionary<string, string> Arguments, Dictionary<string, string> DefaultProperties, bool bPreprocessOnly, BgScriptSchema Schema, ILogger Logger, out BgScript? Graph, string? SingleNodeName = null)
+		public static async Task<BgScript?> ReadAsync(IBgScriptReaderContext Context, string File, Dictionary<string, string> Arguments, Dictionary<string, string> DefaultProperties, bool bPreprocessOnly, BgScriptSchema Schema, ILogger Logger, string? SingleNodeName = null)
 		{
 			// Read the file and build the graph
 			BgScriptReader Reader = new BgScriptReader(Context, DefaultProperties, bPreprocessOnly, Schema, Logger);
-			if (!Reader.TryRead(File, Arguments, Logger, SingleNodeName) || Reader.NumErrors > 0)
+			if (!await Reader.TryReadAsync(File, Arguments, Logger, SingleNodeName) || Reader.NumErrors > 0)
 			{
-				Graph = null;
-				return false;
+				return null;
 			}
 
 			// Make sure all the arguments were valid
@@ -393,8 +392,7 @@ namespace EpicGames.BuildGraph
 			}
 
 			// Return the constructed graph
-			Graph = Reader.Graph;
-			return true;
+			return Reader.Graph;
 		}
 
 		/// <summary>
@@ -404,11 +402,11 @@ namespace EpicGames.BuildGraph
 		/// <param name="Arguments">Arguments passed in to the graph on the command line</param>
 		/// <param name="Logger">Logger for output messages</param>
 		/// <param name="SingleNodeName">The name of the node if only a single node is going to be built, otherwise null.</param>
-		bool TryRead(string File, Dictionary<string, string> Arguments, ILogger Logger, string? SingleNodeName = null)
+		async Task<bool> TryReadAsync(string File, Dictionary<string, string> Arguments, ILogger Logger, string? SingleNodeName = null)
 		{
 			// Get the data for this file
-			byte[] Data;
-			if (!Context.TryRead(File, out Data))
+			byte[]? Data = await Context.ReadAsync(File);
+			if (Data == null)
 			{
 				Logger.LogError("Unable to open file {File}", File);
 				NumErrors++;
@@ -425,7 +423,7 @@ namespace EpicGames.BuildGraph
 
 			// Read the root BuildGraph element
 			this.SingleNodeName = SingleNodeName;
-			ReadGraphBody(Document.DocumentElement, Arguments);
+			await ReadGraphBodyAsync(Document.DocumentElement, Arguments);
 			return true;
 		}
 
@@ -434,68 +432,68 @@ namespace EpicGames.BuildGraph
 		/// </summary>
 		/// <param name="Element">The parent element to read from</param>
 		/// <param name="Arguments">Arguments passed in to the graph on the command line</param>
-		void ReadGraphBody(XmlElement Element, Dictionary<string, string> Arguments)
+		async Task ReadGraphBodyAsync(XmlElement Element, Dictionary<string, string> Arguments)
 		{
 			foreach (BgScriptElement ChildElement in Element.ChildNodes.OfType<BgScriptElement>())
 			{
 				switch (ChildElement.Name)
 				{
 					case "Include":
-						ReadInclude(ChildElement, Arguments);
+						await ReadIncludeAsync(ChildElement, Arguments);
 						break;
 					case "Option":
-						ReadOption(ChildElement, Arguments);
+						await ReadOptionAsync(ChildElement, Arguments);
 						break;
 					case "Property":
-						ReadProperty(ChildElement);
+						await ReadPropertyAsync(ChildElement);
 						break;
 					case "Regex":
-						ReadRegex(ChildElement);
+						await ReadRegexAsync(ChildElement);
 						break;
 					case "EnvVar":
-						ReadEnvVar(ChildElement);
+						await ReadEnvVarAsync(ChildElement);
 						break;
 					case "Macro":
-						ReadMacro(ChildElement);
+						await ReadMacroAsync(ChildElement);
 						break;
 					case "Agent":
-						ReadAgent(ChildElement);
+						await ReadAgentAsync(ChildElement);
 						break;
 					case "Aggregate":
-						ReadAggregate(ChildElement);
+						await ReadAggregateAsync(ChildElement);
 						break;
 					case "Report":
-						ReadReport(ChildElement);
+						await ReadReportAsync(ChildElement);
 						break;
 					case "Badge":
-						ReadBadge(ChildElement);
+						await ReadBadgeAsync(ChildElement);
 						break;
 					case "Label":
-						ReadLabel(ChildElement);
+						await ReadLabelAsync(ChildElement);
 						break;
 					case "Notify":
-						ReadNotifier(ChildElement);
+						await ReadNotifierAsync(ChildElement);
 						break;
 					case "Trace":
-						ReadDiagnostic(ChildElement, LogEventType.Console, null, null);
+						await ReadDiagnosticAsync(ChildElement, LogEventType.Console, null, null);
 						break;
 					case "Warning":
-						ReadDiagnostic(ChildElement, LogEventType.Warning, null, null);
+						await ReadDiagnosticAsync(ChildElement, LogEventType.Warning, null, null);
 						break;
 					case "Error":
-						ReadDiagnostic(ChildElement, LogEventType.Error, null, null);
+						await ReadDiagnosticAsync(ChildElement, LogEventType.Error, null, null);
 						break;
 					case "Do":
-						ReadBlock(ChildElement, x => ReadGraphBody(x, Arguments));
+						await ReadBlockAsync(ChildElement, x => ReadGraphBodyAsync(x, Arguments));
 						break;
 					case "Switch":
-						ReadSwitch(ChildElement, x => ReadGraphBody(x, Arguments));
+						await ReadSwitchAsync(ChildElement, x => ReadGraphBodyAsync(x, Arguments));
 						break;
 					case "ForEach":
-						ReadForEach(ChildElement, x => ReadGraphBody(x, Arguments));
+						await ReadForEachAsync(ChildElement, x => ReadGraphBodyAsync(x, Arguments));
 						break;
 					case "Expand":
-						ReadExpand(ChildElement, x => ReadGraphBody(x, Arguments));
+						await ReadExpandAsync(ChildElement, x => ReadGraphBodyAsync(x, Arguments));
 						break;
 					default:
 						LogError(ChildElement, "Invalid element '{0}'", ChildElement.Name);
@@ -652,13 +650,13 @@ namespace EpicGames.BuildGraph
 		/// </summary>
 		/// <param name="Element">Xml element to read the definition from</param>
 		/// <param name="Arguments">Arguments passed in to the graph on the command line</param>
-		void ReadInclude(BgScriptElement Element, Dictionary<string, string> Arguments)
+		async Task ReadIncludeAsync(BgScriptElement Element, Dictionary<string, string> Arguments)
 		{
-			if (EvaluateCondition(Element))
+			if (await EvaluateConditionAsync(Element))
 			{
 				foreach (string Script in ReadListAttribute(Element, "Script"))
 				{
-					TryRead(CombinePaths(Element.File, Script), Arguments, Logger);
+					await TryReadAsync(CombinePaths(Element.File, Script), Arguments, Logger);
 				}
 			}
 		}
@@ -668,9 +666,9 @@ namespace EpicGames.BuildGraph
 		/// </summary>
 		/// <param name="Element">Xml element to read the definition from</param>
 		/// <param name="Arguments">Arguments passed in to the graph on the command line</param>
-		void ReadOption(BgScriptElement Element, IDictionary<string, string> Arguments)
+		async Task ReadOptionAsync(BgScriptElement Element, IDictionary<string, string> Arguments)
 		{
-			if (EvaluateCondition(Element))
+			if (await EvaluateConditionAsync(Element))
 			{
 				string Name = ReadAttribute(Element, "Name");
 				if (ValidateName(Element, Name))
@@ -731,9 +729,9 @@ namespace EpicGames.BuildGraph
 		/// Reads a property assignment.
 		/// </summary>
 		/// <param name="Element">Xml element to read the definition from</param>
-		void ReadProperty(BgScriptElement Element)
+		async Task ReadPropertyAsync(BgScriptElement Element)
 		{
-			if (EvaluateCondition(Element))
+			if (await EvaluateConditionAsync(Element))
 			{
 				string Name = ReadAttribute(Element, "Name");
 				if (ValidateName(Element, Name))
@@ -766,9 +764,9 @@ namespace EpicGames.BuildGraph
 		/// Reads a Regex assignment.
 		/// </summary>
 		/// <param name="Element">Xml element to read the definition from</param>
-		void ReadRegex(BgScriptElement Element)
+		async Task ReadRegexAsync(BgScriptElement Element)
 		{
-			if (EvaluateCondition(Element))
+			if (await EvaluateConditionAsync(Element))
 			{
 				// Get the pattern
 				string RegexString = ReadAttribute(Element, "Pattern");
@@ -795,7 +793,7 @@ namespace EpicGames.BuildGraph
 						string Input = ReadAttribute(Element, "Input");
 						Match Match = RegexValue.Match(Input);
 
-						bool Optional = BgCondition.Evaluate(ReadAttribute(Element, "Optional"), Context);
+						bool Optional = await BgCondition.EvaluateAsync(ReadAttribute(Element, "Optional"), Context);
 						if (!Match.Success)
 						{
 							if (!Optional)
@@ -820,9 +818,9 @@ namespace EpicGames.BuildGraph
 		/// Reads a property assignment from an environment variable.
 		/// </summary>
 		/// <param name="Element">Xml element to read the definition from</param>
-		void ReadEnvVar(BgScriptElement Element)
+		async Task ReadEnvVarAsync(BgScriptElement Element)
 		{
-			if (EvaluateCondition(Element))
+			if (await EvaluateConditionAsync(Element))
 			{
 				string Name = ReadAttribute(Element, "Name");
 				if (ValidateName(Element, Name))
@@ -837,9 +835,9 @@ namespace EpicGames.BuildGraph
 		/// Reads a macro definition
 		/// </summary>
 		/// <param name="Element">Xml element to read the definition from</param>
-		void ReadMacro(BgScriptElement Element)
+		async Task ReadMacroAsync(BgScriptElement Element)
 		{
-			if(EvaluateCondition(Element))
+			if(await EvaluateConditionAsync(Element))
 			{
 				string Name = ReadAttribute(Element, "Name");
 				if (ValidateName(Element, Name))
@@ -892,11 +890,10 @@ namespace EpicGames.BuildGraph
 		/// Reads the definition for an agent.
 		/// </summary>
 		/// <param name="Element">Xml element to read the definition from</param>
-		/// <param name="Trigger">The controlling trigger for nodes in this agent</param>
-		void ReadAgent(BgScriptElement Element)
+		async Task ReadAgentAsync(BgScriptElement Element)
 		{
 			string? Name;
-			if (EvaluateCondition(Element) && TryReadObjectName(Element, out Name))
+			if (await EvaluateConditionAsync(Element) && TryReadObjectName(Element, out Name))
 			{
 				// Read the valid agent types. This may be omitted if we're continuing an existing agent.
 				string[] Types = ReadListAttribute(Element, "Type");
@@ -925,7 +922,7 @@ namespace EpicGames.BuildGraph
 				}
 
 				// Process all the child elements.
-				ReadAgentBody(Element, Agent);
+				await ReadAgentBodyAsync(Element, Agent);
 			}
 		}
 
@@ -934,7 +931,7 @@ namespace EpicGames.BuildGraph
 		/// </summary>
 		/// <param name="Element">Xml element to read the definition from</param>
 		/// <param name="ParentAgent">The agent to contain the definition</param>
-		void ReadAgentBody(BgScriptElement Element, BgAgent ParentAgent)
+		async Task ReadAgentBodyAsync(BgScriptElement Element, BgAgent ParentAgent)
 		{
 			EnterScope();
 			foreach (BgScriptElement ChildElement in Element.ChildNodes.OfType<BgScriptElement>())
@@ -942,40 +939,40 @@ namespace EpicGames.BuildGraph
 				switch (ChildElement.Name)
 				{
 					case "Property":
-						ReadProperty(ChildElement);
+						await ReadPropertyAsync(ChildElement);
 						break;
 					case "Regex":
-						ReadRegex(ChildElement);
+						await ReadRegexAsync(ChildElement);
 						break;
 					case "Node":
-						ReadNode(ChildElement, ParentAgent);
+						await ReadNodeAsync(ChildElement, ParentAgent);
 						break;
 					case "Aggregate":
-						ReadAggregate(ChildElement);
+						await ReadAggregateAsync(ChildElement);
 						break;
 					case "Trace":
-						ReadDiagnostic(ChildElement, LogEventType.Console, null, ParentAgent);
+						await ReadDiagnosticAsync(ChildElement, LogEventType.Console, null, ParentAgent);
 						break;
 					case "Warning":
-						ReadDiagnostic(ChildElement, LogEventType.Warning, null, ParentAgent);
+						await ReadDiagnosticAsync(ChildElement, LogEventType.Warning, null, ParentAgent);
 						break;
 					case "Error":
-						ReadDiagnostic(ChildElement, LogEventType.Error, null, ParentAgent);
+						await ReadDiagnosticAsync(ChildElement, LogEventType.Error, null, ParentAgent);
 						break;
 					case "Label":
-						ReadLabel(ChildElement);
+						await ReadLabelAsync(ChildElement);
 						break;
 					case "Do":
-						ReadBlock(ChildElement, x => ReadAgentBody(x, ParentAgent));
+						await ReadBlockAsync(ChildElement, x => ReadAgentBodyAsync(x, ParentAgent));
 						break;
 					case "Switch":
-						ReadSwitch(ChildElement, x => ReadAgentBody(x, ParentAgent));
+						await ReadSwitchAsync(ChildElement, x => ReadAgentBodyAsync(x, ParentAgent));
 						break;
 					case "ForEach":
-						ReadForEach(ChildElement, x => ReadAgentBody(x, ParentAgent));
+						await ReadForEachAsync(ChildElement, x => ReadAgentBodyAsync(x, ParentAgent));
 						break;
 					case "Expand":
-						ReadExpand(ChildElement, x => ReadAgentBody(x, ParentAgent));
+						await ReadExpandAsync(ChildElement, x => ReadAgentBodyAsync(x, ParentAgent));
 						break;
 					default:
 						LogError(ChildElement, "Unexpected element type '{0}'", ChildElement.Name);
@@ -989,10 +986,10 @@ namespace EpicGames.BuildGraph
 		/// Reads the definition for an aggregate
 		/// </summary>
 		/// <param name="Element">Xml element to read the definition from</param>
-		void ReadAggregate(BgScriptElement Element)
+		async Task ReadAggregateAsync(BgScriptElement Element)
 		{
 			string? Name;
-			if (EvaluateCondition(Element) && TryReadObjectName(Element, out Name) && CheckNameIsUnique(Element, Name))
+			if (await EvaluateConditionAsync(Element) && TryReadObjectName(Element, out Name) && CheckNameIsUnique(Element, Name))
 			{
 				string[] RequiredNames = ReadListAttribute(Element, "Requires");
 
@@ -1050,10 +1047,10 @@ namespace EpicGames.BuildGraph
 		/// Reads the definition for a report
 		/// </summary>
 		/// <param name="Element">Xml element to read the definition from</param>
-		void ReadReport(BgScriptElement Element)
+		async Task ReadReportAsync(BgScriptElement Element)
 		{
 			string? Name;
-			if (EvaluateCondition(Element) && TryReadObjectName(Element, out Name) && CheckNameIsUnique(Element, Name))
+			if (await EvaluateConditionAsync(Element) && TryReadObjectName(Element, out Name) && CheckNameIsUnique(Element, Name))
 			{
 				string[] RequiredNames = ReadListAttribute(Element, "Requires");
 
@@ -1071,10 +1068,10 @@ namespace EpicGames.BuildGraph
 		/// Reads the definition for a badge
 		/// </summary>
 		/// <param name="Element">Xml element to read the definition from</param>
-		void ReadBadge(BgScriptElement Element)
+		async Task ReadBadgeAsync(BgScriptElement Element)
 		{
 			string? Name;
-			if (EvaluateCondition(Element) && TryReadObjectName(Element, out Name))
+			if (await EvaluateConditionAsync(Element) && TryReadObjectName(Element, out Name))
 			{
 				string[] RequiredNames = ReadListAttribute(Element, "Requires");
 				string[] TargetNames = ReadListAttribute(Element, "Targets");
@@ -1099,9 +1096,9 @@ namespace EpicGames.BuildGraph
 		/// Reads the definition for a label
 		/// </summary>
 		/// <param name="Element">Xml element to read the definition from</param>
-		void ReadLabel(BgScriptElement Element)
+		async Task ReadLabelAsync(BgScriptElement Element)
 		{
-			if (EvaluateCondition(Element))
+			if (await EvaluateConditionAsync(Element))
 			{
 				string Name = ReadAttribute(Element, "Name");
 				if (!String.IsNullOrEmpty(Name))
@@ -1146,10 +1143,10 @@ namespace EpicGames.BuildGraph
 		/// </summary>
 		/// <param name="Element">Xml element to read the definition from</param>
 		/// <param name="ParentAgent">Agent for the node to be added to</param>
-		void ReadNode(BgScriptElement Element, BgAgent ParentAgent)
+		async Task ReadNodeAsync(BgScriptElement Element, BgAgent ParentAgent)
 		{
 			string? Name;
-			if (EvaluateCondition(Element) && TryReadObjectName(Element, out Name))
+			if (await EvaluateConditionAsync(Element) && TryReadObjectName(Element, out Name))
 			{
 				string[] RequiresNames = ReadListAttribute(Element, "Requires");
 				string[] ProducesNames = ReadListAttribute(Element, "Produces");
@@ -1241,7 +1238,7 @@ namespace EpicGames.BuildGraph
 					}
 
 					// Add all the tasks
-					ReadNodeBody(Element, NewNode, ParentAgent);
+					await ReadNodeBodyAsync(Element, NewNode, ParentAgent);
 
 					// Add it to the current agent
 					ParentAgent.Nodes.Add(NewNode);
@@ -1255,7 +1252,7 @@ namespace EpicGames.BuildGraph
 		/// <param name="Element">Xml element to read the definition from</param>
 		/// <param name="NewNode">The new node that has been created</param>
 		/// <param name="ParentAgent">Agent for the node to be added to</param>
-		void ReadNodeBody(XmlElement Element, BgNode NewNode, BgAgent ParentAgent)
+		async Task ReadNodeBodyAsync(XmlElement Element, BgNode NewNode, BgAgent ParentAgent)
 		{
 			EnterScope();
 			foreach (BgScriptElement ChildElement in Element.ChildNodes.OfType<BgScriptElement>())
@@ -1263,34 +1260,34 @@ namespace EpicGames.BuildGraph
 				switch (ChildElement.Name)
 				{
 					case "Property":
-						ReadProperty(ChildElement);
+						await ReadPropertyAsync(ChildElement);
 						break;
 					case "Regex":
-						ReadRegex(ChildElement);
+						await ReadRegexAsync(ChildElement);
 						break;
 					case "Trace":
-						ReadDiagnostic(ChildElement, LogEventType.Console, NewNode, ParentAgent);
+						await ReadDiagnosticAsync(ChildElement, LogEventType.Console, NewNode, ParentAgent);
 						break;
 					case "Warning":
-						ReadDiagnostic(ChildElement, LogEventType.Warning, NewNode, ParentAgent);
+						await ReadDiagnosticAsync(ChildElement, LogEventType.Warning, NewNode, ParentAgent);
 						break;
 					case "Error":
-						ReadDiagnostic(ChildElement, LogEventType.Error, NewNode, ParentAgent);
+						await ReadDiagnosticAsync(ChildElement, LogEventType.Error, NewNode, ParentAgent);
 						break;
 					case "Do":
-						ReadBlock(ChildElement, x => ReadNodeBody(x, NewNode, ParentAgent));
+						await ReadBlockAsync(ChildElement, x => ReadNodeBodyAsync(x, NewNode, ParentAgent));
 						break;
 					case "Switch":
-						ReadSwitch(ChildElement, x => ReadNodeBody(x, NewNode, ParentAgent));
+						await ReadSwitchAsync(ChildElement, x => ReadNodeBodyAsync(x, NewNode, ParentAgent));
 						break;
 					case "ForEach":
-						ReadForEach(ChildElement, x => ReadNodeBody(x, NewNode, ParentAgent));
+						await ReadForEachAsync(ChildElement, x => ReadNodeBodyAsync(x, NewNode, ParentAgent));
 						break;
 					case "Expand":
-						ReadExpand(ChildElement, x => ReadNodeBody(x, NewNode, ParentAgent));
+						await ReadExpandAsync(ChildElement, x => ReadNodeBodyAsync(x, NewNode, ParentAgent));
 						break;
 					default:
-						ReadTask(ChildElement, NewNode);
+						await ReadTaskAsync(ChildElement, NewNode);
 						break;
 				}
 			}
@@ -1301,12 +1298,12 @@ namespace EpicGames.BuildGraph
 		/// Reads a block element
 		/// </summary>
 		/// <param name="Element">Xml element to read the definition from</param>
-		/// <param name="ReadContents">Delegate to read the contents of the element, if the condition evaluates to true</param>
-		void ReadBlock(BgScriptElement Element, Action<BgScriptElement> ReadContents)
+		/// <param name="ReadContentsAsync">Delegate to read the contents of the element, if the condition evaluates to true</param>
+		async Task ReadBlockAsync(BgScriptElement Element, Func<BgScriptElement, Task> ReadContentsAsync)
 		{
-			if (EvaluateCondition(Element))
+			if (await EvaluateConditionAsync(Element))
 			{
-				ReadContents(Element);
+				await ReadContentsAsync(Element);
 			}
 		}
 
@@ -1314,15 +1311,14 @@ namespace EpicGames.BuildGraph
 		/// Reads a "Switch" element 
 		/// </summary>
 		/// <param name="Element">Xml element to read the definition from</param>
-		/// <param name="ReadContents">Delegate to read the contents of the element, if the condition evaluates to true</param>
-		void ReadSwitch(BgScriptElement Element, Action<BgScriptElement> ReadContents)
+		/// <param name="ReadContentsAsync">Delegate to read the contents of the element, if the condition evaluates to true</param>
+		async Task ReadSwitchAsync(BgScriptElement Element, Func<BgScriptElement, Task> ReadContentsAsync)
 		{
 			foreach (BgScriptElement ChildElement in Element.ChildNodes.OfType<BgScriptElement>())
 			{
-				if (ChildElement.Name == "Default" || EvaluateCondition(ChildElement))
+				if (ChildElement.Name == "Default" || await EvaluateConditionAsync(ChildElement))
 				{
-					ReadContents(ChildElement);
-					break;
+					await ReadContentsAsync(ChildElement);
 				}
 			}
 		}
@@ -1331,11 +1327,11 @@ namespace EpicGames.BuildGraph
 		/// Reads a "ForEach" element 
 		/// </summary>
 		/// <param name="Element">Xml element to read the definition from</param>
-		/// <param name="ReadContents">Delegate to read the contents of the element, if the condition evaluates to true</param>
-		void ReadForEach(BgScriptElement Element, Action<BgScriptElement> ReadContents)
+		/// <param name="ReadContentsAsync">Delegate to read the contents of the element, if the condition evaluates to true</param>
+		async Task ReadForEachAsync(BgScriptElement Element, Func<BgScriptElement, Task> ReadContentsAsync)
 		{
 			EnterScope();
-			if(EvaluateCondition(Element))
+			if(await EvaluateConditionAsync(Element))
 			{
 				string Name = ReadAttribute(Element, "Name");
 				string Separator = ReadAttribute(Element, "Separator");
@@ -1361,7 +1357,7 @@ namespace EpicGames.BuildGraph
 						foreach(string Value in Values)
 						{
 							ScopedProperties[ScopedProperties.Count - 1][Name] = Value;
-							ReadContents(Element);
+							await ReadContentsAsync(Element);
 						}
 					}
 				}
@@ -1373,10 +1369,10 @@ namespace EpicGames.BuildGraph
 		/// Reads an "Expand" element 
 		/// </summary>
 		/// <param name="Element">Xml element to read the definition from</param>
-		/// <param name="ReadContents">Delegate to read the contents of the element, if the condition evaluates to true</param>
-		void ReadExpand(BgScriptElement Element, Action<BgScriptElement> ReadContents)
+		/// <param name="ReadContentsAsync">Delegate to read the contents of the element, if the condition evaluates to true</param>
+		async Task ReadExpandAsync(BgScriptElement Element, Func<BgScriptElement, Task> ReadContentsAsync)
 		{
-			if(EvaluateCondition(Element))
+			if(await EvaluateConditionAsync(Element))
 			{
 				string Name = ReadAttribute(Element, "Name");
 
@@ -1424,7 +1420,7 @@ namespace EpicGames.BuildGraph
 						{
 							ScopedProperties[ScopedProperties.Count - 1][Pair.Key] = Arguments[Pair.Value] ?? "";
 						}
-						ReadContents(Macro.Element);
+						await ReadContentsAsync(Macro.Element);
 						LeaveScope();
 					}
 				}
@@ -1436,7 +1432,7 @@ namespace EpicGames.BuildGraph
 		/// </summary>
 		/// <param name="Element">Xml element to read the definition from</param>
 		/// <param name="ParentNode">The node which owns this task</param>
-		void ReadTask(BgScriptElement Element, BgNode ParentNode)
+		async Task ReadTaskAsync(BgScriptElement Element, BgNode ParentNode)
 		{
 			// If we're running a single node and this element's parent isn't the single node to run, ignore the error and return.
 			if (!string.IsNullOrWhiteSpace(SingleNodeName) && ParentNode.Name != SingleNodeName)
@@ -1444,7 +1440,7 @@ namespace EpicGames.BuildGraph
 				return;
 			}
 
-			if (EvaluateCondition(Element))
+			if (await EvaluateConditionAsync(Element))
 			{
 				BgTask Info = new BgTask(Tuple.Create(Element.File, Element.LineNumber), Element.Name);
 				foreach (XmlAttribute? Attribute in Element.Attributes)
@@ -1463,9 +1459,9 @@ namespace EpicGames.BuildGraph
 		/// Reads the definition for an email notifier
 		/// </summary>
 		/// <param name="Element">Xml element to read the definition from</param>
-		void ReadNotifier(BgScriptElement Element)
+		async Task ReadNotifierAsync(BgScriptElement Element)
 		{
-			if (EvaluateCondition(Element))
+			if (await EvaluateConditionAsync(Element))
 			{
 				string[] TargetNames = ReadListAttribute(Element, "Targets");
 				string[] ExceptNames = ReadListAttribute(Element, "Except");
@@ -1559,9 +1555,9 @@ namespace EpicGames.BuildGraph
 		/// <param name="EventType">The diagnostic event type</param>
 		/// <param name="EnclosingNode">The node that this diagnostic is declared in, or null</param>
 		/// <param name="EnclosingAgent">The agent that this diagnostic is declared in, or null</param>
-		void ReadDiagnostic(BgScriptElement Element, LogEventType EventType, BgNode? EnclosingNode, BgAgent? EnclosingAgent)
+		async Task ReadDiagnosticAsync(BgScriptElement Element, LogEventType EventType, BgNode? EnclosingNode, BgAgent? EnclosingAgent)
 		{
-			if (EvaluateCondition(Element))
+			if (await EvaluateConditionAsync(Element))
 			{
 				string Message = ReadAttribute(Element, "Message");
 
@@ -1864,7 +1860,7 @@ namespace EpicGames.BuildGraph
 		/// </summary>
 		/// <param name="Element">The element to check</param>
 		/// <returns>True if the element's condition evaluates to true (or doesn't have a conditional expression), false otherwise</returns>
-		bool EvaluateCondition(BgScriptElement Element)
+		async Task<bool> EvaluateConditionAsync(BgScriptElement Element)
 		{
 			// Check if the element has a conditional attribute
 			const string AttributeName = "If";
@@ -1877,7 +1873,7 @@ namespace EpicGames.BuildGraph
 			try
 			{
 				string Text = ExpandProperties(Element, Element.GetAttribute("If"));
-				return BgCondition.Evaluate(Text, Context);
+				return await BgCondition.EvaluateAsync(Text, Context);
 			}
 			catch (BgConditionException Ex)
 			{
