C# serializer for Unreal Insights .utrace files

diff --git a/Engine/Source/Programs/Shared/EpicGames.Tracing.Tests/EpicGames.Tracing.Tests.csproj b/Engine/Source/Programs/Shared/EpicGames.Tracing.Tests/EpicGames.Tracing.Tests.csproj
new file mode 100644
--- /dev/null
+++ b/Engine/Source/Programs/Shared/EpicGames.Tracing.Tests/EpicGames.Tracing.Tests.csproj
@@ -0,0 +1,40 @@
+<Project Sdk="Microsoft.NET.Sdk">
+
+  <PropertyGroup>
+    <TargetFramework>netcoreapp3.1</TargetFramework>
+    <Configuration Condition=" '$(Configuration)' == '' ">Development</Configuration>
+    <OutputType>Library</OutputType>
+    <Deterministic>true</Deterministic>
+    <RootNamespace>EpicGames.Tracing.Tests</RootNamespace>
+    <GenerateAssemblyCompanyAttribute>false</GenerateAssemblyCompanyAttribute>
+    <GenerateAssemblyInformationalVersionAttribute>false</GenerateAssemblyInformationalVersionAttribute>
+    <GenerateAssemblyProductAttribute>false</GenerateAssemblyProductAttribute>
+    <GenerateAssemblyVersionAttribute>false</GenerateAssemblyVersionAttribute>
+    <Configurations>Debug;Release;Development</Configurations>
+    <IsPackable>false</IsPackable>
+    <Nullable>enable</Nullable>
+  </PropertyGroup>
+  <ItemGroup>
+    <Compile Include="..\..\Shared\MetaData.cs">
+      <Link>Properties\MetaData.cs</Link>
+    </Compile>
+  </ItemGroup>
+
+  <ItemGroup>
+    <PackageReference Include="Microsoft.NET.Test.Sdk" Version="16.5.0" />
+    <PackageReference Include="MSTest.TestAdapter" Version="2.1.0" />
+    <PackageReference Include="MSTest.TestFramework" Version="2.1.0" />
+    <PackageReference Include="coverlet.collector" Version="1.2.0" />
+  </ItemGroup>
+
+  <ItemGroup>
+    <ProjectReference Include="..\EpicGames.Tracing\EpicGames.Tracing.csproj" />
+  </ItemGroup>
+
+  <ItemGroup>
+    <None Update="UnrealInsights\example_trace.decomp.utrace">
+      <CopyToOutputDirectory>PreserveNewest</CopyToOutputDirectory>
+    </None>
+  </ItemGroup>
+
+</Project>
diff --git a/Engine/Source/Programs/Shared/EpicGames.Tracing.Tests/UnrealInsights/EventTest.cs b/Engine/Source/Programs/Shared/EpicGames.Tracing.Tests/UnrealInsights/EventTest.cs
new file mode 100644
--- /dev/null
+++ b/Engine/Source/Programs/Shared/EpicGames.Tracing.Tests/UnrealInsights/EventTest.cs
@@ -0,0 +1,31 @@
+// Copyright Epic Games, Inc. All Rights Reserved.
+
+using System.IO;
+using Microsoft.VisualStudio.TestTools.UnitTesting;
+using EpicGames.Tracing.UnrealInsights;
+using EpicGames.Tracing.UnrealInsights.Events;
+
+namespace EpicGames.Tracing.Tests.UnrealInsights
+{
+	[TestClass]
+	public class EventTest
+	{
+		[TestMethod]
+		public void EnterScopeEventTimestampDeserialize()
+		{
+			{
+				using MemoryStream Ms = new MemoryStream(new byte[] {0x10, 0x05, 0x00, 0x00, 0x00, 0x00, 0x00, 0x00});
+				using BinaryReader Reader = new BinaryReader(Ms);
+				EnterScopeEventTimestamp Event = EnterScopeEventTimestamp.Deserialize(Reader);
+				Assert.AreEqual((ulong)0x05, Event.Timestamp);
+			}
+
+			{
+				using MemoryStream Ms = new MemoryStream(new byte[] {0x10, 0x00, 0x00, 0x00, 0x00, 0x00, 0x00, 0xFF});
+				using BinaryReader Reader = new BinaryReader(Ms);
+				EnterScopeEventTimestamp Event = EnterScopeEventTimestamp.Deserialize(Reader);
+				Assert.AreEqual((ulong)0x00_FF_00_00_00_00_00_00, Event.Timestamp);
+			}
+		}
+	}
+}
\ No newline at end of file
diff --git a/Engine/Source/Programs/Shared/EpicGames.Tracing.Tests/UnrealInsights/EventTypeTest.cs b/Engine/Source/Programs/Shared/EpicGames.Tracing.Tests/UnrealInsights/EventTypeTest.cs
new file mode 100644
--- /dev/null
+++ b/Engine/Source/Programs/Shared/EpicGames.Tracing.Tests/UnrealInsights/EventTypeTest.cs
@@ -0,0 +1,109 @@
+// Copyright Epic Games, Inc. All Rights Reserved.
+
+using System.IO;
+using EpicGames.Tracing.UnrealInsights;
+using EpicGames.Tracing.UnrealInsights.Events;
+using Microsoft.VisualStudio.TestTools.UnitTesting;
+
+namespace EpicGames.Tracing.Tests.UnrealInsights
+{
+	[TestClass]
+	public class EventTypeTest
+	{
+		private string TraceThreadInfoHex = "00 00 4A 00 24 00 04 07 06 0A 00 00 04 00 02 08 04 00 04 00 02 08 08 00 04 00 02 08 0C 00 00 00 88 04 24 54 72 61 63 65 54 68 72 65 61 64 49 6E 66 6F 54 68 72 65 61 64 49 64 53 79 73 74 65 6D 49 64 53 6F 72 74 48 69 6E 74 4E 61 6D 65";
+		private string TraceNewTraceHex = "00 00 55 00 10 00 04 05 06 08 00 00 08 00 03 0A 08 00 08 00 03 0E 10 00 02 00 01 06 12 00 01 00 00 0B 24 54 72 61 63 65 4E 65 77 54 72 61 63 65 53 74 61 72 74 43 79 63 6C 65 43 79 63 6C 65 46 72 65 71 75 65 6E 63 79 45 6E 64 69 61 6E 50 6F 69 6E 74 65 72 53 69 7A 65";
+		private string DiagnosticsSession2Hex = "00 00 9A 00 22 00 08 07 0B 08 00 00 00 00 88 08 00 00 00 00 88 07 00 00 00 00 89 0B 00 00 00 00 89 06 00 00 00 00 89 0C 00 00 04 00 02 0A 04 00 01 00 00 11 05 00 01 00 00 0A 44 69 61 67 6E 6F 73 74 69 63 73 53 65 73 73 69 6F 6E 32 50 6C 61 74 66 6F 72 6D 41 70 70 4E 61 6D 65 43 6F 6D 6D 61 6E 64 4C 69 6E 65 42 72 61 6E 63 68 42 75 69 6C 64 56 65 72 73 69 6F 6E 43 68 61 6E 67 65 6C 69 73 74 43 6F 6E 66 69 67 75 72 61 74 69 6F 6E 54 79 70 65 54 61 72 67 65 74 54 79 70 65";
+		private string LoggingLogMessageHex = "00 00 40 00 1A 00 03 06 07 0A 00 00 08 00 03 08 08 00 08 00 03 05 10 00 00 00 80 0A 4C 6F 67 67 69 6E 67 4C 6F 67 4D 65 73 73 61 67 65 4C 6F 67 50 6F 69 6E 74 43 79 63 6C 65 46 6F 72 6D 61 74 41 72 67 73";
+		
+		[TestMethod]
+		public void Deserialize_Trace_NewTrace()
+		{
+			DeserializeAndAssertEventType(TraceNewTraceHex, TraceNewTraceEvent.EventType, 16);
+		}
+		
+		[TestMethod]
+		public void Serialize_Trace_NewTrace()
+		{
+			SerializeAndAssertEventType(TraceNewTraceHex, TraceNewTraceEvent.EventType, 16);
+		}
+		
+		[TestMethod]
+		public void Deserialize_Trace_ThreadInfo()
+		{
+			DeserializeAndAssertEventType(TraceThreadInfoHex, TraceThreadInfoEvent.EventType, 36);
+		}
+		
+		[TestMethod]
+		public void Serialize_Trace_ThreadInfo()
+		{
+			SerializeAndAssertEventType(TraceThreadInfoHex, TraceThreadInfoEvent.EventType, 36);
+		}
+		
+		[TestMethod]
+		public void Deserialize_Diagnostics_Session2()
+		{
+			DeserializeAndAssertEventType(DiagnosticsSession2Hex, DiagnosticsSession2Event.EventType, 34);
+		}
+		
+		[TestMethod]
+		public void Serialize_Diagnostics_Session2()
+		{
+			SerializeAndAssertEventType(DiagnosticsSession2Hex, DiagnosticsSession2Event.EventType, 34);
+		}
+		
+		[TestMethod]
+		public void Deserialize_Logging_LogMessage()
+		{
+			DeserializeAndAssertEventType(LoggingLogMessageHex, LoggingLogMessage.EventType, 26);
+		}
+		
+		[TestMethod]
+		public void Serialize_Logging_LogMessage()
+		{
+			SerializeAndAssertEventType(LoggingLogMessageHex, LoggingLogMessage.EventType, 26);
+		}
+
+		private static EventType DeserializeAndAssertEventType(string HexData, EventType ExpectedType, ushort ExpectedUid)
+		{
+			byte[] Block = GenericEventTest.HexStringToBytes(HexData);
+			using MemoryStream Ms = new MemoryStream(Block);
+			using BinaryReader Reader = new BinaryReader(Ms);
+
+			TraceImportantEventHeader EventHeader = TraceImportantEventHeader.Deserialize(Reader);
+			Assert.AreEqual(PredefinedEventUid.NewEvent, EventHeader.Uid);
+			
+			(ushort NewEventUid, EventType ActualType)  = EventType.Deserialize(Reader);
+			Reader.EnsureEntireStreamIsConsumed();
+
+			Assert.AreEqual(ExpectedUid, ActualType.NewEventUid);
+			Assert.AreEqual(ExpectedUid, NewEventUid);
+			Assert.AreEqual(ExpectedType.Name, ActualType.Name);
+			Assert.AreEqual(ExpectedType.Size, ActualType.Size);
+			Assert.AreEqual(ExpectedType.Flags, ActualType.Flags);
+			Assert.AreEqual(ExpectedType.Fields.Count, ActualType.Fields.Count);
+			for (int i = 0; i < ExpectedType.Fields.Count; i++)
+			{
+				
+				Assert.AreEqual(ExpectedType.Fields[i].Offset, ActualType.Fields[i].Offset);
+				Assert.AreEqual(ExpectedType.Fields[i].Size, ActualType.Fields[i].Size);
+				Assert.AreEqual(ExpectedType.Fields[i].TypeInfo, ActualType.Fields[i].TypeInfo);
+				Assert.AreEqual(ExpectedType.Fields[i].Name, ActualType.Fields[i].Name);
+				Assert.AreEqual(ExpectedType.Fields[i].NameSize, ActualType.Fields[i].NameSize);
+			}
+
+			return ActualType;
+		}
+
+		private static void SerializeAndAssertEventType(string ExpectedHexData, EventType EventType, ushort Uid)
+		{
+			using MemoryStream Ms = new MemoryStream();
+			using BinaryWriter Writer = new BinaryWriter(Ms);
+			byte[] ExpectedData = GenericEventTest.HexStringToBytes(ExpectedHexData);
+			int ExpectedEventSize = ExpectedData.Length;
+			Assert.AreEqual(ExpectedEventSize, EventType.Size);
+			EventType.Serialize(Uid, Writer);
+			
+			GenericEventTest.AssertHexString(ExpectedHexData, Ms.ToArray());
+		}
+	}
+}
\ No newline at end of file
diff --git a/Engine/Source/Programs/Shared/EpicGames.Tracing.Tests/UnrealInsights/GenericEventTest.cs b/Engine/Source/Programs/Shared/EpicGames.Tracing.Tests/UnrealInsights/GenericEventTest.cs
new file mode 100644
--- /dev/null
+++ b/Engine/Source/Programs/Shared/EpicGames.Tracing.Tests/UnrealInsights/GenericEventTest.cs
@@ -0,0 +1,414 @@
+// Copyright Epic Games, Inc. All Rights Reserved.
+
+using System;
+using System.IO;
+using EpicGames.Tracing.UnrealInsights;
+using EpicGames.Tracing.UnrealInsights.Events;
+using Microsoft.VisualStudio.TestTools.UnitTesting;
+
+namespace EpicGames.Tracing.Tests.UnrealInsights
+{
+	[TestClass]
+	public class GenericEventTest
+	{
+		private string CpuProfilerEventSpecHex = "21 00 77 00 34 25 00 00 01 C1 0D 00 4B 32 4E 6F 64 65 5F 56 61 72 69 61 62 6C 65 47 65 74 20 2F 45 6E 67 69 6E 65 2F 45 6E 67 69 6E 65 53 6B 79 2F 42 50 5F 53 6B 79 5F 53 70 68 65 72 65 2E 42 50 5F 53 6B 79 5F 53 70 68 65 72 65 3A 55 73 65 72 43 6F 6E 73 74 72 75 63 74 69 6F 6E 53 63 72 69 70 74 2E 4B 32 4E 6F 64 65 5F 56 61 72 69 61 62 6C 65 47 65 74 5F 39 38 33 03";
+		private string CpuProfilerEventBatchHex1 = "4F 00 02 20 1E 00 87 A7 EB AE 8A D1 0B 01 A7 02 02 04 FB 23 03 A4 0A 61 04 0B 05 F2 80 42 AA 1C EB 5C 06 C0 BC 03 17 07 02 47 08 C2 02 AB 15 09 AA 03 07 0A AE 01 05 0B D8 0D 05 0C B7 01 0D D8 03 00 03 0E A4 03 6F 0F 89 FF 03 10 F3 C4 04 11 60 06 A7 02 10 AB 19 12 B3 04 10 9F 11 13 0A 04 02 04 ED 03 10 91 11 14 80 01 02 97 02 10 A9 D6 0B 15 10 06 DD 03 10 CB 85 04 16 12 06 D3 01 10 89 16 17 0C 02 D6 02 07 18 DA 84 28 99 01 19 18 05 1A 0A A5 10 1B CD 01 1C 06 88 6A C5 49 1D B8 EB 24 A9 17 20 AF 01 10 CD F9 01 21 71 22 F2 AE 12 85 04 23 D1 0D 24 85 74 25 EA 58 91 0D 25 80 E6 02 A9 1C 25 88 CD 06 A7 27 25 D0 DA 09 8B 20 25 DE AD 0B AF 33 25 C0 A9 09 ED 23 25 C4 81 0A F5 2D 25 9A EC 0A F3 35 25 D4 CF 03 FF 19 25 BE 92 01 B1 0D 25 98 22 06";
+		// Unused at the moment -- private string CpuProfilerEventBatchHex2 = "4F 00 02 40 1E 00 E4 AD A2 B2 8A D1 0B 17 41 12 0B 41 12 0B 41 04 11 41 2C 0F 41 0E 17 41 14 13 41 10 11 41 26 0F 41 12 B1 01 41 28 0D 41 04 0F 41 18 09 41 04 0F 41 16 0B 41 16 19 41 24 0D 41 14 17 41 0E 0B 41 2A 0B 41 48 31 41 06 15 41 06 5B 41 14 0B 41 0C 25 41 14 25 41 14 25 41 10 15 41 12 2D 41 0E 15 41 04 25 41 10 19 41 0C 17 41 16 17 41 1E 17 41 14 0B 41 0E 2D 41 0E 45 41 04 15 41 14 17 41 14 13 41 18 19 41 06 0D 41 10 11 41 10 0B 41 0E 0B 41 0E 0D 41 10 0B 41 10 0F 41 18 0F 41 12 0B 41 18 19 41 0E 1F 41 16 0D 41 14 1F 41 16 33 41 10 0D 41 24 25 41 18 15 41 04 23 41 06 25 41 04 0B 41 14 0F 41 0E 15 41 12 77 41 12 0D 41 14 23 41 18 0F 41 16 3B 41 18 19 41 20 81 01 41 18 4D 41 16 2B 41 16 1F 41 1C 11 41 10 11 41 10 0B 41 0E 15 41 06";
+		private string LoggingLogMessageHex = "35 00 B4 FD CA E1 FB 7F 00 00 A4 AB DB 52 44 17 00 00 02 A2 03 00 02 C2 44 61 00 71 00 50 00 72 00 6F 00 66 00 2E 00 64 00 6C 00 6C 00 00 00 7E 00 00 00 06";
+		private string MemoryMemoryScopeHex1 = "25 00 2D 00 00 27 00 00 00";
+		private string MemoryMemoryScopeHex2 = "25 00 2E 00 00 27 00 00 00";
+		private string MemoryMemoryScopeHex3 = "25 00 3D 59 00 33 00 00 00";
+		private string TraceNewTraceHex = "10 00 13 00 D9 6E DA 52 44 17 00 00 80 96 98 00 00 00 00 00 4D 52 08";
+		private string DiagnosticsSession2Hex = "22 00 EA 00 DC 26 0A 01 03 04 01 A0 00 00 57 69 6E 36 34 01 81 01 00 55 6E 72 65 61 6C 45 64 69 74 6F 72 01 C2 0F 00 20 00 44 00 3A 00 5C 00 64 00 65 00 70 00 6F 00 74 00 5C 00 73 00 74 00 61 00 72 00 73 00 68 00 69 00 70 00 2D 00 6D 00 61 00 69 00 6E 00 5C 00 51 00 41 00 47 00 61 00 6D 00 65 00 5C 00 51 00 41 00 47 00 61 00 6D 00 65 00 2E 00 75 00 70 00 72 00 6F 00 6A 00 65 00 63 00 74 00 20 00 2D 00 74 00 72 00 61 00 63 00 65 00 3D 00 63 00 70 00 75 00 20 00 2D 00 67 00 61 00 6D 00 65 00 01 83 02 00 2B 00 2B 00 55 00 45 00 35 00 2B 00 4D 00 61 00 69 00 6E 00 01 84 05 00 2B 00 2B 00 55 00 45 00 35 00 2B 00 4D 00 61 00 69 00 6E 00 2D 00 43 00 4C 00 2D 00 31 00 37 00 34 00 34 00 32 00 35 00 32 00 34 00 03";
+		private string TraceThreadInfoHex1 = "24 00 1B 00 02 00 00 00 74 9D 02 00 FF FF FF FF 01 43 01 00 47 61 6D 65 54 68 72 65 61 64 03";
+		private string TraceThreadInfoHex2 = "24 00 16 00 03 00 00 00 00 00 00 00 FF FF FF 7F 01 A3 00 00 54 72 61 63 65 03";
+		private string TraceThreadInfoHex3 = "24 00 25 00 04 00 00 00 D8 BB 06 00 50 00 00 00 01 83 02 00 46 6F 72 65 67 72 6F 75 6E 64 20 57 6F 72 6B 65 72 20 23 30 03";
+		
+		[TestMethod]
+		public void Deserialize_Memory_MemoryScope()
+		{
+			GenericEvent Event1 = DeserializeAndAssertGenericEvent(MemoryMemoryScopeHex1, MemoryMemoryScopeEvent.EventType, false, 18, 45);
+			Assert.AreEqual(0x27, Event1.Fields[0].Int!.Value);
+			
+			GenericEvent Event2 = DeserializeAndAssertGenericEvent(MemoryMemoryScopeHex2, MemoryMemoryScopeEvent.EventType, false, 18, 46);
+			Assert.AreEqual(0x27, Event2.Fields[0].Int!.Value);
+			
+			GenericEvent Event3 = DeserializeAndAssertGenericEvent(MemoryMemoryScopeHex3, MemoryMemoryScopeEvent.EventType, false, 18, 5832765);
+			Assert.AreEqual(0x33, Event3.Fields[0].Int!.Value);
+		}
+		
+		[TestMethod]
+		public void Serialize_Memory_MemoryScope()
+		{
+			SerializeAndAssertGenericEvent(MemoryMemoryScopeHex1, MemoryMemoryScopeEvent.EventType, false, 18, new []
+			{
+				GenericEvent.Field.FromInt(0x27),
+			}, 45);
+			
+			SerializeAndAssertGenericEvent(MemoryMemoryScopeHex2, MemoryMemoryScopeEvent.EventType, false, 18, new []
+			{
+				GenericEvent.Field.FromInt(0x27),
+			}, 46);
+			
+			SerializeAndAssertGenericEvent(MemoryMemoryScopeHex3, MemoryMemoryScopeEvent.EventType, false, 18, new []
+			{
+				GenericEvent.Field.FromInt(0x33),
+			}, 5832765);
+		}
+		
+		[TestMethod]
+		public void Deserialize_Trace_NewTrace()
+		{
+			GenericEvent Event = DeserializeAndAssertGenericEvent(TraceNewTraceHex, TraceNewTraceEvent.EventType, true, 16);
+			Assert.AreEqual(25582215261913, Event.Fields[0].Long!.Value);
+			Assert.AreEqual(10000000, Event.Fields[1].Long!.Value);
+			Assert.AreEqual(21069, Event.Fields[2].Int!.Value);
+			Assert.AreEqual(8, Event.Fields[3].Int!.Value);
+		}
+		
+		[TestMethod]
+		public void Serialize_Trace_NewTrace_Generic()
+		{
+			SerializeAndAssertGenericEvent(TraceNewTraceHex, TraceNewTraceEvent.EventType, true, 16, new []
+			{
+				GenericEvent.Field.FromLong(25582215261913),
+				GenericEvent.Field.FromLong(10000000),
+				GenericEvent.Field.FromInt(21069),
+				GenericEvent.Field.FromInt(8),
+			});
+		}
+		
+		[TestMethod]
+		public void Serialize_Trace_NewTrace()
+		{
+			TraceNewTraceEvent Event = new TraceNewTraceEvent(25582215261913, 10000000, 21069, 8);
+			SerializeAndAssertEvent(TraceNewTraceHex, Event, true, 16);
+		}
+		
+		[TestMethod]
+		public void Deserialize_Trace_ThreadInfo()
+		{
+			{
+				GenericEvent Event = DeserializeAndAssertGenericEvent(TraceThreadInfoHex1, TraceThreadInfoEvent.EventType, true, 36);
+				Assert.AreEqual(2, Event.Fields[0].Int!.Value);
+				Assert.AreEqual(171380, Event.Fields[1].Int!.Value);
+				Assert.AreEqual(-1, Event.Fields[2].Int!.Value);
+				Assert.AreEqual("GameThread", Event.Fields[3].String!);
+			}
+			
+			{
+				GenericEvent Event = DeserializeAndAssertGenericEvent(TraceThreadInfoHex2, TraceThreadInfoEvent.EventType, true, 36);
+				Assert.AreEqual(3, Event.Fields[0].Int!.Value);
+				Assert.AreEqual(0, Event.Fields[1].Int!.Value);
+				Assert.AreEqual(2147483647, Event.Fields[2].Int!.Value);
+				Assert.AreEqual("Trace", Event.Fields[3].String!);
+			}
+			
+			{
+				GenericEvent Event = DeserializeAndAssertGenericEvent(TraceThreadInfoHex3, TraceThreadInfoEvent.EventType, true, 36);
+				Assert.AreEqual(4, Event.Fields[0].Int!.Value);
+				Assert.AreEqual(441304, Event.Fields[1].Int!.Value);
+				Assert.AreEqual(80, Event.Fields[2].Int!.Value);
+				Assert.AreEqual("Foreground Worker #0", Event.Fields[3].String!);
+			}
+		}
+		
+		[TestMethod]
+		public void Serialize_Trace_ThreadInfo()
+		{
+			{
+				TraceThreadInfoEvent Event = new TraceThreadInfoEvent(2, 171380, -1, "GameThread");
+				SerializeAndAssertEvent(TraceThreadInfoHex1, Event, true, 36);
+			}
+			
+			{
+				TraceThreadInfoEvent Event = new TraceThreadInfoEvent(3, 0, 2147483647, "Trace");
+				SerializeAndAssertEvent(TraceThreadInfoHex2, Event, true, 36);
+			}
+			
+			{
+				TraceThreadInfoEvent Event = new TraceThreadInfoEvent(4, 441304, 80, "Foreground Worker #0");
+				SerializeAndAssertEvent(TraceThreadInfoHex3, Event, true, 36);
+			}
+		}
+		
+		[TestMethod]
+		public void Deserialize_Diagnostics_Session2()
+		{
+			GenericEvent Event = DeserializeAndAssertGenericEvent(DiagnosticsSession2Hex, DiagnosticsSession2Event.EventType, true, 34);
+			Assert.AreEqual("Win64", Event.Fields[0].String!);
+			Assert.AreEqual("UnrealEditor", Event.Fields[1].String!);
+			Assert.AreEqual(@" D:\depot\starship-main\QAGame\QAGame.uproject -trace=cpu -game", Event.Fields[2].String!);
+			Assert.AreEqual("++UE5+Main", Event.Fields[3].String!);
+			Assert.AreEqual("++UE5+Main-CL-17442524", Event.Fields[4].String!);
+			Assert.AreEqual(17442524, Event.Fields[5].Int!.Value);
+			Assert.AreEqual(3, Event.Fields[6].Int!.Value);
+			Assert.AreEqual(4, Event.Fields[7].Int!.Value);
+		}
+		
+		[TestMethod]
+		public void Serialize_Diagnostics_Session2_Generic()
+		{
+			SerializeAndAssertGenericEvent(DiagnosticsSession2Hex, DiagnosticsSession2Event.EventType, true, 34, new []
+			{
+				GenericEvent.Field.FromString("Win64"),
+				GenericEvent.Field.FromString("UnrealEditor"),
+				GenericEvent.Field.FromString(@" D:\depot\starship-main\QAGame\QAGame.uproject -trace=cpu -game"),
+				GenericEvent.Field.FromString("++UE5+Main"),
+				GenericEvent.Field.FromString("++UE5+Main-CL-17442524"),
+				GenericEvent.Field.FromInt(17442524),
+				GenericEvent.Field.FromInt(3),
+				GenericEvent.Field.FromInt(4)
+			});
+		}
+		
+		[TestMethod]
+		public void Serialize_Diagnostics_Session2()
+		{
+			DiagnosticsSession2Event Event = new DiagnosticsSession2Event("Win64", "UnrealEditor", @" D:\depot\starship-main\QAGame\QAGame.uproject -trace=cpu -game", "++UE5+Main", "++UE5+Main-CL-17442524", 17442524, 3, 4);
+			SerializeAndAssertEvent(DiagnosticsSession2Hex, Event, true, 34);
+		}
+		
+		[TestMethod]
+		public void CpuProfiler_EventSpec_Deserialize()
+		{
+			GenericEvent Event = DeserializeAndAssertGenericEvent(CpuProfilerEventSpecHex, CpuProfilerEventSpecEvent.EventType, true, 33);
+			Assert.AreEqual(9524, Event.Fields[0].Int!.Value);
+			Assert.AreEqual("K2Node_VariableGet /Engine/EngineSky/BP_Sky_Sphere.BP_Sky_Sphere:UserConstructionScript.K2Node_VariableGet_983", Event.Fields[1].String!);
+		}
+		
+		[TestMethod]
+		public void CpuProfiler_EventSpec_Serialize_Generic()
+		{
+			SerializeAndAssertGenericEvent(CpuProfilerEventSpecHex, CpuProfilerEventSpecEvent.EventType, true, 33, new GenericEvent.Field[]
+			{
+				GenericEvent.Field.FromInt(9524),
+				GenericEvent.Field.FromString("K2Node_VariableGet /Engine/EngineSky/BP_Sky_Sphere.BP_Sky_Sphere:UserConstructionScript.K2Node_VariableGet_983")
+			});
+		}
+		
+		[TestMethod]
+		public void CpuProfiler_EventSpec_Serialize()
+		{
+			CpuProfilerEventSpecEvent Event = new CpuProfilerEventSpecEvent(9524, "K2Node_VariableGet /Engine/EngineSky/BP_Sky_Sphere.BP_Sky_Sphere:UserConstructionScript.K2Node_VariableGet_983");
+			SerializeAndAssertEvent(CpuProfilerEventSpecHex, Event, true, 33);
+		}
+		
+		[TestMethod]
+		public void CpuProfiler_EventBatch_Deserialize()
+		{
+			GenericEvent Event = DeserializeAndAssertGenericEvent(CpuProfilerEventBatchHex1, CpuProfilerEventBatchEvent.EventType, false, 39);
+			byte[] EventData = Event.Fields[0].Array!;
+			Assert.AreEqual(241, EventData.Length);
+			
+			CpuProfilerSerializer Serializer = new CpuProfilerSerializer(10000000);
+			Serializer.Read(EventData);
+
+			byte[] EventDataSerialized = Serializer.Write()[0];
+			
+			AssertHexString(EventData, EventDataSerialized);
+		}
+		
+		[TestMethod]
+		public void CpuProfiler_EventBatch_Serialize()
+		{
+			byte[] Data = HexStringToBytes("87 A7 EB AE 8A D1 0B 01 A7 02 02 04 FB 23 03 A4 0A 61 04 0B 05 F2 80 42 AA 1C EB 5C 06 C0 BC 03 17 07 02 47 08 C2 02 AB 15 09 AA 03 07 0A AE 01 05 0B D8 0D 05 0C B7 01 0D D8 03 00 03 0E A4 03 6F 0F 89 FF 03 10 F3 C4 04 11 60 06 A7 02 10 AB 19 12 B3 04 10 9F 11 13 0A 04 02 04 ED 03 10 91 11 14 80 01 02 97 02 10 A9 D6 0B 15 10 06 DD 03 10 CB 85 04 16 12 06 D3 01 10 89 16 17 0C 02 D6 02 07 18 DA 84 28 99 01 19 18 05 1A 0A A5 10 1B CD 01 1C 06 88 6A C5 49 1D B8 EB 24 A9 17 20 AF 01 10 CD F9 01 21 71 22 F2 AE 12 85 04 23 D1 0D 24 85 74 25 EA 58 91 0D 25 80 E6 02 A9 1C 25 88 CD 06 A7 27 25 D0 DA 09 8B 20 25 DE AD 0B AF 33 25 C0 A9 09 ED 23 25 C4 81 0A F5 2D 25 9A EC 0A F3 35 25 D4 CF 03 FF 19 25 BE 92 01 B1 0D 25 98 22");
+			CpuProfilerEventBatchEvent Event = new CpuProfilerEventBatchEvent(Data);
+			SerializeAndAssertEvent(CpuProfilerEventBatchHex1, Event, false, 39);
+		}
+		
+		[TestMethod]
+		public void CpuProfiler_EventBatch_EncodeDecodeTimestamp()
+		{
+			void AssertTimestamp(ulong ExpectedTimestamp, bool ExpectedIsScopeEnter)
+			{
+				ulong Temp = CpuProfilerSerializer.EncodeTimestamp(ExpectedTimestamp, ExpectedIsScopeEnter);
+				(ulong ActualTimestamp, bool ActualIsScopeEnter) = CpuProfilerSerializer.DecodeTimestamp(Temp);
+				Assert.AreEqual(ExpectedTimestamp, ActualTimestamp);
+				Assert.AreEqual(ExpectedIsScopeEnter, ActualIsScopeEnter);
+			}
+
+			AssertTimestamp(0, true);
+			AssertTimestamp(0, false);
+			AssertTimestamp(1, true);
+			AssertTimestamp(1, false);
+			AssertTimestamp(14, true);
+			AssertTimestamp(14, false);
+			AssertTimestamp(255, true);
+			AssertTimestamp(255, false);
+			AssertTimestamp(256, true);
+			AssertTimestamp(256, false);
+			AssertTimestamp(4828284612, true);
+			AssertTimestamp(4828284612, false);
+		}
+		
+		[TestMethod]
+		public void Deserialize_Logging_LogMessage()
+		{
+			GenericEvent Event = DeserializeAndAssertGenericEvent(LoggingLogMessageHex, LoggingLogMessage.EventType, false, 26);
+			Assert.AreEqual(140719801695668, Event.Fields[0].Long!.Value);
+			Assert.AreEqual(25582215343012, Event.Fields[1].Long!.Value);
+			Assert.AreEqual(29, Event.Fields[2].Array!.Length);
+			Assert.AreEqual(0x02, Event.Fields[2].Array![0]);
+			Assert.AreEqual(0xC2, Event.Fields[2].Array![1]);
+			Assert.AreEqual(0x44, Event.Fields[2].Array![2]);
+			Assert.AreEqual(0x7E, Event.Fields[2].Array![25]);
+		}
+		
+		[TestMethod]
+		public void SerializeEventWithAuxAndImportant()
+		{
+			ushort Uid = 33;
+			using MemoryStream Ms = new MemoryStream();
+			using BinaryWriter Writer = new BinaryWriter(Ms);
+
+			GenericEvent.Field[] Fields =
+			{
+				GenericEvent.Field.FromInt(9524),
+				GenericEvent.Field.FromString("K2Node_VariableGet /Engine/EngineSky/BP_Sky_Sphere.BP_Sky_Sphere:UserConstructionScript.K2Node_VariableGet_983"),
+			};
+			GenericEvent Event = new GenericEvent(98989, Fields, CpuProfilerEventSpecEvent.EventType);
+			Assert.AreEqual(0x77, Event.Size);
+			
+			TraceImportantEventHeader EventHeader = new TraceImportantEventHeader(Uid, Event.Size);
+			EventHeader.Serialize(Writer);
+			Event.Serialize(20, Writer);
+
+			AssertHexString(CpuProfilerEventSpecHex, Ms.ToArray());
+		}
+
+		[TestMethod]
+		public void AuxHeaderDeserialize()
+		{
+			// Aux header from field "Name" in CpuProfiler.EventSpec
+			using MemoryStream Ms = new MemoryStream(new byte[] {0x01, 0xC1, 0x0D, 0x00});
+			using BinaryReader Reader = new BinaryReader(Ms);
+
+			var (Uid, FieldIndex, Size) = GenericEvent.DeserializeAuxHeader(Reader.ReadUInt32());
+			Assert.AreEqual(PredefinedEventUid.AuxData, Uid);
+			Assert.AreEqual(1, FieldIndex);
+			Assert.AreEqual(110, Size);
+		}
+
+		[TestMethod]
+		public void AuxHeaderSerialize()
+		{
+			uint Header = GenericEvent.SerializeAuxHeader(1, 1, 110);
+			using MemoryStream Ms = new MemoryStream();
+			using BinaryWriter Writer = new BinaryWriter(Ms);
+			Writer.Write(Header);
+			byte[] Data = Ms.ToArray();
+			Assert.AreEqual(0x01, Data[0]);
+			Assert.AreEqual(0xC1, Data[1]);
+			Assert.AreEqual(0x0D, Data[2]);
+			Assert.AreEqual(0x00, Data[3]);
+			
+			var (Uid, FieldIndex, Size) = GenericEvent.DeserializeAuxHeader(Header);
+			Assert.AreEqual(PredefinedEventUid.AuxData, Uid);
+			Assert.AreEqual(1, FieldIndex);
+			Assert.AreEqual(110, Size);
+		}
+
+		internal static byte[] StringToByteArray(string Hex)
+		{
+			int NumChars = Hex.Length;
+			byte[] Bytes = new byte[NumChars / 2];
+			for (int i = 0; i < NumChars; i += 2)
+				Bytes[i / 2] = Convert.ToByte(Hex.Substring(i, 2), 16);
+			return Bytes;
+		}
+		
+		internal static byte[] HexStringToBytes(string HexData)
+		{
+			return StringToByteArray(HexData.Replace(" ", ""));
+		}
+		
+		public static string ToHexString(byte[] ba)
+		{
+			return BitConverter.ToString(ba).Replace("-"," ");
+		}
+
+		internal static void AssertHexString(string Expected, string Actual)
+		{
+			if (Expected != Actual)
+			{
+				Console.WriteLine("Expected:\n" + Expected);
+				Console.WriteLine("Actual:\n" + Actual);
+			}
+
+			Assert.AreEqual(Expected, Actual);
+		}
+		
+		internal static void AssertHexString(string Expected, byte[] Actual)
+		{
+			AssertHexString(Expected, ToHexString(Actual));
+		}
+		
+		internal static void AssertHexString(byte[] Expected, byte[] Actual)
+		{
+			AssertHexString(ToHexString(Expected), ToHexString(Actual));
+		}
+		
+		private static GenericEvent DeserializeAndAssertGenericEvent(string HexData, EventType EventType, bool IsImportant, ushort ExpectedUid, uint? ExpectedSerial = 0)
+		{
+			byte[] Block = HexStringToBytes(HexData);
+			using MemoryStream Ms = new MemoryStream(Block);
+			using BinaryReader Reader = new BinaryReader(Ms);
+
+			ushort Uid;
+			if (IsImportant)
+			{
+				TraceImportantEventHeader EventHeader = TraceImportantEventHeader.Deserialize(Reader);
+				Assert.AreEqual(ExpectedUid, EventHeader.Uid);
+				Assert.AreEqual(Block.Length - TraceImportantEventHeader.HeaderSize, EventHeader.EventSize);
+				Uid = EventHeader.Uid;
+			}
+			else
+			{
+				Uid = Reader.ReadPackedUid(out _);
+				Assert.AreEqual(ExpectedUid, Uid);	
+			}
+			
+			GenericEvent Event = GenericEvent.Deserialize(Uid, Reader, EventType);
+			Reader.EnsureEntireStreamIsConsumed();
+			
+			Assert.AreEqual(ExpectedSerial, Event.Serial);
+			Assert.AreEqual(EventType.Fields.Count, Event.Fields.Length);
+			return Event;
+		}
+
+		private static void SerializeAndAssertGenericEvent(string ExpectedHexData, EventType EventType, bool IsImportant, ushort Uid, GenericEvent.Field[] Fields, uint Serial = 0)
+		{
+			using MemoryStream Ms = new MemoryStream();
+			using BinaryWriter Writer = new BinaryWriter(Ms);
+
+			byte[] ExpectedData = HexStringToBytes(ExpectedHexData);
+			GenericEvent Event = new GenericEvent(Serial, Fields, EventType);
+
+			if (IsImportant)
+			{
+				int ExpectedEventSize = ExpectedData.Length - TraceImportantEventHeader.HeaderSize;
+				Assert.AreEqual(ExpectedEventSize, Event.Size);
+				TraceImportantEventHeader EventHeader = new TraceImportantEventHeader(Uid, Event.Size);
+				EventHeader.Serialize(Writer);
+			}
+			else
+			{
+				Writer.WritePackedUid(Uid);
+			}
+			
+			Event.Serialize(20, Writer);
+			
+			AssertHexString(ExpectedHexData, Ms.ToArray());
+		}
+		
+		private static void SerializeAndAssertEvent(string ExpectedHexData, ITraceEvent Event, bool IsImportant, ushort Uid)
+		{
+			using MemoryStream Ms = new MemoryStream();
+			using BinaryWriter Writer = new BinaryWriter(Ms);
+			Event.Serialize(Uid, Writer);
+			
+			AssertHexString(ExpectedHexData, Ms.ToArray());
+		}
+	}
+}
\ No newline at end of file
diff --git a/Engine/Source/Programs/Shared/EpicGames.Tracing.Tests/UnrealInsights/TransportPacketTest.cs b/Engine/Source/Programs/Shared/EpicGames.Tracing.Tests/UnrealInsights/TransportPacketTest.cs
new file mode 100644
--- /dev/null
+++ b/Engine/Source/Programs/Shared/EpicGames.Tracing.Tests/UnrealInsights/TransportPacketTest.cs
@@ -0,0 +1,266 @@
+// Copyright Epic Games, Inc. All Rights Reserved.
+
+using System;
+using System.IO;
+using Microsoft.VisualStudio.TestTools.UnitTesting;
+using EpicGames.Tracing.UnrealInsights;
+using EpicGames.Tracing.UnrealInsights.Events;
+
+namespace EpicGames.Tracing.Tests.UnrealInsights
+{
+	[TestClass]
+	public class TransportPacketTest
+	{
+		[TestMethod]
+		public void ShortByteHandling()
+		{
+			ushort Val1 = 0x1234;
+			Assert.AreEqual(0x34, BinaryReaderExtensions.GetLowByte(Val1));
+			Assert.AreEqual(0x12, BinaryReaderExtensions.GetHighByte(Val1));
+			
+			ushort Val2 = BinaryReaderExtensions.MakeShort(0x56, 0x78);
+			Assert.AreEqual(0x56, BinaryReaderExtensions.GetLowByte(Val2));
+			Assert.AreEqual(0x78, BinaryReaderExtensions.GetHighByte(Val2));
+			
+			ushort Val3 = BinaryReaderExtensions.MakeShort(0x25, 0x00);
+			Assert.AreEqual(0x25, Val3);
+		}
+
+		[TestMethod]
+		public void GetPackedUid()
+		{
+			void AssertGetPackedUid(ushort ExpectedUid, bool ExpectedIsTwoByteUid, byte UidLow, byte UidHigh)
+			{
+				Assert.AreEqual(ExpectedUid, BinaryReaderExtensions.GetPackedUid(UidLow, UidHigh, out bool IsTwoByteUid));
+				Assert.AreEqual(ExpectedIsTwoByteUid, IsTwoByteUid);
+			}
+
+			AssertGetPackedUid(0x05, false, 0x0A, 0x08);
+			AssertGetPackedUid(0x13, false, 0x26, 0x00);
+			AssertGetPackedUid(4, false, 0x08, 0x25);
+			AssertGetPackedUid(4, false, 0x08, 0x29);
+			AssertGetPackedUid(5, false, 0x0A, 0x08);
+			AssertGetPackedUid(8, false, 0x10, 0xE5);
+			AssertGetPackedUid(12, false, 0x18, 0xCA);
+			AssertGetPackedUid(5, false, 0x0A, 0x0A);
+			AssertGetPackedUid(12, false, 0x18, 0xF3);
+			
+			AssertGetPackedUid(17, true, 0x23, 0x00);
+			AssertGetPackedUid(18, true, 0x25, 0x00);
+			AssertGetPackedUid(39, true, 0x4F, 0x00);
+			AssertGetPackedUid(41, true, 0x53, 0x00);
+			AssertGetPackedUid(0x12, true, 0x25, 0x00);
+			AssertGetPackedUid(0x13, true, 0x27, 0x00);
+			AssertGetPackedUid(0x4513, true, 0x27, 0x45);
+		}
+		
+		[TestMethod]
+		public void WritePackedUid()
+		{
+			void AssertWriteUid(ushort Uid, bool ExpectTwoByteUid)
+			{
+				using MemoryStream Ms = new MemoryStream();
+				using BinaryWriter Writer = new BinaryWriter(Ms);
+				Writer.WritePackedUid(Uid);
+				byte[] Data = Ms.ToArray();
+				Assert.AreEqual(ExpectTwoByteUid ? 2 : 1, Data.Length);
+				byte UidHigh = Data.Length == 2 ? Data[1] : (byte)0xFF;
+				ushort DeserializedUid = BinaryReaderExtensions.GetPackedUid(Data[0], UidHigh, out bool IsTwoByteUid);
+				Assert.AreEqual(ExpectTwoByteUid, IsTwoByteUid);
+				Assert.AreEqual(Uid, DeserializedUid);
+			}
+			
+			AssertWriteUid(0x00, false);
+			AssertWriteUid(0x0A, false);
+			AssertWriteUid(0x05, false);
+		
+			AssertWriteUid(17, true);
+			AssertWriteUid(18, true);
+			AssertWriteUid(126, true);
+
+			Assert.ThrowsException<NotImplementedException>(() => AssertWriteUid(127, true));
+			Assert.ThrowsException<NotImplementedException>(() => AssertWriteUid(128, true));
+			Assert.ThrowsException<NotImplementedException>(() => AssertWriteUid(4362, true));
+		}
+
+		[TestMethod]
+		public void Deserialize7Bit()
+		{
+			void Assert7Bit(ulong ExpectedValue, string ActualBytesHex)
+			{
+				byte[] ActualBytes = GenericEventTest.StringToByteArray(ActualBytesHex.Replace(" ", ""));
+				using MemoryStream Ms = new MemoryStream(ActualBytes);
+				using BinaryReader Reader = new BinaryReader(Ms);
+				Assert.AreEqual(ExpectedValue, TraceUtils.Read7BitUint(Reader));
+			}
+
+			Assert7Bit(0, "00");
+			Assert7Bit(1, "01");
+			Assert7Bit(127, "7F");
+			Assert7Bit(128, "80 01");
+			Assert7Bit(129, "81 01");
+			
+			Assert7Bit(16383, "FF 7F");
+			Assert7Bit(16384, "80 80 01");
+			Assert7Bit(16385, "81 80 01");
+			
+			Assert7Bit(2097151, "FF FF 7F");
+			Assert7Bit(2097152, "80 80 80 01");
+			Assert7Bit(2097153, "81 80 80 01");
+			
+			Assert7Bit(268435455, "FF FF FF 7F");
+			Assert7Bit(268435456, "80 80 80 80 01");
+			Assert7Bit(268435457, "81 80 80 80 01");
+			
+			Assert7Bit(34359738367, "FF FF FF FF 7F");
+			Assert7Bit(34359738368, "80 80 80 80 80 01");
+			Assert7Bit(34359738369, "81 80 80 80 80 01");
+			
+			Assert7Bit(562949953421311, "FF FF FF FF FF FF 7F");
+			Assert7Bit(562949953421312, "80 80 80 80 80 80 80 01");
+			Assert7Bit(562949953421313, "81 80 80 80 80 80 80 01");
+		}
+		
+		[TestMethod]
+		public void Serialize7Bit()
+		{
+			void Assert7Bit(ulong Value, string ExpectedBytesHex)
+			{
+				byte[] ExpectedBytes = GenericEventTest.StringToByteArray(ExpectedBytesHex.Replace(" ", ""));
+				using MemoryStream Ms = new MemoryStream();
+				using BinaryWriter Writer = new BinaryWriter(Ms);
+
+				int ActualNumBytesWritten = TraceUtils.Write7BitUint(Writer, Value);
+				byte[] ActualBytes = Ms.ToArray();
+
+				Assert.AreEqual(ExpectedBytes.Length, ActualNumBytesWritten);
+				Assert.AreEqual(ExpectedBytes.Length, ActualBytes.Length);
+				GenericEventTest.AssertHexString(ExpectedBytesHex, ActualBytes);
+			}
+
+			Assert7Bit(0, "00");
+			Assert7Bit(1, "01");
+			Assert7Bit(127, "7F");
+			Assert7Bit(128, "80 01");
+			Assert7Bit(129, "81 01");
+			
+			Assert7Bit(16383, "FF 7F");
+			Assert7Bit(16384, "80 80 01");
+			Assert7Bit(16385, "81 80 01");
+			
+			Assert7Bit(2097151, "FF FF 7F");
+			Assert7Bit(2097152, "80 80 80 01");
+			Assert7Bit(2097153, "81 80 80 01");
+			
+			Assert7Bit(268435455, "FF FF FF 7F");
+			Assert7Bit(268435456, "80 80 80 80 01");
+			Assert7Bit(268435457, "81 80 80 80 01");
+			
+			Assert7Bit(34359738367, "FF FF FF FF 7F");
+			Assert7Bit(34359738368, "80 80 80 80 80 01");
+			Assert7Bit(34359738369, "81 80 80 80 80 01");
+			
+			Assert7Bit(562949953421311, "FF FF FF FF FF FF 7F");
+			Assert7Bit(562949953421312, "80 80 80 80 80 80 80 01");
+			Assert7Bit(562949953421313, "81 80 80 80 80 80 80 01");
+		}
+
+		[TestMethod]
+		public void SerializeDeserializeWithStubs()
+		{
+			StubTraceEvent Event1 = new StubTraceEvent(1000, 20);
+			StubTraceEvent Event2 = new StubTraceEvent(2000, 30);
+			ushort Uid1 = 5555;
+			ushort Uid2 = 6666;
+
+			{
+				using MemoryStream Ms = new MemoryStream();
+				using BinaryWriter Writer = new BinaryWriter(Ms);
+				TransportPacket Packet = TransportPacket.Create(0, 0);
+				Packet.Serialize(Writer, new [] { (Uid1, Event1 as ITraceEvent), (Uid2, Event2 as ITraceEvent) });
+				
+				Ms.Position = 0;
+				using BinaryReader Reader = new BinaryReader(Ms);
+
+				TransportPacket DeserializedPacket = TransportPacket.Deserialize(Reader);
+				Assert.AreEqual(0, DeserializedPacket.ThreadId);
+				Assert.AreEqual(Event1.Size + Event2.Size, DeserializedPacket.Data.Length);
+			}
+			
+			{
+				using MemoryStream Ms = new MemoryStream();
+				using BinaryWriter Writer = new BinaryWriter(Ms);
+				TransportPacket Packet = TransportPacket.Create(0, 1);
+				Packet.Serialize(Writer, Array.Empty<(ushort, ITraceEvent)>());
+
+				Ms.Position = 0;
+				using BinaryReader Reader = new BinaryReader(Ms);
+
+				TransportPacket DeserializedPacket = TransportPacket.Deserialize(Reader);
+				Assert.AreEqual(1, DeserializedPacket.ThreadId);
+				Assert.AreEqual(0, DeserializedPacket.Data.Length);
+			}
+			
+			{
+				using MemoryStream Ms = new MemoryStream();
+				using BinaryWriter Writer = new BinaryWriter(Ms);
+				TransportPacket Packet = TransportPacket.Create(0, 12345);
+				Packet.Serialize(Writer, Array.Empty<(ushort, ITraceEvent)>());
+
+				Ms.Position = 0;
+				using BinaryReader Reader = new BinaryReader(Ms);
+
+				TransportPacket DeserializedPacket = TransportPacket.Deserialize(Reader);
+				Assert.AreEqual(12345, DeserializedPacket.ThreadId);
+				Assert.AreEqual(0, DeserializedPacket.Data.Length);
+			}
+		}
+		
+		[TestMethod]
+		public void SerializeDeserialize()
+		{
+			TraceNewTraceEvent Event1 = new TraceNewTraceEvent(25582215261913, 10000000, 21069, 8);
+			DiagnosticsSession2Event Event2 = new DiagnosticsSession2Event("Win64", "UnrealEditor", @" D:\depot\starship-main\QAGame\QAGame.uproject -trace=cpu -game", "++UE5+Main", "++UE5+Main-CL-17442524", 17442524, 3, 4);
+			CpuProfilerEventSpecEvent Event3 = new CpuProfilerEventSpecEvent(9524, "K2Node_VariableGet /Engine/EngineSky/BP_Sky_Sphere.BP_Sky_Sphere:UserConstructionScript.K2Node_VariableGet_983");
+
+			ushort Uid1 = 10;
+			ushort Uid2 = 20;
+			ushort Uid3 = 30;
+
+			{
+				using MemoryStream Ms = new MemoryStream();
+				using BinaryWriter Writer = new BinaryWriter(Ms);
+				TransportPacket Packet = TransportPacket.Create(0, 0);
+				Packet.Serialize(Writer, new [] { (Uid1, Event1 as ITraceEvent), (Uid2, Event2 as ITraceEvent) });
+				
+				Ms.Position = 0;
+				using BinaryReader Reader = new BinaryReader(Ms);
+
+				TransportPacket DeserializedPacket = TransportPacket.Deserialize(Reader);
+				Assert.AreEqual(0, DeserializedPacket.ThreadId);
+				Assert.AreEqual(Event1.Size + Event2.Size, DeserializedPacket.Data.Length);
+				Reader.EnsureEntireStreamIsConsumed();
+			}
+			
+			{
+				using MemoryStream Ms = new MemoryStream();
+				using BinaryWriter Writer = new BinaryWriter(Ms);
+				TransportPacket Packet = TransportPacket.Create(0, 0);
+				Packet.Serialize(Writer, new []
+				{
+					(Uid1, Event1.Type as ITraceEvent),
+					(Uid2, Event2.Type as ITraceEvent),
+					(Uid3, Event3.Type as ITraceEvent)
+				});
+				
+				Ms.Position = 0;
+				using BinaryReader Reader = new BinaryReader(Ms);
+
+				TransportPacket DeserializedPacket = TransportPacket.Deserialize(Reader);
+				Assert.AreEqual(0, DeserializedPacket.ThreadId);
+				Assert.AreEqual(Event1.Type.Size + Event2.Type.Size + Event3.Type.Size, DeserializedPacket.Data.Length);
+				Reader.EnsureEntireStreamIsConsumed();
+			}
+		}
+	}
+}
\ No newline at end of file
diff --git a/Engine/Source/Programs/Shared/EpicGames.Tracing.Tests/UnrealInsights/UnrealInsightsTest.cs b/Engine/Source/Programs/Shared/EpicGames.Tracing.Tests/UnrealInsights/UnrealInsightsTest.cs
new file mode 100644
--- /dev/null
+++ b/Engine/Source/Programs/Shared/EpicGames.Tracing.Tests/UnrealInsights/UnrealInsightsTest.cs
@@ -0,0 +1,246 @@
+// Copyright Epic Games, Inc. All Rights Reserved.
+
+using System.Collections.Generic;
+using System.IO;
+using EpicGames.Tracing.UnrealInsights;
+using EpicGames.Tracing.UnrealInsights.Events;
+using Microsoft.VisualStudio.TestTools.UnitTesting;
+using ITraceEvent = EpicGames.Tracing.UnrealInsights.ITraceEvent;
+using UnrealInsightsReader = EpicGames.Tracing.UnrealInsights.UnrealInsightsReader;
+using UnrealInsightsWriter = EpicGames.Tracing.UnrealInsights.UnrealInsightsWriter;
+
+namespace EpicGames.Tracing.Tests.UnrealInsights
+{
+	public class StubTraceEvent : ITraceEvent
+	{
+		public readonly uint Val1;
+		public readonly byte Val2;
+
+		public StubTraceEvent(uint Val1, byte Val2)
+		{
+			this.Val1 = Val1;
+			this.Val2 = Val2;
+		}
+
+		public ushort Size => sizeof(uint) + sizeof(byte);
+		public EventType Type => EventType.WellKnown(15, "StubEvent");
+
+		public void Serialize(ushort Uid, BinaryWriter Writer)
+		{
+			Writer.Write(Val1);
+			Writer.Write(Val2);
+		}
+		
+		public static StubTraceEvent Deserialize(BinaryReader Reader)
+		{
+			return new StubTraceEvent(Reader.ReadUInt32(), Reader.ReadByte());
+		}
+	}
+	
+	[TestClass]
+	public class TraceEventHeaderTest
+	{
+		[TestMethod]
+		public void SerializeDeserialize()
+		{
+			uint Test1 = 222;
+			uint Test2 = 333;
+			ulong Test3 = 444;
+			ushort EventSize = 4 + 4 + 8;
+			
+			using MemoryStream Ms = new MemoryStream();
+			using BinaryWriter Writer = new BinaryWriter(Ms);
+			new TraceImportantEventHeader(1000, EventSize).Serialize(Writer);
+			Writer.Write(Test1);
+			Writer.Write(Test2);
+			Writer.Write(Test3);
+			
+			Ms.Position = 0;
+			BinaryReader Reader = new BinaryReader(Ms);
+			TraceImportantEventHeader EventHeader = TraceImportantEventHeader.Deserialize(Reader);
+			Assert.AreEqual(1000, EventHeader.Uid);
+			Assert.AreEqual(EventSize, EventHeader.EventSize);
+			Assert.AreEqual(Test1, Reader.ReadUInt32());
+			Assert.AreEqual(Test2, Reader.ReadUInt32());
+			Assert.AreEqual(Test3, Reader.ReadUInt64());
+		}
+	}
+
+	[TestClass]
+	public class UnrealInsightsTest
+	{
+		[TestMethod]
+		public void WriteUtraceFile()
+		{
+			using MemoryStream Ms = new MemoryStream();
+			using BinaryWriter BinaryWriter = new BinaryWriter(Ms);
+
+			UnrealInsightsWriter Writer = new UnrealInsightsWriter();
+			
+			TraceNewTraceEvent NewTraceEvent1 = new TraceNewTraceEvent(1001, 2001, 31, 8);
+			TraceNewTraceEvent NewTraceEvent2 = new TraceNewTraceEvent(1002, 2002, 32, 8);
+			TraceNewTraceEvent NewTraceEvent3 = new TraceNewTraceEvent(1003, 2003, 33, 8);
+			CpuProfilerEventSpecEvent Cpu1 = new CpuProfilerEventSpecEvent(400, "MyCpuEventSpecEvent1");
+
+			Writer.AddEvent(TransportPacket.ThreadIdImportants, NewTraceEvent1);
+			Writer.AddEvent(TransportPacket.ThreadIdImportants, NewTraceEvent2);
+			Writer.AddEvent(TransportPacket.ThreadIdImportants, NewTraceEvent3);
+			Writer.AddEvent(TransportPacket.ThreadIdImportants, Cpu1);
+			Writer.Write(BinaryWriter);
+
+			Ms.Position = 0;
+
+			UnrealInsightsReader Reader = new UnrealInsightsReader();
+			Reader.Read(Ms);
+			Reader.PrintEventSummary();
+			Assert.AreEqual(4, Reader.EventTypes.Count);
+			
+			Assert.AreEqual(0, Reader.EventsPerThread[TransportPacket.ThreadIdEvents].Count);
+			Assert.AreEqual(4, Reader.EventsPerThread[TransportPacket.ThreadIdImportants].Count);
+			
+			{
+				GenericEvent Event = (GenericEvent) Reader.EventsPerThread[1][0];
+				Assert.AreEqual(1001, Event.Fields[0].Long!.Value);
+				Assert.AreEqual(2001, Event.Fields[1].Long!.Value);
+				Assert.AreEqual(31, Event.Fields[2].Int!.Value);
+				Assert.AreEqual(8, Event.Fields[3].Int!.Value);
+			}
+			{
+				GenericEvent Event = (GenericEvent) Reader.EventsPerThread[1][1];
+				Assert.AreEqual(1002, Event.Fields[0].Long!.Value);
+				Assert.AreEqual(2002, Event.Fields[1].Long!.Value);
+				Assert.AreEqual(32, Event.Fields[2].Int!.Value);
+				Assert.AreEqual(8, Event.Fields[3].Int!.Value);
+			}
+			{
+				GenericEvent Event = (GenericEvent) Reader.EventsPerThread[1][2];
+				Assert.AreEqual(1003, Event.Fields[0].Long!.Value);
+				Assert.AreEqual(2003, Event.Fields[1].Long!.Value);
+				Assert.AreEqual(33, Event.Fields[2].Int!.Value);
+				Assert.AreEqual(8, Event.Fields[3].Int!.Value);
+			}
+			{
+				GenericEvent Event = (GenericEvent) Reader.EventsPerThread[1][3];
+				Assert.AreEqual(400, Event.Fields[0].Int!.Value);
+				Assert.AreEqual("MyCpuEventSpecEvent1", Event.Fields[1].String!);
+			}
+		}
+		
+		[TestMethod]
+		[Ignore]
+		public void WriteUtraceExample()
+		{
+			string FileName = "d:\\Temp\\mytrace.utrace";
+			
+			FileStream Fs = File.Open(FileName, FileMode.Create);
+			using BinaryWriter BinaryWriter = new BinaryWriter(Fs);
+			UnrealInsightsWriter Writer = new UnrealInsightsWriter();
+			
+			byte[] CpuBatchData1 = GenericEventTest.HexStringToBytes("87 A7 EB AE 8A D1 0B 01 A7 02 02 04 FB 23 03 A4 0A 61 04 0B 05 F2 80 42 AA 1C EB 5C 06 C0 BC 03 17 07 02 47 08 C2 02 AB 15 09 AA 03 07 0A AE 01 05 0B D8 0D 05 0C B7 01 0D D8 03 00 03 0E A4 03 6F 0F 89 FF 03 10 F3 C4 04 11 60 06 A7 02 10 AB 19 12 B3 04 10 9F 11 13 0A 04 02 04 ED 03 10 91 11 14 80 01 02 97 02 10 A9 D6 0B 15 10 06 DD 03 10 CB 85 04 16 12 06 D3 01 10 89 16 17 0C 02 D6 02 07 18 DA 84 28 99 01 19 18 05 1A 0A A5 10 1B CD 01 1C 06 88 6A C5 49 1D B8 EB 24 A9 17 20 AF 01 10 CD F9 01 21 71 22 F2 AE 12 85 04 23 D1 0D 24 85 74 25 EA 58 91 0D 25 80 E6 02 A9 1C 25 88 CD 06 A7 27 25 D0 DA 09 8B 20 25 DE AD 0B AF 33 25 C0 A9 09 ED 23 25 C4 81 0A F5 2D 25 9A EC 0A F3 35 25 D4 CF 03 FF 19 25 BE 92 01 B1 0D 25 98 22");
+
+			ulong CycleFrequency = 10000000;
+			TraceNewTraceEvent NewTrace1 = new TraceNewTraceEvent(25582215261913, CycleFrequency, 21069, 8);
+			TraceThreadInfoEvent ThreadInfo1 = new TraceThreadInfoEvent(2, 171380, -1, "GameThread");
+			TraceThreadInfoEvent ThreadInfo2 = new TraceThreadInfoEvent(3, 0, 2147483647, "Trace");
+			DiagnosticsSession2Event DiagnosticsSession = new DiagnosticsSession2Event("Finally it works", "HELLO WORLD", "QAGame.uproject -trace=cpu -game", "++UE5+Main", "++UE5+Main-CL-17442524", 1000, 3, 4);
+			CpuProfilerEventSpecEvent CpuEventSpec1 = new CpuProfilerEventSpecEvent(1, "MyCpuEventSpecEvent1");
+			CpuProfilerEventSpecEvent CpuEventSpec2 = new CpuProfilerEventSpecEvent(2, "MyCpuEventSpecEvent2");
+			CpuProfilerEventSpecEvent CpuEventSpec3 = new CpuProfilerEventSpecEvent(3, "MyCpuEventSpecEvent3");
+			CpuProfilerEventSpecEvent CpuEventSpec4 = new CpuProfilerEventSpecEvent(4, "MyCpuEventSpecEvent4");
+			CpuProfilerEventSpecEvent CpuEventSpec5 = new CpuProfilerEventSpecEvent(5, "MyCpuEventSpecEvent5");
+			CpuProfilerEventSpecEvent CpuEventSpec6 = new CpuProfilerEventSpecEvent(6, "MyCpuEventSpecEvent6");
+			CpuProfilerEventSpecEvent CpuEventSpec7 = new CpuProfilerEventSpecEvent(7, "MyCpuEventSpecEvent7");
+			CpuProfilerEventSpecEvent CpuEventSpec8 = new CpuProfilerEventSpecEvent(8, "MyCpuEventSpecEvent8");
+			CpuProfilerEventSpecEvent CpuEventSpecTesting = new CpuProfilerEventSpecEvent(9, "TestingSome");
+			CpuProfilerEventBatchEvent CpuBatch1 = new CpuProfilerEventBatchEvent(CpuBatchData1);
+
+			Writer.AddEvent(TransportPacket.ThreadIdImportants, NewTrace1);
+			Writer.AddEvent(TransportPacket.ThreadIdImportants, ThreadInfo1);
+			Writer.AddEvent(TransportPacket.ThreadIdImportants, ThreadInfo2);
+			Writer.AddEvent(TransportPacket.ThreadIdImportants, DiagnosticsSession);
+			
+			Writer.AddEvent(TransportPacket.ThreadIdImportants, CpuEventSpec1);
+			Writer.AddEvent(TransportPacket.ThreadIdImportants, CpuEventSpec2);
+			Writer.AddEvent(TransportPacket.ThreadIdImportants, CpuEventSpec3);
+			Writer.AddEvent(TransportPacket.ThreadIdImportants, CpuEventSpec4);
+			Writer.AddEvent(TransportPacket.ThreadIdImportants, CpuEventSpec5);
+			Writer.AddEvent(TransportPacket.ThreadIdImportants, CpuEventSpec6);
+			Writer.AddEvent(TransportPacket.ThreadIdImportants, CpuEventSpec7);
+			Writer.AddEvent(TransportPacket.ThreadIdImportants, CpuEventSpec8);
+			Writer.AddEvent(TransportPacket.ThreadIdImportants, CpuEventSpecTesting);
+			
+			Writer.AddEvent(2, CpuBatch1);
+			
+			CpuProfilerSerializer CpuSerializerThread3 = new CpuProfilerSerializer(CycleFrequency);
+			CpuSerializerThread3.ScopeEvents.Add(new CpuProfilerScopeEvent(25582215261913 + 400 * 1000, true, CpuEventSpecTesting.Id));
+			CpuSerializerThread3.ScopeEvents.Add(new CpuProfilerScopeEvent(25582215261913 + (CycleFrequency * CycleFrequency), false, null));
+			
+			List<byte[]> EventBatchData = CpuSerializerThread3.Write();
+			foreach (byte[] Data in EventBatchData)
+			{
+				Writer.AddEvent(3, new CpuProfilerEventBatchEvent(Data));
+			}
+			Writer.Write(BinaryWriter);
+
+			Fs.Close();
+
+			UnrealInsightsReader Reader = new UnrealInsightsReader();
+			Reader.Read(File.Open(FileName, FileMode.Open));
+			Reader.PrintEventSummary();
+
+			Assert.AreEqual(0, Reader.EventsPerThread[TransportPacket.ThreadIdEvents].Count);
+		}
+		
+		[TestMethod]
+		public void ReadUtraceFile()
+		{
+			string ExampleDecompTrace = "UnrealInsights/example_trace.decomp.utrace";
+
+			UnrealInsightsReader Reader = new UnrealInsightsReader();
+			Reader.Read(File.Open(ExampleDecompTrace, FileMode.Open));
+			Reader.PrintEventSummary();
+			
+			Assert.AreEqual(6281, Reader.NumTransportPacketsRead);
+			Dictionary<ushort,List<ITraceEvent>> EventsPerUid = Reader.GetEventsPerUid();
+
+			void AssertEvent(ushort Uid, string ExpectedName, int ExpectedCount)
+			{
+				Assert.IsTrue(Reader.EventTypes.ContainsKey(Uid));
+				Assert.IsTrue(EventsPerUid.ContainsKey(Uid));
+				Assert.AreEqual(ExpectedName, Reader.EventTypes[Uid].Name);
+				Assert.AreEqual(ExpectedCount, EventsPerUid[Uid].Count);
+			}
+
+			AssertEvent(16, "$Trace.NewTrace", 1);
+			AssertEvent(17, "$Trace.ThreadTiming", 71);
+			AssertEvent(18, "Memory.MemoryScope", 16614);
+			AssertEvent(19, "Logging.LogCategory", 741);
+			AssertEvent(20, "Memory.MemoryScopeRealloc", 6262);
+			AssertEvent(21, "Stats.Spec", 972);
+			AssertEvent(22, "CsvProfiler.RegisterCategory", 51);
+			AssertEvent(23, "CsvProfiler.DefineDeclaredStat", 159);
+			AssertEvent(24, "Counters.Spec", 85);
+			AssertEvent(25, "Logging.LogMessageSpec", 6);
+			AssertEvent(26, "Logging.LogMessage", 11);
+			AssertEvent(27, "CsvProfiler.DefineInlineStat", 1);
+			AssertEvent(28, "LoadTime.ClassInfo", 1);
+			AssertEvent(29, "SlateTrace.AddWidget", 1);
+			AssertEvent(30, "SlateTrace.WidgetInfo", 1);
+			AssertEvent(31, "Misc.BookmarkSpec", 2);
+			AssertEvent(32, "Misc.Bookmark", 2);
+			AssertEvent(33, "CpuProfiler.EventSpec", 13003);
+			AssertEvent(34, "Diagnostics.Session2", 1);
+			AssertEvent(35, "Trace.ChannelToggle", 32);
+			AssertEvent(36, "$Trace.ThreadInfo", 100);
+			AssertEvent(37, "Trace.ChannelAnnounce", 34);
+			AssertEvent(38, "$Trace.ThreadGroupBegin", 12);
+			AssertEvent(39, "CpuProfiler.EventBatch", 64354);
+			AssertEvent(40, "$Trace.ThreadGroupEnd", 12);
+			AssertEvent(41, "Cpu.CacheResourceShadersForRendering", 109);
+			AssertEvent(42, "SourceFilters.FilterClass", 2);
+			AssertEvent(43, "WorldSourceFilters.WorldInstance", 1);
+			AssertEvent(44, "WorldSourceFilters.WorldOperation", 2);
+			AssertEvent(45, "CpuProfiler.EndThread", 66);
+		}
+	}
+}
\ No newline at end of file
diff --git a/Engine/Source/Programs/Shared/EpicGames.Tracing/EpicGames.Tracing.csproj b/Engine/Source/Programs/Shared/EpicGames.Tracing/EpicGames.Tracing.csproj
new file mode 100644
--- /dev/null
+++ b/Engine/Source/Programs/Shared/EpicGames.Tracing/EpicGames.Tracing.csproj
@@ -0,0 +1,29 @@
+<Project Sdk="Microsoft.NET.Sdk">
+  <PropertyGroup>
+    <TargetFramework>netcoreapp3.1</TargetFramework>
+    <Configuration Condition=" '$(Configuration)' == '' ">Development</Configuration>
+    <OutputType>Library</OutputType>
+    <Deterministic>true</Deterministic>
+    <RootNamespace>EpicGames.Tracing</RootNamespace>
+    <GenerateAssemblyInfo>false</GenerateAssemblyInfo>
+    <GenerateTargetFrameworkAttribute>false</GenerateTargetFrameworkAttribute>
+    <Configurations>Debug;Release;Development</Configurations>
+    <Nullable>enable</Nullable>
+  </PropertyGroup>
+  <PropertyGroup Condition=" '$(Configuration)|$(Platform)' == 'Debug|AnyCPU' ">
+    <DocumentationFile></DocumentationFile>
+  </PropertyGroup>
+  <PropertyGroup Condition="'$(Configuration)|$(Platform)' == 'Development|AnyCPU'">
+    <DocumentationFile></DocumentationFile>
+    <Optimize>true</Optimize>
+  </PropertyGroup>
+  <PropertyGroup Condition="'$(Configuration)|$(Platform)' == 'Release|AnyCPU'">
+    <DocumentationFile></DocumentationFile>
+    <Optimize>true</Optimize>
+  </PropertyGroup>
+  <ItemGroup>
+    <Compile Include="..\..\Shared\MetaData.cs">
+      <Link>Properties\MetaData.cs</Link>
+    </Compile>
+  </ItemGroup>
+</Project>
\ No newline at end of file
diff --git a/Engine/Source/Programs/Shared/EpicGames.Tracing/UnrealInsights/Events.cs b/Engine/Source/Programs/Shared/EpicGames.Tracing/UnrealInsights/Events.cs
new file mode 100644
--- /dev/null
+++ b/Engine/Source/Programs/Shared/EpicGames.Tracing/UnrealInsights/Events.cs
@@ -0,0 +1,240 @@
+// Copyright Epic Games, Inc. All Rights Reserved.
+
+using System;
+using System.Diagnostics;
+using System.IO;
+using System.Runtime.CompilerServices;
+using EpicGames.Tracing.UnrealInsights.Events;
+
+[assembly: InternalsVisibleTo("EpicGames.Tracing.Tests")]
+
+namespace EpicGames.Tracing.UnrealInsights
+{
+	public static class BinaryReaderExtensions
+	{
+		public static byte[] ReadBytesStrict(this BinaryReader Reader, int Count)
+		{
+			byte[] Result = Reader.ReadBytes(Count);
+			if (Result.Length != Count)
+				throw new EndOfStreamException();
+
+			return Result;
+		}
+
+		/// <summary>
+		/// Check if bit is set for a given byte
+		/// </summary>
+		/// <param name="Value">Value to check</param>
+		/// <param name="Pos">Pos 0 is least significant bit, pos 7 is most</param>
+		/// <returns>True if set</returns>
+		internal static bool IsBitSet(byte Value, int Pos)
+		{
+			return (Value & (1 << Pos)) != 0;
+		}
+
+		internal static bool IsTwoByteUid(byte UidLow)
+		{
+			return IsBitSet(UidLow, 0);
+		}
+
+		internal static byte GetHighByte(ushort a)
+		{
+			return (byte)(a >> 8);
+		}
+
+		internal static byte GetLowByte(ushort a)
+		{
+			return (byte)(a & 0xff);
+		}
+		
+		internal static ushort GetHighWord(uint a)
+		{
+			return (ushort)(a >> 16);
+		}
+
+		internal static ushort GetLowWord(uint a)
+		{
+			return (ushort)(a & 0xffff);
+		}
+
+		internal static ushort MakeShort(byte LowByte, byte HighByte)
+		{
+			return (ushort)((byte)(LowByte & 0xff) | (ushort)(HighByte & 0xff) << 8);
+		}
+
+		internal static ushort GetPackedUid(byte UidLow, byte UidHigh, out bool IsTwoByteUidArg)
+		{
+			// struct packed_uid
+			// {
+			// 	uint8   is_two_byte_uid : 1
+			// 			uid_low         : 7
+			// 	(uint8  uid_high) // if is_two_byte_uid == 1
+			// }
+
+			IsTwoByteUidArg = IsTwoByteUid(UidLow);
+			byte UidLowNoBit = (byte) (UidLow >> 1); // Strip the is_two_byte_uid bit
+
+			return MakeShort(UidLowNoBit, IsTwoByteUidArg ? UidHigh : (byte) 0x00);
+		}
+
+		public static ushort ReadPackedUid(this BinaryReader Reader, out bool IsTwoByteUid)
+		{
+			byte UidLow = Reader.ReadByte();
+			byte UidHigh = Reader.ReadByte();
+			Reader.BaseStream.Position -= 1; // UidLow is always consumed, but UidHigh maybe not.
+			
+			ushort Uid = GetPackedUid(UidLow, UidHigh, out IsTwoByteUid);
+			if (IsTwoByteUid)
+			{
+				Reader.ReadByte();  // Consume UidHigh
+			}
+
+			return Uid;
+		}
+		
+		public static void EnsureEntireStreamIsConsumed(this BinaryReader Reader)
+		{
+			bool IsEntireStreamConsumed = Reader.BaseStream.Position == Reader.BaseStream.Length;
+			if (!IsEntireStreamConsumed)
+			{
+				throw new Exception($"Entire stream/buffer was not consumed. Pos={Reader.BaseStream.Position} Len={Reader.BaseStream.Length}");
+			}
+		}
+		
+		public static bool IsEntireStreamConsumed(this BinaryReader Reader)
+		{
+			return Reader.BaseStream.Position == Reader.BaseStream.Length;
+		}
+	}
+
+	public static class BinaryWriterExtensions
+	{
+		public static void WritePackedUid(this BinaryWriter Writer, ushort Uid)
+		{
+			if (Uid < PredefinedEventUid._WellKnownNum)
+			{
+				byte UidLow = (byte) Uid;
+				UidLow = (byte) (UidLow << 1);
+				// LSB is 0 after shifting, indicating a one-byte UID
+				Writer.Write(UidLow);
+			}
+			else if (Uid < 127)
+			{
+				byte UidLow = BinaryReaderExtensions.GetLowByte(Uid);
+				byte UidHigh = BinaryReaderExtensions.GetHighByte(Uid);
+				UidLow = (byte) (UidLow << 1);
+				UidLow = (byte) (UidLow | 1);
+				Writer.Write(UidLow);
+				Writer.Write(UidHigh);
+			}
+			else
+			{
+				throw new NotImplementedException("Handling of UIDs >= 127 not implemented");
+			}
+		}
+	}
+
+	public interface ITraceEvent
+	{
+		ushort Size { get; }
+		EventType Type { get; }
+		public void Serialize(ushort Uid, BinaryWriter Writer);
+	}
+
+	public class EnterScopeEvent : ITraceEvent
+	{
+		public ushort Size => 0;
+		public EventType Type => EventType.WellKnown(PredefinedEventUid.EnterScope, "EnterScope");
+		public void Serialize(ushort Uid, BinaryWriter Writer) { throw new NotImplementedException(); }
+	}
+	
+	public class LeaveScopeEvent : ITraceEvent
+	{
+		public ushort Size => 0;
+		public EventType Type => EventType.WellKnown(PredefinedEventUid.LeaveScope, "LeaveScope");
+		public void Serialize(ushort Uid, BinaryWriter Writer) { throw new NotImplementedException(); }
+	}
+	
+	public class EnterScopeEventTimestamp : ITraceEvent
+	{
+		public ushort Size => 7;
+		public EventType Type => EventType.WellKnown(PredefinedEventUid.EnterScope_T, "EnterScopeTimestamp");
+		public readonly ulong Timestamp;
+
+		public EnterScopeEventTimestamp(ulong Timestamp)
+		{
+			this.Timestamp = Timestamp;
+		}
+
+		public void Serialize(ushort Uid, BinaryWriter Writer) { throw new NotImplementedException(); }
+
+		public static EnterScopeEventTimestamp Deserialize(BinaryReader Reader)
+		{
+			ulong Value = Reader.ReadUInt64();
+			ushort UidFound = BinaryReaderExtensions.GetPackedUid((byte) (Value & 0xFF), 0x00, out bool _);
+
+			if (UidFound != PredefinedEventUid.EnterScope_T)
+			{
+				throw new ArgumentException($"Bad UID found when deserializing 0x{UidFound:X4}/{UidFound}");
+			}
+
+			ulong Timestamp = Value >> 8;
+			return new EnterScopeEventTimestamp(Timestamp);
+		}
+	}
+	
+	public class LeaveScopeEventTimestamp : ITraceEvent
+	{
+		public ushort Size => 7;
+		public EventType Type => EventType.WellKnown(PredefinedEventUid.EnterScope_T, "LeaveScopeTimestamp");
+		public readonly ulong Timestamp;
+
+		public LeaveScopeEventTimestamp(ulong Timestamp)
+		{
+			this.Timestamp = Timestamp;
+		}
+
+		public void Serialize(ushort Uid, BinaryWriter Writer) { throw new NotImplementedException(); }
+		
+		public static LeaveScopeEventTimestamp Deserialize(BinaryReader Reader)
+		{
+			ulong Value = Reader.ReadUInt64();
+			ushort UidFound = BinaryReaderExtensions.GetPackedUid((byte) (Value & 0xFF), 0x00, out bool _);
+
+			if (UidFound != PredefinedEventUid.LeaveScope_T)
+			{
+				throw new ArgumentException($"Bad UID found when deserializing 0x{UidFound:X4}/{UidFound}");
+			}
+
+			ulong Timestamp = Value >> 8;
+			return new LeaveScopeEventTimestamp(Timestamp);
+		}
+	}
+
+	public class TraceImportantEventHeader
+	{
+		public readonly ushort Uid;
+		public readonly ushort EventSize;
+
+		public const ushort HeaderSize = sizeof(ushort) + sizeof(ushort); 
+
+		public TraceImportantEventHeader(ushort Uid, ushort EventSize)
+		{
+			Debug.Assert(HeaderSize == 4);
+			this.Uid = Uid;
+			this.EventSize = EventSize;
+		}
+
+		public void Serialize(BinaryWriter Writer)
+		{
+			Writer.Write(Uid);
+			Writer.Write(EventSize);
+		}
+
+		public static TraceImportantEventHeader Deserialize(BinaryReader Reader)
+		{
+			ushort Uid = Reader.ReadUInt16();
+			return new TraceImportantEventHeader(Uid, Reader.ReadUInt16());
+		}
+	}
+}
\ No newline at end of file
diff --git a/Engine/Source/Programs/Shared/EpicGames.Tracing/UnrealInsights/Events/CpuProfiler.cs b/Engine/Source/Programs/Shared/EpicGames.Tracing/UnrealInsights/Events/CpuProfiler.cs
new file mode 100644
--- /dev/null
+++ b/Engine/Source/Programs/Shared/EpicGames.Tracing/UnrealInsights/Events/CpuProfiler.cs
@@ -0,0 +1,202 @@
+// Copyright Epic Games, Inc. All Rights Reserved.
+
+using System.Collections.Generic;
+using System.IO;
+using System.Linq;
+
+namespace EpicGames.Tracing.UnrealInsights.Events
+{
+	public class CpuProfilerEventSpecEvent : ITraceEvent
+	{
+		public static readonly EventType EventType;
+		
+		public ushort Size => (ushort) (GenericEvent.Size + TraceImportantEventHeader.HeaderSize);
+		public EventType Type => EventType;
+
+		public readonly uint Id;
+		public readonly string Name;
+
+		private readonly GenericEvent GenericEvent;
+
+		public CpuProfilerEventSpecEvent(uint Id, string Name)
+		{
+			this.Id = Id;
+			this.Name = Name;
+			
+			GenericEvent.Field[] Fields =
+			{
+				GenericEvent.Field.FromInt((int) Id),
+				GenericEvent.Field.FromString(Name),
+			};
+			
+			GenericEvent = new GenericEvent(0, Fields, EventType);
+		}
+
+		static CpuProfilerEventSpecEvent()
+		{
+			EventType = new EventType("CpuProfiler", "EventSpec", EventType.FlagImportant | EventType.FlagMaybeHasAux | EventType.FlagNoSync);
+			EventType.AddEventType(0, 4, EventTypeField.TypeInt32, "Id");
+			EventType.AddEventType(4, 0, EventTypeField.TypeAnsiString, "Name");
+		}
+
+		public void Serialize(ushort Uid, BinaryWriter Writer)
+		{
+			new TraceImportantEventHeader(Uid, GenericEvent.Size).Serialize(Writer);
+			GenericEvent.Serialize(Uid, Writer);
+		}
+	}
+	
+	public class CpuProfilerEventBatchEvent : ITraceEvent
+	{
+		public static readonly EventType EventType;
+		
+		public ushort Size => (ushort) (GenericEvent.Size + sizeof(ushort));
+		public EventType Type => EventType;
+
+		public readonly GenericEvent GenericEvent;
+
+		public CpuProfilerEventBatchEvent(byte[] Data)
+		{
+			GenericEvent.Field[] Fields = {GenericEvent.Field.FromArray(Data)};
+			GenericEvent = new GenericEvent(0, Fields, EventType);
+		}
+		
+		static CpuProfilerEventBatchEvent()
+		{
+			EventType = new EventType("CpuProfiler", "EventBatch", EventType.FlagMaybeHasAux | EventType.FlagNoSync);
+			EventType.AddEventType(0, 0, EventTypeField.TypeArray, "Data");
+		}
+
+		public void Serialize(ushort Uid, BinaryWriter Writer)
+		{
+			Writer.WritePackedUid(Uid);
+			GenericEvent.Serialize(Uid, Writer);
+		}
+	}
+	
+	public class CpuProfilerScopeEvent
+	{
+		public readonly ulong Timestamp;
+		public readonly bool IsEnterScope;
+		public readonly uint? SpecId;
+
+		public CpuProfilerScopeEvent(ulong Timestamp, bool IsEnterScope, uint? SpecId)
+		{
+			this.Timestamp = Timestamp;
+			this.IsEnterScope = IsEnterScope;
+			this.SpecId = SpecId;
+		}
+
+		public override string ToString()
+		{
+			string Name = IsEnterScope ? "EnterScope" : "ExitScope";
+			return $"{nameof(CpuProfilerScopeEvent)}({Name} {nameof(Timestamp)}={Timestamp} {nameof(SpecId)}={SpecId})";
+		}
+	}
+	
+	/// <summary>
+	/// Stateful serializer that can serialize/deserialize the data argument in CpuProfiler.EventBatch events
+	/// </summary>
+	public class CpuProfilerSerializer
+	{
+		private readonly ulong CycleFrequency;
+		public readonly List<CpuProfilerScopeEvent> ScopeEvents = new List<CpuProfilerScopeEvent>();
+
+		public CpuProfilerSerializer(ulong CycleFrequency)
+		{
+			this.CycleFrequency = CycleFrequency;
+		}
+		
+		public void Read(byte[] EventBatchData)
+		{
+			using MemoryStream Ms = new MemoryStream(EventBatchData);
+			using BinaryReader Reader = new BinaryReader(Ms);
+			
+			ulong AbsTimestamp = 0;
+			ulong LastTimestamp = 0;
+
+			while (Reader.BaseStream.Position != Reader.BaseStream.Length)
+			{
+				ulong Value = TraceUtils.Read7BitUint(Reader);
+				(ulong Timestamp, bool IsEnterScope) = DecodeTimestamp(Value);
+
+				if (AbsTimestamp == 0)
+				{
+					AbsTimestamp = Timestamp; // First timestamp is always absolute
+				}
+				else
+				{
+					Timestamp = LastTimestamp + Timestamp * CycleFrequency; // Make timestamp absolute
+				}
+				LastTimestamp = Timestamp;
+
+				if (IsEnterScope)
+				{
+					ulong SpecId = TraceUtils.Read7BitUint(Reader);
+					CpuProfilerScopeEvent ScopeEvent = new CpuProfilerScopeEvent(Timestamp, true, (uint?) SpecId);
+					ScopeEvents.Add(ScopeEvent);
+				}
+				else
+				{
+					CpuProfilerScopeEvent ScopeEvent = new CpuProfilerScopeEvent(Timestamp, false, null);
+					ScopeEvents.Add(ScopeEvent);
+				}
+			}
+		}
+		
+		public List<byte[]> Write()
+		{
+			List<MemoryStream> Streams = new List<MemoryStream>();
+			MemoryStream Stream = new MemoryStream();
+			BinaryWriter Writer = new BinaryWriter(Stream);
+			Streams.Add(Stream);
+			
+			ulong LastAbsTimestamp = 0;
+			foreach (CpuProfilerScopeEvent Scope in ScopeEvents)
+			{
+				ulong Timestamp;
+				if (LastAbsTimestamp == 0)
+				{
+					Timestamp = Scope.Timestamp;
+				}
+				else
+				{
+					Timestamp = (Scope.Timestamp - LastAbsTimestamp) / CycleFrequency;
+				}
+				LastAbsTimestamp = Scope.Timestamp;
+				
+				TraceUtils.Write7BitUint(Writer, EncodeTimestamp(Timestamp, Scope.IsEnterScope));
+				if (Scope.IsEnterScope)
+				{
+					TraceUtils.Write7BitUint(Writer, (ulong) Scope.SpecId!);
+				}
+
+				// Split buffers after 250 bytes
+				if (Stream.Position > 250)
+				{
+					Writer.Close();
+					Stream = new MemoryStream();
+					Writer = new BinaryWriter(Stream);
+					LastAbsTimestamp = 0;
+				}
+			}
+
+			Writer.Close();
+
+			return Streams.Select(x => x.ToArray()).ToList();
+		}
+		
+		internal static ulong EncodeTimestamp(ulong Timestamp, bool IsScopeEnter)
+		{
+			ulong Value = (Timestamp << 1) | (uint) (IsScopeEnter ? 1 : 0);
+			return Value;
+		}
+		
+		internal static (ulong Timestamp, bool IsScopeEnter) DecodeTimestamp(ulong Value)
+		{
+			bool IsScopeEnter = (Value & 1) != 0;
+			ulong Timestamp = Value >> 1; // Strip the IsScopeEnter bit
+			return (Timestamp, IsScopeEnter);
+		}
+	}
+}
\ No newline at end of file
diff --git a/Engine/Source/Programs/Shared/EpicGames.Tracing/UnrealInsights/Events/Diagnostics.cs b/Engine/Source/Programs/Shared/EpicGames.Tracing/UnrealInsights/Events/Diagnostics.cs
new file mode 100644
--- /dev/null
+++ b/Engine/Source/Programs/Shared/EpicGames.Tracing/UnrealInsights/Events/Diagnostics.cs
@@ -0,0 +1,51 @@
+// Copyright Epic Games, Inc. All Rights Reserved.
+
+using System.IO;
+
+namespace EpicGames.Tracing.UnrealInsights.Events
+{
+	public class DiagnosticsSession2Event : ITraceEvent
+	{
+		public static readonly EventType EventType;
+		
+		public ushort Size => (ushort) (GenericEvent.Size + TraceImportantEventHeader.HeaderSize);
+		public EventType Type => EventType;
+		private readonly GenericEvent GenericEvent;
+
+		public DiagnosticsSession2Event(string Platform, string AppName, string CommandLine, string Branch, string BuildVersion, int ChangeList, int ConfigurationType, int TargetType)
+		{
+			GenericEvent.Field[] Fields =
+			{
+				GenericEvent.Field.FromString(Platform),
+				GenericEvent.Field.FromString(AppName),
+				GenericEvent.Field.FromString(CommandLine),
+				GenericEvent.Field.FromString(Branch),
+				GenericEvent.Field.FromString(BuildVersion),
+				GenericEvent.Field.FromInt(ChangeList),
+				GenericEvent.Field.FromInt(ConfigurationType),
+				GenericEvent.Field.FromInt(TargetType),
+			};
+
+			GenericEvent = new GenericEvent(0, Fields, EventType);
+		}
+		
+		static DiagnosticsSession2Event()
+		{
+			EventType = new EventType("Diagnostics", "Session2", EventType.FlagImportant | EventType.FlagMaybeHasAux | EventType.FlagNoSync);
+			EventType.AddEventType(0, 0, EventTypeField.TypeAnsiString, "Platform");
+			EventType.AddEventType(0, 0, EventTypeField.TypeAnsiString, "AppName");
+			EventType.AddEventType(0, 0, EventTypeField.TypeWideString, "CommandLine");
+			EventType.AddEventType(0, 0, EventTypeField.TypeWideString, "Branch");
+			EventType.AddEventType(0, 0, EventTypeField.TypeWideString, "BuildVersion");
+			EventType.AddEventType(0, 4, EventTypeField.TypeInt32, "Changelist");
+			EventType.AddEventType(4, 1, EventTypeField.TypeInt8, "ConfigurationType");
+			EventType.AddEventType(5, 1, EventTypeField.TypeInt8, "TargetType");
+		}
+
+		public void Serialize(ushort Uid, BinaryWriter Writer)
+		{
+			new TraceImportantEventHeader(Uid, GenericEvent.Size).Serialize(Writer);
+			GenericEvent.Serialize(Uid, Writer);
+		}
+	}
+}
\ No newline at end of file
diff --git a/Engine/Source/Programs/Shared/EpicGames.Tracing/UnrealInsights/Events/GenericEvent.cs b/Engine/Source/Programs/Shared/EpicGames.Tracing/UnrealInsights/Events/GenericEvent.cs
new file mode 100644
--- /dev/null
+++ b/Engine/Source/Programs/Shared/EpicGames.Tracing/UnrealInsights/Events/GenericEvent.cs
@@ -0,0 +1,277 @@
+// Copyright Epic Games, Inc. All Rights Reserved.
+
+using System;
+using System.IO;
+using System.Text;
+
+namespace EpicGames.Tracing.UnrealInsights.Events
+{
+	/// <summary>
+	/// Represents a generic event in the stream with fields stored dynamically
+	/// </summary>
+	public class GenericEvent : ITraceEvent
+	{
+		public class Field
+		{
+			public readonly bool? Bool;
+			public readonly int? Int;
+			public readonly long? Long;
+			public readonly double? Float;
+			public readonly string? String;
+			public readonly byte[]? Array;
+
+			private Field(bool Value) { Bool = Value; }
+			private Field(int Value) { Int = Value; }
+			private Field(long Value) { Long = Value; }
+			private Field(float Value) { Float = Value; }
+			private Field(double Value) { Float = Value; }
+			private Field(string Value) { String = Value; }
+			private Field(byte[] Value) { Array = Value; }
+			public static Field FromBool(bool Value) => new Field(Value);
+			public static Field FromInt(int Value) => new Field(Value);
+			public static Field FromLong(long Value) => new Field(Value);
+			public static Field FromFloat(float Value) => new Field(Value);
+			public static Field FromDouble(double Value) => new Field(Value);
+			public static Field FromString(string Value) => new Field(Value);
+			public static Field FromArray(byte[] Value) => new Field(Value);
+		}
+
+		public readonly uint Serial;
+		public readonly Field[] Fields;
+		public EventType Type => EventType;
+		private EventType EventType;
+
+		public GenericEvent(uint Serial, Field[] Fields, EventType EventType)
+		{
+			this.Serial = Serial;
+			this.Fields = Fields;
+			this.EventType = EventType;
+		}
+
+		public ushort Size
+		{
+			get
+			{
+				ushort TotalSize = 0;
+				if (EventType.HasSerial) TotalSize += 3; // 24-bit serial
+				TotalSize += EventType.GetEventSize();
+				if (EventType.MaybeHasAux())
+				{
+					for (int i = 0; i < EventType.Fields.Count; i++)
+					{
+						EventTypeField FieldType = EventType.Fields[i];
+						if (!FieldType.IsAuxData())
+							continue;
+
+						ushort AuxDataSize;
+						if (FieldType.TypeInfo == EventTypeField.TypeAnsiString)
+						{
+							AuxDataSize = (ushort) Fields[i].String!.Length;
+						}
+						else if (FieldType.TypeInfo == EventTypeField.TypeWideString)
+						{
+							AuxDataSize = (ushort) (Fields[i].String!.Length * 2);
+						}
+						else
+						{
+							AuxDataSize = (ushort) Fields[i].Array!.Length;
+						}
+
+						TotalSize += sizeof(uint); // AuxHeader
+						TotalSize += AuxDataSize;
+					}
+					TotalSize += sizeof(byte); // AuxTerminal UID
+				}
+				
+				return TotalSize;
+			}
+		}
+
+		public void Serialize(ushort Uid, BinaryWriter Writer)
+		{
+			if (EventType.HasSerial)
+			{
+				byte SerialLow = (byte) (Serial & 0xFF);
+				Writer.Write(SerialLow);
+				Writer.Write(BinaryReaderExtensions.GetHighWord(Serial));
+			}
+			
+			for (int i = 0; i < EventType.Fields.Count; i++)
+			{
+				EventTypeField FieldType = EventType.Fields[i];
+
+				if (FieldType.TypeInfo == EventTypeField.TypeInt8) { Writer.Write((byte)Fields[i].Int!); }
+				else if (FieldType.TypeInfo == EventTypeField.TypeInt16) { Writer.Write((ushort)Fields[i].Int!); }
+				else if (FieldType.TypeInfo == EventTypeField.TypeInt32) { Writer.Write((uint)Fields[i].Int!); }
+				else if (FieldType.TypeInfo == EventTypeField.TypeInt64) { Writer.Write((ulong)Fields[i].Long!); }
+				else if (FieldType.TypeInfo == EventTypeField.TypePointer) { Writer.Write((ulong)Fields[i].Long!); }
+				else if (FieldType.TypeInfo == EventTypeField.TypeFloat32) { Writer.Write((float)Fields[i].Float!); }
+				else if (FieldType.TypeInfo == EventTypeField.TypeFloat64) { Writer.Write((double)Fields[i].Float!); }
+				else if (FieldType.TypeInfo == EventTypeField.TypeAnsiString) {  } // Write later as aux-data
+				else if (FieldType.TypeInfo == EventTypeField.TypeWideString) {  } // Write later as aux-data
+				else if (FieldType.TypeInfo == EventTypeField.TypeArray) {  } // Write later as aux-data
+				else if (FieldType.TypeInfo == EventTypeField.TypeBool) { Writer.Write(Fields[i].Bool!.Value ? (byte)1 : (byte)0); }
+				else { throw new Exception($"Found unknown TypeInfo {FieldType.TypeInfo}"); }
+			}
+			
+			if (EventType.MaybeHasAux())
+			{
+				for (int i = 0; i < EventType.Fields.Count; i++)
+				{
+					EventTypeField FieldType = EventType.Fields[i];
+
+					if (FieldType.Size > 0)
+						continue; // Skip any non-aux fields
+
+					ushort AuxUid = PredefinedEventUid.AuxData;
+					if (!EventType.IsImportant())
+					{
+						// Well-known UIDs are only shifted in non-important events
+						AuxUid = (ushort) (AuxUid << 1);
+					}
+
+					byte[] Data;
+					if (FieldType.TypeInfo == EventTypeField.TypeAnsiString)
+					{
+						Data = Encoding.ASCII.GetBytes(Fields[i].String!);
+					}
+					else if (FieldType.TypeInfo == EventTypeField.TypeWideString)
+					{
+						Data = Encoding.Unicode.GetBytes(Fields[i].String!);
+					}
+					else
+					{
+						Data = Fields[i].Array!;
+					}
+
+					uint AuxHeader = 0;
+					AuxHeader = SetBits(AuxHeader, AuxUid, 0, 8);
+					AuxHeader = SetBits(AuxHeader, (uint) i, 8, 5);
+					AuxHeader = SetBits(AuxHeader, (uint) Data.Length, 13, 19);
+					
+					Writer.Write(AuxHeader);
+					Writer.Write(Data);
+				}
+				
+				ushort AuxDataTerminalUid = PredefinedEventUid.AuxDataTerminal;
+				
+				if (!EventType.IsImportant())
+				{
+					// Well-known UIDs are only shifted in non-important events
+					AuxDataTerminalUid = (ushort) (AuxDataTerminalUid << 1);
+				}
+				Writer.Write((byte)AuxDataTerminalUid);
+			}
+		}
+
+		public static GenericEvent Deserialize(ushort Uid, BinaryReader Reader, EventType EventType)
+		{
+			Field[] Fields = new Field[EventType.Fields.Count];
+			uint Serial = 0;
+			
+			if (!EventType.IsNoSync() && !EventType.IsImportant())
+			{
+				// Read uint24 as serial
+				byte SerialLow = Reader.ReadByte();
+				ushort SerialHigh = Reader.ReadUInt16();
+				Serial = SerialLow | ((uint)SerialHigh << 16);
+			}
+
+			for (int i = 0; i < EventType.Fields.Count; i++)
+			{
+				EventTypeField FieldType = EventType.Fields[i];
+
+				if (FieldType.TypeInfo == EventTypeField.TypeInt8) { Fields[i] = Field.FromInt(Reader.ReadByte()); }
+				else if (FieldType.TypeInfo == EventTypeField.TypeInt16) { Fields[i] = Field.FromInt(Reader.ReadUInt16()); }
+				else if (FieldType.TypeInfo == EventTypeField.TypeInt32) { Fields[i] = Field.FromInt((int)Reader.ReadUInt32()); }
+				else if (FieldType.TypeInfo == EventTypeField.TypeInt64) { Fields[i] = Field.FromLong((long) Reader.ReadUInt64()); }
+				else if (FieldType.TypeInfo == EventTypeField.TypePointer) { Fields[i] = Field.FromLong((long) Reader.ReadUInt64()); }
+				else if (FieldType.TypeInfo == EventTypeField.TypeFloat32) { Fields[i] = Field.FromFloat(Reader.ReadSingle()); }
+				else if (FieldType.TypeInfo == EventTypeField.TypeFloat64) { Fields[i] = Field.FromDouble(Reader.ReadDouble()); }
+				else if (FieldType.TypeInfo == EventTypeField.TypeAnsiString) {  } // Read later as aux-data
+				else if (FieldType.TypeInfo == EventTypeField.TypeWideString) {  } // Read later as aux-data
+				else if (FieldType.TypeInfo == EventTypeField.TypeArray) {  } // Read later as aux-data
+				else if (FieldType.TypeInfo == EventTypeField.TypeBool) { Fields[i] = Field.FromBool(Reader.ReadByte() == 1); }
+				else { throw new Exception($"Found unknown TypeInfo {FieldType.TypeInfo}"); }
+			}
+
+			if (EventType.MaybeHasAux())
+			{
+				for (;;)
+				{
+					ushort AuxUid = Reader.ReadByte();
+					if (!EventType.IsImportant())
+					{
+						// Well-known UIDs are only shifted in non-important events
+						AuxUid = (ushort) (AuxUid >> 1);
+					}
+
+					if (AuxUid == PredefinedEventUid.AuxData)
+					{
+						Reader.BaseStream.Position -= 1; // Include the UID for parsing rest of header as a single uint32
+						uint AuxHeader = Reader.ReadUInt32();
+
+						int FieldIndex = (int) ReadBits(AuxHeader, 8, 5);
+						uint Size = ReadBits(AuxHeader, 13, 19);
+
+						byte[] AuxData = Reader.ReadBytesStrict((int) Size);
+						
+						if (EventType.Fields[FieldIndex].TypeInfo == EventTypeField.TypeAnsiString)
+						{
+							Fields[FieldIndex] = Field.FromString(Encoding.ASCII.GetString(AuxData));
+						}
+						else if (EventType.Fields[FieldIndex].TypeInfo == EventTypeField.TypeWideString)
+						{
+							Fields[FieldIndex] = Field.FromString(Encoding.Unicode.GetString(AuxData));
+						}
+						else
+						{
+							Fields[FieldIndex] = Field.FromArray(AuxData);
+						}
+					}
+					else if (AuxUid == PredefinedEventUid.AuxDataTerminal)
+					{
+						break;
+					}
+					else
+					{
+						throw new Exception($"Invalid AuxUid found: {AuxUid} / 0x{AuxUid:X4}");
+					}
+				}
+			}
+
+			return new GenericEvent(Serial, Fields, EventType);
+		}
+		
+		private static uint SetBits(uint Word, uint Value, int Pos, int Size)
+		{
+			uint mask = (((uint)1 << Size) - 1) << Pos;
+			Word &= ~mask;
+			Word |= (Value << Pos) & mask;
+			return Word;
+		}
+
+		private static uint ReadBits(uint Word, int Pos, int Size)
+		{
+			uint Mask = (((uint)1 << Size) - 1) << Pos;
+			return (Word & Mask) >> Pos;
+		}
+
+		internal static (ushort Uid, int FieldIndex, int Size) DeserializeAuxHeader(uint Header)
+		{
+			ushort Uid = (ushort) ReadBits(Header, 0, 8);
+			int FieldIndex = (int) ReadBits(Header, 8, 5);
+			int Size = (int) ReadBits(Header, 13, 19);
+			return (Uid, FieldIndex, Size);
+		}
+		
+		internal static uint SerializeAuxHeader(ushort Uid, int FieldIndex, int Size)
+		{
+			uint Header = 0;
+			Header = SetBits(Header, Uid, 0, 8);
+			Header = SetBits(Header, (uint) FieldIndex, 8, 5);
+			Header = SetBits(Header, (uint) Size, 13, 19);
+			return Header;
+		}
+	}
+}
\ No newline at end of file
diff --git a/Engine/Source/Programs/Shared/EpicGames.Tracing/UnrealInsights/Events/Logging.cs b/Engine/Source/Programs/Shared/EpicGames.Tracing/UnrealInsights/Events/Logging.cs
new file mode 100644
--- /dev/null
+++ b/Engine/Source/Programs/Shared/EpicGames.Tracing/UnrealInsights/Events/Logging.cs
@@ -0,0 +1,28 @@
+// Copyright Epic Games, Inc. All Rights Reserved.
+
+using System;
+using System.IO;
+
+namespace EpicGames.Tracing.UnrealInsights.Events
+{
+	public class LoggingLogMessage : ITraceEvent
+	{
+		public static readonly EventType EventType;
+
+		public ushort Size => throw new NotImplementedException();
+		public EventType Type => EventType;
+
+		static LoggingLogMessage()
+		{
+			EventType = new EventType("Logging", "LogMessage", EventType.FlagMaybeHasAux | EventType.FlagNoSync);
+			EventType.AddEventType(0, 8, EventTypeField.TypeInt64, "LogPoint");
+			EventType.AddEventType(8, 8, EventTypeField.TypeInt64, "Cycle");
+			EventType.AddEventType(16, 0, EventTypeField.TypeArray, "FormatArgs");
+		}
+		
+		public void Serialize(ushort Uid, BinaryWriter Writer)
+		{
+			throw new NotImplementedException();
+		}
+	}
+}
\ No newline at end of file
diff --git a/Engine/Source/Programs/Shared/EpicGames.Tracing/UnrealInsights/Events/Memory.cs b/Engine/Source/Programs/Shared/EpicGames.Tracing/UnrealInsights/Events/Memory.cs
new file mode 100644
--- /dev/null
+++ b/Engine/Source/Programs/Shared/EpicGames.Tracing/UnrealInsights/Events/Memory.cs
@@ -0,0 +1,26 @@
+// Copyright Epic Games, Inc. All Rights Reserved.
+
+using System;
+using System.IO;
+
+namespace EpicGames.Tracing.UnrealInsights.Events
+{
+	public class MemoryMemoryScopeEvent : ITraceEvent
+	{
+		public static readonly EventType EventType;
+		
+		public ushort Size => throw new NotImplementedException();
+		public EventType Type => EventType;
+		
+		static MemoryMemoryScopeEvent()
+		{
+			EventType = new EventType("Memory", "MemoryScope", EventType.FlagNone);
+			EventType.AddEventType(0, 4, EventTypeField.TypeInt32, "Id");
+		}
+
+		public void Serialize(ushort Uid, BinaryWriter Writer)
+		{
+			throw new NotImplementedException();
+		}
+	}
+}
\ No newline at end of file
diff --git a/Engine/Source/Programs/Shared/EpicGames.Tracing/UnrealInsights/Events/NewEvent.cs b/Engine/Source/Programs/Shared/EpicGames.Tracing/UnrealInsights/Events/NewEvent.cs
new file mode 100644
--- /dev/null
+++ b/Engine/Source/Programs/Shared/EpicGames.Tracing/UnrealInsights/Events/NewEvent.cs
@@ -0,0 +1,320 @@
+// Copyright Epic Games, Inc. All Rights Reserved.
+
+using System;
+using System.Collections.Generic;
+using System.IO;
+using System.Linq;
+using System.Text;
+
+namespace EpicGames.Tracing.UnrealInsights.Events
+{
+	public class EventTypeField
+	{
+		public ushort Offset { get; private set; }
+		public ushort Size { get; private set; }
+		public byte TypeInfo { get; private set; }
+		public byte NameSize { get; private set; }
+		public string Name { get; private set; } = "<unknown>";
+		public const ushort StructSize = 2 + 2 + 1 + 1;
+
+		private static readonly byte Field_CategoryMask = ToOctalByte("300");
+		private static readonly byte Field_Integer = ToOctalByte("000");
+		private static readonly byte Field_Float = ToOctalByte("100");
+		private static readonly byte Field_Array = ToOctalByte("200");
+
+		private static readonly byte Field_Pow2SizeMask = ToOctalByte("003");
+		private static readonly byte Field_8 = ToOctalByte("000");
+		private static readonly byte Field_16 = ToOctalByte("001");
+		private static readonly byte Field_32 = ToOctalByte("002");
+		private static readonly byte Field_64 = ToOctalByte("003");
+		private static readonly byte Field_Ptr = ToOctalByte("003"); // Assume 64-bit
+
+		private static readonly byte Field_SpecialMask = ToOctalByte("030");
+		private static readonly byte Field_Pod = ToOctalByte("000");
+		private static readonly byte Field_String = ToOctalByte("010");
+
+		public static readonly byte TypeBool = (byte) (Field_Pod | Field_Integer | Field_8);
+		public static readonly byte TypeInt8 = (byte) (Field_Pod | Field_Integer | Field_8);
+		public static readonly byte TypeInt16 = (byte) (Field_Pod | Field_Integer | Field_16);
+		public static readonly byte TypeInt32 = (byte) (Field_Pod | Field_Integer | Field_32);
+		public static readonly byte TypeInt64 = (byte) (Field_Pod | Field_Integer | Field_64);
+		public static readonly byte TypePointer = (byte) (Field_Pod | Field_Integer | Field_Ptr);
+		public static readonly byte TypeFloat32 = (byte) (Field_Pod | Field_Float | Field_32);
+		public static readonly byte TypeFloat64 = (byte) (Field_Pod | Field_Float | Field_64);
+		public static readonly byte TypeAnsiString = (byte) (Field_String | Field_Integer | Field_Array | Field_8);
+		public static readonly byte TypeWideString = (byte) (Field_String | Field_Integer | Field_Array | Field_16);
+		public static readonly byte TypeArray = (byte) (Field_Array);
+
+		public static byte ToOctalByte(string Value)
+		{
+			return (byte) Convert.ToInt32(Value, 8);
+		}
+
+		public static string TypeInfoToString(byte ByteInfo)
+		{
+			if (ByteInfo == TypeInt8) return "int8";
+			if (ByteInfo == TypeInt16) return "int16";
+			if (ByteInfo == TypeInt32) return "int32";
+			if (ByteInfo == TypeInt64) return "int64";
+			if (ByteInfo == TypePointer) return "ptr";
+			if (ByteInfo == TypeFloat32) return "float32";
+			if (ByteInfo == TypeFloat64) return "float64";
+			if (ByteInfo == TypeAnsiString) return "ansi_str";
+			if (ByteInfo == TypeWideString) return "wide_str";
+			if (ByteInfo == TypeArray) return "array";
+			if (ByteInfo == TypeBool) return "bool";
+			throw new Exception($"Unable to convert type info {ByteInfo} to string");
+		}
+
+		public EventTypeField(ushort Offset, ushort Size, byte TypeInfo, byte NameSize)
+		{
+			this.Offset = Offset;
+			this.Size = Size;
+			this.TypeInfo = TypeInfo;
+			this.NameSize = NameSize;
+			ValidateSize();
+		}
+
+		public EventTypeField(ushort Offset, ushort Size, byte TypeInfo, string Name)
+		{
+			this.Offset = Offset;
+			this.Size = Size;
+			this.TypeInfo = TypeInfo;
+			SetName(Name);
+			ValidateSize();
+		}
+
+		public bool IsArray()
+		{
+			return (TypeInfo & Field_Array) != 0;
+		}
+		
+		public bool IsAuxData()
+		{
+			return IsArray();
+		}
+
+		public void SetName(string Name)
+		{
+			this.NameSize = (byte) Encoding.UTF8.GetBytes(Name).Length;
+			this.Name = Name;
+		}
+
+		public void ValidateSize()
+		{
+			if (TypeInfo == TypeAnsiString || TypeInfo == TypeWideString || TypeInfo == TypeArray)
+			{
+				if (Size != 0)
+				{
+					// Array data is encoded in aux-data, following the event
+					throw new ArgumentException("Array-based field must have a size of 0");
+				}
+			}
+		}
+
+		public override string ToString()
+		{
+			return $"Field({Name} Offset={Offset} Size={Size} Type={TypeInfoToString(TypeInfo)} TypeInfo={TypeInfo})";
+		}
+	}
+
+	public class EventType : ITraceEvent
+	{
+		public ushort NewEventUid { get; set; }
+		public string LoggerName { get; private set; }
+		public string EventName { get; private set; }
+		public byte Flags { get; private set; }
+		public List<EventTypeField> Fields { get; private set; }
+		
+		public string Name => $"{LoggerName}.{EventName}";
+
+		private int EventSize;
+
+		public const byte FlagNone = 0;
+		public const byte FlagImportant = 1 << 0;
+		public const byte FlagMaybeHasAux = 1 << 1;
+		public const byte FlagNoSync = 1 << 2;
+
+		public EventType(ushort NewEventUid, string LoggerName, string EventName, byte Flags, List<EventTypeField> Fields)
+		{
+			this.NewEventUid = NewEventUid;
+			this.LoggerName = LoggerName;
+			this.EventName = EventName;
+			this.Flags = Flags;
+			this.Fields = Fields;
+			EventSize += this.Fields.Sum(f => f.Size);
+		}
+
+		public EventType(string LoggerName, string EventName, byte Flags)
+		{
+			this.NewEventUid = 0;
+			this.LoggerName = LoggerName;
+			this.EventName = EventName;
+			this.Flags = Flags;
+			this.Fields = new List<EventTypeField>();
+		}
+
+		public static EventType WellKnown(ushort Uid, string Name)
+		{
+			return new EventType(Uid, "WellKnown", Name, 0, new List<EventTypeField>());
+		}
+		
+		private static EventType Self()
+		{
+			return new EventType(0, "EventType", "Self", 0, new List<EventTypeField>());
+		}
+
+		public void AddEventType(ushort Offset, ushort Size, byte TypeInfo, string Name)
+		{
+			EventTypeField Field = new EventTypeField(Offset, Size, TypeInfo, Name);
+			Fields.Add(Field);
+			EventSize += Field.Size;
+		}
+		
+		public bool IsWellKnown()
+		{
+			return NewEventUid < PredefinedEventUid._WellKnownNum;
+		}
+
+		public bool IsImportant()
+		{
+			return (Flags & FlagImportant) != 0;
+		}
+
+		public bool MaybeHasAux()
+		{
+			return (Flags & FlagMaybeHasAux) != 0;
+		}
+
+		public bool IsNoSync()
+		{
+			return (Flags & FlagNoSync) != 0;
+		}
+		
+		public bool HasSerial => !IsNoSync() && !IsImportant();
+
+		public int NumAuxFields => Fields.FindAll(x => x.IsAuxData()).Count;
+
+		/// <summary>
+		/// Size of events as seen in the stream
+		/// For size of this NewEvent event, see GetSize()
+		/// </summary>
+		/// <returns></returns>
+		public ushort GetEventSize()
+		{
+			return (ushort) EventSize;
+		}
+
+		public override string ToString()
+		{
+			return $"EventType({Name})";
+		}
+		
+		public string ToStringDetailed()
+		{
+			string FieldText = string.Join(',', Fields.Select(x => x.ToString()));
+			string FlagText = "";
+			if (IsImportant()) FlagText += ",Important";
+			if (MaybeHasAux()) FlagText += ",MaybeHasAux";
+			if (IsNoSync()) FlagText += ",NoSync";
+			FlagText = FlagText.Trim(',');
+			return $"EventType(Uid={NewEventUid} {LoggerName} {EventName} Flags={FlagText} FlagRaw={Flags} Fields={FieldText})";
+		}
+
+		public ushort Size
+		{
+			get
+			{
+				byte[] LoggerNameBytes = Encoding.UTF8.GetBytes(LoggerName);
+				byte[] EventNameBytes = Encoding.UTF8.GetBytes(EventName);
+
+				ushort FieldArraySize = (ushort) (Fields.Count * EventTypeField.StructSize);
+				ushort NamesSize = (ushort) (LoggerNameBytes.Length + EventNameBytes.Length + Fields.Sum(x => x.NameSize));
+				ushort NewEventSize = (ushort) (2 + 1 + 1 + 1 + 1 + FieldArraySize + NamesSize);
+				
+				NewEventSize += TraceImportantEventHeader.HeaderSize;
+				return NewEventSize;
+			}
+		}
+
+		public EventType Type => Self();
+
+		public override int GetHashCode()
+		{
+			return HashCode.Combine(NewEventUid, LoggerName, EventName, Flags);
+		}
+
+		protected bool Equals(EventType other)
+		{
+			return NewEventUid == other.NewEventUid && LoggerName == other.LoggerName && EventName == other.EventName && Flags == other.Flags;
+		}
+
+		public override bool Equals(object? obj)
+		{
+			if (ReferenceEquals(null, obj)) return false;
+			if (ReferenceEquals(this, obj)) return true;
+			if (obj.GetType() != this.GetType()) return false;
+			return Equals((EventType) obj);
+		}
+
+		public void Serialize(ushort Uid, BinaryWriter Writer)
+		{
+			new TraceImportantEventHeader(PredefinedEventUid.NewEvent, (ushort) (Size - TraceImportantEventHeader.HeaderSize)).Serialize(Writer);
+			
+			byte[] LoggerNameBytes = Encoding.UTF8.GetBytes(LoggerName);
+			byte[] EventNameBytes = Encoding.UTF8.GetBytes(EventName);
+
+			Writer.Write(Uid); // UID of the new event to declare
+			Writer.Write((byte) Fields.Count);
+			Writer.Write(Flags);
+			Writer.Write((byte) LoggerNameBytes.Length);
+			Writer.Write((byte) EventNameBytes.Length);
+
+			foreach (EventTypeField Field in Fields)
+			{
+				Writer.Write(Field.Offset);
+				Writer.Write(Field.Size);
+				Writer.Write(Field.TypeInfo);
+				Writer.Write(Field.NameSize);
+			}
+
+			Writer.Write(LoggerNameBytes);
+			Writer.Write(EventNameBytes);
+
+			foreach (EventTypeField Field in Fields)
+			{
+				Writer.Write(Encoding.UTF8.GetBytes(Field.Name));
+			}
+		}
+
+		public static (ushort, EventType) Deserialize(BinaryReader Reader)
+		{
+			ushort NewEventUid = Reader.ReadUInt16();
+			byte FieldCount = Reader.ReadByte();
+			byte Flags = Reader.ReadByte();
+			byte LoggerNameSize = Reader.ReadByte();
+			byte EventNameSize = Reader.ReadByte();
+			
+			EventTypeField[] Fields = new EventTypeField[FieldCount];
+
+			for (int i = 0; i < FieldCount; i++)
+			{
+				ushort Offset = Reader.ReadUInt16();
+				ushort Size = Reader.ReadUInt16();
+				byte TypeInfo = Reader.ReadByte();
+				byte NameSize = Reader.ReadByte();
+				Fields[i] = new EventTypeField(Offset, Size, TypeInfo, NameSize);
+			}
+
+			string LoggerName = Encoding.UTF8.GetString(Reader.ReadBytesStrict(LoggerNameSize));
+			string EventName = Encoding.UTF8.GetString(Reader.ReadBytesStrict(EventNameSize));
+
+			for (int i = 0; i < FieldCount; i++)
+			{
+				Fields[i].SetName(Encoding.UTF8.GetString(Reader.ReadBytesStrict(Fields[i].NameSize)));
+			}
+
+			return (NewEventUid, new EventType(NewEventUid, LoggerName, EventName, Flags, Fields.ToList()));
+		}
+	}
+}
\ No newline at end of file
diff --git a/Engine/Source/Programs/Shared/EpicGames.Tracing/UnrealInsights/Events/Trace.cs b/Engine/Source/Programs/Shared/EpicGames.Tracing/UnrealInsights/Events/Trace.cs
new file mode 100644
--- /dev/null
+++ b/Engine/Source/Programs/Shared/EpicGames.Tracing/UnrealInsights/Events/Trace.cs
@@ -0,0 +1,91 @@
+// Copyright Epic Games, Inc. All Rights Reserved.
+
+using System.IO;
+
+namespace EpicGames.Tracing.UnrealInsights.Events
+{
+	public class TraceNewTraceEvent : ITraceEvent
+	{
+		public static readonly EventType EventType;
+		
+		public ulong StartCycle { get; private set; }
+		public ulong CycleFrequency { get; private set; }
+		public ushort Endian { get; private set; }
+		public byte PointerSize { get; private set; }
+
+		public ushort Size => (ushort) (EventType.GetEventSize() + TraceImportantEventHeader.HeaderSize);
+		public EventType Type => EventType;
+		
+		public TraceNewTraceEvent(ulong StartCycle, ulong CycleFrequency, ushort Endian, byte PointerSize)
+		{
+			this.StartCycle = StartCycle;
+			this.CycleFrequency = CycleFrequency;
+			this.Endian = Endian;
+			this.PointerSize = PointerSize;
+		}
+		
+		static TraceNewTraceEvent()
+		{
+			EventType = new EventType("$Trace", "NewTrace", EventType.FlagImportant | EventType.FlagNoSync);
+			EventType.AddEventType(0, 8, EventTypeField.TypeInt64, "StartCycle");
+			EventType.AddEventType(8, 8, EventTypeField.TypeInt64, "CycleFrequency");
+			EventType.AddEventType(16, 2, EventTypeField.TypeInt16, "Endian");
+			EventType.AddEventType(18, 1, EventTypeField.TypeInt8, "PointerSize");
+		}
+
+		public void Serialize(ushort Uid, BinaryWriter Writer)
+		{
+			new TraceImportantEventHeader(Uid, EventType.GetEventSize()).Serialize(Writer);
+			Writer.Write(StartCycle);
+			Writer.Write(CycleFrequency);
+			Writer.Write(Endian);
+			Writer.Write(PointerSize);
+		}
+		
+		public static TraceNewTraceEvent Deserialize(BinaryReader Reader)
+		{
+			ulong StartCycle = Reader.ReadUInt64();
+			ulong CycleFrequency = Reader.ReadUInt64();
+			ushort Endian = Reader.ReadUInt16();
+			byte PointerSize = Reader.ReadByte();
+			return new TraceNewTraceEvent(StartCycle, CycleFrequency, Endian, PointerSize);
+		}
+	}
+	
+	public class TraceThreadInfoEvent : ITraceEvent
+	{
+		public static readonly EventType EventType;
+
+		public ushort Size => (ushort) (GenericEvent.Size + TraceImportantEventHeader.HeaderSize);
+		public EventType Type => EventType;
+		private readonly GenericEvent GenericEvent;
+		
+		public TraceThreadInfoEvent(int ThreadId, int SystemId, int SortHint, string Name)
+		{
+			GenericEvent.Field[] Fields =
+			{
+				GenericEvent.Field.FromInt((int) ThreadId),
+				GenericEvent.Field.FromInt((int) SystemId),
+				GenericEvent.Field.FromInt((int) SortHint),
+				GenericEvent.Field.FromString(Name),
+			};
+
+			GenericEvent = new GenericEvent(0, Fields, EventType);
+		}
+
+		static TraceThreadInfoEvent()
+		{
+			EventType = new EventType("$Trace", "ThreadInfo", EventType.FlagImportant | EventType.FlagMaybeHasAux | EventType.FlagNoSync);
+			EventType.AddEventType(0, 4, EventTypeField.TypeInt32, "ThreadId");
+			EventType.AddEventType(4, 4, EventTypeField.TypeInt32, "SystemId");
+			EventType.AddEventType(8, 4, EventTypeField.TypeInt32, "SortHint");
+			EventType.AddEventType(12, 0, EventTypeField.TypeAnsiString, "Name");
+		}
+
+		public void Serialize(ushort Uid, BinaryWriter Writer)
+		{
+			new TraceImportantEventHeader(Uid, GenericEvent.Size).Serialize(Writer);
+			GenericEvent.Serialize(Uid, Writer);
+		}
+	}
+}
\ No newline at end of file
diff --git a/Engine/Source/Programs/Shared/EpicGames.Tracing/UnrealInsights/Reader.cs b/Engine/Source/Programs/Shared/EpicGames.Tracing/UnrealInsights/Reader.cs
new file mode 100644
--- /dev/null
+++ b/Engine/Source/Programs/Shared/EpicGames.Tracing/UnrealInsights/Reader.cs
@@ -0,0 +1,200 @@
+// Copyright Epic Games, Inc. All Rights Reserved.
+
+using System;
+using System.Collections.Generic;
+using System.IO;
+using System.Linq;
+using EpicGames.Tracing.UnrealInsights.Events;
+
+namespace EpicGames.Tracing.UnrealInsights
+{
+	public class PredefinedEventUid
+	{
+		public const ushort NewEvent = 0;
+		public const ushort AuxData = 1;
+		// public const ushort _AuxData_Unused = 2;
+		public const ushort AuxDataTerminal = 3;
+		public const ushort EnterScope = 4;
+		public const ushort LeaveScope = 5;
+		// public const ushort _Unused6 = 6;
+		// public const ushort _Unused7 = 7;
+		public const ushort EnterScope_T = 8;
+		// public const ushort _EnterScope_T_Unused0 = 9;
+		// public const ushort _EnterScope_T_Unused1 = 10;
+		// public const ushort _EnterScope_T_Unused2 = 11;
+		public const ushort LeaveScope_T = 12;
+		// public const ushort _LeaveScope_T_Unused0 = 13;
+		// public const ushort _LeaveScope_T_Unused1 = 14;
+		// public const ushort _LeaveScope_T_Unused2 = 15;
+		public const ushort _WellKnownNum = 16;
+	}
+	
+	public class UnrealInsightsReader
+	{
+		public readonly Dictionary<ushort, EventType> EventTypes = new Dictionary<ushort, EventType>();
+		public readonly Dictionary<ushort, List<ITraceEvent>> EventsPerThread = new Dictionary<ushort, List<ITraceEvent>>();
+
+		internal int NumTransportPacketsRead { get; private set; } = 0;
+
+		public UnrealInsightsReader()
+		{
+			EventTypes[PredefinedEventUid.EnterScope_T] = new EnterScopeEventTimestamp(0).Type;
+			EventTypes[PredefinedEventUid.LeaveScope_T] = new LeaveScopeEventTimestamp(0).Type;
+		}
+
+		public void Read(Stream Stream)
+		{
+			using BinaryReader Reader = new BinaryReader(Stream);
+			StreamHeader.Deserialize(Reader);
+
+			Dictionary<ushort, MemoryStream> ThreadStreams = Demux(Reader);
+			Reader.EnsureEntireStreamIsConsumed(); // There should be no more trailing bytes once all packets are consumed.
+			
+			void ReadThreadStream(ushort ThreadId)
+			{
+				MemoryStream ThreadStream = ThreadStreams[ThreadId];
+				using BinaryReader ThreadReader = new BinaryReader(ThreadStream);
+				EventsPerThread[ThreadId] = ReadEvents(ThreadId, ThreadReader);
+				ThreadReader.EnsureEntireStreamIsConsumed();
+				ThreadStream.Dispose();
+				ThreadStreams.Remove(ThreadId);
+			}
+
+			// Read the NewEvents first to ensure event types are registered
+			ReadThreadStream(TransportPacket.ThreadIdEvents);
+			
+			// Read the important events next
+			ReadThreadStream(TransportPacket.ThreadIdImportants);
+			
+			foreach (ushort ThreadId in ThreadStreams.Keys)
+			{
+				ReadThreadStream(ThreadId);
+			}
+		}
+
+		/// <summary>
+		/// Demultiplex the stream into one continuous buffer per thread.
+		/// This avoids dealing with blocks that span multiple transport packets.
+		///
+		/// It's a naive and simple enough for our parsing purposes but it definitely won't read larger trace files.
+		/// </summary>
+		/// <param name="Reader">Reader to consume</param>
+		/// <returns>A stream of data per thread</returns>
+		public Dictionary<ushort, MemoryStream> Demux(BinaryReader Reader)
+		{
+			Dictionary<ushort, MemoryStream> ThreadStreams = new Dictionary<ushort, MemoryStream>();
+			while (Reader.BaseStream.Position < Reader.BaseStream.Length)
+			{
+				TransportPacket Packet = TransportPacket.Deserialize(Reader);
+
+				if (!ThreadStreams.TryGetValue(Packet.GetThreadId(), out MemoryStream? ThreadStream))
+				{
+					ThreadStream = new MemoryStream(5 * 1024 * 1024);
+					ThreadStreams[Packet.GetThreadId()] = ThreadStream;
+				}
+
+				ThreadStream.Write(Packet.Data);
+				NumTransportPacketsRead++;
+			}
+			Reader.EnsureEntireStreamIsConsumed();
+			
+			// Reset stream positions so it can be read from the beginning once returned
+			ThreadStreams.Values.ToList().ForEach(Ms => Ms.Position = 0);
+
+			return ThreadStreams;
+		}
+
+		public List<ITraceEvent> ReadEvents(int ThreadId, BinaryReader Reader)
+		{
+			List<ITraceEvent> EventsRead = new List<ITraceEvent>();
+			while (Reader.BaseStream.Position != Reader.BaseStream.Length)
+			{
+				ushort Uid;
+				ushort Size = 0;
+				if (ThreadId == TransportPacket.ThreadIdEvents || ThreadId == TransportPacket.ThreadIdImportants)
+				{
+					Uid = Reader.ReadUInt16();
+					Size = Reader.ReadUInt16(); // TODO: validate and respect the size
+				}
+				else
+				{
+					Uid = Reader.ReadPackedUid(out _);
+				}
+				
+				if (Uid >= PredefinedEventUid._WellKnownNum)
+				{
+					if (!EventTypes.TryGetValue(Uid, out EventType? EventType))
+					{
+						throw new Exception($"No event type registered for UID {Uid} / 0x{Uid:X4}");
+					}
+				
+					GenericEvent Event = GenericEvent.Deserialize(Uid, Reader, EventType);
+					EventsRead.Add(Event);
+				}
+				else
+				{
+					switch (Uid)
+					{
+						case PredefinedEventUid.NewEvent:
+							if (ThreadId != TransportPacket.ThreadIdEvents)
+								throw new Exception("NewEvents are only allowed on thread ID " + TransportPacket.ThreadIdEvents);
+
+							(ushort NewEventUid, EventType NewEvent) = EventType.Deserialize(Reader);
+							EventTypes[NewEventUid] = NewEvent;
+							break;
+						
+						case PredefinedEventUid.EnterScope:
+							break;
+
+						case PredefinedEventUid.LeaveScope:
+							break;
+						
+						case PredefinedEventUid.EnterScope_T:
+							Reader.BaseStream.Position -= 1; // Reset so UID can be read inside event's deserialize method
+							EventsRead.Add(EnterScopeEventTimestamp.Deserialize(Reader));
+							break;
+
+						case PredefinedEventUid.LeaveScope_T:
+							Reader.BaseStream.Position -= 1; // Reset so UID can be read inside event's deserialize method
+							EventsRead.Add(LeaveScopeEventTimestamp.Deserialize(Reader));
+							break;
+
+						default:
+							throw new Exception($"Cannot handle unknown UID {Uid}/0x{Uid:X4}");
+					}
+				}
+			}
+
+			return EventsRead;
+		}
+		
+		public Dictionary<ushort, List<ITraceEvent>> GetEventsPerUid()
+		{
+			Dictionary<ushort, List<ITraceEvent>> EventsPerUid = new Dictionary<ushort, List<ITraceEvent>>();
+			
+			foreach ((ushort _, List<ITraceEvent> ThreadEvents) in EventsPerThread)
+			{
+				foreach (ITraceEvent Event in ThreadEvents)
+				{
+					ushort Uid = EventTypes.First(x => Event.Type.Name == x.Value.Name).Key;
+					if (!EventsPerUid.TryGetValue(Uid, out List<ITraceEvent>? Events))
+					{
+						Events = new List<ITraceEvent>();
+						EventsPerUid[Uid] = Events;
+					}
+					Events.Add(Event);
+				}
+			}
+		
+			return EventsPerUid;
+		}
+
+		public void PrintEventSummary()
+		{
+			foreach ((ushort Uid, List<ITraceEvent> Events) in GetEventsPerUid())
+			{
+				Console.WriteLine($"{Uid,4} {EventTypes[Uid].Name,-45} {Events.Count}");
+			}
+		}
+	}
+}
\ No newline at end of file
diff --git a/Engine/Source/Programs/Shared/EpicGames.Tracing/UnrealInsights/StreamHeader.cs b/Engine/Source/Programs/Shared/EpicGames.Tracing/UnrealInsights/StreamHeader.cs
new file mode 100644
--- /dev/null
+++ b/Engine/Source/Programs/Shared/EpicGames.Tracing/UnrealInsights/StreamHeader.cs
@@ -0,0 +1,62 @@
+// Copyright Epic Games, Inc. All Rights Reserved.
+
+using System;
+using System.IO;
+
+namespace EpicGames.Tracing.UnrealInsights
+{
+	public class StreamHeader
+	{
+		public const uint MagicTrc = 1414677317; // TRCE
+		public const uint MagicTrc2 = 1414677298; // TRC2
+
+		public uint MagicFourCc;
+		public ushort MetadataSize;
+		public ushort MetadataField0;
+		public ushort ControlPort;
+
+		public byte TransportVersion;
+		public byte ProtocolVersion;
+
+		private StreamHeader()
+		{
+		}
+
+		public static StreamHeader Default()
+		{
+			StreamHeader Handshake = new StreamHeader();
+			Handshake.MagicFourCc = MagicTrc;
+			Handshake.TransportVersion = 3;
+			Handshake.ProtocolVersion = 5;
+			return Handshake;
+		}
+
+		public void Serialize(BinaryWriter Writer)
+		{
+			Writer.Write(MagicTrc);
+			Writer.Write(TransportVersion);
+			Writer.Write(ProtocolVersion);
+		}
+
+		public static StreamHeader Deserialize(BinaryReader Reader)
+		{
+			StreamHeader Header = new StreamHeader();
+			Header.MagicFourCc = Reader.ReadUInt32();
+			if (Header.MagicFourCc != MagicTrc && Header.MagicFourCc != MagicTrc2)
+				throw new ArgumentException("Only support magic number TRCE and TRC2");
+
+			if (Header.MagicFourCc == MagicTrc2)
+			{
+				Header.MetadataSize = Reader.ReadUInt16();
+				if (Header.MetadataSize != 4)
+					throw new ArgumentException("Only support metadata size of 4 bytes (got " + Header.MetadataSize + ")");
+				Header.MetadataField0 = Reader.ReadUInt16();
+				Header.ControlPort = Reader.ReadUInt16();
+			}
+
+			Header.TransportVersion = Reader.ReadByte();
+			Header.ProtocolVersion = Reader.ReadByte();
+			return Header;
+		}
+	}
+}
\ No newline at end of file
diff --git a/Engine/Source/Programs/Shared/EpicGames.Tracing/UnrealInsights/TraceUtils.cs b/Engine/Source/Programs/Shared/EpicGames.Tracing/UnrealInsights/TraceUtils.cs
new file mode 100644
--- /dev/null
+++ b/Engine/Source/Programs/Shared/EpicGames.Tracing/UnrealInsights/TraceUtils.cs
@@ -0,0 +1,44 @@
+// Copyright Epic Games, Inc. All Rights Reserved.
+
+using System.IO;
+
+namespace EpicGames.Tracing.UnrealInsights
+{
+	public class TraceUtils
+	{
+		public static ulong Read7BitUint(BinaryReader Reader)
+		{
+			ulong Value = 0;
+			ulong ByteIndex = 0;
+			bool HasMoreBytes;
+			
+			do
+			{
+				byte ByteValue = Reader.ReadByte();
+				HasMoreBytes = (ByteValue & 0x80) != 0;
+				Value |= (ulong)(ByteValue & 0x7f) << (int)(ByteIndex * 7);
+				++ByteIndex;
+			} while (HasMoreBytes);
+			
+			return Value;
+		}
+
+		public static int Write7BitUint(BinaryWriter Writer, ulong Value)
+		{
+			int NumBytesWritten = 0;
+			do
+			{
+				byte HasMoreBytesBit = (byte) (Value > 0x7F ? 1 : 0);
+				byte HasMoreBytes = (byte)(HasMoreBytesBit << 7);
+				byte ByteToWrite = (byte)((Value & 0x7F) | HasMoreBytes);
+				
+				Writer.Write(ByteToWrite);
+				Value >>= 7;
+				NumBytesWritten++;
+				
+			} while (Value > 0);
+
+			return NumBytesWritten;
+		}
+	}
+}
\ No newline at end of file
diff --git a/Engine/Source/Programs/Shared/EpicGames.Tracing/UnrealInsights/TransportPacket.cs b/Engine/Source/Programs/Shared/EpicGames.Tracing/UnrealInsights/TransportPacket.cs
new file mode 100644
--- /dev/null
+++ b/Engine/Source/Programs/Shared/EpicGames.Tracing/UnrealInsights/TransportPacket.cs
@@ -0,0 +1,102 @@
+// Copyright Epic Games, Inc. All Rights Reserved.
+
+using System;
+using System.Collections.Generic;
+using System.IO;
+using System.Linq;
+
+namespace EpicGames.Tracing.UnrealInsights
+{
+	// Handles FTidPacketBase, TTidPacket, TTidPacketEncoded
+	public class TransportPacket
+	{
+		public const int ThreadIdEvents = 0;
+		public const int ThreadIdInternal = 1;
+		public const int ThreadIdImportants = ThreadIdInternal;
+		
+		private const ushort EncodedMarker = 0x8000;
+		private const ushort PartialMarker = 0x4000;
+		private const ushort ThreadIdMask = PartialMarker - 1;
+		
+		public ushort PacketSize { get; private set; }
+		public ushort ThreadId { get; private set; }
+		public ushort DecodedSize { get; private set; }
+		public byte[] Data { get; private set; } = Array.Empty<byte>();
+	
+		private TransportPacket()
+		{
+		}
+		
+		public bool IsEncoded()
+		{
+			return (ThreadId & EncodedMarker) != 0;
+		}
+		
+		public bool IsPartial()
+		{
+			return (ThreadId & PartialMarker) != 0;
+		}
+		
+		public ushort GetThreadId()
+		{
+			return (ushort) (ThreadId & ThreadIdMask);
+		}
+		
+		public static bool IsNormalThread(ushort ThreadId)
+		{
+			if (ThreadId == ThreadIdEvents) return false;
+			if (ThreadId == ThreadIdInternal) return false;
+			if (ThreadId == ThreadIdImportants) return false;
+			return true;
+		}
+		
+		public void Serialize(BinaryWriter Writer, IEnumerable<(ushort, ITraceEvent)> Events)
+		{
+			ushort TotalSize = (ushort) Events.Sum(x =>
+			{
+				ITraceEvent Event = x.Item2;
+				return Event.Size;
+			});
+			TotalSize += 4; // The two uint16 writes below are included
+			
+			Writer.Write(TotalSize);
+			Writer.Write(ThreadId);
+
+			foreach ((ushort Uid, ITraceEvent Event) in Events)
+			{
+				Event.Serialize(Uid, Writer);
+			}
+		}
+
+		public override string ToString()
+		{
+			return $"TransportPacket(ThreadId={GetThreadId()} PacketSize={PacketSize} IsEncoded={IsEncoded()} IsPartial={IsPartial()})";
+		}
+
+		public static TransportPacket Create(ushort PacketSize, ushort ThreadId)
+		{
+			TransportPacket Packet = new TransportPacket();
+			Packet.PacketSize = PacketSize;
+			Packet.ThreadId = ThreadId;
+			return Packet;
+		}
+
+		public static TransportPacket Deserialize(BinaryReader Reader)
+		{
+			TransportPacket Packet = new TransportPacket();
+			Packet.PacketSize = Reader.ReadUInt16();
+			Packet.ThreadId = Reader.ReadUInt16();
+
+			int HeaderSize = sizeof(ushort) + sizeof(ushort); // PacketSize + ThreadId
+			if (Packet.IsEncoded())
+			{
+				Packet.DecodedSize = Reader.ReadUInt16();
+				HeaderSize += sizeof(ushort); // sizeof(DecodedSize) 
+			}
+
+			int BytesToRead = Packet.PacketSize - HeaderSize;
+			Packet.Data = Reader.ReadBytesStrict(BytesToRead);
+			return Packet;
+		}
+	}
+}
\ No newline at end of file
diff --git a/Engine/Source/Programs/Shared/EpicGames.Tracing/UnrealInsights/Writer.cs b/Engine/Source/Programs/Shared/EpicGames.Tracing/UnrealInsights/Writer.cs
new file mode 100644
--- /dev/null
+++ b/Engine/Source/Programs/Shared/EpicGames.Tracing/UnrealInsights/Writer.cs
@@ -0,0 +1,109 @@
+// Copyright Epic Games, Inc. All Rights Reserved.
+
+using System;
+using System.Collections.Generic;
+using System.IO;
+using System.Linq;
+using EpicGames.Tracing.UnrealInsights.Events;
+
+namespace EpicGames.Tracing.UnrealInsights
+{
+	public class UnrealInsightsWriter
+	{
+		public readonly Dictionary<EventType, ushort> EventTypeToUids = new Dictionary<EventType, ushort>();
+		public readonly Dictionary<ushort, EventType> UidToEventTypes = new Dictionary<ushort, EventType>();
+		public readonly Dictionary<ushort, List<ITraceEvent>> ThreadToEvents = new Dictionary<ushort, List<ITraceEvent>>();
+		
+		private ushort UidCounter = PredefinedEventUid._WellKnownNum; // IDs below 16 are reserved for well-known events
+		private object Lock = new object();
+
+		public UnrealInsightsWriter()
+		{
+			ThreadToEvents[TransportPacket.ThreadIdEvents] = new List<ITraceEvent>();
+
+			EnterScopeEvent EnterScopeEvent = new EnterScopeEvent();
+			RegisterEventType(PredefinedEventUid.EnterScope, EnterScopeEvent.Type);
+			
+			LeaveScopeEvent LeaveScopeEvent = new LeaveScopeEvent();
+			RegisterEventType(PredefinedEventUid.LeaveScope, LeaveScopeEvent.Type);
+			
+			EnterScopeEvent EnterScopeTimestampEvent = new EnterScopeEvent();
+			RegisterEventType(PredefinedEventUid.EnterScope_T, EnterScopeTimestampEvent.Type);
+			
+			LeaveScopeEventTimestamp LeaveScopeTimestampEvent = new LeaveScopeEventTimestamp(0);
+			RegisterEventType(PredefinedEventUid.LeaveScope_T, LeaveScopeTimestampEvent.Type);
+		}
+
+		public void AddEvent(ushort ThreadId, ITraceEvent Event)
+		{
+			if (ThreadId == TransportPacket.ThreadIdEvents)
+			{
+				throw new ArgumentException("Cannot add events directly to new-event thread");
+			}
+			
+			lock (Lock)
+			{
+				if (!ThreadToEvents.TryGetValue(ThreadId, out List<ITraceEvent>? Events))
+				{
+					Events = new List<ITraceEvent>();
+					ThreadToEvents[ThreadId] = Events;
+				}
+
+				if (!EventTypeToUids.ContainsKey(Event.Type))
+				{
+					ushort NewUid = UidCounter++;
+					RegisterEventType(NewUid, Event.Type);
+					ThreadToEvents[TransportPacket.ThreadIdEvents].Add(Event);
+				}
+				
+				Events.Add(Event);
+			}
+		}
+
+		private void RegisterEventType(ushort Uid, EventType EventType)
+		{
+			UidToEventTypes[Uid] = EventType;
+			EventTypeToUids[EventType] = Uid;
+		}
+
+		public void Write(BinaryWriter Writer)
+		{
+			lock (Lock)
+			{
+				WriteHeader(Writer);
+				WriteThread(TransportPacket.ThreadIdImportants, Writer);
+
+				foreach (ushort ThreadId in ThreadToEvents.Keys)
+				{
+					if (ThreadId == TransportPacket.ThreadIdEvents || ThreadId == TransportPacket.ThreadIdImportants)
+						continue;
+					
+					WriteThread(ThreadId, Writer);
+				}
+			}
+		}
+
+		private void WriteHeader(BinaryWriter Writer)
+		{
+			StreamHeader.Default().Serialize(Writer);
+			TransportPacket TransportPacket = TransportPacket.Create(0, 0);
+			IEnumerable<(ushort Uid, ITraceEvent Event)> NewEvents = ThreadToEvents[TransportPacket.ThreadIdEvents].Select(e =>
+			{
+				ushort NewEventUid = EventTypeToUids[e.Type];
+				ITraceEvent MyEventType = e.Type;
+				return (NewEventUid, MyEventType);
+			});
+			TransportPacket.Serialize(Writer, NewEvents);
+		}
+
+		private void WriteThread(ushort ThreadId, BinaryWriter Writer)
+		{
+			foreach (ITraceEvent Event in ThreadToEvents[ThreadId])
+			{
+				ushort Uid = EventTypeToUids[Event.Type];
+				TransportPacket TransportPacket = TransportPacket.Create(0, ThreadId);
+				TransportPacket.Serialize(Writer, new [] { (Uid, Event) });
+			}
+		}
+	}
+}
\ No newline at end of file
