BuildGraph: Rename BgScript to BgGraph, to distinguish the concrete preprocessed graph from the input un-preprocessed script.

diff --git a/Engine/Source/Programs/AutomationTool/BuildGraph/BuildGraph.cs b/Engine/Source/Programs/AutomationTool/BuildGraph/BuildGraph.cs
--- a/Engine/Source/Programs/AutomationTool/BuildGraph/BuildGraph.cs
+++ b/Engine/Source/Programs/AutomationTool/BuildGraph/BuildGraph.cs
@@ -404,7 +404,7 @@ namespace AutomationTool
 			ScriptFileName = FullScriptFile.MakeRelativeTo(Unreal.RootDirectory).Replace('\\', '/');
 
 			// Read the script from disk
-			BgScript Graph = BgScriptReader.ReadAsync(Context, ScriptFileName, Arguments, DefaultProperties, PreprocessedFileName != null, Schema, Logger, SingleNodeName).Result;
+			BgGraph Graph = BgScriptReader.ReadAsync(Context, ScriptFileName, Arguments, DefaultProperties, PreprocessedFileName != null, Schema, Logger, SingleNodeName).Result;
 			if(Graph == null)
 			{
 				return ExitCode.Error_Unknown;
@@ -591,7 +591,7 @@ namespace AutomationTool
 			// Print out all the diagnostic messages which still apply, unless we're running a step as part of a build system or just listing the contents of the file. 
 			if(SingleNode == null && (!bListOnly || bShowDiagnostics))
 			{
-				foreach(BgScriptDiagnostic Diagnostic in Graph.Diagnostics)
+				foreach(BgGraphDiagnostic Diagnostic in Graph.Diagnostics)
 				{
 					if(Diagnostic.EventType == LogEventType.Console)
 					{
@@ -659,7 +659,7 @@ namespace AutomationTool
 			return ExitCode.Success;
 		}
 
-		bool CreateTaskInstances(BgScript Graph, Dictionary<string, ScriptTaskBinding> NameToTask, Dictionary<BgTask, CustomTask> TaskInfoToTask)
+		bool CreateTaskInstances(BgGraph Graph, Dictionary<string, ScriptTaskBinding> NameToTask, Dictionary<BgTask, CustomTask> TaskInfoToTask)
 		{
 			bool bResult = true;
 			foreach (BgAgent Agent in Graph.Agents)
@@ -672,7 +672,7 @@ namespace AutomationTool
 			return bResult;
 		}
 
-		bool CreateTaskInstances(BgScript Graph, BgNode Node, Dictionary<string, ScriptTaskBinding> NameToTask, Dictionary<BgTask, CustomTask> TaskInfoToTask)
+		bool CreateTaskInstances(BgGraph Graph, BgNode Node, Dictionary<string, ScriptTaskBinding> NameToTask, Dictionary<BgTask, CustomTask> TaskInfoToTask)
 		{
 			bool bResult = true;
 			foreach (BgTask TaskInfo in Node.Tasks)
@@ -973,7 +973,7 @@ namespace AutomationTool
 		/// </summary>
 		/// <param name="Graph">The graph instance</param>
 		/// <param name="Storage">The temp storage backend which stores the shared state</param>
-		HashSet<BgNode> FindCompletedNodes(BgScript Graph, TempStorage Storage)
+		HashSet<BgNode> FindCompletedNodes(BgGraph Graph, TempStorage Storage)
 		{
 			HashSet<BgNode> CompletedNodes = new HashSet<BgNode>();
 			foreach(BgNode Node in Graph.Agents.SelectMany(x => x.Nodes))
@@ -993,7 +993,7 @@ namespace AutomationTool
 		/// <param name="Graph">The graph instance</param>
 		/// <param name="Storage">The temp storage backend which stores the shared state</param>
 		/// <returns>True if everything built successfully</returns>
-		bool BuildAllNodes(JobContext Job, BgScript Graph, Dictionary<BgTask, CustomTask> TaskInfoToTask, TempStorage Storage)
+		bool BuildAllNodes(JobContext Job, BgGraph Graph, Dictionary<BgTask, CustomTask> TaskInfoToTask, TempStorage Storage)
 		{
 			// Build a flat list of nodes to execute, in order
 			BgNode[] NodesToExecute = Graph.Agents.SelectMany(x => x.Nodes).ToArray();
@@ -1038,7 +1038,7 @@ namespace AutomationTool
 		/// <param name="Storage">The temp storage backend which stores the shared state</param>
 		/// <param name="bWithBanner">Whether to write a banner before and after this node's log output</param>
 		/// <returns>True if the node built successfully, false otherwise.</returns>
-		bool BuildNode(JobContext Job, BgScript Graph, BgNode Node, Dictionary<BgTask, CustomTask> TaskInfoToTask, TempStorage Storage, bool bWithBanner)
+		bool BuildNode(JobContext Job, BgGraph Graph, BgNode Node, Dictionary<BgTask, CustomTask> TaskInfoToTask, TempStorage Storage, bool bWithBanner)
 		{
 			DirectoryReference RootDir = new DirectoryReference(CommandUtils.CmdEnv.LocalRoot);
 
diff --git a/Engine/Source/Programs/Shared/EpicGames.BuildGraph/BgGraph.cs b/Engine/Source/Programs/Shared/EpicGames.BuildGraph/BgGraph.cs
new file mode 100644
--- /dev/null
+++ b/Engine/Source/Programs/Shared/EpicGames.BuildGraph/BgGraph.cs
@@ -0,0 +1,786 @@
+// Copyright Epic Games, Inc. All Rights Reserved.
+
+using System;
+using System.Collections.Generic;
+using System.Text;
+using System.IO;
+using System.Reflection;
+using System.Xml;
+using System.Linq;
+using System.Diagnostics;
+using EpicGames.Core;
+using Microsoft.Extensions.Logging;
+using System.Diagnostics.CodeAnalysis;
+
+namespace EpicGames.BuildGraph
+{
+	/// <summary>
+	/// Options for how the graph should be printed
+	/// </summary>
+	public enum GraphPrintOptions
+	{
+		/// <summary>
+		/// Includes a list of the graph options
+		/// </summary>
+		ShowCommandLineOptions = 0x1,
+
+		/// <summary>
+		/// Includes the list of dependencies for each node
+		/// </summary>
+		ShowDependencies = 0x2,
+
+		/// <summary>
+		/// Includes the list of notifiers for each node
+		/// </summary>
+		ShowNotifications = 0x4,
+	}
+
+	/// <summary>
+	/// Diagnostic message from the graph script. These messages are parsed at startup, then culled along with the rest of the graph nodes before output. Doing so
+	/// allows errors and warnings which are only output if a node is part of the graph being executed.
+	/// </summary>
+	public class BgGraphDiagnostic
+	{
+		/// <summary>
+		/// The diagnostic event type
+		/// </summary>
+		public LogEventType EventType;
+
+		/// <summary>
+		/// The message to display
+		/// </summary>
+		public string Message;
+
+		/// <summary>
+		/// The node which this diagnostic is declared in. If the node is culled from the graph, the message will not be displayed.
+		/// </summary>
+		public BgNode? EnclosingNode;
+
+		/// <summary>
+		/// The agent that this diagnostic is declared in. If the entire agent is culled from the graph, the message will not be displayed.
+		/// </summary>
+		public BgAgent? EnclosingAgent;
+
+		/// <summary>
+		/// Constructor
+		/// </summary>
+		public BgGraphDiagnostic(LogEventType EventType, string Message, BgNode? EnclosingNode, BgAgent? EnclosingAgent)
+		{
+			this.EventType = EventType;
+			this.Message = Message;
+			this.EnclosingNode = EnclosingNode;
+			this.EnclosingAgent = EnclosingAgent;
+		}
+	}
+
+	/// <summary>
+	/// Represents a graph option. These are expanded during preprocessing, but are retained in order to display help messages.
+	/// </summary>
+	public class BgScriptOption
+	{
+		/// <summary>
+		/// Name of this option
+		/// </summary>
+		public string Name;
+
+		/// <summary>
+		/// Description for this option
+		/// </summary>
+		public string Description;
+
+		/// <summary>
+		/// Default value for this option
+		/// </summary>
+		public string DefaultValue;
+
+		/// <summary>
+		/// Constructor
+		/// </summary>
+		/// <param name="Name">The name of this option</param>
+		/// <param name="Description">Description of the option, for display on help pages</param>
+		/// <param name="DefaultValue">Default value for the option</param>
+		public BgScriptOption(string Name, string Description, string DefaultValue)
+		{
+			this.Name = Name;
+			this.Description = Description;
+			this.DefaultValue = DefaultValue;
+		}
+
+		/// <summary>
+		/// Returns a name of this option for debugging
+		/// </summary>
+		/// <returns>Name of the option</returns>
+		public override string ToString()
+		{
+			return Name;
+		}
+	}
+
+	/// <summary>
+	/// Definition of a graph.
+	/// </summary>
+	public class BgGraph
+	{
+		/// <summary>
+		/// List of options, in the order they were specified
+		/// </summary>
+		public List<BgScriptOption> Options = new List<BgScriptOption>();
+
+		/// <summary>
+		/// List of agents containing nodes to execute
+		/// </summary>
+		public List<BgAgent> Agents = new List<BgAgent>();
+
+		/// <summary>
+		/// Mapping from name to agent
+		/// </summary>
+		public Dictionary<string, BgAgent> NameToAgent = new Dictionary<string, BgAgent>(StringComparer.OrdinalIgnoreCase);
+
+		/// <summary>
+		/// Mapping of names to the corresponding node.
+		/// </summary>
+		public Dictionary<string, BgNode> NameToNode = new Dictionary<string, BgNode>(StringComparer.OrdinalIgnoreCase);
+
+		/// <summary>
+		/// Mapping of names to the corresponding report.
+		/// </summary>
+		public Dictionary<string, BgReport> NameToReport = new Dictionary<string, BgReport>(StringComparer.OrdinalIgnoreCase);
+
+		/// <summary>
+		/// Mapping of names to their corresponding node output.
+		/// </summary>
+		public Dictionary<string, BgNodeOutput> TagNameToNodeOutput = new Dictionary<string, BgNodeOutput>(StringComparer.OrdinalIgnoreCase);
+
+		/// <summary>
+		/// Mapping of aggregate names to their respective nodes
+		/// </summary>
+		public Dictionary<string, BgAggregate> NameToAggregate = new Dictionary<string, BgAggregate>(StringComparer.OrdinalIgnoreCase);
+
+		/// <summary>
+		/// List of badges that can be displayed for this build
+		/// </summary>
+		public List<BgBadge> Badges = new List<BgBadge>();
+
+		/// <summary>
+		/// List of labels that can be displayed for this build
+		/// </summary>
+		public List<BgLabel> Labels = new List<BgLabel>();
+
+		/// <summary>
+		/// Diagnostic messages for this graph
+		/// </summary>
+		public List<BgGraphDiagnostic> Diagnostics = new List<BgGraphDiagnostic>();
+
+		/// <summary>
+		/// Default constructor
+		/// </summary>
+		public BgGraph()
+		{
+		}
+
+		/// <summary>
+		/// Checks whether a given name already exists
+		/// </summary>
+		/// <param name="Name">The name to check.</param>
+		/// <returns>True if the name exists, false otherwise.</returns>
+		public bool ContainsName(string Name)
+		{
+			return NameToNode.ContainsKey(Name) || NameToReport.ContainsKey(Name) || NameToAggregate.ContainsKey(Name);
+		}
+
+		/// <summary>
+		/// Tries to resolve the given name to one or more nodes. Checks for aggregates, and actual nodes.
+		/// </summary>
+		/// <param name="Name">The name to search for</param>
+		/// <param name="OutNodes">If the name is a match, receives an array of nodes and their output names</param>
+		/// <returns>True if the name was found, false otherwise.</returns>
+		public bool TryResolveReference(string Name, [NotNullWhen(true)] out BgNode[]? OutNodes)
+		{
+			// Check if it's a tag reference or node reference
+			if (Name.StartsWith("#"))
+			{
+				// Check if it's a regular node or output name
+				BgNodeOutput? Output;
+				if (TagNameToNodeOutput.TryGetValue(Name, out Output))
+				{
+					OutNodes = new BgNode[] { Output.ProducingNode };
+					return true;
+				}
+			}
+			else
+			{
+				// Check if it's a regular node or output name
+				BgNode? Node;
+				if (NameToNode.TryGetValue(Name, out Node))
+				{
+					OutNodes = new BgNode[] { Node };
+					return true;
+				}
+
+				// Check if it's an aggregate name
+				BgAggregate? Aggregate;
+				if (NameToAggregate.TryGetValue(Name, out Aggregate))
+				{
+					OutNodes = Aggregate.RequiredNodes.ToArray();
+					return true;
+				}
+
+				// Check if it's a group name
+				BgAgent? Agent;
+				if (NameToAgent.TryGetValue(Name, out Agent))
+				{
+					OutNodes = Agent.Nodes.ToArray();
+					return true;
+				}
+			}
+
+			// Otherwise fail
+			OutNodes = null;
+			return false;
+		}
+
+		/// <summary>
+		/// Tries to resolve the given name to one or more node outputs. Checks for aggregates, and actual nodes.
+		/// </summary>
+		/// <param name="Name">The name to search for</param>
+		/// <param name="OutOutputs">If the name is a match, receives an array of nodes and their output names</param>
+		/// <returns>True if the name was found, false otherwise.</returns>
+		public bool TryResolveInputReference(string Name, [NotNullWhen(true)] out BgNodeOutput[]? OutOutputs)
+		{
+			// Check if it's a tag reference or node reference
+			if (Name.StartsWith("#"))
+			{
+				// Check if it's a regular node or output name
+				BgNodeOutput? Output;
+				if (TagNameToNodeOutput.TryGetValue(Name, out Output))
+				{
+					OutOutputs = new BgNodeOutput[] { Output };
+					return true;
+				}
+			}
+			else
+			{
+				// Check if it's a regular node or output name
+				BgNode? Node;
+				if (NameToNode.TryGetValue(Name, out Node))
+				{
+					OutOutputs = Node.Outputs.Union(Node.Inputs).ToArray();
+					return true;
+				}
+
+				// Check if it's an aggregate name
+				BgAggregate? Aggregate;
+				if (NameToAggregate.TryGetValue(Name, out Aggregate))
+				{
+					OutOutputs = Aggregate.RequiredNodes.SelectMany(x => x.Outputs.Union(x.Inputs)).Distinct().ToArray();
+					return true;
+				}
+			}
+
+			// Otherwise fail
+			OutOutputs = null;
+			return false;
+		}
+
+		/// <summary>
+		/// Cull the graph to only include the given nodes and their dependencies
+		/// </summary>
+		/// <param name="TargetNodes">A set of target nodes to build</param>
+		public void Select(IEnumerable<BgNode> TargetNodes)
+		{
+			// Find this node and all its dependencies
+			HashSet<BgNode> RetainNodes = new HashSet<BgNode>(TargetNodes);
+			foreach (BgNode TargetNode in TargetNodes)
+			{
+				RetainNodes.UnionWith(TargetNode.InputDependencies);
+			}
+
+			// Remove all the nodes which are not marked to be kept
+			foreach (BgAgent Agent in Agents)
+			{
+				Agent.Nodes = Agent.Nodes.Where(x => RetainNodes.Contains(x)).ToList();
+			}
+
+			// Remove all the empty agents
+			Agents.RemoveAll(x => x.Nodes.Count == 0);
+
+			// Trim down the list of nodes for each report to the ones that are being built
+			foreach (BgReport Report in NameToReport.Values)
+			{
+				Report.Nodes.RemoveWhere(x => !RetainNodes.Contains(x));
+			}
+
+			// Remove all the empty reports
+			NameToReport = NameToReport.Where(x => x.Value.Nodes.Count > 0).ToDictionary(Pair => Pair.Key, Pair => Pair.Value, StringComparer.InvariantCultureIgnoreCase);
+
+			// Remove all the order dependencies which are no longer part of the graph. Since we don't need to build them, we don't need to wait for them
+			foreach (BgNode Node in RetainNodes)
+			{
+				Node.OrderDependencies = Node.OrderDependencies.Where(x => RetainNodes.Contains(x)).ToArray();
+			}
+
+			// Create a new list of aggregates for everything that's left
+			Dictionary<string, BgAggregate> NewNameToAggregate = new Dictionary<string, BgAggregate>(NameToAggregate.Comparer);
+			foreach (BgAggregate Aggregate in NameToAggregate.Values)
+			{
+				if (Aggregate.RequiredNodes.All(x => RetainNodes.Contains(x)))
+				{
+					NewNameToAggregate[Aggregate.Name] = Aggregate;
+				}
+			}
+			NameToAggregate = NewNameToAggregate;
+
+			// Remove any labels that are no longer value
+			foreach (BgLabel Label in Labels)
+			{
+				Label.RequiredNodes.RemoveWhere(x => !RetainNodes.Contains(x));
+				Label.IncludedNodes.RemoveWhere(x => !RetainNodes.Contains(x));
+			}
+			Labels.RemoveAll(x => x.RequiredNodes.Count == 0);
+
+			// Remove any badges which do not have all their dependencies
+			Badges.RemoveAll(x => x.Nodes.Any(y => !RetainNodes.Contains(y)));
+
+			// Remove any diagnostics which are no longer part of the graph
+			Diagnostics.RemoveAll(x => (x.EnclosingNode != null && !RetainNodes.Contains(x.EnclosingNode)) || (x.EnclosingAgent != null && !Agents.Contains(x.EnclosingAgent)));
+		}
+
+		/// <summary>
+		/// Writes a preprocessed build graph to a script file
+		/// </summary>
+		/// <param name="File">The file to load</param>
+		/// <param name="SchemaFile">Schema file for validation</param>
+		public void Write(FileReference File, FileReference SchemaFile)
+		{
+			XmlWriterSettings Settings = new XmlWriterSettings();
+			Settings.Indent = true;
+			Settings.IndentChars = "\t";
+
+			using (XmlWriter Writer = XmlWriter.Create(File.FullName, Settings))
+			{
+				Writer.WriteStartElement("BuildGraph", "http://www.epicgames.com/BuildGraph");
+
+				if (SchemaFile != null)
+				{
+					Writer.WriteAttributeString("schemaLocation", "http://www.w3.org/2001/XMLSchema-instance", "http://www.epicgames.com/BuildGraph " + SchemaFile.MakeRelativeTo(File.Directory));
+				}
+
+				foreach (BgAgent Agent in Agents)
+				{
+					Agent.Write(Writer);
+				}
+
+				foreach (BgAggregate Aggregate in NameToAggregate.Values)
+				{
+					// If the aggregate has no required elements, skip it.
+					if (Aggregate.RequiredNodes.Count == 0)
+					{
+						continue;
+					}
+
+					Writer.WriteStartElement("Aggregate");
+					Writer.WriteAttributeString("Name", Aggregate.Name);
+					Writer.WriteAttributeString("Requires", String.Join(";", Aggregate.RequiredNodes.Select(x => x.Name)));
+					Writer.WriteEndElement();
+				}
+
+				foreach (BgLabel Label in Labels)
+				{
+					Writer.WriteStartElement("Label");
+					if (Label.DashboardCategory != null)
+					{
+						Writer.WriteAttributeString("Category", Label.DashboardCategory);
+					}
+					Writer.WriteAttributeString("Name", Label.DashboardName);
+					Writer.WriteAttributeString("Requires", String.Join(";", Label.RequiredNodes.Select(x => x.Name)));
+
+					HashSet<BgNode> IncludedNodes = new HashSet<BgNode>(Label.IncludedNodes);
+					IncludedNodes.ExceptWith(Label.IncludedNodes.SelectMany(x => x.InputDependencies));
+					IncludedNodes.ExceptWith(Label.RequiredNodes);
+					if (IncludedNodes.Count > 0)
+					{
+						Writer.WriteAttributeString("Include", String.Join(";", IncludedNodes.Select(x => x.Name)));
+					}
+
+					HashSet<BgNode> ExcludedNodes = new HashSet<BgNode>(Label.IncludedNodes);
+					ExcludedNodes.UnionWith(Label.IncludedNodes.SelectMany(x => x.InputDependencies));
+					ExcludedNodes.ExceptWith(Label.IncludedNodes);
+					ExcludedNodes.ExceptWith(ExcludedNodes.ToArray().SelectMany(x => x.InputDependencies));
+					if (ExcludedNodes.Count > 0)
+					{
+						Writer.WriteAttributeString("Exclude", String.Join(";", ExcludedNodes.Select(x => x.Name)));
+					}
+					Writer.WriteEndElement();
+				}
+
+				foreach (BgReport Report in NameToReport.Values)
+				{
+					Writer.WriteStartElement("Report");
+					Writer.WriteAttributeString("Name", Report.Name);
+					Writer.WriteAttributeString("Requires", String.Join(";", Report.Nodes.Select(x => x.Name)));
+					Writer.WriteEndElement();
+				}
+
+				foreach (BgBadge Badge in Badges)
+				{
+					Writer.WriteStartElement("Badge");
+					Writer.WriteAttributeString("Name", Badge.Name);
+					if (Badge.Project != null)
+					{
+						Writer.WriteAttributeString("Project", Badge.Project);
+					}
+					if (Badge.Change != 0)
+					{
+						Writer.WriteAttributeString("Change", Badge.Change.ToString());
+					}
+					Writer.WriteAttributeString("Requires", String.Join(";", Badge.Nodes.Select(x => x.Name)));
+					Writer.WriteEndElement();
+				}
+
+				Writer.WriteEndElement();
+			}
+		}
+
+		/// <summary>
+		/// Export the build graph to a Json file, for parallel execution by the build system
+		/// </summary>
+		/// <param name="File">Output file to write</param>
+		/// <param name="CompletedNodes">Set of nodes which have been completed</param>
+		public void Export(FileReference File, HashSet<BgNode> CompletedNodes)
+		{
+			// Find all the nodes which we're actually going to execute. We'll use this to filter the graph.
+			HashSet<BgNode> NodesToExecute = new HashSet<BgNode>();
+			foreach (BgNode Node in Agents.SelectMany(x => x.Nodes))
+			{
+				if (!CompletedNodes.Contains(Node))
+				{
+					NodesToExecute.Add(Node);
+				}
+			}
+
+			// Open the output file
+			using (JsonWriter JsonWriter = new JsonWriter(File.FullName))
+			{
+				JsonWriter.WriteObjectStart();
+
+				// Write all the agents
+				JsonWriter.WriteArrayStart("Groups");
+				foreach (BgAgent Agent in Agents)
+				{
+					BgNode[] Nodes = Agent.Nodes.Where(x => NodesToExecute.Contains(x)).ToArray();
+					if (Nodes.Length > 0)
+					{
+						JsonWriter.WriteObjectStart();
+						JsonWriter.WriteValue("Name", Agent.Name);
+						JsonWriter.WriteArrayStart("Agent Types");
+						foreach (string AgentType in Agent.PossibleTypes)
+						{
+							JsonWriter.WriteValue(AgentType);
+						}
+						JsonWriter.WriteArrayEnd();
+						JsonWriter.WriteArrayStart("Nodes");
+						foreach (BgNode Node in Nodes)
+						{
+							JsonWriter.WriteObjectStart();
+							JsonWriter.WriteValue("Name", Node.Name);
+							JsonWriter.WriteValue("DependsOn", String.Join(";", Node.GetDirectOrderDependencies().Where(x => NodesToExecute.Contains(x))));
+							JsonWriter.WriteValue("RunEarly", Node.bRunEarly);
+							JsonWriter.WriteObjectStart("Notify");
+							JsonWriter.WriteValue("Default", String.Join(";", Node.NotifyUsers));
+							JsonWriter.WriteValue("Submitters", String.Join(";", Node.NotifySubmitters));
+							JsonWriter.WriteValue("Warnings", Node.bNotifyOnWarnings);
+							JsonWriter.WriteObjectEnd();
+							JsonWriter.WriteObjectEnd();
+						}
+						JsonWriter.WriteArrayEnd();
+						JsonWriter.WriteObjectEnd();
+					}
+				}
+				JsonWriter.WriteArrayEnd();
+
+				// Write all the badges
+				JsonWriter.WriteArrayStart("Badges");
+				foreach (BgBadge Badge in Badges)
+				{
+					BgNode[] Dependencies = Badge.Nodes.Where(x => NodesToExecute.Contains(x)).ToArray();
+					if (Dependencies.Length > 0)
+					{
+						// Reduce that list to the smallest subset of direct dependencies
+						HashSet<BgNode> DirectDependencies = new HashSet<BgNode>(Dependencies);
+						foreach (BgNode Dependency in Dependencies)
+						{
+							DirectDependencies.ExceptWith(Dependency.OrderDependencies);
+						}
+
+						JsonWriter.WriteObjectStart();
+						JsonWriter.WriteValue("Name", Badge.Name);
+						if (!String.IsNullOrEmpty(Badge.Project))
+						{
+							JsonWriter.WriteValue("Project", Badge.Project);
+						}
+						if (Badge.Change != 0)
+						{
+							JsonWriter.WriteValue("Change", Badge.Change);
+						}
+						JsonWriter.WriteValue("AllDependencies", String.Join(";", Agents.SelectMany(x => x.Nodes).Where(x => Dependencies.Contains(x)).Select(x => x.Name)));
+						JsonWriter.WriteValue("DirectDependencies", String.Join(";", DirectDependencies.Select(x => x.Name)));
+						JsonWriter.WriteObjectEnd();
+					}
+				}
+				JsonWriter.WriteArrayEnd();
+
+				// Write all the triggers and reports. 
+				JsonWriter.WriteArrayStart("Reports");
+				foreach (BgReport Report in NameToReport.Values)
+				{
+					BgNode[] Dependencies = Report.Nodes.Where(x => NodesToExecute.Contains(x)).ToArray();
+					if (Dependencies.Length > 0)
+					{
+						// Reduce that list to the smallest subset of direct dependencies
+						HashSet<BgNode> DirectDependencies = new HashSet<BgNode>(Dependencies);
+						foreach (BgNode Dependency in Dependencies)
+						{
+							DirectDependencies.ExceptWith(Dependency.OrderDependencies);
+						}
+
+						JsonWriter.WriteObjectStart();
+						JsonWriter.WriteValue("Name", Report.Name);
+						JsonWriter.WriteValue("AllDependencies", String.Join(";", Agents.SelectMany(x => x.Nodes).Where(x => Dependencies.Contains(x)).Select(x => x.Name)));
+						JsonWriter.WriteValue("DirectDependencies", String.Join(";", DirectDependencies.Select(x => x.Name)));
+						JsonWriter.WriteValue("Notify", String.Join(";", Report.NotifyUsers));
+						JsonWriter.WriteValue("IsTrigger", false);
+						JsonWriter.WriteObjectEnd();
+					}
+				}
+				JsonWriter.WriteArrayEnd();
+
+				JsonWriter.WriteObjectEnd();
+			}
+		}
+
+		/// <summary>
+		/// Export the build graph to a Json file for parsing by Horde
+		/// </summary>
+		/// <param name="File">Output file to write</param>
+		public void ExportForHorde(FileReference File)
+		{
+			DirectoryReference.CreateDirectory(File.Directory);
+			using (JsonWriter JsonWriter = new JsonWriter(File.FullName))
+			{
+				JsonWriter.WriteObjectStart();
+				JsonWriter.WriteArrayStart("Groups");
+				foreach (BgAgent Agent in Agents)
+				{
+					JsonWriter.WriteObjectStart();
+					JsonWriter.WriteArrayStart("Types");
+					foreach (string PossibleType in Agent.PossibleTypes)
+					{
+						JsonWriter.WriteValue(PossibleType);
+					}
+					JsonWriter.WriteArrayEnd();
+					JsonWriter.WriteArrayStart("Nodes");
+					foreach (BgNode Node in Agent.Nodes)
+					{
+						JsonWriter.WriteObjectStart();
+						JsonWriter.WriteValue("Name", Node.Name);
+						JsonWriter.WriteValue("RunEarly", Node.bRunEarly);
+						JsonWriter.WriteValue("Warnings", Node.bNotifyOnWarnings);
+
+						JsonWriter.WriteArrayStart("InputDependencies");
+						foreach (string InputDependency in Node.GetDirectInputDependencies().Select(x => x.Name))
+						{
+							JsonWriter.WriteValue(InputDependency);
+						}
+						JsonWriter.WriteArrayEnd();
+
+						JsonWriter.WriteArrayStart("OrderDependencies");
+						foreach (string OrderDependency in Node.GetDirectOrderDependencies().Select(x => x.Name))
+						{
+							JsonWriter.WriteValue(OrderDependency);
+						}
+						JsonWriter.WriteArrayEnd();
+
+						JsonWriter.WriteObjectEnd();
+					}
+					JsonWriter.WriteArrayEnd();
+					JsonWriter.WriteObjectEnd();
+				}
+				JsonWriter.WriteArrayEnd();
+
+				JsonWriter.WriteArrayStart("Aggregates");
+				foreach (BgAggregate Aggregate in NameToAggregate.Values)
+				{
+					JsonWriter.WriteObjectStart();
+					JsonWriter.WriteValue("Name", Aggregate.Name);
+					JsonWriter.WriteArrayStart("Nodes");
+					foreach (BgNode RequiredNode in Aggregate.RequiredNodes.OrderBy(x => x.Name))
+					{
+						JsonWriter.WriteValue(RequiredNode.Name);
+					}
+					JsonWriter.WriteArrayEnd();
+					JsonWriter.WriteObjectEnd();
+				}
+				JsonWriter.WriteArrayEnd();
+
+				JsonWriter.WriteArrayStart("Labels");
+				foreach (BgLabel Label in Labels)
+				{
+					JsonWriter.WriteObjectStart();
+					if (!String.IsNullOrEmpty(Label.DashboardName))
+					{
+						JsonWriter.WriteValue("Name", Label.DashboardName);
+					}
+					if (!String.IsNullOrEmpty(Label.DashboardCategory))
+					{
+						JsonWriter.WriteValue("Category", Label.DashboardCategory);
+					}
+					if (!String.IsNullOrEmpty(Label.UgsBadge))
+					{
+						JsonWriter.WriteValue("UgsBadge", Label.UgsBadge);
+					}
+					if (!String.IsNullOrEmpty(Label.UgsProject))
+					{
+						JsonWriter.WriteValue("UgsProject", Label.UgsProject);
+					}
+					if (Label.Change != BgLabelChange.Current)
+					{
+						JsonWriter.WriteValue("Change", Label.Change.ToString());
+					}
+
+					JsonWriter.WriteArrayStart("RequiredNodes");
+					foreach (BgNode RequiredNode in Label.RequiredNodes.OrderBy(x => x.Name))
+					{
+						JsonWriter.WriteValue(RequiredNode.Name);
+					}
+					JsonWriter.WriteArrayEnd();
+					JsonWriter.WriteArrayStart("IncludedNodes");
+					foreach (BgNode IncludedNode in Label.IncludedNodes.OrderBy(x => x.Name))
+					{
+						JsonWriter.WriteValue(IncludedNode.Name);
+					}
+					JsonWriter.WriteArrayEnd();
+					JsonWriter.WriteObjectEnd();
+				}
+				JsonWriter.WriteArrayEnd();
+
+				JsonWriter.WriteArrayStart("Badges");
+				foreach (BgBadge Badge in Badges)
+				{
+					HashSet<BgNode> Dependencies = Badge.Nodes;
+					if (Dependencies.Count > 0)
+					{
+						// Reduce that list to the smallest subset of direct dependencies
+						HashSet<BgNode> DirectDependencies = new HashSet<BgNode>(Dependencies);
+						foreach (BgNode Dependency in Dependencies)
+						{
+							DirectDependencies.ExceptWith(Dependency.OrderDependencies);
+						}
+
+						JsonWriter.WriteObjectStart();
+						JsonWriter.WriteValue("Name", Badge.Name);
+						if (!String.IsNullOrEmpty(Badge.Project))
+						{
+							JsonWriter.WriteValue("Project", Badge.Project);
+						}
+						if (Badge.Change != 0)
+						{
+							JsonWriter.WriteValue("Change", Badge.Change);
+						}
+						JsonWriter.WriteValue("Dependencies", String.Join(";", DirectDependencies.Select(x => x.Name)));
+						JsonWriter.WriteObjectEnd();
+					}
+				}
+				JsonWriter.WriteArrayEnd();
+
+				JsonWriter.WriteObjectEnd();
+			}
+		}
+
+		/// <summary>
+		/// Print the contents of the graph
+		/// </summary>
+		/// <param name="CompletedNodes">Set of nodes which are already complete</param>
+		/// <param name="PrintOptions">Options for how to print the graph</param>
+		/// <param name="Logger"></param>
+		public void Print(HashSet<BgNode> CompletedNodes, GraphPrintOptions PrintOptions, ILogger Logger)
+		{
+			// Print the options
+			if ((PrintOptions & GraphPrintOptions.ShowCommandLineOptions) != 0)
+			{
+				// Get the list of messages
+				List<KeyValuePair<string, string>> Parameters = new List<KeyValuePair<string, string>>();
+				foreach (BgScriptOption Option in Options)
+				{
+					string Name = String.Format("-set:{0}=...", Option.Name);
+
+					StringBuilder Description = new StringBuilder(Option.Description);
+					if (!String.IsNullOrEmpty(Option.DefaultValue))
+					{
+						Description.AppendFormat(" (Default: {0})", Option.DefaultValue);
+					}
+
+					Parameters.Add(new KeyValuePair<string, string>(Name, Description.ToString()));
+				}
+
+				// Format them to the log
+				if (Parameters.Count > 0)
+				{
+					Logger.LogInformation("");
+					Logger.LogInformation("Options:");
+					Logger.LogInformation("");
+					HelpUtils.PrintTable(Parameters, 4, 24, Logger);
+				}
+			}
+
+			// Output all the triggers in order
+			Logger.LogInformation("");
+			Logger.LogInformation("Graph:");
+			foreach (BgAgent Agent in Agents)
+			{
+				Logger.LogInformation("        Agent: {0} ({1})", Agent.Name, String.Join(";", Agent.PossibleTypes));
+				foreach (BgNode Node in Agent.Nodes)
+				{
+					Logger.LogInformation("            Node: {0}{1}", Node.Name, CompletedNodes.Contains(Node) ? " (completed)" : Node.bRunEarly ? " (early)" : "");
+					if (PrintOptions.HasFlag(GraphPrintOptions.ShowDependencies))
+					{
+						HashSet<BgNode> InputDependencies = new HashSet<BgNode>(Node.GetDirectInputDependencies());
+						foreach (BgNode InputDependency in InputDependencies)
+						{
+							Logger.LogInformation("                input> {0}", InputDependency.Name);
+						}
+						HashSet<BgNode> OrderDependencies = new HashSet<BgNode>(Node.GetDirectOrderDependencies());
+						foreach (BgNode OrderDependency in OrderDependencies.Except(InputDependencies))
+						{
+							Logger.LogInformation("                after> {0}", OrderDependency.Name);
+						}
+					}
+					if (PrintOptions.HasFlag(GraphPrintOptions.ShowNotifications))
+					{
+						string Label = Node.bNotifyOnWarnings ? "warnings" : "errors";
+						foreach (string User in Node.NotifyUsers)
+						{
+							Logger.LogInformation("                {0}> {1}", Label, User);
+						}
+						foreach (string Submitter in Node.NotifySubmitters)
+						{
+							Logger.LogInformation("                {0}> submitters to {1}", Label, Submitter);
+						}
+					}
+				}
+			}
+			Logger.LogInformation("");
+
+			// Print out all the non-empty aggregates
+			BgAggregate[] Aggregates = NameToAggregate.Values.OrderBy(x => x.Name).ToArray();
+			if (Aggregates.Length > 0)
+			{
+				Logger.LogInformation("Aggregates:");
+				foreach (string AggregateName in Aggregates.Select(x => x.Name))
+				{
+					Logger.LogInformation("    {0}", AggregateName);
+				}
+				Logger.LogInformation("");
+			}
+		}
+	}
+}
diff --git a/Engine/Source/Programs/Shared/EpicGames.BuildGraph/BgScript.cs b/Engine/Source/Programs/Shared/EpicGames.BuildGraph/BgScript.cs
--- a/Engine/Source/Programs/Shared/EpicGames.BuildGraph/BgScript.cs
+++ b/Engine/Source/Programs/Shared/EpicGames.BuildGraph/BgScript.cs
@@ -14,773 +14,4 @@ using System.Diagnostics.CodeAnalysis;
 
 namespace EpicGames.BuildGraph
 {
-	/// <summary>
-	/// Options for how the graph should be printed
-	/// </summary>
-	public enum GraphPrintOptions
-	{
-		/// <summary>
-		/// Includes a list of the graph options
-		/// </summary>
-		ShowCommandLineOptions = 0x1,
-
-		/// <summary>
-		/// Includes the list of dependencies for each node
-		/// </summary>
-		ShowDependencies = 0x2,
-
-		/// <summary>
-		/// Includes the list of notifiers for each node
-		/// </summary>
-		ShowNotifications = 0x4,
-	}
-
-	/// <summary>
-	/// Diagnostic message from the graph script. These messages are parsed at startup, then culled along with the rest of the graph nodes before output. Doing so
-	/// allows errors and warnings which are only output if a node is part of the graph being executed.
-	/// </summary>
-	public class BgScriptDiagnostic
-	{
-		/// <summary>
-		/// The diagnostic event type
-		/// </summary>
-		public LogEventType EventType;
-
-		/// <summary>
-		/// The message to display
-		/// </summary>
-		public string Message;
-
-		/// <summary>
-		/// The node which this diagnostic is declared in. If the node is culled from the graph, the message will not be displayed.
-		/// </summary>
-		public BgNode? EnclosingNode;
-
-		/// <summary>
-		/// The agent that this diagnostic is declared in. If the entire agent is culled from the graph, the message will not be displayed.
-		/// </summary>
-		public BgAgent? EnclosingAgent;
-
-		/// <summary>
-		/// Constructor
-		/// </summary>
-		public BgScriptDiagnostic(LogEventType EventType, string Message, BgNode? EnclosingNode, BgAgent? EnclosingAgent)
-		{
-			this.EventType = EventType;
-			this.Message = Message;
-			this.EnclosingNode = EnclosingNode;
-			this.EnclosingAgent = EnclosingAgent;
-		}
-	}
-
-	/// <summary>
-	/// Represents a graph option. These are expanded during preprocessing, but are retained in order to display help messages.
-	/// </summary>
-	public class BgScriptOption
-	{
-		/// <summary>
-		/// Name of this option
-		/// </summary>
-		public string Name;
-
-		/// <summary>
-		/// Description for this option
-		/// </summary>
-		public string Description;
-
-		/// <summary>
-		/// Default value for this option
-		/// </summary>
-		public string DefaultValue;
-
-		/// <summary>
-		/// Constructor
-		/// </summary>
-		/// <param name="Name">The name of this option</param>
-		/// <param name="Description">Description of the option, for display on help pages</param>
-		/// <param name="DefaultValue">Default value for the option</param>
-		public BgScriptOption(string Name, string Description, string DefaultValue)
-		{
-			this.Name = Name;
-			this.Description = Description;
-			this.DefaultValue = DefaultValue;
-		}
-
-		/// <summary>
-		/// Returns a name of this option for debugging
-		/// </summary>
-		/// <returns>Name of the option</returns>
-		public override string ToString()
-		{
-			return Name;
-		}
-	}
-
-	/// <summary>
-	/// Definition of a graph.
-	/// </summary>
-	public class BgScript
-	{
-		/// <summary>
-		/// List of options, in the order they were specified
-		/// </summary>
-		public List<BgScriptOption> Options = new List<BgScriptOption>();
-
-		/// <summary>
-		/// List of agents containing nodes to execute
-		/// </summary>
-		public List<BgAgent> Agents = new List<BgAgent>();
-
-		/// <summary>
-		/// Mapping from name to agent
-		/// </summary>
-		public Dictionary<string, BgAgent> NameToAgent = new Dictionary<string, BgAgent>(StringComparer.OrdinalIgnoreCase);
-
-		/// <summary>
-		/// Mapping of names to the corresponding node.
-		/// </summary>
-		public Dictionary<string, BgNode> NameToNode = new Dictionary<string,BgNode>(StringComparer.OrdinalIgnoreCase);
-
-		/// <summary>
-		/// Mapping of names to the corresponding report.
-		/// </summary>
-		public Dictionary<string, BgReport> NameToReport = new Dictionary<string, BgReport>(StringComparer.OrdinalIgnoreCase);
-
-		/// <summary>
-		/// Mapping of names to their corresponding node output.
-		/// </summary>
-		public Dictionary<string, BgNodeOutput> TagNameToNodeOutput = new Dictionary<string,BgNodeOutput>(StringComparer.OrdinalIgnoreCase);
-
-		/// <summary>
-		/// Mapping of aggregate names to their respective nodes
-		/// </summary>
-		public Dictionary<string, BgAggregate> NameToAggregate = new Dictionary<string, BgAggregate>(StringComparer.OrdinalIgnoreCase);
-
-		/// <summary>
-		/// List of badges that can be displayed for this build
-		/// </summary>
-		public List<BgBadge> Badges = new List<BgBadge>();
-
-		/// <summary>
-		/// List of labels that can be displayed for this build
-		/// </summary>
-		public List<BgLabel> Labels = new List<BgLabel>();
-
-		/// <summary>
-		/// Diagnostic messages for this graph
-		/// </summary>
-		public List<BgScriptDiagnostic> Diagnostics = new List<BgScriptDiagnostic>();
-
-		/// <summary>
-		/// Default constructor
-		/// </summary>
-		public BgScript()
-		{
-		}
-
-		/// <summary>
-		/// Checks whether a given name already exists
-		/// </summary>
-		/// <param name="Name">The name to check.</param>
-		/// <returns>True if the name exists, false otherwise.</returns>
-		public bool ContainsName(string Name)
-		{
-			return NameToNode.ContainsKey(Name) || NameToReport.ContainsKey(Name) || NameToAggregate.ContainsKey(Name);
-		}
-
-		/// <summary>
-		/// Tries to resolve the given name to one or more nodes. Checks for aggregates, and actual nodes.
-		/// </summary>
-		/// <param name="Name">The name to search for</param>
-		/// <param name="OutNodes">If the name is a match, receives an array of nodes and their output names</param>
-		/// <returns>True if the name was found, false otherwise.</returns>
-		public bool TryResolveReference(string Name, [NotNullWhen(true)] out BgNode[]? OutNodes)
-		{
-			// Check if it's a tag reference or node reference
-			if(Name.StartsWith("#"))
-			{
-				// Check if it's a regular node or output name
-				BgNodeOutput? Output;
-				if(TagNameToNodeOutput.TryGetValue(Name, out Output))
-				{
-					OutNodes = new BgNode[]{ Output.ProducingNode };
-					return true;
-				}
-			}
-			else
-			{
-				// Check if it's a regular node or output name
-				BgNode? Node;
-				if(NameToNode.TryGetValue(Name, out Node))
-				{
-					OutNodes = new BgNode[]{ Node };
-					return true;
-				}
-
-				// Check if it's an aggregate name
-				BgAggregate? Aggregate;
-				if(NameToAggregate.TryGetValue(Name, out Aggregate))
-				{
-					OutNodes = Aggregate.RequiredNodes.ToArray();
-					return true;
-				}
-
-				// Check if it's a group name
-				BgAgent? Agent;
-				if(NameToAgent.TryGetValue(Name, out Agent))
-				{
-					OutNodes = Agent.Nodes.ToArray();
-					return true;
-				}
-			}
-
-			// Otherwise fail
-			OutNodes = null;
-			return false;
-		}
-
-		/// <summary>
-		/// Tries to resolve the given name to one or more node outputs. Checks for aggregates, and actual nodes.
-		/// </summary>
-		/// <param name="Name">The name to search for</param>
-		/// <param name="OutOutputs">If the name is a match, receives an array of nodes and their output names</param>
-		/// <returns>True if the name was found, false otherwise.</returns>
-		public bool TryResolveInputReference(string Name, [NotNullWhen(true)] out BgNodeOutput[]? OutOutputs)
-		{
-			// Check if it's a tag reference or node reference
-			if(Name.StartsWith("#"))
-			{
-				// Check if it's a regular node or output name
-				BgNodeOutput? Output;
-				if(TagNameToNodeOutput.TryGetValue(Name, out Output))
-				{
-					OutOutputs = new BgNodeOutput[]{ Output };
-					return true;
-				}
-			}
-			else
-			{
-				// Check if it's a regular node or output name
-				BgNode? Node;
-				if(NameToNode.TryGetValue(Name, out Node))
-				{
-					OutOutputs = Node.Outputs.Union(Node.Inputs).ToArray();
-					return true;
-				}
-
-				// Check if it's an aggregate name
-				BgAggregate? Aggregate;
-				if(NameToAggregate.TryGetValue(Name, out Aggregate))
-				{
-					OutOutputs = Aggregate.RequiredNodes.SelectMany(x => x.Outputs.Union(x.Inputs)).Distinct().ToArray();
-					return true;
-				}
-			}
-
-			// Otherwise fail
-			OutOutputs = null;
-			return false;
-		}
-
-		/// <summary>
-		/// Cull the graph to only include the given nodes and their dependencies
-		/// </summary>
-		/// <param name="TargetNodes">A set of target nodes to build</param>
-		public void Select(IEnumerable<BgNode> TargetNodes)
-		{
-			// Find this node and all its dependencies
-			HashSet<BgNode> RetainNodes = new HashSet<BgNode>(TargetNodes);
-			foreach(BgNode TargetNode in TargetNodes)
-			{
-				RetainNodes.UnionWith(TargetNode.InputDependencies);
-			}
-
-			// Remove all the nodes which are not marked to be kept
-			foreach(BgAgent Agent in Agents)
-			{
-				Agent.Nodes = Agent.Nodes.Where(x => RetainNodes.Contains(x)).ToList();
-			}
-
-			// Remove all the empty agents
-			Agents.RemoveAll(x => x.Nodes.Count == 0);
-
-			// Trim down the list of nodes for each report to the ones that are being built
-			foreach (BgReport Report in NameToReport.Values)
-			{
-				Report.Nodes.RemoveWhere(x => !RetainNodes.Contains(x));
-			}
-
-			// Remove all the empty reports
-			NameToReport = NameToReport.Where(x => x.Value.Nodes.Count > 0).ToDictionary(Pair => Pair.Key, Pair => Pair.Value, StringComparer.InvariantCultureIgnoreCase);
-
-			// Remove all the order dependencies which are no longer part of the graph. Since we don't need to build them, we don't need to wait for them
-			foreach(BgNode Node in RetainNodes)
-			{
-				Node.OrderDependencies = Node.OrderDependencies.Where(x => RetainNodes.Contains(x)).ToArray();
-			}
-
-			// Create a new list of aggregates for everything that's left
-			Dictionary<string, BgAggregate> NewNameToAggregate = new Dictionary<string, BgAggregate>(NameToAggregate.Comparer);
-			foreach(BgAggregate Aggregate in NameToAggregate.Values)
-			{
-				if (Aggregate.RequiredNodes.All(x => RetainNodes.Contains(x)))
-				{
-					NewNameToAggregate[Aggregate.Name] = Aggregate;
-				}
-			}
-			NameToAggregate = NewNameToAggregate;
-
-			// Remove any labels that are no longer value
-			foreach (BgLabel Label in Labels)
-			{
-				Label.RequiredNodes.RemoveWhere(x => !RetainNodes.Contains(x));
-				Label.IncludedNodes.RemoveWhere(x => !RetainNodes.Contains(x));
-			}
-			Labels.RemoveAll(x => x.RequiredNodes.Count == 0);
-
-			// Remove any badges which do not have all their dependencies
-			Badges.RemoveAll(x => x.Nodes.Any(y => !RetainNodes.Contains(y)));
-
-			// Remove any diagnostics which are no longer part of the graph
-			Diagnostics.RemoveAll(x => (x.EnclosingNode != null && !RetainNodes.Contains(x.EnclosingNode)) || (x.EnclosingAgent != null && !Agents.Contains(x.EnclosingAgent)));
-		}
-
-		/// <summary>
-		/// Writes a preprocessed build graph to a script file
-		/// </summary>
-		/// <param name="File">The file to load</param>
-		/// <param name="SchemaFile">Schema file for validation</param>
-		public void Write(FileReference File, FileReference SchemaFile)
-		{
-			XmlWriterSettings Settings = new XmlWriterSettings();
-			Settings.Indent = true;
-			Settings.IndentChars = "\t";
-
-			using (XmlWriter Writer = XmlWriter.Create(File.FullName, Settings))
-			{
-				Writer.WriteStartElement("BuildGraph", "http://www.epicgames.com/BuildGraph");
-
-				if (SchemaFile != null)
-				{
-					Writer.WriteAttributeString("schemaLocation", "http://www.w3.org/2001/XMLSchema-instance", "http://www.epicgames.com/BuildGraph " + SchemaFile.MakeRelativeTo(File.Directory));
-				}
-
-				foreach (BgAgent Agent in Agents)
-				{
-					Agent.Write(Writer);
-				}
-
-				foreach (BgAggregate Aggregate in NameToAggregate.Values)
-				{
-					// If the aggregate has no required elements, skip it.
-					if (Aggregate.RequiredNodes.Count == 0)
-					{
-						continue;
-					}
-
-					Writer.WriteStartElement("Aggregate");
-					Writer.WriteAttributeString("Name", Aggregate.Name);
-					Writer.WriteAttributeString("Requires", String.Join(";", Aggregate.RequiredNodes.Select(x => x.Name)));
-					Writer.WriteEndElement();
-				}
-
-				foreach(BgLabel Label in Labels)
-				{
-					Writer.WriteStartElement("Label");
-					if (Label.DashboardCategory != null)
-					{
-						Writer.WriteAttributeString("Category", Label.DashboardCategory);
-					}
-					Writer.WriteAttributeString("Name", Label.DashboardName);
-					Writer.WriteAttributeString("Requires", String.Join(";", Label.RequiredNodes.Select(x => x.Name)));
-
-					HashSet<BgNode> IncludedNodes = new HashSet<BgNode>(Label.IncludedNodes);
-					IncludedNodes.ExceptWith(Label.IncludedNodes.SelectMany(x => x.InputDependencies));
-					IncludedNodes.ExceptWith(Label.RequiredNodes);
-					if (IncludedNodes.Count > 0)
-					{
-						Writer.WriteAttributeString("Include", String.Join(";", IncludedNodes.Select(x => x.Name)));
-					}
-
-					HashSet<BgNode> ExcludedNodes = new HashSet<BgNode>(Label.IncludedNodes);
-					ExcludedNodes.UnionWith(Label.IncludedNodes.SelectMany(x => x.InputDependencies));
-					ExcludedNodes.ExceptWith(Label.IncludedNodes);
-					ExcludedNodes.ExceptWith(ExcludedNodes.ToArray().SelectMany(x => x.InputDependencies));
-					if (ExcludedNodes.Count > 0)
-					{
-						Writer.WriteAttributeString("Exclude", String.Join(";", ExcludedNodes.Select(x => x.Name)));
-					}
-					Writer.WriteEndElement();
-				}
-
-				foreach (BgReport Report in NameToReport.Values)
-				{
-					Writer.WriteStartElement("Report");
-					Writer.WriteAttributeString("Name", Report.Name);
-					Writer.WriteAttributeString("Requires", String.Join(";", Report.Nodes.Select(x => x.Name)));
-					Writer.WriteEndElement();
-				}
-
-				foreach (BgBadge Badge in Badges)
-				{
-					Writer.WriteStartElement("Badge");
-					Writer.WriteAttributeString("Name", Badge.Name);
-					if (Badge.Project != null)
-					{
-						Writer.WriteAttributeString("Project", Badge.Project);
-					}
-					if (Badge.Change != 0)
-					{
-						Writer.WriteAttributeString("Change", Badge.Change.ToString());
-					}
-					Writer.WriteAttributeString("Requires", String.Join(";", Badge.Nodes.Select(x => x.Name)));
-					Writer.WriteEndElement();
-				}
-
-				Writer.WriteEndElement();
-			}
-		}
-
-		/// <summary>
-		/// Export the build graph to a Json file, for parallel execution by the build system
-		/// </summary>
-		/// <param name="File">Output file to write</param>
-		/// <param name="CompletedNodes">Set of nodes which have been completed</param>
-		public void Export(FileReference File, HashSet<BgNode> CompletedNodes)
-		{
-			// Find all the nodes which we're actually going to execute. We'll use this to filter the graph.
-			HashSet<BgNode> NodesToExecute = new HashSet<BgNode>();
-			foreach(BgNode Node in Agents.SelectMany(x => x.Nodes))
-			{
-				if(!CompletedNodes.Contains(Node))
-				{
-					NodesToExecute.Add(Node);
-				}
-			}
-
-			// Open the output file
-			using(JsonWriter JsonWriter = new JsonWriter(File.FullName))
-			{
-				JsonWriter.WriteObjectStart();
-
-				// Write all the agents
-				JsonWriter.WriteArrayStart("Groups");
-				foreach(BgAgent Agent in Agents)
-				{
-					BgNode[] Nodes = Agent.Nodes.Where(x => NodesToExecute.Contains(x)).ToArray();
-					if(Nodes.Length > 0)
-					{
-						JsonWriter.WriteObjectStart();
-						JsonWriter.WriteValue("Name", Agent.Name);
-						JsonWriter.WriteArrayStart("Agent Types");
-						foreach(string AgentType in Agent.PossibleTypes)
-						{
-							JsonWriter.WriteValue(AgentType);
-						}
-						JsonWriter.WriteArrayEnd();
-						JsonWriter.WriteArrayStart("Nodes");
-						foreach(BgNode Node in Nodes)
-						{
-							JsonWriter.WriteObjectStart();
-							JsonWriter.WriteValue("Name", Node.Name);
-							JsonWriter.WriteValue("DependsOn", String.Join(";", Node.GetDirectOrderDependencies().Where(x => NodesToExecute.Contains(x))));
-							JsonWriter.WriteValue("RunEarly", Node.bRunEarly);
-							JsonWriter.WriteObjectStart("Notify");
-							JsonWriter.WriteValue("Default", String.Join(";", Node.NotifyUsers));
-							JsonWriter.WriteValue("Submitters", String.Join(";", Node.NotifySubmitters));
-							JsonWriter.WriteValue("Warnings", Node.bNotifyOnWarnings);
-							JsonWriter.WriteObjectEnd();
-							JsonWriter.WriteObjectEnd();
-						}
-						JsonWriter.WriteArrayEnd();
-						JsonWriter.WriteObjectEnd();
-					}
-				}
-				JsonWriter.WriteArrayEnd();
-
-				// Write all the badges
-				JsonWriter.WriteArrayStart("Badges");
-				foreach (BgBadge Badge in Badges)
-				{
-					BgNode[] Dependencies = Badge.Nodes.Where(x => NodesToExecute.Contains(x)).ToArray();
-					if (Dependencies.Length > 0)
-					{
-						// Reduce that list to the smallest subset of direct dependencies
-						HashSet<BgNode> DirectDependencies = new HashSet<BgNode>(Dependencies);
-						foreach (BgNode Dependency in Dependencies)
-						{
-							DirectDependencies.ExceptWith(Dependency.OrderDependencies);
-						}
-
-						JsonWriter.WriteObjectStart();
-						JsonWriter.WriteValue("Name", Badge.Name);
-						if (!String.IsNullOrEmpty(Badge.Project))
-						{
-							JsonWriter.WriteValue("Project", Badge.Project);
-						}
-						if(Badge.Change != 0)
-						{
-							JsonWriter.WriteValue("Change", Badge.Change);
-						}
-						JsonWriter.WriteValue("AllDependencies", String.Join(";", Agents.SelectMany(x => x.Nodes).Where(x => Dependencies.Contains(x)).Select(x => x.Name)));
-						JsonWriter.WriteValue("DirectDependencies", String.Join(";", DirectDependencies.Select(x => x.Name)));
-						JsonWriter.WriteObjectEnd();
-					}
-				}
-				JsonWriter.WriteArrayEnd();
-
-				// Write all the triggers and reports. 
-				JsonWriter.WriteArrayStart("Reports");
-				foreach (BgReport Report in NameToReport.Values)
-				{
-					BgNode[] Dependencies = Report.Nodes.Where(x => NodesToExecute.Contains(x)).ToArray();
-					if (Dependencies.Length > 0)
-					{
-						// Reduce that list to the smallest subset of direct dependencies
-						HashSet<BgNode> DirectDependencies = new HashSet<BgNode>(Dependencies);
-						foreach (BgNode Dependency in Dependencies)
-						{
-							DirectDependencies.ExceptWith(Dependency.OrderDependencies);
-						}
-
-						JsonWriter.WriteObjectStart();
-						JsonWriter.WriteValue("Name", Report.Name);
-						JsonWriter.WriteValue("AllDependencies", String.Join(";", Agents.SelectMany(x => x.Nodes).Where(x => Dependencies.Contains(x)).Select(x => x.Name)));
-						JsonWriter.WriteValue("DirectDependencies", String.Join(";", DirectDependencies.Select(x => x.Name)));
-						JsonWriter.WriteValue("Notify", String.Join(";", Report.NotifyUsers));
-						JsonWriter.WriteValue("IsTrigger", false);
-						JsonWriter.WriteObjectEnd();
-					}
-				}
-				JsonWriter.WriteArrayEnd();
-
-				JsonWriter.WriteObjectEnd();
-			}
-		}
-
-		/// <summary>
-		/// Export the build graph to a Json file for parsing by Horde
-		/// </summary>
-		/// <param name="File">Output file to write</param>
-		public void ExportForHorde(FileReference File)
-		{
-			DirectoryReference.CreateDirectory(File.Directory);
-			using (JsonWriter JsonWriter = new JsonWriter(File.FullName))
-			{
-				JsonWriter.WriteObjectStart();
-				JsonWriter.WriteArrayStart("Groups");
-				foreach (BgAgent Agent in Agents)
-				{
-					JsonWriter.WriteObjectStart();
-					JsonWriter.WriteArrayStart("Types");
-					foreach(string PossibleType in Agent.PossibleTypes)
-					{
-						JsonWriter.WriteValue(PossibleType);
-					}
-					JsonWriter.WriteArrayEnd();
-					JsonWriter.WriteArrayStart("Nodes");
-					foreach (BgNode Node in Agent.Nodes)
-					{
-						JsonWriter.WriteObjectStart();
-						JsonWriter.WriteValue("Name", Node.Name);
-						JsonWriter.WriteValue("RunEarly", Node.bRunEarly);
-						JsonWriter.WriteValue("Warnings", Node.bNotifyOnWarnings);
-
-						JsonWriter.WriteArrayStart("InputDependencies");
-						foreach (string InputDependency in Node.GetDirectInputDependencies().Select(x => x.Name))
-						{
-							JsonWriter.WriteValue(InputDependency);
-						}
-						JsonWriter.WriteArrayEnd();
-
-						JsonWriter.WriteArrayStart("OrderDependencies");
-						foreach (string OrderDependency in Node.GetDirectOrderDependencies().Select(x => x.Name))
-						{
-							JsonWriter.WriteValue(OrderDependency);
-						}
-						JsonWriter.WriteArrayEnd();
-
-						JsonWriter.WriteObjectEnd();
-					}
-					JsonWriter.WriteArrayEnd();
-					JsonWriter.WriteObjectEnd();
-				}
-				JsonWriter.WriteArrayEnd();
-
-				JsonWriter.WriteArrayStart("Aggregates");
-				foreach (BgAggregate Aggregate in NameToAggregate.Values)
-				{
-					JsonWriter.WriteObjectStart();
-					JsonWriter.WriteValue("Name", Aggregate.Name);
-					JsonWriter.WriteArrayStart("Nodes");
-					foreach (BgNode RequiredNode in Aggregate.RequiredNodes.OrderBy(x => x.Name))
-					{
-						JsonWriter.WriteValue(RequiredNode.Name);
-					}
-					JsonWriter.WriteArrayEnd();
-					JsonWriter.WriteObjectEnd();
-				}
-				JsonWriter.WriteArrayEnd();
-
-				JsonWriter.WriteArrayStart("Labels");
-				foreach (BgLabel Label in Labels)
-				{
-					JsonWriter.WriteObjectStart();
-					if (!String.IsNullOrEmpty(Label.DashboardName))
-					{
-						JsonWriter.WriteValue("Name", Label.DashboardName);
-					}
-					if (!String.IsNullOrEmpty(Label.DashboardCategory))
-					{
-						JsonWriter.WriteValue("Category", Label.DashboardCategory);
-					}
-					if (!String.IsNullOrEmpty(Label.UgsBadge))
-					{
-						JsonWriter.WriteValue("UgsBadge", Label.UgsBadge);
-					}
-					if (!String.IsNullOrEmpty(Label.UgsProject))
-					{
-						JsonWriter.WriteValue("UgsProject", Label.UgsProject);
-					}
-					if (Label.Change != BgLabelChange.Current)
-					{
-						JsonWriter.WriteValue("Change", Label.Change.ToString());
-					}
-
-					JsonWriter.WriteArrayStart("RequiredNodes");
-					foreach (BgNode RequiredNode in Label.RequiredNodes.OrderBy(x => x.Name))
-					{
-						JsonWriter.WriteValue(RequiredNode.Name);
-					}
-					JsonWriter.WriteArrayEnd();
-					JsonWriter.WriteArrayStart("IncludedNodes");
-					foreach (BgNode IncludedNode in Label.IncludedNodes.OrderBy(x => x.Name))
-					{
-						JsonWriter.WriteValue(IncludedNode.Name);
-					}
-					JsonWriter.WriteArrayEnd();
-					JsonWriter.WriteObjectEnd();
-				}
-				JsonWriter.WriteArrayEnd();
-
-				JsonWriter.WriteArrayStart("Badges");
-				foreach (BgBadge Badge in Badges)
-				{
-					HashSet<BgNode> Dependencies = Badge.Nodes;
-					if (Dependencies.Count > 0)
-					{
-						// Reduce that list to the smallest subset of direct dependencies
-						HashSet<BgNode> DirectDependencies = new HashSet<BgNode>(Dependencies);
-						foreach (BgNode Dependency in Dependencies)
-						{
-							DirectDependencies.ExceptWith(Dependency.OrderDependencies);
-						}
-
-						JsonWriter.WriteObjectStart();
-						JsonWriter.WriteValue("Name", Badge.Name);
-						if (!String.IsNullOrEmpty(Badge.Project))
-						{
-							JsonWriter.WriteValue("Project", Badge.Project);
-						}
-						if (Badge.Change != 0)
-						{
-							JsonWriter.WriteValue("Change", Badge.Change);
-						}
-						JsonWriter.WriteValue("Dependencies", String.Join(";", DirectDependencies.Select(x => x.Name)));
-						JsonWriter.WriteObjectEnd();
-					}
-				}
-				JsonWriter.WriteArrayEnd();
-
-				JsonWriter.WriteObjectEnd();
-			}
-		}
-
-		/// <summary>
-		/// Print the contents of the graph
-		/// </summary>
-		/// <param name="CompletedNodes">Set of nodes which are already complete</param>
-		/// <param name="PrintOptions">Options for how to print the graph</param>
-		/// <param name="Logger"></param>
-		public void Print(HashSet<BgNode> CompletedNodes, GraphPrintOptions PrintOptions, ILogger Logger)
-		{
-			// Print the options
-			if((PrintOptions & GraphPrintOptions.ShowCommandLineOptions) != 0)
-			{
-				// Get the list of messages
-				List<KeyValuePair<string, string>> Parameters = new List<KeyValuePair<string, string>>();
-				foreach(BgScriptOption Option in Options)
-				{
-					string Name = String.Format("-set:{0}=...", Option.Name);
-
-					StringBuilder Description = new StringBuilder(Option.Description);
-					if(!String.IsNullOrEmpty(Option.DefaultValue))
-					{
-						Description.AppendFormat(" (Default: {0})", Option.DefaultValue);
-					}
-
-					Parameters.Add(new KeyValuePair<string, string>(Name, Description.ToString()));
-				}
-
-				// Format them to the log
-				if(Parameters.Count > 0)
-				{
-					Logger.LogInformation("");
-					Logger.LogInformation("Options:");
-					Logger.LogInformation("");
-					HelpUtils.PrintTable(Parameters, 4, 24, Logger);
-				}
-			}
-
-			// Output all the triggers in order
-			Logger.LogInformation("");
-			Logger.LogInformation("Graph:");
-			foreach(BgAgent Agent in Agents)
-			{
-				Logger.LogInformation("        Agent: {0} ({1})", Agent.Name, String.Join(";", Agent.PossibleTypes));
-				foreach(BgNode Node in Agent.Nodes)
-				{
-					Logger.LogInformation("            Node: {0}{1}", Node.Name, CompletedNodes.Contains(Node)? " (completed)" : Node.bRunEarly? " (early)" : "");
-					if(PrintOptions.HasFlag(GraphPrintOptions.ShowDependencies))
-					{
-						HashSet<BgNode> InputDependencies = new HashSet<BgNode>(Node.GetDirectInputDependencies());
-						foreach(BgNode InputDependency in InputDependencies)
-						{
-							Logger.LogInformation("                input> {0}", InputDependency.Name);
-						}
-						HashSet<BgNode> OrderDependencies = new HashSet<BgNode>(Node.GetDirectOrderDependencies());
-						foreach(BgNode OrderDependency in OrderDependencies.Except(InputDependencies))
-						{
-							Logger.LogInformation("                after> {0}", OrderDependency.Name);
-						}
-					}
-					if(PrintOptions.HasFlag(GraphPrintOptions.ShowNotifications))
-					{
-						string Label = Node.bNotifyOnWarnings? "warnings" : "errors";
-						foreach(string User in Node.NotifyUsers)
-						{
-							Logger.LogInformation("                {0}> {1}", Label, User);
-						}
-						foreach(string Submitter in Node.NotifySubmitters)
-						{
-							Logger.LogInformation("                {0}> submitters to {1}", Label, Submitter);
-						}
-					}
-				}
-			}
-			Logger.LogInformation("");
-
-			// Print out all the non-empty aggregates
-			BgAggregate[] Aggregates = NameToAggregate.Values.OrderBy(x => x.Name).ToArray();
-			if(Aggregates.Length > 0)
-			{
-				Logger.LogInformation("Aggregates:");
-				foreach(string AggregateName in Aggregates.Select(x => x.Name))
-				{
-					Logger.LogInformation("    {0}", AggregateName);
-				}
-				Logger.LogInformation("");
-			}
-		}
-	}
 }
diff --git a/Engine/Source/Programs/Shared/EpicGames.BuildGraph/BgScriptReader.cs b/Engine/Source/Programs/Shared/EpicGames.BuildGraph/BgScriptReader.cs
--- a/Engine/Source/Programs/Shared/EpicGames.BuildGraph/BgScriptReader.cs
+++ b/Engine/Source/Programs/Shared/EpicGames.BuildGraph/BgScriptReader.cs
@@ -291,7 +291,7 @@ namespace EpicGames.BuildGraph
 		/// <summary>
 		/// The current graph
 		/// </summary>
-		BgScript Graph = new BgScript();
+		BgGraph Graph = new BgGraph();
 
 		/// <summary>
 		/// List of property name to value lookups. Modifications to properties are scoped to nodes and agents. EnterScope() pushes an empty dictionary onto the end of this list, and LeaveScope() removes one. 
@@ -370,7 +370,7 @@ namespace EpicGames.BuildGraph
 		/// <param name="Logger">Logger for output messages</param>
 		/// <param name="SingleNodeName">If a single node will be processed, the name of that node.</param>
 		/// <returns>True if the graph was read, false if there were errors</returns>
-		public static async Task<BgScript?> ReadAsync(IBgScriptReaderContext Context, string File, Dictionary<string, string> Arguments, Dictionary<string, string> DefaultProperties, bool bPreprocessOnly, BgScriptSchema Schema, ILogger Logger, string? SingleNodeName = null)
+		public static async Task<BgGraph?> ReadAsync(IBgScriptReaderContext Context, string File, Dictionary<string, string> Arguments, Dictionary<string, string> DefaultProperties, bool bPreprocessOnly, BgScriptSchema Schema, ILogger Logger, string? SingleNodeName = null)
 		{
 			// Read the file and build the graph
 			BgScriptReader Reader = new BgScriptReader(Context, DefaultProperties, bPreprocessOnly, Schema, Logger);
@@ -1561,7 +1561,7 @@ namespace EpicGames.BuildGraph
 			{
 				string Message = ReadAttribute(Element, "Message");
 
-				BgScriptDiagnostic Diagnostic = new BgScriptDiagnostic(EventType, $"{Element.File}({Element.LineNumber}): {Message}", EnclosingNode, EnclosingAgent);
+				BgGraphDiagnostic Diagnostic = new BgGraphDiagnostic(EventType, $"{Element.File}({Element.LineNumber}): {Message}", EnclosingNode, EnclosingAgent);
 				Graph.Diagnostics.Add(Diagnostic);
 			}
 		}
