BuildGraph: Changes to support introspection of BuildGraph property dependencies without full script evaluation.

- Include directives must appear at the root level of the graph, and no longer support conditional inclusion.
- Macros may no longer be conditionally defined.
- Macros must have constant names. Ability to extend existing macros via the <Extend> element has been added to support pluggable extension points required by installed build script.

diff --git a/Engine/Build/InstalledEngineFilters.xml b/Engine/Build/InstalledEngineFilters.xml
--- a/Engine/Build/InstalledEngineFilters.xml
+++ b/Engine/Build/InstalledEngineFilters.xml
@@ -4,12 +4,11 @@
 	<!-- List of patterns that should always be excluded when making an installed build. Filter the default list of exclusions to include confidential platforms if they're enabled. -->
 	<Property Name="ConfidentialExceptions" Value=""/>
 	<ForEach Name="RestrictedFolderName" Values="$(RestrictedFolderNames)">
-		<ForEach Name="Platform" Values="$(ExtensionPlatforms)">
-			<Expand Name="FilterRestrictedFolders$(Platform)"/>
-		</ForEach>
+		<Expand Name="Platform_FilterRestrictedFolders"/>
+
 		<Do If="'$(RestrictedFolderName)' != ''">
-				<Property Name="ConfidentialExceptions" Value="$(ConfidentialExceptions);" If="'$(ConfidentialExceptions)' != ''"/>
-				<Property Name="ConfidentialExceptions" Value="$(ConfidentialExceptions).../$(RestrictedFolderName)/..."/>
+			<Property Name="ConfidentialExceptions" Value="$(ConfidentialExceptions);" If="'$(ConfidentialExceptions)' != ''"/>
+			<Property Name="ConfidentialExceptions" Value="$(ConfidentialExceptions).../$(RestrictedFolderName)/..."/>
 		</Do>
 	</ForEach>
 
@@ -189,9 +188,7 @@
 	<Property Name="AutomationToolPaths" Value="Engine/Source/Programs/AutomationTool/..."/>
 	
 	<!-- Platform extensions filter overrides -->
-	<ForEach Name="Platform" Values="$(ExtensionPlatforms)">
-		<Expand Name="$(Platform)FilterOverrides"/>
-	</ForEach>
+	<Expand Name="Platform_FilterOverrides"/>
 	
 	<ForEach Name="Filter" Values="$(CopyEditorEngineOrPlatformFilter)">
 		<Property Name="CopyEditorFilter" Value="$(CopyEditorFilter);Engine/$(Filter);Engine/Platforms/*/$(Filter)"/>
diff --git a/Engine/Source/Programs/Shared/EpicGames.BuildGraph/BgSchema.cs b/Engine/Source/Programs/Shared/EpicGames.BuildGraph/BgSchema.cs
--- a/Engine/Source/Programs/Shared/EpicGames.BuildGraph/BgSchema.cs
+++ b/Engine/Source/Programs/Shared/EpicGames.BuildGraph/BgSchema.cs
@@ -135,6 +135,7 @@ namespace EpicGames.BuildGraph
 	public enum ScriptSchemaStandardType
 	{
 		Graph,
+		NestedGraph,
 		Agent,
 		AgentBody,
 		Node,
@@ -151,10 +152,12 @@ namespace EpicGames.BuildGraph
 		Regex,
 		Macro,
 		MacroBody,
+		Extend,
 		Expand,
 		Trace,
 		Warning,
 		Error,
+		Path,
 		Name,
 		NameList,
 		Tag,
@@ -196,7 +199,17 @@ namespace EpicGames.BuildGraph
 		/// <summary>
 		/// Characters which are not permitted in names.
 		/// </summary>
-		public const string IllegalNameCharacters = "^<>:\"/\\|?*";
+		public const string IllegalNameCharacters = "^<>:\"/\\|?*$";
+
+		/// <summary>
+		/// Characters which are not permitted in paths.
+		/// </summary>
+		public const string IllegalPathCharacters = ":$";
+
+		/// <summary>
+		/// Pattern which matches any name; alphanumeric characters, with single embedded spaces.
+		/// </summary>
+		const string PathPattern = "[^ " + IllegalPathCharacters + "]+";
 
 		/// <summary>
 		/// Pattern which matches any name; alphanumeric characters, with single embedded spaces.
@@ -313,7 +326,8 @@ namespace EpicGames.BuildGraph
 			NewSchema.TargetNamespace = NamespaceURI;
 			NewSchema.ElementFormDefault = XmlSchemaForm.Qualified;
 			NewSchema.Items.Add(CreateSchemaElement(RootElementName, ScriptSchemaStandardType.Graph));
-			NewSchema.Items.Add(CreateGraphType());
+			NewSchema.Items.Add(CreateGraphType(ScriptSchemaStandardType.Graph));
+			NewSchema.Items.Add(CreateGraphType(ScriptSchemaStandardType.NestedGraph));
 			NewSchema.Items.Add(CreateAgentType());
 			NewSchema.Items.Add(CreateAgentBodyType());
 			NewSchema.Items.Add(CreateNodeType());
@@ -331,10 +345,12 @@ namespace EpicGames.BuildGraph
 			NewSchema.Items.Add(CreateRegexType());
 			NewSchema.Items.Add(CreateMacroType());
 			NewSchema.Items.Add(CreateMacroBodyType(TaskNameToType));
+			NewSchema.Items.Add(CreateExtendType());
 			NewSchema.Items.Add(CreateExpandType());
 			NewSchema.Items.Add(CreateDiagnosticType(ScriptSchemaStandardType.Trace));
 			NewSchema.Items.Add(CreateDiagnosticType(ScriptSchemaStandardType.Warning));
 			NewSchema.Items.Add(CreateDiagnosticType(ScriptSchemaStandardType.Error));
+			NewSchema.Items.Add(CreateSimpleTypeFromRegex(GetTypeName(ScriptSchemaStandardType.Path), PathPattern));
 			NewSchema.Items.Add(CreateSimpleTypeFromRegex(GetTypeName(ScriptSchemaStandardType.Name), "(" + NamePattern + "|" + StringWithPropertiesPattern + ")"));
 			NewSchema.Items.Add(CreateSimpleTypeFromRegex(GetTypeName(ScriptSchemaStandardType.NameList), "(" + NameListPattern + "|" + StringWithPropertiesPattern + ")"));
 			NewSchema.Items.Add(CreateSimpleTypeFromRegex(GetTypeName(ScriptSchemaStandardType.Tag), "(" + TagPattern + "|" + StringWithPropertiesPattern + ")"));
@@ -470,6 +486,7 @@ namespace EpicGames.BuildGraph
 			GraphChoice.Items.Add(CreateSchemaElement("Property", ScriptSchemaStandardType.Property));
 			GraphChoice.Items.Add(CreateSchemaElement("Regex", ScriptSchemaStandardType.Regex));
 			GraphChoice.Items.Add(CreateSchemaElement("Macro", ScriptSchemaStandardType.Macro));
+			GraphChoice.Items.Add(CreateSchemaElement("Extend", ScriptSchemaStandardType.Extend));
 			GraphChoice.Items.Add(CreateSchemaElement("Agent", ScriptSchemaStandardType.Agent));
 			GraphChoice.Items.Add(CreateSchemaElement("Aggregate", ScriptSchemaStandardType.Aggregate));
 			GraphChoice.Items.Add(CreateSchemaElement("Report", ScriptSchemaStandardType.Report));
@@ -490,6 +507,45 @@ namespace EpicGames.BuildGraph
 			return GraphType;
 		}
 
+		/// <summary>
+		/// Creates the schema type representing the graph type
+		/// </summary>
+		/// <returns>Type definition for a graph</returns>
+		static XmlSchemaType CreateGraphType(ScriptSchemaStandardType Type)
+		{
+			XmlSchemaChoice GraphChoice = new XmlSchemaChoice();
+			GraphChoice.MinOccurs = 0;
+			GraphChoice.MaxOccursString = "unbounded";
+			if (Type == ScriptSchemaStandardType.Graph)
+			{
+				GraphChoice.Items.Add(CreateSchemaElement("Include", ScriptSchemaStandardType.Include));
+				GraphChoice.Items.Add(CreateSchemaElement("Option", ScriptSchemaStandardType.Option));
+				GraphChoice.Items.Add(CreateSchemaElement("Macro", ScriptSchemaStandardType.Macro));
+				GraphChoice.Items.Add(CreateSchemaElement("Extend", ScriptSchemaStandardType.Extend));
+			}
+			GraphChoice.Items.Add(CreateSchemaElement("Property", ScriptSchemaStandardType.Property));
+			GraphChoice.Items.Add(CreateSchemaElement("Regex", ScriptSchemaStandardType.Regex));
+			GraphChoice.Items.Add(CreateSchemaElement("EnvVar", ScriptSchemaStandardType.EnvVar));
+			GraphChoice.Items.Add(CreateSchemaElement("Agent", ScriptSchemaStandardType.Agent));
+			GraphChoice.Items.Add(CreateSchemaElement("Aggregate", ScriptSchemaStandardType.Aggregate));
+			GraphChoice.Items.Add(CreateSchemaElement("Report", ScriptSchemaStandardType.Report));
+			GraphChoice.Items.Add(CreateSchemaElement("Badge", ScriptSchemaStandardType.Badge));
+			GraphChoice.Items.Add(CreateSchemaElement("Label", ScriptSchemaStandardType.Label));
+			GraphChoice.Items.Add(CreateSchemaElement("Notify", ScriptSchemaStandardType.Notify));
+			GraphChoice.Items.Add(CreateSchemaElement("Trace", ScriptSchemaStandardType.Trace));
+			GraphChoice.Items.Add(CreateSchemaElement("Warning", ScriptSchemaStandardType.Warning));
+			GraphChoice.Items.Add(CreateSchemaElement("Error", ScriptSchemaStandardType.Error));
+			GraphChoice.Items.Add(CreateSchemaElement("Expand", ScriptSchemaStandardType.Expand));
+			GraphChoice.Items.Add(CreateDoElement(ScriptSchemaStandardType.NestedGraph));
+			GraphChoice.Items.Add(CreateSwitchElement(ScriptSchemaStandardType.NestedGraph));
+			GraphChoice.Items.Add(CreateForEachElement(ScriptSchemaStandardType.NestedGraph));
+
+			XmlSchemaComplexType GraphType = new XmlSchemaComplexType();
+			GraphType.Name = GetTypeName(Type);
+			GraphType.Particle = GraphChoice;
+			return GraphType;
+		}
+
 		/// <summary>
 		/// Creates the schema type representing the agent type
 		/// </summary>
@@ -699,8 +755,7 @@ namespace EpicGames.BuildGraph
 		{
 			XmlSchemaComplexType PropertyType = new XmlSchemaComplexType();
 			PropertyType.Name = GetTypeName(ScriptSchemaStandardType.Include);
-			PropertyType.Attributes.Add(CreateSchemaAttribute("Script", ScriptSchemaStandardType.BalancedString, XmlSchemaUse.Required));
-			PropertyType.Attributes.Add(CreateSchemaAttribute("If", ScriptSchemaStandardType.BalancedString, XmlSchemaUse.Optional));
+			PropertyType.Attributes.Add(CreateSchemaAttribute("Script", ScriptSchemaStandardType.Path, XmlSchemaUse.Required));
 			return PropertyType;
 		}
 
@@ -789,7 +844,7 @@ namespace EpicGames.BuildGraph
 			Extension.Attributes.Add(CreateSchemaAttribute("Name", ScriptSchemaStandardType.Name, XmlSchemaUse.Required));
 			Extension.Attributes.Add(CreateSchemaAttribute("Arguments", ScriptSchemaStandardType.BalancedString, XmlSchemaUse.Optional));
 			Extension.Attributes.Add(CreateSchemaAttribute("OptionalArguments", ScriptSchemaStandardType.BalancedString, XmlSchemaUse.Optional));
-			Extension.Attributes.Add(CreateSchemaAttribute("If", ScriptSchemaStandardType.BalancedString, XmlSchemaUse.Optional));
+			Extension.Attributes.Add(CreateSchemaAttribute("Append", ScriptSchemaStandardType.BalancedString, XmlSchemaUse.Optional));
 
 			XmlSchemaComplexContent ContentModel = new XmlSchemaComplexContent();
 			ContentModel.Content = Extension;
@@ -845,6 +900,25 @@ namespace EpicGames.BuildGraph
 			return NodeType;
 		}
 
+		/// <summary>
+		/// Creates the schema type representing the macro type
+		/// </summary>
+		/// <returns>Type definition for a node</returns>
+		static XmlSchemaType CreateExtendType()
+		{
+			XmlSchemaComplexContentExtension Extension = new XmlSchemaComplexContentExtension();
+			Extension.BaseTypeName = GetQualifiedTypeName(ScriptSchemaStandardType.MacroBody);
+			Extension.Attributes.Add(CreateSchemaAttribute("Name", ScriptSchemaStandardType.Name, XmlSchemaUse.Required));
+
+			XmlSchemaComplexContent ContentModel = new XmlSchemaComplexContent();
+			ContentModel.Content = Extension;
+
+			XmlSchemaComplexType ComplexType = new XmlSchemaComplexType();
+			ComplexType.Name = GetTypeName(ScriptSchemaStandardType.Extend);
+			ComplexType.ContentModel = ContentModel;
+			return ComplexType;
+		}
+
 		/// <summary>
 		/// Creates the schema type representing a macro expansion
 		/// </summary>
diff --git a/Engine/Source/Programs/Shared/EpicGames.BuildGraph/BgScriptReader.cs b/Engine/Source/Programs/Shared/EpicGames.BuildGraph/BgScriptReader.cs
--- a/Engine/Source/Programs/Shared/EpicGames.BuildGraph/BgScriptReader.cs
+++ b/Engine/Source/Programs/Shared/EpicGames.BuildGraph/BgScriptReader.cs
@@ -195,7 +195,7 @@ namespace EpicGames.BuildGraph
 		/// <summary>
 		/// Element where the function was declared
 		/// </summary>
-		public readonly BgScriptElement Element;
+		public List<BgScriptElement> Elements = new List<BgScriptElement>();
 
 		/// <summary>
 		/// The total number of arguments
@@ -222,7 +222,7 @@ namespace EpicGames.BuildGraph
 		public BgScriptMacro(string Name, BgScriptElement Element, Dictionary<string, int> ArgumentNameToIndex, int NumRequiredArguments)
 		{
 			this.Name = Name;
-			this.Element = Element;
+			this.Elements.Add(Element);
 			this.NumArguments = ArgumentNameToIndex.Count;
 			this.NumRequiredArguments = NumRequiredArguments;
 			this.ArgumentNameToIndex = ArgumentNameToIndex;
@@ -462,6 +462,9 @@ namespace EpicGames.BuildGraph
 					case "Macro":
 						await ReadMacroAsync(ChildElement);
 						break;
+					case "Extend":
+						await ReadExtendAsync(ChildElement);
+						break;
 					case "Agent":
 						await ReadAgentAsync(ChildElement);
 						break;
@@ -858,26 +861,23 @@ namespace EpicGames.BuildGraph
 		/// <param name="Element">Xml element to read the definition from</param>
 		async Task ReadMacroAsync(BgScriptElement Element)
 		{
-			if(await EvaluateConditionAsync(Element))
+			string Name = Element.GetAttribute("Name");
+			if (ValidateName(Element, Name))
 			{
-				string Name = ReadAttribute(Element, "Name");
-				if (ValidateName(Element, Name))
+				BgScriptMacro? OriginalDefinition;
+				if(MacroNameToDefinition.TryGetValue(Name, out OriginalDefinition))
 				{
-					BgScriptMacro? OriginalDefinition;
-					if(MacroNameToDefinition.TryGetValue(Name, out OriginalDefinition))
-					{
-						LogError(Element, "Function '{0}' has already been declared (see {1} line {2})", OriginalDefinition.Element.File, OriginalDefinition.Element.LineNumber);
-					}
-					else
-					{
-						Dictionary<string, int> ArgumentNameToIndex = new Dictionary<string, int>();
-						ReadMacroArguments(Element, "Arguments", ArgumentNameToIndex);
+					LogError(Element, "Macro '{0}' has already been declared (see {1} line {2})", Name, OriginalDefinition.Elements[0].File, OriginalDefinition.Elements[0].LineNumber);
+				}
+				else
+				{
+					Dictionary<string, int> ArgumentNameToIndex = new Dictionary<string, int>();
+					ReadMacroArguments(Element, "Arguments", ArgumentNameToIndex);
 
-						int NumRequiredArguments = ArgumentNameToIndex.Count;
-						ReadMacroArguments(Element, "OptionalArguments", ArgumentNameToIndex);
+					int NumRequiredArguments = ArgumentNameToIndex.Count;
+					ReadMacroArguments(Element, "OptionalArguments", ArgumentNameToIndex);
 
-						MacroNameToDefinition.Add(Name, new BgScriptMacro(Name, Element, ArgumentNameToIndex, NumRequiredArguments));
-					}
+					MacroNameToDefinition.Add(Name, new BgScriptMacro(Name, Element, ArgumentNameToIndex, NumRequiredArguments));
 				}
 			}
 		}
@@ -907,6 +907,28 @@ namespace EpicGames.BuildGraph
 			}
 		}
 
+		/// <summary>
+		/// Reads a macro definition
+		/// </summary>
+		/// <param name="Element">Xml element to read the definition from</param>
+		async Task ReadExtendAsync(BgScriptElement Element)
+		{
+			if (await EvaluateConditionAsync(Element))
+			{
+				string Name = ReadAttribute(Element, "Name");
+
+				BgScriptMacro? OriginalDefinition;
+				if (MacroNameToDefinition.TryGetValue(Name, out OriginalDefinition))
+				{
+					OriginalDefinition.Elements.Add(Element);
+				}
+				else
+				{
+					LogError(Element, "Macro '{0}' has not been declared", Name);
+				}
+			}
+		}
+
 		/// <summary>
 		/// Reads the definition for an agent.
 		/// </summary>
@@ -1395,54 +1417,59 @@ namespace EpicGames.BuildGraph
 		{
 			if(await EvaluateConditionAsync(Element))
 			{
-				string Name = ReadAttribute(Element, "Name");
-
-				BgScriptMacro? Macro;
-				if(!MacroNameToDefinition.TryGetValue(Name, out Macro))
-				{
-					LogError(Element, "Macro '{0}' does not exist", Name);
-				}
-				else
+				string Name = Element.GetAttribute("Name");
+				if (ValidateName(Element, Name))
 				{
-					// Parse the argument list
-					string[] Arguments = new string[Macro.ArgumentNameToIndex.Count];
-					foreach(XmlAttribute? Attribute in Element.Attributes)
+					BgScriptMacro? Macro;
+					if (!MacroNameToDefinition.TryGetValue(Name, out Macro))
+					{
+						LogError(Element, "Macro '{0}' does not exist", Name);
+					}
+					else
 					{
-						if(Attribute != null && Attribute.Name != "Name" && Attribute.Name != "If")
+						// Parse the argument list
+						string[] Arguments = new string[Macro.ArgumentNameToIndex.Count];
+						foreach (XmlAttribute? Attribute in Element.Attributes)
 						{
-							int Index;
-							if(Macro.ArgumentNameToIndex.TryGetValue(Attribute.Name, out Index))
+							if (Attribute != null && Attribute.Name != "Name" && Attribute.Name != "If")
 							{
-								Arguments[Index] = ExpandProperties(Element, Attribute.Value);
-							}
-							else
-							{
-								LogWarning(Element, "Macro '{0}' does not take an argument '{1}'", Name, Attribute.Name);
+								int Index;
+								if (Macro.ArgumentNameToIndex.TryGetValue(Attribute.Name, out Index))
+								{
+									Arguments[Index] = ExpandProperties(Element, Attribute.Value);
+								}
+								else
+								{
+									LogWarning(Element, "Macro '{0}' does not take an argument '{1}'", Name, Attribute.Name);
+								}
 							}
 						}
-					}
 
-					// Make sure none of the required arguments are missing
-					bool bHasMissingArguments = false;
-					for(int Idx = 0; Idx < Macro.NumRequiredArguments; Idx++)
-					{
-						if(Arguments[Idx] == null)
+						// Make sure none of the required arguments are missing
+						bool bHasMissingArguments = false;
+						for (int Idx = 0; Idx < Macro.NumRequiredArguments; Idx++)
 						{
-							LogWarning(Element, "Macro '{0}' is missing argument '{1}'", Macro.Name, Macro.ArgumentNameToIndex.First(x => x.Value == Idx).Key);
-							bHasMissingArguments = true;
+							if (Arguments[Idx] == null)
+							{
+								LogWarning(Element, "Macro '{0}' is missing argument '{1}'", Macro.Name, Macro.ArgumentNameToIndex.First(x => x.Value == Idx).Key);
+								bHasMissingArguments = true;
+							}
 						}
-					}
 
-					// Expand the function
-					if(!bHasMissingArguments)
-					{
-						EnterScope();
-						foreach(KeyValuePair<string, int> Pair in Macro.ArgumentNameToIndex)
+						// Expand the function
+						if (!bHasMissingArguments)
 						{
-							ScopedProperties[ScopedProperties.Count - 1][Pair.Key] = Arguments[Pair.Value] ?? "";
+							EnterScope();
+							foreach (KeyValuePair<string, int> Pair in Macro.ArgumentNameToIndex)
+							{
+								ScopedProperties[ScopedProperties.Count - 1][Pair.Key] = Arguments[Pair.Value] ?? "";
+							}
+							foreach (BgScriptElement MacroElement in Macro.Elements)
+							{
+								await ReadContentsAsync(MacroElement);
+							}
+							LeaveScope();
 						}
-						await ReadContentsAsync(Macro.Element);
-						LeaveScope();
 					}
 				}
 			}
