Fix build.

diff --git a/Engine/Source/Programs/Shared/EpicGames.BuildGraph/BgSchema.cs b/Engine/Source/Programs/Shared/EpicGames.BuildGraph/BgSchema.cs
--- a/Engine/Source/Programs/Shared/EpicGames.BuildGraph/BgSchema.cs
+++ b/Engine/Source/Programs/Shared/EpicGames.BuildGraph/BgSchema.cs
@@ -199,7 +199,7 @@ namespace EpicGames.BuildGraph
 		/// <summary>
 		/// Characters which are not permitted in names.
 		/// </summary>
-		public const string IllegalNameCharacters = "^<>:\"/\\|?*$";
+		public const string IllegalNameCharacters = "^<>:\"/\\|?*";
 
 		/// <summary>
 		/// Characters which are not permitted in paths.
diff --git a/Engine/Source/Programs/Shared/EpicGames.BuildGraph/BgScriptReader.cs b/Engine/Source/Programs/Shared/EpicGames.BuildGraph/BgScriptReader.cs
--- a/Engine/Source/Programs/Shared/EpicGames.BuildGraph/BgScriptReader.cs
+++ b/Engine/Source/Programs/Shared/EpicGames.BuildGraph/BgScriptReader.cs
@@ -460,7 +460,7 @@ namespace EpicGames.BuildGraph
 						await ReadEnvVarAsync(ChildElement);
 						break;
 					case "Macro":
-						await ReadMacroAsync(ChildElement);
+						ReadMacro(ChildElement);
 						break;
 					case "Extend":
 						await ReadExtendAsync(ChildElement);
@@ -859,7 +859,7 @@ namespace EpicGames.BuildGraph
 		/// Reads a macro definition
 		/// </summary>
 		/// <param name="Element">Xml element to read the definition from</param>
-		async Task ReadMacroAsync(BgScriptElement Element)
+		void ReadMacro(BgScriptElement Element)
 		{
 			string Name = Element.GetAttribute("Name");
 			if (ValidateName(Element, Name))
