BuildGraph: Add a dedicated type to store source file location for script messages.

diff --git a/Engine/Source/Programs/AutomationTool/BuildGraph/BuildGraph.cs b/Engine/Source/Programs/AutomationTool/BuildGraph/BuildGraph.cs
--- a/Engine/Source/Programs/AutomationTool/BuildGraph/BuildGraph.cs
+++ b/Engine/Source/Programs/AutomationTool/BuildGraph/BuildGraph.cs
@@ -606,15 +606,15 @@ namespace AutomationTool
 				{
 					if(Diagnostic.EventType == LogEventType.Console)
 					{
-						CommandUtils.LogInformation(Diagnostic.Message);
+						CommandUtils.LogWarning("{0}({1}): {2}", Diagnostic.Location.File, Diagnostic.Location.LineNumber, Diagnostic.Message);
 					}
 					else if(Diagnostic.EventType == LogEventType.Warning)
 					{
-						CommandUtils.LogWarning(Diagnostic.Message);
+						CommandUtils.LogWarning("{0}({1}): warning: {2}", Diagnostic.Location.File, Diagnostic.Location.LineNumber, Diagnostic.Message);
 					}
 					else
 					{
-						CommandUtils.LogError(Diagnostic.Message);
+						CommandUtils.LogError("{0}({1}): error: {2}", Diagnostic.Location.File, Diagnostic.Location.LineNumber, Diagnostic.Message);
 					}
 				}
 				if(Graph.Diagnostics.Any(x => x.EventType == LogEventType.Error))
@@ -771,7 +771,7 @@ namespace AutomationTool
 			CustomTask NewTask = (CustomTask)Activator.CreateInstance(Task.TaskClass, ParametersObject);
 
 			// Set up the source location for diagnostics
-			NewTask.SourceLocation = TaskInfo.SourceLocation;
+			NewTask.SourceLocation = TaskInfo.Location;
 
 			// Make sure all the read tags are local or listed as a dependency
 			foreach (string ReadTagName in NewTask.FindConsumedTagNames())
@@ -840,7 +840,7 @@ namespace AutomationTool
 
 		void OutputBindingError(BgTask Task, string Format, params object[] Args)
 		{
-			Logger.LogScriptError(Context.GetNativePath(Task.SourceLocation.Item1), Task.SourceLocation.Item2, Format, Args);
+			Logger.LogScriptError(Task.Location, Format, Args);
 		}
 
 		/// <summary>
@@ -1261,7 +1261,7 @@ namespace AutomationTool
 							ExceptionUtils.AddContext(Ex, "while executing task {0}", Tasks[Idx].GetTraceString());
 							if (Tasks[Idx].SourceLocation != null)
 							{
-								ExceptionUtils.AddContext(Ex, "at {0}({1})", Tasks[Idx].SourceLocation.Item1, Tasks[Idx].SourceLocation.Item2);
+								ExceptionUtils.AddContext(Ex, "at {0}({1})", Tasks[Idx].SourceLocation.File, Tasks[Idx].SourceLocation.LineNumber);
 							}
 							throw;
 						}
@@ -1289,7 +1289,7 @@ namespace AutomationTool
 							}
 							if (Tasks[FirstIdx].SourceLocation != null)
 							{
-								ExceptionUtils.AddContext(Ex, "at {0}({1})", Tasks[FirstIdx].SourceLocation.Item1, Tasks[FirstIdx].SourceLocation.Item2);
+								ExceptionUtils.AddContext(Ex, "at {0}({1})", Tasks[FirstIdx].SourceLocation.File, Tasks[FirstIdx].SourceLocation.LineNumber);
 							}
 							throw;
 						}
diff --git a/Engine/Source/Programs/AutomationTool/BuildGraph/Task.cs b/Engine/Source/Programs/AutomationTool/BuildGraph/Task.cs
--- a/Engine/Source/Programs/AutomationTool/BuildGraph/Task.cs
+++ b/Engine/Source/Programs/AutomationTool/BuildGraph/Task.cs
@@ -100,11 +100,7 @@ namespace AutomationTool
 		/// <summary>
 		/// Line number in a source file that this task was declared. Optional; used for log messages.
 		/// </summary>
-		public Tuple<string, int> SourceLocation
-		{
-			get;
-			set;
-		}
+		public BgScriptLocation SourceLocation { get; set; }
 
 		/// <summary>
 		/// Execute this node.
@@ -188,8 +184,8 @@ namespace AutomationTool
 		/// <param name="Prefix">Prefix for metadata entries</param>
 		public virtual void GetTraceMetadata(ITraceSpan Span, string Prefix)
 		{
-			Span.AddMetadata(Prefix + "source.file", SourceLocation.Item1);
-			Span.AddMetadata(Prefix + "source.line", SourceLocation.Item2.ToString());
+			Span.AddMetadata(Prefix + "source.file", SourceLocation.File);
+			Span.AddMetadata(Prefix + "source.line", SourceLocation.LineNumber.ToString());
 		}
 		
 		/// <summary>
@@ -199,8 +195,8 @@ namespace AutomationTool
 		/// <param name="Prefix">Prefix for metadata entries</param>
 		public virtual void GetTraceMetadata(ISpan Span, string Prefix)
 		{
-			Span.SetTag(Prefix + "source.file", SourceLocation.Item1);
-			Span.SetTag(Prefix + "source.line", SourceLocation.Item2);
+			Span.SetTag(Prefix + "source.file", SourceLocation.File);
+			Span.SetTag(Prefix + "source.line", SourceLocation.LineNumber);
 		}
 
 		/// <summary>
diff --git a/Engine/Source/Programs/Shared/EpicGames.BuildGraph/BgGraph.cs b/Engine/Source/Programs/Shared/EpicGames.BuildGraph/BgGraph.cs
--- a/Engine/Source/Programs/Shared/EpicGames.BuildGraph/BgGraph.cs
+++ b/Engine/Source/Programs/Shared/EpicGames.BuildGraph/BgGraph.cs
@@ -41,6 +41,11 @@ namespace EpicGames.BuildGraph
 	/// </summary>
 	public class BgGraphDiagnostic
 	{
+		/// <summary>
+		/// Location of the diagnostic
+		/// </summary>
+		public BgScriptLocation Location;
+
 		/// <summary>
 		/// The diagnostic event type
 		/// </summary>
@@ -64,8 +69,9 @@ namespace EpicGames.BuildGraph
 		/// <summary>
 		/// Constructor
 		/// </summary>
-		public BgGraphDiagnostic(LogEventType EventType, string Message, BgNode? EnclosingNode, BgAgent? EnclosingAgent)
+		public BgGraphDiagnostic(BgScriptLocation Location, LogEventType EventType, string Message, BgNode? EnclosingNode, BgAgent? EnclosingAgent)
 		{
+			this.Location = Location;
 			this.EventType = EventType;
 			this.Message = Message;
 			this.EnclosingNode = EnclosingNode;
diff --git a/Engine/Source/Programs/Shared/EpicGames.BuildGraph/BgScriptReader.cs b/Engine/Source/Programs/Shared/EpicGames.BuildGraph/BgScriptReader.cs
--- a/Engine/Source/Programs/Shared/EpicGames.BuildGraph/BgScriptReader.cs
+++ b/Engine/Source/Programs/Shared/EpicGames.BuildGraph/BgScriptReader.cs
@@ -17,6 +17,40 @@ using Microsoft.Extensions.Logging;
 
 namespace EpicGames.BuildGraph
 {
+	/// <summary>
+	/// Location of an element within a file
+	/// </summary>
+	public class BgScriptLocation
+	{
+		/// <summary>
+		/// The file containing this element
+		/// </summary>
+		public string File { get; }
+
+		/// <summary>
+		/// Native file object
+		/// </summary>
+		public object NativeFile { get; }
+
+		/// <summary>
+		/// The line number containing this element
+		/// </summary>
+		public int LineNumber { get; }
+
+		/// <summary>
+		/// Constructor
+		/// </summary>
+		/// <param name="File"></param>
+		/// <param name="NativeFile"></param>
+		/// <param name="LineNumber"></param>
+		public BgScriptLocation(string File, object NativeFile, int LineNumber)
+		{
+			this.File = File;
+			this.NativeFile = NativeFile;
+			this.LineNumber = LineNumber;
+		}
+	}
+
 	/// <summary>
 	/// Implementation of XmlDocument which preserves line numbers for its elements
 	/// </summary>
@@ -62,7 +96,8 @@ namespace EpicGames.BuildGraph
 		/// </summary>
 		public override XmlElement CreateElement(string Prefix, string LocalName, string NamespaceUri)
 		{
-			return new BgScriptElement(File, NativeFile, LineInfo!.LineNumber, Prefix, LocalName, NamespaceUri, this);
+			BgScriptLocation Location = new BgScriptLocation(File, NativeFile, LineInfo!.LineNumber);
+			return new BgScriptElement(Location, Prefix, LocalName, NamespaceUri, this);
 		}
 
 		/// <summary>
@@ -100,7 +135,8 @@ namespace EpicGames.BuildGraph
 				{
 					if (!Document.bHasErrors)
 					{
-						Logger.LogScriptError(NativeFile, Ex.LineNumber, "{Message}", Ex.Message);
+						BgScriptLocation Location = new BgScriptLocation(File, NativeFile, Ex.LineNumber);
+						Logger.LogScriptError(Location, "{Message}", Ex.Message);
 						Document.bHasErrors = true;
 					}
 				}
@@ -115,13 +151,15 @@ namespace EpicGames.BuildGraph
 				// Check that the root element is valid. If not, we didn't actually validate against the schema.
 				if (Document.DocumentElement.Name != BgScriptSchema.RootElementName)
 				{
-					Logger.LogScriptError(NativeFile, 1, "Script does not have a root element called '{ElementName}'", BgScriptSchema.RootElementName);
+					BgScriptLocation Location = new BgScriptLocation(File, NativeFile, 1);
+					Logger.LogScriptError(Location, "Script does not have a root element called '{ElementName}'", BgScriptSchema.RootElementName);
 					OutDocument = null;
 					return false;
 				}
 				if (Document.DocumentElement.NamespaceURI != BgScriptSchema.NamespaceURI)
 				{
-					Logger.LogScriptError(NativeFile, 1, "Script root element is not in the '{Namespace}' namespace (add the xmlns=\"{NewNamespace}\" attribute)", BgScriptSchema.NamespaceURI, BgScriptSchema.NamespaceURI);
+					BgScriptLocation Location = new BgScriptLocation(File, NativeFile, 1);
+					Logger.LogScriptError(Location, "Script root element is not in the '{Namespace}' namespace (add the xmlns=\"{NewNamespace}\" attribute)", BgScriptSchema.NamespaceURI, BgScriptSchema.NamespaceURI);
 					OutDocument = null;
 					return false;
 				}
@@ -138,13 +176,14 @@ namespace EpicGames.BuildGraph
 		/// <param name="Args">Standard argument for ValidationEventHandler</param>
 		void ValidationEvent(object Sender, ValidationEventArgs Args)
 		{
+			BgScriptLocation Location = new BgScriptLocation(File, NativeFile, Args.Exception.LineNumber);
 			if (Args.Severity == XmlSeverityType.Warning)
 			{
-				Logger.LogScriptWarning(NativeFile, Args.Exception.LineNumber, "{Message}", Args.Message);
+				Logger.LogScriptWarning(Location, "{Message}", Args.Message);
 			}
 			else
 			{
-				Logger.LogScriptError(NativeFile, Args.Exception.LineNumber, "{Message}", Args.Message);
+				Logger.LogScriptError(Location, "{Message}", Args.Message);
 				bHasErrors = true;
 			}
 		}
@@ -156,29 +195,17 @@ namespace EpicGames.BuildGraph
 	class BgScriptElement : XmlElement
 	{
 		/// <summary>
-		/// The file containing this element
-		/// </summary>
-		public string File { get; }
-
-		/// <summary>
-		/// Native file object
+		/// Location of the element within the file
 		/// </summary>
-		public object NativeFile { get; }
-
-		/// <summary>
-		/// The line number containing this element
-		/// </summary>
-		public int LineNumber { get; }
+		public BgScriptLocation Location { get; }
 
 		/// <summary>
 		/// Constructor
 		/// </summary>
-		public BgScriptElement(string File, object NativeFile, int LineNumber, string Prefix, string LocalName, string NamespaceUri, BgScriptDocument Document)
+		public BgScriptElement(BgScriptLocation Location, string Prefix, string LocalName, string NamespaceUri, BgScriptDocument Document)
 			: base(Prefix, LocalName, NamespaceUri, Document)
 		{
-			this.File = File;
-			this.NativeFile = NativeFile;
-			this.LineNumber = LineNumber;
+			this.Location = Location;
 		}
 	}
 
@@ -265,20 +292,20 @@ namespace EpicGames.BuildGraph
 	/// </summary>
 	public static class BgScriptExtensions
 	{
-		public static void LogScriptError(this ILogger Logger, object File, int LineNumber, string Format, params object[] Args)
+		public static void LogScriptError(this ILogger Logger, BgScriptLocation Location, string Format, params object[] Args)
 		{
 			object[] AllArgs = new object[Args.Length + 2];
-			AllArgs[0] = File;
-			AllArgs[1] = LineNumber;
+			AllArgs[0] = Location.NativeFile;
+			AllArgs[1] = Location.LineNumber;
 			Args.CopyTo(AllArgs, 2);
 			Logger.LogError($"{{Script}}({{Line}}): error: {Format}", AllArgs);
 		}
 
-		public static void LogScriptWarning(this ILogger Logger, object File, int LineNumber, string Format, params object[] Args)
+		public static void LogScriptWarning(this ILogger Logger, BgScriptLocation Location, string Format, params object[] Args)
 		{
 			object[] AllArgs = new object[Args.Length + 2];
-			AllArgs[0] = File;
-			AllArgs[1] = LineNumber;
+			AllArgs[0] = Location.NativeFile;
+			AllArgs[1] = Location.LineNumber;
 			Args.CopyTo(AllArgs, 2);
 			Logger.LogWarning($"{{Script}}({{Line}}): warning: {Format}", AllArgs);
 		}
@@ -666,7 +693,7 @@ namespace EpicGames.BuildGraph
 				HashSet<string> Files = new HashSet<string>();
 				foreach (string Script in ReadListAttribute(Element, "Script"))
 				{
-					string IncludePath = CombinePaths(Element.File, Script);
+					string IncludePath = CombinePaths(Element.Location.File, Script);
 					if (Regex.IsMatch(IncludePath, @"\*|\?|\.\.\."))
 					{
 						Files.UnionWith(await Context.FindAsync(IncludePath));
@@ -867,7 +894,8 @@ namespace EpicGames.BuildGraph
 				BgScriptMacro? OriginalDefinition;
 				if(MacroNameToDefinition.TryGetValue(Name, out OriginalDefinition))
 				{
-					LogError(Element, "Macro '{0}' has already been declared (see {1} line {2})", Name, OriginalDefinition.Elements[0].File, OriginalDefinition.Elements[0].LineNumber);
+					BgScriptLocation Location = OriginalDefinition.Elements[0].Location;
+					LogError(Element, "Macro '{0}' has already been declared (see {1} line {2})", Name, Location.File, Location.LineNumber);
 				}
 				else
 				{
@@ -1490,7 +1518,7 @@ namespace EpicGames.BuildGraph
 
 			if (await EvaluateConditionAsync(Element))
 			{
-				BgTask Info = new BgTask(Tuple.Create(Element.File, Element.LineNumber), Element.Name);
+				BgTask Info = new BgTask(Element.Location, Element.Name);
 				foreach (XmlAttribute? Attribute in Element.Attributes)
 				{
 					if (String.Compare(Attribute!.Name, "If", StringComparison.InvariantCultureIgnoreCase) != 0)
@@ -1609,7 +1637,7 @@ namespace EpicGames.BuildGraph
 			{
 				string Message = ReadAttribute(Element, "Message");
 
-				BgGraphDiagnostic Diagnostic = new BgGraphDiagnostic(EventType, $"{Element.File}({Element.LineNumber}): {Message}", EnclosingNode, EnclosingAgent);
+				BgGraphDiagnostic Diagnostic = new BgGraphDiagnostic(Element.Location, EventType, Message, EnclosingNode, EnclosingAgent);
 				Graph.Diagnostics.Add(Diagnostic);
 			}
 		}
@@ -1888,7 +1916,7 @@ namespace EpicGames.BuildGraph
 		/// <param name="Args">Optional arguments</param>
 		void LogError(BgScriptElement Element, string Format, params object[] Args)
 		{
-			Logger.LogScriptError(Element.NativeFile, Element.LineNumber, Format, Args);
+			Logger.LogScriptError(Element.Location, Format, Args);
 			NumErrors++;
 		}
 
@@ -1900,7 +1928,7 @@ namespace EpicGames.BuildGraph
 		/// <param name="Args">Optional arguments</param>
 		void LogWarning(BgScriptElement Element, string Format, params object[] Args)
 		{
-			Logger.LogScriptWarning(Element.NativeFile, Element.LineNumber, Format, Args);
+			Logger.LogScriptWarning(Element.Location, Format, Args);
 		}
 
 		/// <summary>
diff --git a/Engine/Source/Programs/Shared/EpicGames.BuildGraph/BgTask.cs b/Engine/Source/Programs/Shared/EpicGames.BuildGraph/BgTask.cs
--- a/Engine/Source/Programs/Shared/EpicGames.BuildGraph/BgTask.cs
+++ b/Engine/Source/Programs/Shared/EpicGames.BuildGraph/BgTask.cs
@@ -17,11 +17,7 @@ namespace EpicGames.BuildGraph
 		/// <summary>
 		/// Line number in a source file that this task was declared. Optional; used for log messages.
 		/// </summary>
-		public Tuple<string, int> SourceLocation
-		{
-			get;
-			set;
-		}
+		public BgScriptLocation Location { get; }
 
 		/// <summary>
 		/// Name of the task
@@ -36,9 +32,9 @@ namespace EpicGames.BuildGraph
 		/// <summary>
 		/// Constructor
 		/// </summary>
-		public BgTask(Tuple<string, int> SourceLocation, string Name)
+		public BgTask(BgScriptLocation Location, string Name)
 		{
-			this.SourceLocation = SourceLocation;
+			this.Location = Location;
 			this.Name = Name;
 		}
 
