Horde: Remove stray newline from the end of json formatted log messages.

diff --git a/Engine/Source/Programs/Shared/EpicGames.Core/LogEvent.cs b/Engine/Source/Programs/Shared/EpicGames.Core/LogEvent.cs
--- a/Engine/Source/Programs/Shared/EpicGames.Core/LogEvent.cs
+++ b/Engine/Source/Programs/Shared/EpicGames.Core/LogEvent.cs
@@ -491,8 +491,6 @@ namespace EpicGames.Core
 			{
 				Write(Writer);
 			}
-			Buffer.GetSpan(1)[0] = (byte)'\n';
-			Buffer.Advance(1);
 			return Buffer.WrittenSpan.ToArray();
 		}
 
