BuildGraph: Allow macros in include directives again.

diff --git a/Engine/Source/Programs/Shared/EpicGames.BuildGraph/BgSchema.cs b/Engine/Source/Programs/Shared/EpicGames.BuildGraph/BgSchema.cs
--- a/Engine/Source/Programs/Shared/EpicGames.BuildGraph/BgSchema.cs
+++ b/Engine/Source/Programs/Shared/EpicGames.BuildGraph/BgSchema.cs
@@ -755,7 +755,8 @@ namespace EpicGames.BuildGraph
 		{
 			XmlSchemaComplexType PropertyType = new XmlSchemaComplexType();
 			PropertyType.Name = GetTypeName(ScriptSchemaStandardType.Include);
-			PropertyType.Attributes.Add(CreateSchemaAttribute("Script", ScriptSchemaStandardType.Path, XmlSchemaUse.Required));
+			PropertyType.Attributes.Add(CreateSchemaAttribute("Script", ScriptSchemaStandardType.BalancedString, XmlSchemaUse.Required));
+			PropertyType.Attributes.Add(CreateSchemaAttribute("If", ScriptSchemaStandardType.BalancedString, XmlSchemaUse.Optional));
 			return PropertyType;
 		}
 
