BuildGraph: Revert some changes causing problems with some existing scripts.

diff --git a/Engine/Source/Programs/AutomationTool/BuildGraph/BuildGraph.cs b/Engine/Source/Programs/AutomationTool/BuildGraph/BuildGraph.cs
--- a/Engine/Source/Programs/AutomationTool/BuildGraph/BuildGraph.cs
+++ b/Engine/Source/Programs/AutomationTool/BuildGraph/BuildGraph.cs
@@ -117,7 +117,7 @@ namespace AutomationTool
 		{
 			try
 			{
-				return Task.FromResult(FileReference.Exists(FileReference.Combine(Unreal.RootDirectory, Path)));
+				return Task.FromResult(FileReference.Exists(FileReference.Combine(Unreal.RootDirectory, Path)) || DirectoryReference.Exists(DirectoryReference.Combine(Unreal.RootDirectory, Path)));
 			}
 			catch
 			{
diff --git a/Engine/Source/Programs/Shared/EpicGames.BuildGraph/BgSchema.cs b/Engine/Source/Programs/Shared/EpicGames.BuildGraph/BgSchema.cs
--- a/Engine/Source/Programs/Shared/EpicGames.BuildGraph/BgSchema.cs
+++ b/Engine/Source/Programs/Shared/EpicGames.BuildGraph/BgSchema.cs
@@ -135,7 +135,6 @@ namespace EpicGames.BuildGraph
 	public enum ScriptSchemaStandardType
 	{
 		Graph,
-		NestedGraph,
 		Agent,
 		AgentBody,
 		Node,
@@ -157,7 +156,6 @@ namespace EpicGames.BuildGraph
 		Trace,
 		Warning,
 		Error,
-		Path,
 		Name,
 		NameList,
 		Tag,
@@ -201,16 +199,6 @@ namespace EpicGames.BuildGraph
 		/// </summary>
 		public const string IllegalNameCharacters = "^<>:\"/\\|?*";
 
-		/// <summary>
-		/// Characters which are not permitted in paths.
-		/// </summary>
-		public const string IllegalPathCharacters = ":$";
-
-		/// <summary>
-		/// Pattern which matches any name; alphanumeric characters, with single embedded spaces.
-		/// </summary>
-		const string PathPattern = "[^ " + IllegalPathCharacters + "]+";
-
 		/// <summary>
 		/// Pattern which matches any name; alphanumeric characters, with single embedded spaces.
 		/// </summary>
@@ -326,8 +314,7 @@ namespace EpicGames.BuildGraph
 			NewSchema.TargetNamespace = NamespaceURI;
 			NewSchema.ElementFormDefault = XmlSchemaForm.Qualified;
 			NewSchema.Items.Add(CreateSchemaElement(RootElementName, ScriptSchemaStandardType.Graph));
-			NewSchema.Items.Add(CreateGraphType(ScriptSchemaStandardType.Graph));
-			NewSchema.Items.Add(CreateGraphType(ScriptSchemaStandardType.NestedGraph));
+			NewSchema.Items.Add(CreateGraphType());
 			NewSchema.Items.Add(CreateAgentType());
 			NewSchema.Items.Add(CreateAgentBodyType());
 			NewSchema.Items.Add(CreateNodeType());
@@ -350,7 +337,6 @@ namespace EpicGames.BuildGraph
 			NewSchema.Items.Add(CreateDiagnosticType(ScriptSchemaStandardType.Trace));
 			NewSchema.Items.Add(CreateDiagnosticType(ScriptSchemaStandardType.Warning));
 			NewSchema.Items.Add(CreateDiagnosticType(ScriptSchemaStandardType.Error));
-			NewSchema.Items.Add(CreateSimpleTypeFromRegex(GetTypeName(ScriptSchemaStandardType.Path), PathPattern));
 			NewSchema.Items.Add(CreateSimpleTypeFromRegex(GetTypeName(ScriptSchemaStandardType.Name), "(" + NamePattern + "|" + StringWithPropertiesPattern + ")"));
 			NewSchema.Items.Add(CreateSimpleTypeFromRegex(GetTypeName(ScriptSchemaStandardType.NameList), "(" + NameListPattern + "|" + StringWithPropertiesPattern + ")"));
 			NewSchema.Items.Add(CreateSimpleTypeFromRegex(GetTypeName(ScriptSchemaStandardType.Tag), "(" + TagPattern + "|" + StringWithPropertiesPattern + ")"));
@@ -507,45 +493,6 @@ namespace EpicGames.BuildGraph
 			return GraphType;
 		}
 
-		/// <summary>
-		/// Creates the schema type representing the graph type
-		/// </summary>
-		/// <returns>Type definition for a graph</returns>
-		static XmlSchemaType CreateGraphType(ScriptSchemaStandardType Type)
-		{
-			XmlSchemaChoice GraphChoice = new XmlSchemaChoice();
-			GraphChoice.MinOccurs = 0;
-			GraphChoice.MaxOccursString = "unbounded";
-			if (Type == ScriptSchemaStandardType.Graph)
-			{
-				GraphChoice.Items.Add(CreateSchemaElement("Include", ScriptSchemaStandardType.Include));
-				GraphChoice.Items.Add(CreateSchemaElement("Option", ScriptSchemaStandardType.Option));
-				GraphChoice.Items.Add(CreateSchemaElement("Macro", ScriptSchemaStandardType.Macro));
-				GraphChoice.Items.Add(CreateSchemaElement("Extend", ScriptSchemaStandardType.Extend));
-			}
-			GraphChoice.Items.Add(CreateSchemaElement("Property", ScriptSchemaStandardType.Property));
-			GraphChoice.Items.Add(CreateSchemaElement("Regex", ScriptSchemaStandardType.Regex));
-			GraphChoice.Items.Add(CreateSchemaElement("EnvVar", ScriptSchemaStandardType.EnvVar));
-			GraphChoice.Items.Add(CreateSchemaElement("Agent", ScriptSchemaStandardType.Agent));
-			GraphChoice.Items.Add(CreateSchemaElement("Aggregate", ScriptSchemaStandardType.Aggregate));
-			GraphChoice.Items.Add(CreateSchemaElement("Report", ScriptSchemaStandardType.Report));
-			GraphChoice.Items.Add(CreateSchemaElement("Badge", ScriptSchemaStandardType.Badge));
-			GraphChoice.Items.Add(CreateSchemaElement("Label", ScriptSchemaStandardType.Label));
-			GraphChoice.Items.Add(CreateSchemaElement("Notify", ScriptSchemaStandardType.Notify));
-			GraphChoice.Items.Add(CreateSchemaElement("Trace", ScriptSchemaStandardType.Trace));
-			GraphChoice.Items.Add(CreateSchemaElement("Warning", ScriptSchemaStandardType.Warning));
-			GraphChoice.Items.Add(CreateSchemaElement("Error", ScriptSchemaStandardType.Error));
-			GraphChoice.Items.Add(CreateSchemaElement("Expand", ScriptSchemaStandardType.Expand));
-			GraphChoice.Items.Add(CreateDoElement(ScriptSchemaStandardType.NestedGraph));
-			GraphChoice.Items.Add(CreateSwitchElement(ScriptSchemaStandardType.NestedGraph));
-			GraphChoice.Items.Add(CreateForEachElement(ScriptSchemaStandardType.NestedGraph));
-
-			XmlSchemaComplexType GraphType = new XmlSchemaComplexType();
-			GraphType.Name = GetTypeName(Type);
-			GraphType.Particle = GraphChoice;
-			return GraphType;
-		}
-
 		/// <summary>
 		/// Creates the schema type representing the agent type
 		/// </summary>
@@ -845,7 +792,6 @@ namespace EpicGames.BuildGraph
 			Extension.Attributes.Add(CreateSchemaAttribute("Name", ScriptSchemaStandardType.Name, XmlSchemaUse.Required));
 			Extension.Attributes.Add(CreateSchemaAttribute("Arguments", ScriptSchemaStandardType.BalancedString, XmlSchemaUse.Optional));
 			Extension.Attributes.Add(CreateSchemaAttribute("OptionalArguments", ScriptSchemaStandardType.BalancedString, XmlSchemaUse.Optional));
-			Extension.Attributes.Add(CreateSchemaAttribute("Append", ScriptSchemaStandardType.BalancedString, XmlSchemaUse.Optional));
 
 			XmlSchemaComplexContent ContentModel = new XmlSchemaComplexContent();
 			ContentModel.Content = Extension;
@@ -1178,4 +1124,4 @@ namespace EpicGames.BuildGraph
 			return SchemaType;
 		}
 	}
-}
+}
\ No newline at end of file
diff --git a/Engine/Source/Programs/Shared/EpicGames.BuildGraph/BgScriptReader.cs b/Engine/Source/Programs/Shared/EpicGames.BuildGraph/BgScriptReader.cs
--- a/Engine/Source/Programs/Shared/EpicGames.BuildGraph/BgScriptReader.cs
+++ b/Engine/Source/Programs/Shared/EpicGames.BuildGraph/BgScriptReader.cs
@@ -1445,7 +1445,7 @@ namespace EpicGames.BuildGraph
 		{
 			if(await EvaluateConditionAsync(Element))
 			{
-				string Name = Element.GetAttribute("Name");
+				string Name = ReadAttribute(Element, "Name");
 				if (ValidateName(Element, Name))
 				{
 					BgScriptMacro? Macro;
@@ -1998,4 +1998,4 @@ namespace EpicGames.BuildGraph
 			return Result;
 		}
 	}
-}
+}
\ No newline at end of file
