Horde: Trap exceptions when jobs output exceptions that look like structured log messages, but fail to parse.

diff --git a/Engine/Source/Programs/Shared/EpicGames.Core/JsonLogEvent.cs b/Engine/Source/Programs/Shared/EpicGames.Core/JsonLogEvent.cs
--- a/Engine/Source/Programs/Shared/EpicGames.Core/JsonLogEvent.cs
+++ b/Engine/Source/Programs/Shared/EpicGames.Core/JsonLogEvent.cs
@@ -52,51 +52,55 @@ namespace EpicGames.Core
 		/// <returns></returns>
 		public static bool TryParse(ReadOnlyMemory<byte> Data, out JsonLogEvent Event)
 		{
-			LogLevel Level = LogLevel.None;
-			int EventId = 0;
-
-			Utf8JsonReader Reader = new Utf8JsonReader(Data.Span);
-			if (Reader.Read() && Reader.TokenType == JsonTokenType.StartObject)
+			try
 			{
-				ReadOnlySpan<byte> PropertyName;
-				for (; ReadNextPropertyName(ref Reader, out PropertyName); Reader.Skip())
+				LogLevel Level = LogLevel.None;
+				int EventId = 0;
+
+				Utf8JsonReader Reader = new Utf8JsonReader(Data.Span);
+				if (Reader.Read() && Reader.TokenType == JsonTokenType.StartObject)
 				{
-					if (Utf8StringComparer.OrdinalIgnoreCase.Equals(PropertyName, LevelString.Span))
+					ReadOnlySpan<byte> PropertyName;
+					for (; ReadNextPropertyName(ref Reader, out PropertyName); Reader.Skip())
 					{
-						if (!Enum.TryParse(Reader.GetString(), out Level))
+						if (Utf8StringComparer.OrdinalIgnoreCase.Equals(PropertyName, LevelString.Span))
 						{
-							break;
+							if (!Enum.TryParse(Reader.GetString(), out Level))
+							{
+								break;
+							}
 						}
-					}
-					else if (Utf8StringComparer.OrdinalIgnoreCase.Equals(PropertyName, PropertiesString.Span) && Reader.TokenType == JsonTokenType.StartObject)
-					{
-						for (; ReadNextPropertyName(ref Reader, out PropertyName); Reader.Skip())
+						else if (Utf8StringComparer.OrdinalIgnoreCase.Equals(PropertyName, PropertiesString.Span) && Reader.TokenType == JsonTokenType.StartObject)
 						{
-							if (Utf8StringComparer.OrdinalIgnoreCase.Equals(EventIdString.Span, PropertyName) && Reader.TokenType == JsonTokenType.StartObject)
+							for (; ReadNextPropertyName(ref Reader, out PropertyName); Reader.Skip())
 							{
-								for (; ReadNextPropertyName(ref Reader, out PropertyName); Reader.Skip())
+								if (Utf8StringComparer.OrdinalIgnoreCase.Equals(EventIdString.Span, PropertyName) && Reader.TokenType == JsonTokenType.StartObject)
 								{
-									if (Utf8StringComparer.OrdinalIgnoreCase.Equals(IdString.Span, PropertyName) && Reader.TryGetInt32(out int NewEventId))
+									for (; ReadNextPropertyName(ref Reader, out PropertyName); Reader.Skip())
 									{
-										EventId = NewEventId;
+										if (Utf8StringComparer.OrdinalIgnoreCase.Equals(IdString.Span, PropertyName) && Reader.TryGetInt32(out int NewEventId))
+										{
+											EventId = NewEventId;
+										}
 									}
 								}
 							}
 						}
 					}
 				}
-			}
 
-			if (Reader.TokenType == JsonTokenType.EndObject && Level != LogLevel.None && Reader.BytesConsumed == Data.Length)
-			{
-				Event = new JsonLogEvent(Level, new EventId(EventId), Data);
-				return true;
+				if (Reader.TokenType == JsonTokenType.EndObject && Level != LogLevel.None && Reader.BytesConsumed == Data.Length)
+				{
+					Event = new JsonLogEvent(Level, new EventId(EventId), Data);
+					return true;
+				}
 			}
-			else
+			catch
 			{
-				Event = default;
-				return false;
 			}
+
+			Event = default;
+			return false;
 		}
 
 		static bool ReadNextPropertyName(ref Utf8JsonReader Reader, out ReadOnlySpan<byte> PropertyName)
