Horde: Fix tests in serialization library.

diff --git a/Engine/Source/Programs/Shared/EpicGames.Serialization.Tests/EngineTests.cs b/Engine/Source/Programs/Shared/EpicGames.Serialization.Tests/EngineTests.cs
--- a/Engine/Source/Programs/Shared/EpicGames.Serialization.Tests/EngineTests.cs
+++ b/Engine/Source/Programs/Shared/EpicGames.Serialization.Tests/EngineTests.cs
@@ -58,8 +58,8 @@ namespace EpicGames.Serialization.Tests
 			[CbFieldType.Float64] = CbFieldAccessors.FromStruct<double>(x => x.IsFloat(), (x, y) => x.AsDouble(y)),
 			[CbFieldType.BoolTrue] = CbFieldAccessors.FromStruct<bool>(x => x.IsBool(), (x, y) => x.AsBool(y)),
 			[CbFieldType.BoolFalse] = CbFieldAccessors.FromStruct<bool>(x => x.IsBool(), (x, y) => x.AsBool(y)),
-			[CbFieldType.ObjectAttachment] = new CbFieldAccessors(IoHash.Zero, x => x.IsObjectAttachment(), x => x.AsObjectAttachment(), (x, Default) => x.AsObjectAttachment((IoHash)Default)),
-			[CbFieldType.BinaryAttachment] = new CbFieldAccessors(IoHash.Zero, x => x.IsBinaryAttachment(), x => x.AsBinaryAttachment(), (x, Default) => x.AsBinaryAttachment((IoHash)Default)),
+			[CbFieldType.ObjectAttachment] = new CbFieldAccessors(new CbObjectAttachment(IoHash.Zero), x => x.IsObjectAttachment(), x => x.AsObjectAttachment(), (x, Default) => x.AsObjectAttachment((CbObjectAttachment)Default)),
+			[CbFieldType.BinaryAttachment] = new CbFieldAccessors(new CbBinaryAttachment(IoHash.Zero), x => x.IsBinaryAttachment(), x => x.AsBinaryAttachment(), (x, Default) => x.AsBinaryAttachment((CbBinaryAttachment)Default)),
 			[CbFieldType.Hash] = new CbFieldAccessors(IoHash.Zero, x => x.IsHash(), x => x.AsHash(), (x, Default) => x.AsHash((IoHash)Default)),
 			[CbFieldType.Uuid] = CbFieldAccessors.FromStruct<Guid>(x => x.IsUuid(), (x, y) => x.AsUuid(y)),
 			[CbFieldType.DateTime] = CbFieldAccessors.FromStruct<DateTime>(new DateTime(0, DateTimeKind.Utc), x => x.IsDateTime(), (x, y) => x.AsDateTime(y)),
@@ -842,18 +842,18 @@ namespace EpicGames.Serialization.Tests
 			TestField(CbFieldType.ObjectAttachment, ZeroBytes);
 
 			// Test CbField(ObjectAttachment, NonZero)
-			TestField(CbFieldType.ObjectAttachment, SequentialBytes, new IoHash(SequentialBytes));
+			TestField(CbFieldType.ObjectAttachment, SequentialBytes, (CbObjectAttachment)new IoHash(SequentialBytes));
 
 			// Test CbField(ObjectAttachment, NonZero) AsAttachment
 			{
 				CbField Field = new CbField(SequentialBytes, CbFieldType.ObjectAttachment);
-				TestField(CbFieldType.ObjectAttachment, Field, new IoHash(SequentialBytes), new IoHash(), CbFieldError.None);
+				TestField(CbFieldType.ObjectAttachment, Field, (CbObjectAttachment)new IoHash(SequentialBytes), (CbObjectAttachment)new IoHash(), CbFieldError.None);
 			}
 
 			// Test CbField(None) as ObjectAttachment
 			{
 				CbField DefaultField = new CbField();
-				TestFieldError(CbFieldType.ObjectAttachment, DefaultField, CbFieldError.TypeError, new IoHash(SequentialBytes));
+				TestFieldError(CbFieldType.ObjectAttachment, DefaultField, CbFieldError.TypeError, (CbObjectAttachment)new IoHash(SequentialBytes));
 			}
 		}
 
@@ -867,18 +867,18 @@ namespace EpicGames.Serialization.Tests
 			TestField(CbFieldType.BinaryAttachment, ZeroBytes);
 
 			// Test CbField(BinaryAttachment, NonZero)
-			TestField(CbFieldType.BinaryAttachment, SequentialBytes, new IoHash(SequentialBytes));
+			TestField(CbFieldType.BinaryAttachment, SequentialBytes, (CbBinaryAttachment)new IoHash(SequentialBytes));
 
 			// Test CbField(BinaryAttachment, NonZero) AsAttachment
 			{
 				CbField Field = new CbField(SequentialBytes, CbFieldType.BinaryAttachment);
-				TestField(CbFieldType.BinaryAttachment, Field, new IoHash(SequentialBytes), new IoHash(), CbFieldError.None);
+				TestField(CbFieldType.BinaryAttachment, Field, (CbBinaryAttachment)new IoHash(SequentialBytes), (CbBinaryAttachment)new IoHash(), CbFieldError.None);
 			}
 
 			// Test CbField(None) as BinaryAttachment
 			{
 				CbField DefaultField = new CbField();
-				TestFieldError(CbFieldType.BinaryAttachment, DefaultField, CbFieldError.TypeError, new IoHash(SequentialBytes));
+				TestFieldError(CbFieldType.BinaryAttachment, DefaultField, CbFieldError.TypeError, (CbBinaryAttachment)new IoHash(SequentialBytes));
 			}
 		}
 
diff --git a/Engine/Source/Programs/Shared/EpicGames.Serialization/CompactBinary.cs b/Engine/Source/Programs/Shared/EpicGames.Serialization/CompactBinary.cs
--- a/Engine/Source/Programs/Shared/EpicGames.Serialization/CompactBinary.cs
+++ b/Engine/Source/Programs/Shared/EpicGames.Serialization/CompactBinary.cs
@@ -221,6 +221,18 @@ namespace EpicGames.Serialization
 		/// <inheritdoc/>
 		public override string ToString() => Hash.ToString();
 
+		/// <inheritdoc/>
+		public override bool Equals(object? Obj) => Obj is CbBinaryAttachment Other && Other.Hash == Hash;
+
+		/// <inheritdoc/>
+		public override int GetHashCode() => Hash.GetHashCode();
+
+		/// <inheritdoc/>
+		public static bool operator ==(CbBinaryAttachment Lhs, CbBinaryAttachment Rhs) => Lhs.Hash == Rhs.Hash;
+
+		/// <inheritdoc/>
+		public static bool operator !=(CbBinaryAttachment Lhs, CbBinaryAttachment Rhs) => Lhs.Hash != Rhs.Hash;
+
 		/// <summary>
 		/// Convert a hash to a binary attachment 
 		/// </summary>
@@ -288,6 +300,18 @@ namespace EpicGames.Serialization
 		/// <inheritdoc/>
 		public override string ToString() => Hash.ToString();
 
+		/// <inheritdoc/>
+		public override bool Equals(object? Obj) => Obj is CbObjectAttachment Other && Other.Hash == Hash;
+
+		/// <inheritdoc/>
+		public override int GetHashCode() => Hash.GetHashCode();
+
+		/// <inheritdoc/>
+		public static bool operator ==(CbObjectAttachment Lhs, CbObjectAttachment Rhs) => Lhs.Hash == Rhs.Hash;
+
+		/// <inheritdoc/>
+		public static bool operator !=(CbObjectAttachment Lhs, CbObjectAttachment Rhs) => Lhs.Hash != Rhs.Hash;
+
 		/// <summary>
 		/// Use an object attachment as a hash
 		/// </summary>
