BuildGraph: Minor cleanup.

diff --git a/Engine/Source/Programs/Shared/EpicGames.BuildGraph/BgScriptReader.cs b/Engine/Source/Programs/Shared/EpicGames.BuildGraph/BgScriptReader.cs
--- a/Engine/Source/Programs/Shared/EpicGames.BuildGraph/BgScriptReader.cs
+++ b/Engine/Source/Programs/Shared/EpicGames.BuildGraph/BgScriptReader.cs
@@ -54,17 +54,17 @@ namespace EpicGames.BuildGraph
 	/// <summary>
 	/// Implementation of XmlDocument which preserves line numbers for its elements
 	/// </summary>
-	class BgScriptDocument : XmlDocument
+	public class BgScriptDocument : XmlDocument
 	{
 		/// <summary>
 		/// The file being read
 		/// </summary>
-		string File;
+		string File { get; }
 
 		/// <summary>
 		/// Native file representation
 		/// </summary>
-		object NativeFile;
+		object NativeFile { get; }
 
 		/// <summary>
 		/// Interface to the LineInfo on the active XmlReader
@@ -79,7 +79,7 @@ namespace EpicGames.BuildGraph
 		/// <summary>
 		/// Logger for validation errors
 		/// </summary>
-		ILogger Logger;
+		ILogger Logger { get; }
 
 		/// <summary>
 		/// Private constructor. Use ScriptDocument.Load to read an XML document.
@@ -256,6 +256,9 @@ namespace EpicGames.BuildGraph
 		}
 	}
 
+	/// <summary>
+	/// Context used for parsing scripts
+	/// </summary>
 	public interface IBgScriptReaderContext
 	{
 		/// <summary>
@@ -292,6 +295,9 @@ namespace EpicGames.BuildGraph
 	/// </summary>
 	public static class BgScriptExtensions
 	{
+		/// <summary>
+		/// Utility method to log a script error at a particular location
+		/// </summary>
 		public static void LogScriptError(this ILogger Logger, BgScriptLocation Location, string Format, params object[] Args)
 		{
 			object[] AllArgs = new object[Args.Length + 2];
@@ -301,6 +307,9 @@ namespace EpicGames.BuildGraph
 			Logger.LogError($"{{Script}}({{Line}}): error: {Format}", AllArgs);
 		}
 
+		/// <summary>
+		/// Utility method to log a script warning at a particular location
+		/// </summary>
 		public static void LogScriptWarning(this ILogger Logger, BgScriptLocation Location, string Format, params object[] Args)
 		{
 			object[] AllArgs = new object[Args.Length + 2];
@@ -343,6 +352,11 @@ namespace EpicGames.BuildGraph
 		/// </summary>
 		Dictionary<string, BgScriptMacro> MacroNameToDefinition = new Dictionary<string, BgScriptMacro>();
 
+		/// <summary>
+		/// Arguments for evaluating the graph
+		/// </summary>
+		Dictionary<string, string> Arguments;
+
 		/// <summary>
 		/// Preprocess the file, but do not expand any values that are not portable (eg. paths on the local machine)
 		/// </summary>
@@ -373,13 +387,17 @@ namespace EpicGames.BuildGraph
 		/// </summary>
 		/// <param name="Context">Context object</param>
 		/// <param name="DefaultProperties">Default properties available to the script</param>
+		/// <param name="Arguments">Arguments passed in to the graph on the command line</param>
 		/// <param name="bPreprocessOnly">Preprocess the file, but do not expand any values that are not portable (eg. paths on the local machine)</param>
+		/// <param name="SingleNodeName">If a single node will be processed, the name of that node.</param>
 		/// <param name="Schema">Schema for the script</param>
 		/// <param name="Logger">Logger for diagnostic messages</param>
-		private BgScriptReader(IBgScriptReaderContext Context, IDictionary<string, string> DefaultProperties, bool bPreprocessOnly, BgScriptSchema Schema, ILogger Logger)
+		private BgScriptReader(IBgScriptReaderContext Context, IDictionary<string, string> DefaultProperties, IReadOnlyDictionary<string, string> Arguments, bool bPreprocessOnly, string? SingleNodeName, BgScriptSchema Schema, ILogger Logger)
 		{
 			this.Context = Context;
+			this.Arguments = new Dictionary<string, string>(Arguments, StringComparer.OrdinalIgnoreCase);
 			this.bPreprocessOnly = bPreprocessOnly;
+			this.SingleNodeName = SingleNodeName; 
 			this.Schema = Schema;
 			this.Logger = Logger;
 
@@ -406,8 +424,8 @@ namespace EpicGames.BuildGraph
 		public static async Task<BgGraph?> ReadAsync(IBgScriptReaderContext Context, string File, Dictionary<string, string> Arguments, Dictionary<string, string> DefaultProperties, bool bPreprocessOnly, BgScriptSchema Schema, ILogger Logger, string? SingleNodeName = null)
 		{
 			// Read the file and build the graph
-			BgScriptReader Reader = new BgScriptReader(Context, DefaultProperties, bPreprocessOnly, Schema, Logger);
-			if (!await Reader.TryReadAsync(File, Arguments, Logger, SingleNodeName) || Reader.NumErrors > 0)
+			BgScriptReader Reader = new BgScriptReader(Context, DefaultProperties, Arguments, bPreprocessOnly, SingleNodeName, Schema, Logger);
+			if (!await Reader.TryReadAsync(File) || Reader.NumErrors > 0)
 			{
 				return null;
 			}
@@ -432,10 +450,7 @@ namespace EpicGames.BuildGraph
 		/// Read the script from the given file
 		/// </summary>
 		/// <param name="File">File to read from</param>
-		/// <param name="Arguments">Arguments passed in to the graph on the command line</param>
-		/// <param name="Logger">Logger for output messages</param>
-		/// <param name="SingleNodeName">The name of the node if only a single node is going to be built, otherwise null.</param>
-		async Task<bool> TryReadAsync(string File, Dictionary<string, string> Arguments, ILogger Logger, string? SingleNodeName = null)
+		async Task<bool> TryReadAsync(string File)
 		{
 			// Get the data for this file
 			byte[]? Data = await Context.ReadAsync(File);
@@ -455,7 +470,6 @@ namespace EpicGames.BuildGraph
 			}
 
 			// Read the root BuildGraph element
-			this.SingleNodeName = SingleNodeName;
 			await ReadGraphBodyAsync(Document.DocumentElement, Arguments);
 			return true;
 		}
@@ -538,31 +552,13 @@ namespace EpicGames.BuildGraph
 			}
 		}
 
-		/// <summary>
-		/// Handles validation messages from validating the document against its schema
-		/// </summary>
-		/// <param name="Sender">The source of the event</param>
-		/// <param name="Args">Event arguments</param>
-		void ValidationHandler(object Sender, ValidationEventArgs Args)
-		{
-			if (Args.Severity == XmlSeverityType.Warning)
-			{
-				Logger.LogWarning("Script: {Message}", Args.Message);
-			}
-			else
-			{
-				Logger.LogError("Script: {Message}", Args.Message);
-				NumErrors++;
-			}
-		}
-
 		/// <summary>
 		/// Push a new property scope onto the stack
 		/// </summary>
 		void EnterScope()
 		{
-			ScopedProperties.Add(new Dictionary<string, string>(StringComparer.InvariantCultureIgnoreCase));
-			ShadowProperties.Add(new HashSet<string>(StringComparer.InvariantCultureIgnoreCase));
+			ScopedProperties.Add(new Dictionary<string, string>(StringComparer.OrdinalIgnoreCase));
+			ShadowProperties.Add(new HashSet<string>(StringComparer.OrdinalIgnoreCase));
 		}
 
 		/// <summary>
@@ -707,7 +703,7 @@ namespace EpicGames.BuildGraph
 				foreach(string File in Files.OrderBy(x => x, StringComparer.OrdinalIgnoreCase))
 				{
 					Logger.LogDebug("Including file {File}", File);
-					await TryReadAsync(File, Arguments, Logger);
+					await TryReadAsync(File);
 				}
 			}
 		}
