BuildGraph: Separate out a base class for the script parser, allowing the structure of a script to be examined without processing anything.

diff --git a/Engine/Source/Programs/AutomationTool/BuildGraph/BuildGraph.cs b/Engine/Source/Programs/AutomationTool/BuildGraph/BuildGraph.cs
--- a/Engine/Source/Programs/AutomationTool/BuildGraph/BuildGraph.cs
+++ b/Engine/Source/Programs/AutomationTool/BuildGraph/BuildGraph.cs
@@ -415,7 +415,7 @@ namespace AutomationTool
 			ScriptFileName = FullScriptFile.MakeRelativeTo(Unreal.RootDirectory).Replace('\\', '/');
 
 			// Read the script from disk
-			BgGraph Graph = BgScriptReader.ReadAsync(Context, ScriptFileName, Arguments, DefaultProperties, PreprocessedFileName != null, Schema, Logger, SingleNodeName).Result;
+			BgGraph Graph = BgScriptReader.ReadAsync(Context, ScriptFileName, Arguments, DefaultProperties, Schema, Logger, SingleNodeName).Result;
 			if(Graph == null)
 			{
 				return ExitCode.Error_Unknown;
diff --git a/Engine/Source/Programs/Shared/EpicGames.BuildGraph/BgScriptReader.cs b/Engine/Source/Programs/Shared/EpicGames.BuildGraph/BgScriptReader.cs
--- a/Engine/Source/Programs/Shared/EpicGames.BuildGraph/BgScriptReader.cs
+++ b/Engine/Source/Programs/Shared/EpicGames.BuildGraph/BgScriptReader.cs
@@ -192,7 +192,7 @@
 	/// <summary>
 	/// Implementation of XmlElement which preserves line numbers
 	/// </summary>
-	class BgScriptElement : XmlElement
+	public class BgScriptElement : XmlElement
 	{
 		/// <summary>
 		/// Location of the element within the file
@@ -323,23 +323,18 @@
 	/// <summary>
 	/// Reader for build graph definitions. Instanced to contain temporary state; public interface is through ScriptReader.TryRead().
 	/// </summary>
-	public class BgScriptReader
+	public abstract class BgScriptReaderBase
 	{
 		/// <summary>
 		/// Interface used for reading files
 		/// </summary>
-		IBgScriptReaderContext Context;
+		protected IBgScriptReaderContext Context { get; }
 
 		/// <summary>
-		/// The current graph
-		/// </summary>
-		BgGraph Graph = new BgGraph();
-
-		/// <summary>
 		/// List of property name to value lookups. Modifications to properties are scoped to nodes and agents. EnterScope() pushes an empty dictionary onto the end of this list, and LeaveScope() removes one. 
 		/// ExpandProperties() searches from last to first lookup when trying to resolve a property name, and takes the first it finds.
 		/// </summary>
-		List<Dictionary<string, string>> ScopedProperties = new List<Dictionary<string, string>>();
+		protected List<Dictionary<string, string>> ScopedProperties { get; } = new List<Dictionary<string, string>>();
 
 		/// <summary>
 		/// When declaring a property in a nested scope, we enter its name into a set for each parent scope which prevents redeclaration in an OUTER scope later. Subsequent NESTED scopes can redeclare it.
@@ -353,104 +348,40 @@
 		Dictionary<string, BgScriptMacro> MacroNameToDefinition = new Dictionary<string, BgScriptMacro>();
 
 		/// <summary>
-		/// Arguments for evaluating the graph
-		/// </summary>
-		Dictionary<string, string> Arguments;
-
-		/// <summary>
-		/// Preprocess the file, but do not expand any values that are not portable (eg. paths on the local machine)
-		/// </summary>
-		bool bPreprocessOnly;
-
-		/// <summary>
 		/// Schema for the script
 		/// </summary>
-		BgScriptSchema Schema;
+		BgScriptSchema Schema { get; }
 
 		/// <summary>
 		/// Logger for diagnostic messages
 		/// </summary>
-		ILogger Logger;
+		protected ILogger Logger { get; }
 
 		/// <summary>
 		/// The number of errors encountered during processing so far
 		/// </summary>
-		int NumErrors;
+		public int NumErrors { get; private set; }
 
 		/// <summary>
-		/// The name of the node if only a single node is going to be built, otherwise null.
-		/// </summary>
-		string? SingleNodeName;
-
-		/// <summary>
 		/// Private constructor. Use ScriptReader.TryRead() to read a script file.
 		/// </summary>
 		/// <param name="Context">Context object</param>
-		/// <param name="DefaultProperties">Default properties available to the script</param>
-		/// <param name="Arguments">Arguments passed in to the graph on the command line</param>
-		/// <param name="bPreprocessOnly">Preprocess the file, but do not expand any values that are not portable (eg. paths on the local machine)</param>
-		/// <param name="SingleNodeName">If a single node will be processed, the name of that node.</param>
 		/// <param name="Schema">Schema for the script</param>
 		/// <param name="Logger">Logger for diagnostic messages</param>
-		private BgScriptReader(IBgScriptReaderContext Context, IDictionary<string, string> DefaultProperties, IReadOnlyDictionary<string, string> Arguments, bool bPreprocessOnly, string? SingleNodeName, BgScriptSchema Schema, ILogger Logger)
+		protected BgScriptReaderBase(IBgScriptReaderContext Context, BgScriptSchema Schema, ILogger Logger)
 		{
 			this.Context = Context;
-			this.Arguments = new Dictionary<string, string>(Arguments, StringComparer.OrdinalIgnoreCase);
-			this.bPreprocessOnly = bPreprocessOnly;
-			this.SingleNodeName = SingleNodeName; 
 			this.Schema = Schema;
 			this.Logger = Logger;
 
 			EnterScope();
-
-			foreach(KeyValuePair<string, string> Pair in DefaultProperties)
-			{
-				ScopedProperties[ScopedProperties.Count - 1].Add(Pair.Key, Pair.Value);
-			}
 		}
 
 		/// <summary>
-		/// Try to read a script file from the given file.
-		/// </summary>
-		/// <param name="Context">Supplies context about the parse</param>
-		/// <param name="File">File to read from</param>
-		/// <param name="Arguments">Arguments passed in to the graph on the command line</param>
-		/// <param name="DefaultProperties">Default properties available to the script</param>
-		/// <param name="bPreprocessOnly">Preprocess the file, but do not expand any values that are not portable (eg. paths on the local machine)</param>
-		/// <param name="Schema">Schema for the script</param>
-		/// <param name="Logger">Logger for output messages</param>
-		/// <param name="SingleNodeName">If a single node will be processed, the name of that node.</param>
-		/// <returns>True if the graph was read, false if there were errors</returns>
-		public static async Task<BgGraph?> ReadAsync(IBgScriptReaderContext Context, string File, Dictionary<string, string> Arguments, Dictionary<string, string> DefaultProperties, bool bPreprocessOnly, BgScriptSchema Schema, ILogger Logger, string? SingleNodeName = null)
-		{
-			// Read the file and build the graph
-			BgScriptReader Reader = new BgScriptReader(Context, DefaultProperties, Arguments, bPreprocessOnly, SingleNodeName, Schema, Logger);
-			if (!await Reader.TryReadAsync(File) || Reader.NumErrors > 0)
-			{
-				return null;
-			}
-
-			// Make sure all the arguments were valid
-			HashSet<string> ValidArgumentNames = new HashSet<string>(Reader.Graph.Options.Select(x => x.Name), StringComparer.OrdinalIgnoreCase);
-			ValidArgumentNames.Add("PreflightChange");
-
-			foreach(string ArgumentName in Arguments.Keys)
-			{
-				if (!ValidArgumentNames.Contains(ArgumentName))
-				{
-					Logger.LogWarning("Unknown argument '{ArgumentName}' for '{Script}'", ArgumentName, Context.GetNativePath(File));
-				}
-			}
-
-			// Return the constructed graph
-			return Reader.Graph;
-		}
-
-		/// <summary>
 		/// Read the script from the given file
 		/// </summary>
 		/// <param name="File">File to read from</param>
-		async Task<bool> TryReadAsync(string File)
+		protected async Task<bool> TryReadAsync(string File)
 		{
 			// Get the data for this file
 			byte[]? Data = await Context.ReadAsync(File);
@@ -470,7 +401,7 @@
 			}
 
 			// Read the root BuildGraph element
-			await ReadGraphBodyAsync(Document.DocumentElement, Arguments);
+			await ReadGraphBodyAsync(Document.DocumentElement);
 			return true;
 		}
 
@@ -479,17 +410,17 @@
 		/// </summary>
 		/// <param name="Element">The parent element to read from</param>
 		/// <param name="Arguments">Arguments passed in to the graph on the command line</param>
-		async Task ReadGraphBodyAsync(XmlElement Element, Dictionary<string, string> Arguments)
+		async Task ReadGraphBodyAsync(XmlElement Element)
 		{
 			foreach (BgScriptElement ChildElement in Element.ChildNodes.OfType<BgScriptElement>())
 			{
 				switch (ChildElement.Name)
 				{
 					case "Include":
-						await ReadIncludeAsync(ChildElement, Arguments);
+						await ReadIncludeAsync(ChildElement);
 						break;
 					case "Option":
-						await ReadOptionAsync(ChildElement, Arguments);
+						await ReadOptionAsync(ChildElement);
 						break;
 					case "Property":
 						await ReadPropertyAsync(ChildElement);
@@ -525,25 +456,25 @@
 						await ReadNotifierAsync(ChildElement);
 						break;
 					case "Trace":
-						await ReadDiagnosticAsync(ChildElement, LogEventType.Console, null, null);
+						await ReadDiagnosticAsync(ChildElement, LogEventType.Console);
 						break;
 					case "Warning":
-						await ReadDiagnosticAsync(ChildElement, LogEventType.Warning, null, null);
+						await ReadDiagnosticAsync(ChildElement, LogEventType.Warning);
 						break;
 					case "Error":
-						await ReadDiagnosticAsync(ChildElement, LogEventType.Error, null, null);
+						await ReadDiagnosticAsync(ChildElement, LogEventType.Error);
 						break;
 					case "Do":
-						await ReadBlockAsync(ChildElement, x => ReadGraphBodyAsync(x, Arguments));
+						await ReadBlockAsync(ChildElement, ReadGraphBodyAsync);
 						break;
 					case "Switch":
-						await ReadSwitchAsync(ChildElement, x => ReadGraphBodyAsync(x, Arguments));
+						await ReadSwitchAsync(ChildElement, ReadGraphBodyAsync);
 						break;
 					case "ForEach":
-						await ReadForEachAsync(ChildElement, x => ReadGraphBodyAsync(x, Arguments));
+						await ReadForEachAsync(ChildElement, ReadGraphBodyAsync);
 						break;
 					case "Expand":
-						await ReadExpandAsync(ChildElement, x => ReadGraphBodyAsync(x, Arguments));
+						await ReadExpandAsync(ChildElement, ReadGraphBodyAsync);
 						break;
 					default:
 						LogError(ChildElement, "Invalid element '{0}'", ChildElement.Name);
@@ -555,7 +486,7 @@
 		/// <summary>
 		/// Push a new property scope onto the stack
 		/// </summary>
-		void EnterScope()
+		protected void EnterScope()
 		{
 			ScopedProperties.Add(new Dictionary<string, string>(StringComparer.OrdinalIgnoreCase));
 			ShadowProperties.Add(new HashSet<string>(StringComparer.OrdinalIgnoreCase));
@@ -564,7 +495,7 @@
 		/// <summary>
 		/// Pop a property scope from the stack
 		/// </summary>
-		void LeaveScope()
+		protected void LeaveScope()
 		{
 			ScopedProperties.RemoveAt(ScopedProperties.Count - 1);
 			ShadowProperties.RemoveAt(ShadowProperties.Count - 1);
@@ -576,7 +507,7 @@
 		/// <param name="Element">Element containing the property assignment. Used for error messages if the property is shadowed in another scope.</param>
 		/// <param name="Name">Name of the property</param>
 		/// <param name="Value">Value for the property</param>
-		void SetPropertyValue(BgScriptElement Element, string Name, string Value)
+		protected void SetPropertyValue(BgScriptElement Element, string Name, string Value)
 		{
 			// Find the scope containing this property, defaulting to the current scope
 			int ScopeIdx = 0;
@@ -607,7 +538,7 @@
 		/// <param name="Name">Name of the property</param>
 		/// <param name="Value">On success, contains the value of the property. Set to null otherwise.</param>
 		/// <returns>True if the property was found, false otherwise</returns>
-		bool TryGetPropertyValue(string Name, out string? Value)
+		protected bool TryGetPropertyValue(string Name, out string? Value)
 		{
 			int ValueLength = 0;
 			if (Name.Contains(":"))
@@ -681,8 +612,7 @@
 		/// Read an include directive, and the contents of the target file
 		/// </summary>
 		/// <param name="Element">Xml element to read the definition from</param>
-		/// <param name="Arguments">Arguments passed in to the graph on the command line</param>
-		async Task ReadIncludeAsync(BgScriptElement Element, Dictionary<string, string> Arguments)
+		async Task ReadIncludeAsync(BgScriptElement Element)
 		{
 			if (await EvaluateConditionAsync(Element))
 			{
@@ -712,171 +642,25 @@
 		/// Reads the definition of a graph option; a parameter which can be set by the user on the command-line or via an environment variable.
 		/// </summary>
 		/// <param name="Element">Xml element to read the definition from</param>
-		/// <param name="Arguments">Arguments passed in to the graph on the command line</param>
-		async Task ReadOptionAsync(BgScriptElement Element, IDictionary<string, string> Arguments)
-		{
-			if (await EvaluateConditionAsync(Element))
-			{
-				string Name = ReadAttribute(Element, "Name");
-				if (ValidateName(Element, Name))
-				{
-					// Make sure we're at global scope
-					if(ScopedProperties.Count > 1)
-					{
-						throw new Exception("Incorrect scope depth for reading option settings");
-					}
-
-					// Check if the property already exists. If it does, we don't need to register it as an option.
-					string? ExistingValue;
-					if(TryGetPropertyValue(Name, out ExistingValue) && ExistingValue != null)
-					{
-						// If there's a restriction on this definition, check it matches
-						string Restrict = ReadAttribute(Element, "Restrict");
-						if(!String.IsNullOrEmpty(Restrict) && !Regex.IsMatch(ExistingValue, "^" + Restrict + "$", RegexOptions.IgnoreCase))
-						{
-							LogError(Element, "'{0} is already set to '{1}', which does not match the given restriction ('{2}')", Name, ExistingValue, Restrict);
-						}
-					}
-					else
-					{
-						// Create a new option object to store the settings
-						string Description = ReadAttribute(Element, "Description");
-						string DefaultValue = ReadAttribute(Element, "DefaultValue");
-						BgScriptOption Option = new BgScriptOption(Name, Description, DefaultValue);
-						Graph.Options.Add(Option);
-
-						// Get the value of this property
-						string? Value;
-						if(!Arguments.TryGetValue(Name, out Value))
-						{
-							Value = Option.DefaultValue;
-						}
-						SetPropertyValue(Element, Name, Value);
-
-						// If there's a restriction on it, check it's valid
-						string Restrict = ReadAttribute(Element, "Restrict");
-						if(!String.IsNullOrEmpty(Restrict))
-						{
-							string Pattern = "^(" + Restrict + ")$";
-							if(!Regex.IsMatch(Value, Pattern, RegexOptions.IgnoreCase))
-							{
-								LogError(Element, "'{0}' is not a valid value for '{1}' (required: '{2}')", Value, Name, Restrict);
-							}
-							if(Option.DefaultValue != Value && !Regex.IsMatch(Option.DefaultValue, Pattern, RegexOptions.IgnoreCase))
-							{
-								LogError(Element, "Default value '{0}' is not valid for '{1}' (required: '{2}')", Option.DefaultValue, Name, Restrict);
-							}
-						}
-					}
-				}
-			}
-		}
+		protected abstract Task ReadOptionAsync(BgScriptElement Element);
 
 		/// <summary>
 		/// Reads a property assignment.
 		/// </summary>
 		/// <param name="Element">Xml element to read the definition from</param>
-		async Task ReadPropertyAsync(BgScriptElement Element)
-		{
-			if (await EvaluateConditionAsync(Element))
-			{
-				string Name = ReadAttribute(Element, "Name");
-				if (ValidateName(Element, Name))
-				{
-					string Value = ReadAttribute(Element, "Value");
-					if(Element.HasChildNodes)
-					{
-						// Read the element content, and append each line to the value as a semicolon delimited list
-						StringBuilder Builder = new StringBuilder(Value);
-						foreach(string Line in Element.InnerText.Split('\n'))
-						{
-							string TrimLine = ExpandProperties(Element, Line.Trim());
-							if(TrimLine.Length > 0)
-							{
-								if(Builder.Length > 0)
-								{
-									Builder.Append(";");
-								}
-								Builder.Append(TrimLine);
-							}
-						}
-						Value = Builder.ToString();
-					}
-					SetPropertyValue(Element, Name, Value);
-				}
-			}
-		}
+		protected abstract Task ReadPropertyAsync(BgScriptElement Element);
 
 		/// <summary>
 		/// Reads a Regex assignment.
 		/// </summary>
 		/// <param name="Element">Xml element to read the definition from</param>
-		async Task ReadRegexAsync(BgScriptElement Element)
-		{
-			if (await EvaluateConditionAsync(Element))
-			{
-				// Get the pattern
-				string RegexString = ReadAttribute(Element, "Pattern");
-
-				// Make sure its a valid regex.
-				Regex? RegexValue = ParseRegex(Element, RegexString);
-				if (RegexValue != null)
-				{
-					// read the names in 
-					string[] CaptureNames = ReadListAttribute(Element, "Capture");
-
-					// get number of groups we passed in
-					int[] GroupNumbers = RegexValue.GetGroupNumbers();
-
-					// make sure the number of property names is the same as the number of match groups
-					// this includes the entire string match group as [0], so don't count that one.
-					if (CaptureNames.Length != GroupNumbers.Count() - 1)
-					{
-						LogError(Element, "MatchGroup count: {0} does not match the number of names specified: {1}", GroupNumbers.Count() - 1, CaptureNames.Length);
-					}
-					else
-					{
-						// apply the regex to the value
-						string Input = ReadAttribute(Element, "Input");
-						Match Match = RegexValue.Match(Input);
-
-						bool Optional = await BgCondition.EvaluateAsync(ReadAttribute(Element, "Optional"), Context);
-						if (!Match.Success)
-						{
-							if (!Optional)
-							{
-								LogError(Element, "Regex {0} did not find a match against input string {1}", RegexString, Input);
-							}
-						}
-						else
-						{
-							// assign each property to the group it matches, skip over [0]
-							for (int MatchIdx = 1; MatchIdx < GroupNumbers.Count(); MatchIdx++)
-							{
-								SetPropertyValue(Element, CaptureNames[MatchIdx - 1], Match.Groups[MatchIdx].Value);
-							}
-						}
-					}
-				}
-			}
-		}
+		protected abstract Task ReadRegexAsync(BgScriptElement Element);
 
 		/// <summary>
 		/// Reads a property assignment from an environment variable.
 		/// </summary>
 		/// <param name="Element">Xml element to read the definition from</param>
-		async Task ReadEnvVarAsync(BgScriptElement Element)
-		{
-			if (await EvaluateConditionAsync(Element))
-			{
-				string Name = ReadAttribute(Element, "Name");
-				if (ValidateName(Element, Name))
-				{
-					string Value = Environment.GetEnvironmentVariable(Name) ?? "";
-					SetPropertyValue(Element, Name, Value);
-				}
-			}
-		}
+		protected abstract Task ReadEnvVarAsync(BgScriptElement Element);
 
 		/// <summary>
 		/// Reads a macro definition
@@ -957,48 +741,14 @@
 		/// Reads the definition for an agent.
 		/// </summary>
 		/// <param name="Element">Xml element to read the definition from</param>
-		async Task ReadAgentAsync(BgScriptElement Element)
-		{
-			string? Name;
-			if (await EvaluateConditionAsync(Element) && TryReadObjectName(Element, out Name))
-			{
-				// Read the valid agent types. This may be omitted if we're continuing an existing agent.
-				string[] Types = ReadListAttribute(Element, "Type");
+		protected abstract Task ReadAgentAsync(BgScriptElement Element);
 
-				// Create the agent object, or continue an existing one
-				BgAgent? Agent;
-				if (Graph.NameToAgent.TryGetValue(Name, out Agent))
-				{
-					if (Types.Length > 0 && Agent.PossibleTypes.Length > 0)
-					{
-						if (Types.Length != Agent.PossibleTypes.Length || !Types.SequenceEqual(Agent.PossibleTypes, StringComparer.InvariantCultureIgnoreCase))
-						{
-							LogError(Element, "Agent types ({0}) were different than previous agent definition with types ({1}). Must either be empty or match exactly.", string.Join(",", Types), string.Join(",", Agent.PossibleTypes));
-						}
-					}
-				}
-				else
-				{
-					if (Types.Length == 0)
-					{
-						LogError(Element, "Missing type for agent '{0}'", Name);
-					}
-					Agent = new BgAgent(Name, Types);
-					Graph.NameToAgent.Add(Name, Agent);
-					Graph.Agents.Add(Agent);
-				}
-
-				// Process all the child elements.
-				await ReadAgentBodyAsync(Element, Agent);
-			}
-		}
-
 		/// <summary>
 		/// Read the contents of an agent definition
 		/// </summary>
 		/// <param name="Element">Xml element to read the definition from</param>
 		/// <param name="ParentAgent">The agent to contain the definition</param>
-		async Task ReadAgentBodyAsync(BgScriptElement Element, BgAgent ParentAgent)
+		protected async Task ReadAgentBodyAsync(BgScriptElement Element)
 		{
 			EnterScope();
 			foreach (BgScriptElement ChildElement in Element.ChildNodes.OfType<BgScriptElement>())
@@ -1012,34 +762,34 @@
 						await ReadRegexAsync(ChildElement);
 						break;
 					case "Node":
-						await ReadNodeAsync(ChildElement, ParentAgent);
+						await ReadNodeAsync(ChildElement);
 						break;
 					case "Aggregate":
 						await ReadAggregateAsync(ChildElement);
 						break;
 					case "Trace":
-						await ReadDiagnosticAsync(ChildElement, LogEventType.Console, null, ParentAgent);
+						await ReadDiagnosticAsync(ChildElement, LogEventType.Console);
 						break;
 					case "Warning":
-						await ReadDiagnosticAsync(ChildElement, LogEventType.Warning, null, ParentAgent);
+						await ReadDiagnosticAsync(ChildElement, LogEventType.Warning);
 						break;
 					case "Error":
-						await ReadDiagnosticAsync(ChildElement, LogEventType.Error, null, ParentAgent);
+						await ReadDiagnosticAsync(ChildElement, LogEventType.Error);
 						break;
 					case "Label":
 						await ReadLabelAsync(ChildElement);
 						break;
 					case "Do":
-						await ReadBlockAsync(ChildElement, x => ReadAgentBodyAsync(x, ParentAgent));
+						await ReadBlockAsync(ChildElement, ReadAgentBodyAsync);
 						break;
 					case "Switch":
-						await ReadSwitchAsync(ChildElement, x => ReadAgentBodyAsync(x, ParentAgent));
+						await ReadSwitchAsync(ChildElement, ReadAgentBodyAsync);
 						break;
 					case "ForEach":
-						await ReadForEachAsync(ChildElement, x => ReadAgentBodyAsync(x, ParentAgent));
+						await ReadForEachAsync(ChildElement, ReadAgentBodyAsync);
 						break;
 					case "Expand":
-						await ReadExpandAsync(ChildElement, x => ReadAgentBodyAsync(x, ParentAgent));
+						await ReadExpandAsync(ChildElement, ReadAgentBodyAsync);
 						break;
 					default:
 						LogError(ChildElement, "Unexpected element type '{0}'", ChildElement.Name);
@@ -1053,273 +803,37 @@
 		/// Reads the definition for an aggregate
 		/// </summary>
 		/// <param name="Element">Xml element to read the definition from</param>
-		async Task ReadAggregateAsync(BgScriptElement Element)
-		{
-			string? Name;
-			if (await EvaluateConditionAsync(Element) && TryReadObjectName(Element, out Name) && CheckNameIsUnique(Element, Name))
-			{
-				string[] RequiredNames = ReadListAttribute(Element, "Requires");
-
-				BgAggregate NewAggregate = new BgAggregate(Name);
-				foreach (BgNode ReferencedNode in ResolveReferences(Element, RequiredNames))
-				{
-					NewAggregate.RequiredNodes.Add(ReferencedNode);
-				}
-				Graph.NameToAggregate[Name] = NewAggregate;
-
-				string LabelCategoryName = ReadAttribute(Element, "Label");
-				if (!String.IsNullOrEmpty(LabelCategoryName))
-				{
-					BgLabel Label;
-
-					// Create the label
-					int SlashIdx = LabelCategoryName.IndexOf('/');
-					if (SlashIdx != -1)
-					{
-						Label = new BgLabel(LabelCategoryName.Substring(SlashIdx + 1), LabelCategoryName.Substring(0, SlashIdx), null, null, BgLabelChange.Current);
-					}
-					else
-					{
-						Label = new BgLabel(LabelCategoryName, "Other", null, null, BgLabelChange.Current);
-					}
-
-					// Find all the included nodes
-					foreach (BgNode RequiredNode in NewAggregate.RequiredNodes)
-					{
-						Label.RequiredNodes.Add(RequiredNode);
-						Label.IncludedNodes.Add(RequiredNode);
-						Label.IncludedNodes.UnionWith(RequiredNode.OrderDependencies);
-					}
-
-					string[] IncludedNames = ReadListAttribute(Element, "Include");
-					foreach (BgNode IncludedNode in ResolveReferences(Element, IncludedNames))
-					{
-						Label.IncludedNodes.Add(IncludedNode);
-						Label.IncludedNodes.UnionWith(IncludedNode.OrderDependencies);
-					}
-
-					string[] ExcludedNames = ReadListAttribute(Element, "Exclude");
-					foreach (BgNode ExcludedNode in ResolveReferences(Element, ExcludedNames))
-					{
-						Label.IncludedNodes.Remove(ExcludedNode);
-						Label.IncludedNodes.ExceptWith(ExcludedNode.OrderDependencies);
-					}
-
-					Graph.Labels.Add(Label);
-				}
-			}
-		}
+		protected abstract Task ReadAggregateAsync(BgScriptElement Element);
 
 		/// <summary>
 		/// Reads the definition for a report
 		/// </summary>
 		/// <param name="Element">Xml element to read the definition from</param>
-		async Task ReadReportAsync(BgScriptElement Element)
-		{
-			string? Name;
-			if (await EvaluateConditionAsync(Element) && TryReadObjectName(Element, out Name) && CheckNameIsUnique(Element, Name))
-			{
-				string[] RequiredNames = ReadListAttribute(Element, "Requires");
-
-				BgReport NewReport = new BgReport(Name);
-				foreach (BgNode ReferencedNode in ResolveReferences(Element, RequiredNames))
-				{
-					NewReport.Nodes.Add(ReferencedNode);
-					NewReport.Nodes.UnionWith(ReferencedNode.OrderDependencies);
-				}
-				Graph.NameToReport.Add(Name, NewReport);
-			}
-		}
+		protected abstract Task ReadReportAsync(BgScriptElement Element);
 
 		/// <summary>
 		/// Reads the definition for a badge
 		/// </summary>
 		/// <param name="Element">Xml element to read the definition from</param>
-		async Task ReadBadgeAsync(BgScriptElement Element)
-		{
-			string? Name;
-			if (await EvaluateConditionAsync(Element) && TryReadObjectName(Element, out Name))
-			{
-				string[] RequiredNames = ReadListAttribute(Element, "Requires");
-				string[] TargetNames = ReadListAttribute(Element, "Targets");
-				string Project = ReadAttribute(Element, "Project");
-				int Change = ReadIntegerAttribute(Element, "Change", 0);
-
-				BgBadge NewBadge = new BgBadge(Name, Project, Change);
-				foreach (BgNode ReferencedNode in ResolveReferences(Element, RequiredNames))
-				{
-					NewBadge.Nodes.Add(ReferencedNode);
-				}
-				foreach (BgNode ReferencedNode in ResolveReferences(Element, TargetNames))
-				{
-					NewBadge.Nodes.Add(ReferencedNode);
-					NewBadge.Nodes.UnionWith(ReferencedNode.OrderDependencies);
-				}
-				Graph.Badges.Add(NewBadge);
-			}
-		}
+		protected abstract Task ReadBadgeAsync(BgScriptElement Element);
 
 		/// <summary>
 		/// Reads the definition for a label
 		/// </summary>
 		/// <param name="Element">Xml element to read the definition from</param>
-		async Task ReadLabelAsync(BgScriptElement Element)
-		{
-			if (await EvaluateConditionAsync(Element))
-			{
-				string Name = ReadAttribute(Element, "Name");
-				if (!String.IsNullOrEmpty(Name))
-				{
-					ValidateName(Element, Name);
-				}
-
-				string Category = ReadAttribute(Element, "Category");
-
-				string[] RequiredNames = ReadListAttribute(Element, "Requires");
-				string[] IncludedNames = ReadListAttribute(Element, "Include");
-				string[] ExcludedNames = ReadListAttribute(Element, "Exclude");
-
-				string UgsBadge = ReadAttribute(Element, "UgsBadge");
-				string UgsProject = ReadAttribute(Element, "UgsProject");
-
-				BgLabelChange Change = ReadEnumAttribute<BgLabelChange>(Element, "Change", BgLabelChange.Current);
-
-				BgLabel NewLabel = new BgLabel(Name, Category, UgsBadge, UgsProject, Change);
-				foreach (BgNode ReferencedNode in ResolveReferences(Element, RequiredNames))
-				{
-					NewLabel.RequiredNodes.Add(ReferencedNode);
-					NewLabel.IncludedNodes.Add(ReferencedNode);
-					NewLabel.IncludedNodes.UnionWith(ReferencedNode.OrderDependencies);
-				}
-				foreach (BgNode IncludedNode in ResolveReferences(Element, IncludedNames))
-				{
-					NewLabel.IncludedNodes.Add(IncludedNode);
-					NewLabel.IncludedNodes.UnionWith(IncludedNode.OrderDependencies);
-				}
-				foreach (BgNode ExcludedNode in ResolveReferences(Element, ExcludedNames))
-				{
-					NewLabel.IncludedNodes.Remove(ExcludedNode);
-					NewLabel.IncludedNodes.ExceptWith(ExcludedNode.OrderDependencies);
-				}
-				Graph.Labels.Add(NewLabel);
-			}
-		}
+		protected abstract Task ReadLabelAsync(BgScriptElement Element);
 
 		/// <summary>
 		/// Reads the definition for a node, and adds it to the given agent
 		/// </summary>
 		/// <param name="Element">Xml element to read the definition from</param>
-		/// <param name="ParentAgent">Agent for the node to be added to</param>
-		async Task ReadNodeAsync(BgScriptElement Element, BgAgent ParentAgent)
-		{
-			string? Name;
-			if (await EvaluateConditionAsync(Element) && TryReadObjectName(Element, out Name))
-			{
-				string[] RequiresNames = ReadListAttribute(Element, "Requires");
-				string[] ProducesNames = ReadListAttribute(Element, "Produces");
-				string[] AfterNames = ReadListAttribute(Element, "After");
-				string[] TokenFileNames = ReadListAttribute(Element, "Token");
-				bool bRunEarly = ReadBooleanAttribute(Element, "RunEarly", false);
-				bool bNotifyOnWarnings = ReadBooleanAttribute(Element, "NotifyOnWarnings", true);
-
-				// Resolve all the inputs we depend on
-				HashSet<BgNodeOutput> Inputs = ResolveInputReferences(Element, RequiresNames);
-
-				// Gather up all the input dependencies, and check they're all upstream of the current node
-				HashSet<BgNode> InputDependencies = new HashSet<BgNode>();
-				foreach (BgNode InputDependency in Inputs.Select(x => x.ProducingNode).Distinct())
-				{
-					InputDependencies.Add(InputDependency);
-				}
-
-				// Remove all the lock names from the list of required names
-				HashSet<FileReference> RequiredTokens = new HashSet<FileReference>(TokenFileNames.Select(x => new FileReference(x)));
-
-				// Recursively include all their dependencies too
-				foreach (BgNode InputDependency in InputDependencies.ToArray())
-				{
-					RequiredTokens.UnionWith(InputDependency.RequiredTokens);
-					InputDependencies.UnionWith(InputDependency.InputDependencies);
-				}
-
-				// Validate all the outputs
-				List<string> ValidOutputNames = new List<string>();
-				foreach (string ProducesName in ProducesNames)
-				{
-					BgNodeOutput? ExistingOutput;
-					if(Graph.TagNameToNodeOutput.TryGetValue(ProducesName, out ExistingOutput))
-					{
-						LogError(Element, "Output tag '{0}' is already generated by node '{1}'", ProducesName, ExistingOutput.ProducingNode.Name);
-					}
-					else if(!ProducesName.StartsWith("#"))
-					{
-						LogError(Element, "Output tag names must begin with a '#' character ('{0}')", ProducesName);
-					}
-					else
-					{
-						ValidOutputNames.Add(ProducesName);
-					}
-				}
-
-				// Gather up all the order dependencies
-				HashSet<BgNode> OrderDependencies = new HashSet<BgNode>(InputDependencies);
-				OrderDependencies.UnionWith(ResolveReferences(Element, AfterNames));
-
-				// Recursively include all their order dependencies too
-				foreach (BgNode OrderDependency in OrderDependencies.ToArray())
-				{
-					OrderDependencies.UnionWith(OrderDependency.OrderDependencies);
-				}
-
-				// Check that we're not dependent on anything completing that is declared after the initial declaration of this agent.
-				int AgentIdx = Graph.Agents.IndexOf(ParentAgent);
-				for (int Idx = AgentIdx + 1; Idx < Graph.Agents.Count; Idx++)
-				{
-					foreach (BgNode Node in Graph.Agents[Idx].Nodes.Where(x => OrderDependencies.Contains(x)))
-					{
-						LogError(Element, "Node '{0}' has a dependency on '{1}', which was declared after the initial definition of '{2}'.", Name, Node.Name, ParentAgent.Name);
-					}
-				}
-
-				// Construct and register the node
-				if (CheckNameIsUnique(Element, Name))
-				{
-					// Add it to the node lookup
-					BgNode NewNode = new BgNode(Name, Inputs.ToArray(), ValidOutputNames.ToArray(), InputDependencies.ToArray(), OrderDependencies.ToArray(), RequiredTokens.ToArray());
-					NewNode.bRunEarly = bRunEarly;
-					NewNode.bNotifyOnWarnings = bNotifyOnWarnings;
-					Graph.NameToNode.Add(Name, NewNode);
-
-					// Register all the output tags in the global name table.
-					foreach(BgNodeOutput Output in NewNode.Outputs)
-					{
-						BgNodeOutput? ExistingOutput;
-						if(Graph.TagNameToNodeOutput.TryGetValue(Output.TagName, out ExistingOutput))
-						{
-							LogError(Element, "Node '{0}' already has an output called '{1}'", ExistingOutput.ProducingNode.Name, Output.TagName);
-						}
-						else
-						{
-							Graph.TagNameToNodeOutput.Add(Output.TagName, Output);
-						}
-					}
-
-					// Add all the tasks
-					await ReadNodeBodyAsync(Element, NewNode, ParentAgent);
-
-					// Add it to the current agent
-					ParentAgent.Nodes.Add(NewNode);
-				}
-			}
-		}
+		protected abstract Task ReadNodeAsync(BgScriptElement Element);
 
 		/// <summary>
 		/// Reads the contents of a node element
 		/// </summary>
 		/// <param name="Element">Xml element to read the definition from</param>
-		/// <param name="NewNode">The new node that has been created</param>
-		/// <param name="ParentAgent">Agent for the node to be added to</param>
-		async Task ReadNodeBodyAsync(XmlElement Element, BgNode NewNode, BgAgent ParentAgent)
+		protected async Task ReadNodeBodyAsync(XmlElement Element)
 		{
 			EnterScope();
 			foreach (BgScriptElement ChildElement in Element.ChildNodes.OfType<BgScriptElement>())
@@ -1333,28 +847,28 @@
 						await ReadRegexAsync(ChildElement);
 						break;
 					case "Trace":
-						await ReadDiagnosticAsync(ChildElement, LogEventType.Console, NewNode, ParentAgent);
+						await ReadDiagnosticAsync(ChildElement, LogEventType.Console);
 						break;
 					case "Warning":
-						await ReadDiagnosticAsync(ChildElement, LogEventType.Warning, NewNode, ParentAgent);
+						await ReadDiagnosticAsync(ChildElement, LogEventType.Warning);
 						break;
 					case "Error":
-						await ReadDiagnosticAsync(ChildElement, LogEventType.Error, NewNode, ParentAgent);
+						await ReadDiagnosticAsync(ChildElement, LogEventType.Error);
 						break;
 					case "Do":
-						await ReadBlockAsync(ChildElement, x => ReadNodeBodyAsync(x, NewNode, ParentAgent));
+						await ReadBlockAsync(ChildElement, ReadNodeBodyAsync);
 						break;
 					case "Switch":
-						await ReadSwitchAsync(ChildElement, x => ReadNodeBodyAsync(x, NewNode, ParentAgent));
+						await ReadSwitchAsync(ChildElement, ReadNodeBodyAsync);
 						break;
 					case "ForEach":
-						await ReadForEachAsync(ChildElement, x => ReadNodeBodyAsync(x, NewNode, ParentAgent));
+						await ReadForEachAsync(ChildElement, ReadNodeBodyAsync);
 						break;
 					case "Expand":
-						await ReadExpandAsync(ChildElement, x => ReadNodeBodyAsync(x, NewNode, ParentAgent));
+						await ReadExpandAsync(ChildElement, ReadNodeBodyAsync);
 						break;
 					default:
-						await ReadTaskAsync(ChildElement, NewNode);
+						await ReadTaskAsync(ChildElement);
 						break;
 				}
 			}
@@ -1379,16 +893,7 @@
 		/// </summary>
 		/// <param name="Element">Xml element to read the definition from</param>
 		/// <param name="ReadContentsAsync">Delegate to read the contents of the element, if the condition evaluates to true</param>
-		async Task ReadSwitchAsync(BgScriptElement Element, Func<BgScriptElement, Task> ReadContentsAsync)
-		{
-			foreach (BgScriptElement ChildElement in Element.ChildNodes.OfType<BgScriptElement>())
-			{
-				if (ChildElement.Name == "Default" || await EvaluateConditionAsync(ChildElement))
-				{
-					await ReadContentsAsync(ChildElement);
-				}
-			}
-		}
+		protected abstract Task ReadSwitchAsync(BgScriptElement Element, Func<BgScriptElement, Task> ReadContentsAsync);
 
 		/// <summary>
 		/// Reads a "ForEach" element 
@@ -1503,221 +1008,28 @@
 		/// Reads a task definition from the given element, and add it to the given list
 		/// </summary>
 		/// <param name="Element">Xml element to read the definition from</param>
-		/// <param name="ParentNode">The node which owns this task</param>
-		async Task ReadTaskAsync(BgScriptElement Element, BgNode ParentNode)
-		{
-			// If we're running a single node and this element's parent isn't the single node to run, ignore the error and return.
-			if (!string.IsNullOrWhiteSpace(SingleNodeName) && ParentNode.Name != SingleNodeName)
-			{
-				return;
-			}
-
-			if (await EvaluateConditionAsync(Element))
-			{
-				BgTask Info = new BgTask(Element.Location, Element.Name);
-				foreach (XmlAttribute? Attribute in Element.Attributes)
-				{
-					if (String.Compare(Attribute!.Name, "If", StringComparison.InvariantCultureIgnoreCase) != 0)
-					{
-						string ExpandedValue = ExpandProperties(Element, Attribute.Value);
-						Info.Arguments.Add(Attribute.Name, ExpandedValue);
-					}
-				}
-				ParentNode.Tasks.Add(Info);
-			}
-		}
+		protected abstract Task ReadTaskAsync(BgScriptElement Element);
 
 		/// <summary>
 		/// Reads the definition for an email notifier
 		/// </summary>
 		/// <param name="Element">Xml element to read the definition from</param>
-		async Task ReadNotifierAsync(BgScriptElement Element)
-		{
-			if (await EvaluateConditionAsync(Element))
-			{
-				string[] TargetNames = ReadListAttribute(Element, "Targets");
-				string[] ExceptNames = ReadListAttribute(Element, "Except");
-				string[] IndividualNodeNames = ReadListAttribute(Element, "Nodes");
-				string[] ReportNames = ReadListAttribute(Element, "Reports");
-				string[] Users = ReadListAttribute(Element, "Users");
-				string[] Submitters = ReadListAttribute(Element, "Submitters");
-				bool? bWarnings = Element.HasAttribute("Warnings") ? (bool?)ReadBooleanAttribute(Element, "Warnings", true) : null;
-				bool bAbsolute = Element.HasAttribute("Absolute") ? ReadBooleanAttribute(Element, "Absolute", true) : false;
-
-				// Find the list of targets which are included, and recurse through all their dependencies
-				HashSet<BgNode> Nodes = new HashSet<BgNode>();
-				if (TargetNames != null)
-				{
-					HashSet<BgNode> TargetNodes = ResolveReferences(Element, TargetNames);
-					foreach (BgNode Node in TargetNodes)
-					{
-						Nodes.Add(Node);
-						Nodes.UnionWith(Node.InputDependencies);
-					}
-				}
-
-				// Add all the individually referenced nodes
-				if (IndividualNodeNames != null)
-				{
-					HashSet<BgNode> IndividualNodes = ResolveReferences(Element, IndividualNodeNames);
-					Nodes.UnionWith(IndividualNodes);
-				}
-
-				// Exclude all the exceptions
-				if (ExceptNames != null)
-				{
-					HashSet<BgNode> ExceptNodes = ResolveReferences(Element, ExceptNames);
-					Nodes.ExceptWith(ExceptNodes);
-				}
-
-				// Update all the referenced nodes with the settings
-				foreach (BgNode Node in Nodes)
-				{
-					if (Users != null)
-					{
-						if (bAbsolute)
-						{
-							Node.NotifyUsers = new HashSet<string>(Users);
-						}
-						else
-						{
-							Node.NotifyUsers.UnionWith(Users);
-						}
-					}
-					if (Submitters != null)
-					{
-						if (bAbsolute)
-						{
-							Node.NotifySubmitters = new HashSet<string>(Submitters);
-						}
-						else
-						{
-							Node.NotifySubmitters.UnionWith(Submitters);
-						}
-					}
-					if (bWarnings.HasValue)
-					{
-						Node.bNotifyOnWarnings = bWarnings.Value;
-					}
-				}
-
-				// Add the users to the list of reports
-				if (ReportNames != null)
-				{
-					foreach (string ReportName in ReportNames)
-					{
-						BgReport? Report;
-						if (Graph.NameToReport.TryGetValue(ReportName, out Report))
-						{
-							Report.NotifyUsers.UnionWith(Users);
-						}
-						else
-						{
-							LogError(Element, "Report '{0}' has not been defined", ReportName);
-						}
-					}
-				}
-			}
-		}
+		protected abstract Task ReadNotifierAsync(BgScriptElement Element);
 
 		/// <summary>
 		/// Reads a warning from the given element, evaluates the condition on it, and writes it to the log if the condition passes.
 		/// </summary>
 		/// <param name="Element">Xml element to read the definition from</param>
 		/// <param name="EventType">The diagnostic event type</param>
-		/// <param name="EnclosingNode">The node that this diagnostic is declared in, or null</param>
-		/// <param name="EnclosingAgent">The agent that this diagnostic is declared in, or null</param>
-		async Task ReadDiagnosticAsync(BgScriptElement Element, LogEventType EventType, BgNode? EnclosingNode, BgAgent? EnclosingAgent)
-		{
-			if (await EvaluateConditionAsync(Element))
-			{
-				string Message = ReadAttribute(Element, "Message");
-
-				BgGraphDiagnostic Diagnostic = new BgGraphDiagnostic(Element.Location, EventType, Message, EnclosingNode, EnclosingAgent);
-				Graph.Diagnostics.Add(Diagnostic);
-			}
-		}
+		protected abstract Task ReadDiagnosticAsync(BgScriptElement Element, LogEventType EventType);
 
 		/// <summary>
-		/// Checks that the given name does not already used to refer to a node, and print an error if it is.
-		/// </summary>
-		/// <param name="Element">Xml element to read from</param>
-		/// <param name="Name">Name of the alias</param>
-		/// <returns>True if the name was registered correctly, false otherwise.</returns>
-		bool CheckNameIsUnique(BgScriptElement Element, string Name)
-		{
-			// Get the nodes that it maps to
-			if (Graph.ContainsName(Name))
-			{
-				LogError(Element, "'{0}' is already defined; cannot add a second time", Name);
-				return false;
-			}
-			return true;
-		}
-
-		/// <summary>
-		/// Resolve a list of references to a set of nodes
-		/// </summary>
-		/// <param name="Element">Element used to locate any errors</param>
-		/// <param name="ReferenceNames">Sequence of names to look up</param>
-		/// <returns>Hashset of all the nodes included by the given names</returns>
-		HashSet<BgNode> ResolveReferences(BgScriptElement Element, IEnumerable<string> ReferenceNames)
-		{
-			HashSet<BgNode> Nodes = new HashSet<BgNode>();
-			foreach (string ReferenceName in ReferenceNames)
-			{
-				BgNode[]? OtherNodes;
-				if (Graph.TryResolveReference(ReferenceName, out OtherNodes))
-				{
-					Nodes.UnionWith(OtherNodes);
-				}
-				else if (!ReferenceName.StartsWith("#") && Graph.TagNameToNodeOutput.ContainsKey("#" + ReferenceName))
-				{
-					LogError(Element, "Reference to '{0}' cannot be resolved; did you mean '#{0}'?", ReferenceName);
-				}
-				else
-				{
-					LogError(Element, "Reference to '{0}' cannot be resolved; check it has been defined.", ReferenceName);
-				}
-			}
-			return Nodes;
-		}
-
-		/// <summary>
-		/// Resolve a list of references to a set of nodes
-		/// </summary>
-		/// <param name="Element">Element used to locate any errors</param>
-		/// <param name="ReferenceNames">Sequence of names to look up</param>
-		/// <returns>Set of all the nodes included by the given names</returns>
-		HashSet<BgNodeOutput> ResolveInputReferences(BgScriptElement Element, IEnumerable<string> ReferenceNames)
-		{
-			HashSet<BgNodeOutput> Inputs = new HashSet<BgNodeOutput>();
-			foreach (string ReferenceName in ReferenceNames)
-			{
-				BgNodeOutput[]? ReferenceInputs;
-				if (Graph.TryResolveInputReference(ReferenceName, out ReferenceInputs))
-				{
-					Inputs.UnionWith(ReferenceInputs);
-				}
-				else if (!ReferenceName.StartsWith("#") && Graph.TagNameToNodeOutput.ContainsKey("#" + ReferenceName))
-				{
-					LogError(Element, "Reference to '{0}' cannot be resolved; did you mean '#{0}'?", ReferenceName);
-				}
-				else
-				{
-					LogError(Element, "Reference to '{0}' cannot be resolved; check it has been defined.", ReferenceName);
-				}
-			}
-			return Inputs;
-		}
-
-		/// <summary>
 		/// Reads an object name from its defining element. Outputs an error if the name is missing.
 		/// </summary>
 		/// <param name="Element">Element to read the name for</param>
 		/// <param name="Name">Output variable to receive the name of the object</param>
 		/// <returns>True if the object had a valid name (assigned to the Name variable), false if the name was invalid or missing.</returns>
-		bool TryReadObjectName(BgScriptElement Element, [NotNullWhen(true)] out string? Name)
+		protected bool TryReadObjectName(BgScriptElement Element, [NotNullWhen(true)] out string? Name)
 		{
 			// Check the name attribute is present
 			if (!Element.HasAttribute("Name"))
@@ -1746,7 +1058,7 @@
 		/// <param name="Element">The element that contains the name</param>
 		/// <param name="Name">The name to check</param>
 		/// <returns>True if the name is valid</returns>
-		bool ValidateName(BgScriptElement Element, string Name)
+		protected bool ValidateName(BgScriptElement Element, string Name)
 		{
 			// Check it's not empty
 			if (Name.Length == 0)
@@ -1778,7 +1090,7 @@
 		/// <param name="Element">The element that contains the regex</param>
 		/// <param name="Regex">The pattern to construct</param>
 		/// <returns>The regex if is valid, otherwise null</returns>
-		Regex? ParseRegex(BgScriptElement Element, string Regex)
+		protected Regex? ParseRegex(BgScriptElement Element, string Regex)
 		{
 			if(Regex.Length == 0)
 			{
@@ -1802,7 +1114,7 @@
 		/// <param name="Element">Element to read the attribute from</param>
 		/// <param name="Name">Name of the attribute</param>
 		/// <returns>Array of names, with all leading and trailing whitespace removed</returns>
-		string ReadAttribute(BgScriptElement Element, string Name)
+		protected string ReadAttribute(BgScriptElement Element, string Name)
 		{
 			return ExpandProperties(Element, Element.GetAttribute(Name));
 		}
@@ -1814,7 +1126,7 @@
 		/// <param name="Name"></param>
 		/// <param name="Separator"></param>
 		/// <returns>Array of names, with all leading and trailing whitespace removed</returns>
-		string[] ReadListAttribute(BgScriptElement Element, string Name, char Separator = ';')
+		protected string[] ReadListAttribute(BgScriptElement Element, string Name, char Separator = ';')
 		{
 			string Value = ReadAttribute(Element, Name);
 			return Value.Split(new char[] { Separator }).Select(x => x.Trim()).Where(x => x.Length > 0).ToArray();
@@ -1827,7 +1139,7 @@
 		/// <param name="Name">Name of the attribute</param>
 		/// <param name="bDefaultValue">Default value if the attribute is missing</param>
 		/// <returns>The value of the attribute field</returns>
-		bool ReadBooleanAttribute(BgScriptElement Element, string Name, bool bDefaultValue)
+		protected bool ReadBooleanAttribute(BgScriptElement Element, string Name, bool bDefaultValue)
 		{
 			bool bResult = bDefaultValue;
 			if (Element.HasAttribute(Name))
@@ -1856,7 +1168,7 @@
 		/// <param name="Name">Name of the attribute</param>
 		/// <param name="DefaultValue">Default value for the integer, if the attribute is missing</param>
 		/// <returns>The value of the attribute field</returns>
-		int ReadIntegerAttribute(BgScriptElement Element, string Name, int DefaultValue)
+		protected int ReadIntegerAttribute(BgScriptElement Element, string Name, int DefaultValue)
 		{
 			int Result = DefaultValue;
 			if (Element.HasAttribute(Name))
@@ -1884,7 +1196,7 @@
 		/// <param name="Name">Name of the attribute</param>
 		/// <param name="DefaultValue">Default value for the enum, if the attribute is missing</param>
 		/// <returns>The value of the attribute field</returns>
-		T ReadEnumAttribute<T>(BgScriptElement Element, string Name, T DefaultValue) where T : struct
+		protected T ReadEnumAttribute<T>(BgScriptElement Element, string Name, T DefaultValue) where T : struct
 		{
 			T Result = DefaultValue;
 			if (Element.HasAttribute(Name))
@@ -1910,7 +1222,7 @@
 		/// <param name="Element">The script element causing the error</param>
 		/// <param name="Format">Standard String.Format()-style format string</param>
 		/// <param name="Args">Optional arguments</param>
-		void LogError(BgScriptElement Element, string Format, params object[] Args)
+		protected void LogError(BgScriptElement Element, string Format, params object[] Args)
 		{
 			Logger.LogScriptError(Element.Location, Format, Args);
 			NumErrors++;
@@ -1922,7 +1234,7 @@
 		/// <param name="Element">The script element causing the error</param>
 		/// <param name="Format">Standard String.Format()-style format string</param>
 		/// <param name="Args">Optional arguments</param>
-		void LogWarning(BgScriptElement Element, string Format, params object[] Args)
+		protected void LogWarning(BgScriptElement Element, string Format, params object[] Args)
 		{
 			Logger.LogScriptWarning(Element.Location, Format, Args);
 		}
@@ -1932,7 +1244,7 @@
 		/// </summary>
 		/// <param name="Element">The element to check</param>
 		/// <returns>True if the element's condition evaluates to true (or doesn't have a conditional expression), false otherwise</returns>
-		async Task<bool> EvaluateConditionAsync(BgScriptElement Element)
+		protected async Task<bool> EvaluateConditionAsync(BgScriptElement Element)
 		{
 			// Check if the element has a conditional attribute
 			const string AttributeName = "If";
@@ -1960,7 +1272,7 @@
 		/// <param name="Element">The element containing the string. Used for diagnostic messages.</param>
 		/// <param name="Text">The input string to expand properties in</param>
 		/// <returns>The expanded string</returns>
-		string ExpandProperties(BgScriptElement Element, string Text)
+		protected string ExpandProperties(BgScriptElement Element, string Text)
 		{
 			string Result = Text;
 			// Iterate in reverse order to handle cases where there are nested expansions like $(Outer$(Inner))
@@ -1994,4 +1306,786 @@
 			return Result;
 		}
 	}
+
+	/// <summary>
+	/// Reader for build graph definitions. Instanced to contain temporary state; public interface is through ScriptReader.TryRead().
+	/// </summary>
+	public class BgScriptReader : BgScriptReaderBase
+	{
+		/// <summary>
+		/// The current graph
+		/// </summary>
+		BgGraph Graph = new BgGraph();
+
+		/// <summary>
+		/// Arguments for evaluating the graph
+		/// </summary>
+		Dictionary<string, string> Arguments;
+
+		/// <summary>
+		/// The name of the node if only a single node is going to be built, otherwise null.
+		/// </summary>
+		string? SingleNodeName;
+
+		BgAgent? EnclosingAgent;
+		BgNode? EnclosingNode;
+
+		/// <summary>
+		/// Private constructor. Use ScriptReader.TryRead() to read a script file.
+		/// </summary>
+		/// <param name="Context">Context object</param>
+		/// <param name="DefaultProperties">Default properties available to the script</param>
+		/// <param name="Arguments">Arguments passed in to the graph on the command line</param>
+		/// <param name="SingleNodeName">If a single node will be processed, the name of that node.</param>
+		/// <param name="Schema">Schema for the script</param>
+		/// <param name="Logger">Logger for diagnostic messages</param>
+		private BgScriptReader(IBgScriptReaderContext Context, IDictionary<string, string> DefaultProperties, IReadOnlyDictionary<string, string> Arguments, string? SingleNodeName, BgScriptSchema Schema, ILogger Logger)
+			: base(Context, Schema, Logger)
+		{
+			this.Arguments = new Dictionary<string, string>(Arguments, StringComparer.OrdinalIgnoreCase);
+			this.SingleNodeName = SingleNodeName;
+
+			foreach (KeyValuePair<string, string> Pair in DefaultProperties)
+			{
+				SetPropertyValue(null!, Pair.Key, Pair.Value);
+			}
+		}
+
+		/// <summary>
+		/// Try to read a script file from the given file.
+		/// </summary>
+		/// <param name="Context">Supplies context about the parse</param>
+		/// <param name="File">File to read from</param>
+		/// <param name="Arguments">Arguments passed in to the graph on the command line</param>
+		/// <param name="DefaultProperties">Default properties available to the script</param>
+		/// <param name="Schema">Schema for the script</param>
+		/// <param name="Logger">Logger for output messages</param>
+		/// <param name="SingleNodeName">If a single node will be processed, the name of that node.</param>
+		/// <returns>True if the graph was read, false if there were errors</returns>
+		public static async Task<BgGraph?> ReadAsync(IBgScriptReaderContext Context, string File, Dictionary<string, string> Arguments, Dictionary<string, string> DefaultProperties, BgScriptSchema Schema, ILogger Logger, string? SingleNodeName = null)
+		{
+			// Read the file and build the graph
+			BgScriptReader Reader = new BgScriptReader(Context, DefaultProperties, Arguments, SingleNodeName, Schema, Logger);
+			if (!await Reader.TryReadAsync(File) || Reader.NumErrors > 0)
+			{
+				return null;
+			}
+
+			// Make sure all the arguments were valid
+			HashSet<string> ValidArgumentNames = new HashSet<string>(Reader.Graph.Options.Select(x => x.Name), StringComparer.OrdinalIgnoreCase);
+			ValidArgumentNames.Add("PreflightChange");
+
+			foreach (string ArgumentName in Arguments.Keys)
+			{
+				if (!ValidArgumentNames.Contains(ArgumentName))
+				{
+					Logger.LogWarning("Unknown argument '{ArgumentName}' for '{Script}'", ArgumentName, Context.GetNativePath(File));
+				}
+			}
+
+			// Return the constructed graph
+			return Reader.Graph;
+		}
+
+		/// <summary>
+		/// Reads the definition of a graph option; a parameter which can be set by the user on the command-line or via an environment variable.
+		/// </summary>
+		/// <param name="Element">Xml element to read the definition from</param>
+		protected override async Task ReadOptionAsync(BgScriptElement Element)
+		{
+			if (await EvaluateConditionAsync(Element))
+			{
+				string Name = ReadAttribute(Element, "Name");
+				if (ValidateName(Element, Name))
+				{
+					// Make sure we're at global scope
+					if (ScopedProperties.Count > 1)
+					{
+						throw new Exception("Incorrect scope depth for reading option settings");
+					}
+
+					// Check if the property already exists. If it does, we don't need to register it as an option.
+					string? ExistingValue;
+					if (TryGetPropertyValue(Name, out ExistingValue) && ExistingValue != null)
+					{
+						// If there's a restriction on this definition, check it matches
+						string Restrict = ReadAttribute(Element, "Restrict");
+						if (!String.IsNullOrEmpty(Restrict) && !Regex.IsMatch(ExistingValue, "^" + Restrict + "$", RegexOptions.IgnoreCase))
+						{
+							LogError(Element, "'{0} is already set to '{1}', which does not match the given restriction ('{2}')", Name, ExistingValue, Restrict);
+						}
+					}
+					else
+					{
+						// Create a new option object to store the settings
+						string Description = ReadAttribute(Element, "Description");
+						string DefaultValue = ReadAttribute(Element, "DefaultValue");
+						BgScriptOption Option = new BgScriptOption(Name, Description, DefaultValue);
+						Graph.Options.Add(Option);
+
+						// Get the value of this property
+						string? Value;
+						if (!Arguments.TryGetValue(Name, out Value))
+						{
+							Value = Option.DefaultValue;
+						}
+						SetPropertyValue(Element, Name, Value);
+
+						// If there's a restriction on it, check it's valid
+						string Restrict = ReadAttribute(Element, "Restrict");
+						if (!String.IsNullOrEmpty(Restrict))
+						{
+							string Pattern = "^(" + Restrict + ")$";
+							if (!Regex.IsMatch(Value, Pattern, RegexOptions.IgnoreCase))
+							{
+								LogError(Element, "'{0}' is not a valid value for '{1}' (required: '{2}')", Value, Name, Restrict);
+							}
+							if (Option.DefaultValue != Value && !Regex.IsMatch(Option.DefaultValue, Pattern, RegexOptions.IgnoreCase))
+							{
+								LogError(Element, "Default value '{0}' is not valid for '{1}' (required: '{2}')", Option.DefaultValue, Name, Restrict);
+							}
+						}
+					}
+				}
+			}
+		}
+
+		/// <summary>
+		/// Reads a property assignment.
+		/// </summary>
+		/// <param name="Element">Xml element to read the definition from</param>
+		protected override async Task ReadPropertyAsync(BgScriptElement Element)
+		{
+			if (await EvaluateConditionAsync(Element))
+			{
+				string Name = ReadAttribute(Element, "Name");
+				if (ValidateName(Element, Name))
+				{
+					string Value = ReadAttribute(Element, "Value");
+					if (Element.HasChildNodes)
+					{
+						// Read the element content, and append each line to the value as a semicolon delimited list
+						StringBuilder Builder = new StringBuilder(Value);
+						foreach (string Line in Element.InnerText.Split('\n'))
+						{
+							string TrimLine = ExpandProperties(Element, Line.Trim());
+							if (TrimLine.Length > 0)
+							{
+								if (Builder.Length > 0)
+								{
+									Builder.Append(";");
+								}
+								Builder.Append(TrimLine);
+							}
+						}
+						Value = Builder.ToString();
+					}
+					SetPropertyValue(Element, Name, Value);
+				}
+			}
+		}
+
+		/// <summary>
+		/// Reads a Regex assignment.
+		/// </summary>
+		/// <param name="Element">Xml element to read the definition from</param>
+		protected override async Task ReadRegexAsync(BgScriptElement Element)
+		{
+			if (await EvaluateConditionAsync(Element))
+			{
+				// Get the pattern
+				string RegexString = ReadAttribute(Element, "Pattern");
+
+				// Make sure its a valid regex.
+				Regex? RegexValue = ParseRegex(Element, RegexString);
+				if (RegexValue != null)
+				{
+					// read the names in 
+					string[] CaptureNames = ReadListAttribute(Element, "Capture");
+
+					// get number of groups we passed in
+					int[] GroupNumbers = RegexValue.GetGroupNumbers();
+
+					// make sure the number of property names is the same as the number of match groups
+					// this includes the entire string match group as [0], so don't count that one.
+					if (CaptureNames.Length != GroupNumbers.Count() - 1)
+					{
+						LogError(Element, "MatchGroup count: {0} does not match the number of names specified: {1}", GroupNumbers.Count() - 1, CaptureNames.Length);
+					}
+					else
+					{
+						// apply the regex to the value
+						string Input = ReadAttribute(Element, "Input");
+						Match Match = RegexValue.Match(Input);
+
+						bool Optional = await BgCondition.EvaluateAsync(ReadAttribute(Element, "Optional"), Context);
+						if (!Match.Success)
+						{
+							if (!Optional)
+							{
+								LogError(Element, "Regex {0} did not find a match against input string {1}", RegexString, Input);
+							}
+						}
+						else
+						{
+							// assign each property to the group it matches, skip over [0]
+							for (int MatchIdx = 1; MatchIdx < GroupNumbers.Count(); MatchIdx++)
+							{
+								SetPropertyValue(Element, CaptureNames[MatchIdx - 1], Match.Groups[MatchIdx].Value);
+							}
+						}
+					}
+				}
+			}
+		}
+
+		/// <summary>
+		/// Reads a property assignment from an environment variable.
+		/// </summary>
+		/// <param name="Element">Xml element to read the definition from</param>
+		protected override async Task ReadEnvVarAsync(BgScriptElement Element)
+		{
+			if (await EvaluateConditionAsync(Element))
+			{
+				string Name = ReadAttribute(Element, "Name");
+				if (ValidateName(Element, Name))
+				{
+					string Value = Environment.GetEnvironmentVariable(Name) ?? "";
+					SetPropertyValue(Element, Name, Value);
+				}
+			}
+		}
+
+		/// <summary>
+		/// Reads the definition for an agent.
+		/// </summary>
+		/// <param name="Element">Xml element to read the definition from</param>
+		protected override async Task ReadAgentAsync(BgScriptElement Element)
+		{
+			string? Name;
+			if (await EvaluateConditionAsync(Element) && TryReadObjectName(Element, out Name))
+			{
+				// Read the valid agent types. This may be omitted if we're continuing an existing agent.
+				string[] Types = ReadListAttribute(Element, "Type");
+
+				// Create the agent object, or continue an existing one
+				BgAgent? Agent;
+				if (Graph.NameToAgent.TryGetValue(Name, out Agent))
+				{
+					if (Types.Length > 0 && Agent.PossibleTypes.Length > 0)
+					{
+						if (Types.Length != Agent.PossibleTypes.Length || !Types.SequenceEqual(Agent.PossibleTypes, StringComparer.InvariantCultureIgnoreCase))
+						{
+							LogError(Element, "Agent types ({0}) were different than previous agent definition with types ({1}). Must either be empty or match exactly.", string.Join(",", Types), string.Join(",", Agent.PossibleTypes));
+						}
+					}
+				}
+				else
+				{
+					if (Types.Length == 0)
+					{
+						LogError(Element, "Missing type for agent '{0}'", Name);
+					}
+					Agent = new BgAgent(Name, Types);
+					Graph.NameToAgent.Add(Name, Agent);
+					Graph.Agents.Add(Agent);
+				}
+
+				// Process all the child elements.
+				EnclosingAgent = Agent;
+				await ReadAgentBodyAsync(Element);
+				EnclosingAgent = null;
+			}
+		}
+
+		/// <summary>
+		/// Reads the definition for an aggregate
+		/// </summary>
+		/// <param name="Element">Xml element to read the definition from</param>
+		protected override async Task ReadAggregateAsync(BgScriptElement Element)
+		{
+			string? Name;
+			if (await EvaluateConditionAsync(Element) && TryReadObjectName(Element, out Name) && CheckNameIsUnique(Element, Name))
+			{
+				string[] RequiredNames = ReadListAttribute(Element, "Requires");
+
+				BgAggregate NewAggregate = new BgAggregate(Name);
+				foreach (BgNode ReferencedNode in ResolveReferences(Element, RequiredNames))
+				{
+					NewAggregate.RequiredNodes.Add(ReferencedNode);
+				}
+				Graph.NameToAggregate[Name] = NewAggregate;
+
+				string LabelCategoryName = ReadAttribute(Element, "Label");
+				if (!String.IsNullOrEmpty(LabelCategoryName))
+				{
+					BgLabel Label;
+
+					// Create the label
+					int SlashIdx = LabelCategoryName.IndexOf('/');
+					if (SlashIdx != -1)
+					{
+						Label = new BgLabel(LabelCategoryName.Substring(SlashIdx + 1), LabelCategoryName.Substring(0, SlashIdx), null, null, BgLabelChange.Current);
+					}
+					else
+					{
+						Label = new BgLabel(LabelCategoryName, "Other", null, null, BgLabelChange.Current);
+					}
+
+					// Find all the included nodes
+					foreach (BgNode RequiredNode in NewAggregate.RequiredNodes)
+					{
+						Label.RequiredNodes.Add(RequiredNode);
+						Label.IncludedNodes.Add(RequiredNode);
+						Label.IncludedNodes.UnionWith(RequiredNode.OrderDependencies);
+					}
+
+					string[] IncludedNames = ReadListAttribute(Element, "Include");
+					foreach (BgNode IncludedNode in ResolveReferences(Element, IncludedNames))
+					{
+						Label.IncludedNodes.Add(IncludedNode);
+						Label.IncludedNodes.UnionWith(IncludedNode.OrderDependencies);
+					}
+
+					string[] ExcludedNames = ReadListAttribute(Element, "Exclude");
+					foreach (BgNode ExcludedNode in ResolveReferences(Element, ExcludedNames))
+					{
+						Label.IncludedNodes.Remove(ExcludedNode);
+						Label.IncludedNodes.ExceptWith(ExcludedNode.OrderDependencies);
+					}
+
+					Graph.Labels.Add(Label);
+				}
+			}
+		}
+
+		/// <summary>
+		/// Reads the definition for a report
+		/// </summary>
+		/// <param name="Element">Xml element to read the definition from</param>
+		protected override async Task ReadReportAsync(BgScriptElement Element)
+		{
+			string? Name;
+			if (await EvaluateConditionAsync(Element) && TryReadObjectName(Element, out Name) && CheckNameIsUnique(Element, Name))
+			{
+				string[] RequiredNames = ReadListAttribute(Element, "Requires");
+
+				BgReport NewReport = new BgReport(Name);
+				foreach (BgNode ReferencedNode in ResolveReferences(Element, RequiredNames))
+				{
+					NewReport.Nodes.Add(ReferencedNode);
+					NewReport.Nodes.UnionWith(ReferencedNode.OrderDependencies);
+				}
+				Graph.NameToReport.Add(Name, NewReport);
+			}
+		}
+
+		/// <summary>
+		/// Reads the definition for a badge
+		/// </summary>
+		/// <param name="Element">Xml element to read the definition from</param>
+		protected override async Task ReadBadgeAsync(BgScriptElement Element)
+		{
+			string? Name;
+			if (await EvaluateConditionAsync(Element) && TryReadObjectName(Element, out Name))
+			{
+				string[] RequiredNames = ReadListAttribute(Element, "Requires");
+				string[] TargetNames = ReadListAttribute(Element, "Targets");
+				string Project = ReadAttribute(Element, "Project");
+				int Change = ReadIntegerAttribute(Element, "Change", 0);
+
+				BgBadge NewBadge = new BgBadge(Name, Project, Change);
+				foreach (BgNode ReferencedNode in ResolveReferences(Element, RequiredNames))
+				{
+					NewBadge.Nodes.Add(ReferencedNode);
+				}
+				foreach (BgNode ReferencedNode in ResolveReferences(Element, TargetNames))
+				{
+					NewBadge.Nodes.Add(ReferencedNode);
+					NewBadge.Nodes.UnionWith(ReferencedNode.OrderDependencies);
+				}
+				Graph.Badges.Add(NewBadge);
+			}
+		}
+
+		/// <summary>
+		/// Reads the definition for a label
+		/// </summary>
+		/// <param name="Element">Xml element to read the definition from</param>
+		protected override async Task ReadLabelAsync(BgScriptElement Element)
+		{
+			if (await EvaluateConditionAsync(Element))
+			{
+				string Name = ReadAttribute(Element, "Name");
+				if (!String.IsNullOrEmpty(Name))
+				{
+					ValidateName(Element, Name);
+				}
+
+				string Category = ReadAttribute(Element, "Category");
+
+				string[] RequiredNames = ReadListAttribute(Element, "Requires");
+				string[] IncludedNames = ReadListAttribute(Element, "Include");
+				string[] ExcludedNames = ReadListAttribute(Element, "Exclude");
+
+				string UgsBadge = ReadAttribute(Element, "UgsBadge");
+				string UgsProject = ReadAttribute(Element, "UgsProject");
+
+				BgLabelChange Change = ReadEnumAttribute<BgLabelChange>(Element, "Change", BgLabelChange.Current);
+
+				BgLabel NewLabel = new BgLabel(Name, Category, UgsBadge, UgsProject, Change);
+				foreach (BgNode ReferencedNode in ResolveReferences(Element, RequiredNames))
+				{
+					NewLabel.RequiredNodes.Add(ReferencedNode);
+					NewLabel.IncludedNodes.Add(ReferencedNode);
+					NewLabel.IncludedNodes.UnionWith(ReferencedNode.OrderDependencies);
+				}
+				foreach (BgNode IncludedNode in ResolveReferences(Element, IncludedNames))
+				{
+					NewLabel.IncludedNodes.Add(IncludedNode);
+					NewLabel.IncludedNodes.UnionWith(IncludedNode.OrderDependencies);
+				}
+				foreach (BgNode ExcludedNode in ResolveReferences(Element, ExcludedNames))
+				{
+					NewLabel.IncludedNodes.Remove(ExcludedNode);
+					NewLabel.IncludedNodes.ExceptWith(ExcludedNode.OrderDependencies);
+				}
+				Graph.Labels.Add(NewLabel);
+			}
+		}
+
+		/// <summary>
+		/// Reads the definition for a node, and adds it to the given agent
+		/// </summary>
+		/// <param name="Element">Xml element to read the definition from</param>
+		protected override async Task ReadNodeAsync(BgScriptElement Element)
+		{
+			string? Name;
+			if (await EvaluateConditionAsync(Element) && TryReadObjectName(Element, out Name))
+			{
+				string[] RequiresNames = ReadListAttribute(Element, "Requires");
+				string[] ProducesNames = ReadListAttribute(Element, "Produces");
+				string[] AfterNames = ReadListAttribute(Element, "After");
+				string[] TokenFileNames = ReadListAttribute(Element, "Token");
+				bool bRunEarly = ReadBooleanAttribute(Element, "RunEarly", false);
+				bool bNotifyOnWarnings = ReadBooleanAttribute(Element, "NotifyOnWarnings", true);
+
+				// Resolve all the inputs we depend on
+				HashSet<BgNodeOutput> Inputs = ResolveInputReferences(Element, RequiresNames);
+
+				// Gather up all the input dependencies, and check they're all upstream of the current node
+				HashSet<BgNode> InputDependencies = new HashSet<BgNode>();
+				foreach (BgNode InputDependency in Inputs.Select(x => x.ProducingNode).Distinct())
+				{
+					InputDependencies.Add(InputDependency);
+				}
+
+				// Remove all the lock names from the list of required names
+				HashSet<FileReference> RequiredTokens = new HashSet<FileReference>(TokenFileNames.Select(x => new FileReference(x)));
+
+				// Recursively include all their dependencies too
+				foreach (BgNode InputDependency in InputDependencies.ToArray())
+				{
+					RequiredTokens.UnionWith(InputDependency.RequiredTokens);
+					InputDependencies.UnionWith(InputDependency.InputDependencies);
+				}
+
+				// Validate all the outputs
+				List<string> ValidOutputNames = new List<string>();
+				foreach (string ProducesName in ProducesNames)
+				{
+					BgNodeOutput? ExistingOutput;
+					if (Graph.TagNameToNodeOutput.TryGetValue(ProducesName, out ExistingOutput))
+					{
+						LogError(Element, "Output tag '{0}' is already generated by node '{1}'", ProducesName, ExistingOutput.ProducingNode.Name);
+					}
+					else if (!ProducesName.StartsWith("#"))
+					{
+						LogError(Element, "Output tag names must begin with a '#' character ('{0}')", ProducesName);
+					}
+					else
+					{
+						ValidOutputNames.Add(ProducesName);
+					}
+				}
+
+				// Gather up all the order dependencies
+				HashSet<BgNode> OrderDependencies = new HashSet<BgNode>(InputDependencies);
+				OrderDependencies.UnionWith(ResolveReferences(Element, AfterNames));
+
+				// Recursively include all their order dependencies too
+				foreach (BgNode OrderDependency in OrderDependencies.ToArray())
+				{
+					OrderDependencies.UnionWith(OrderDependency.OrderDependencies);
+				}
+
+				// Check that we're not dependent on anything completing that is declared after the initial declaration of this agent.
+				int AgentIdx = Graph.Agents.IndexOf(EnclosingAgent!);
+				for (int Idx = AgentIdx + 1; Idx < Graph.Agents.Count; Idx++)
+				{
+					foreach (BgNode Node in Graph.Agents[Idx].Nodes.Where(x => OrderDependencies.Contains(x)))
+					{
+						LogError(Element, "Node '{0}' has a dependency on '{1}', which was declared after the initial definition of '{2}'.", Name, Node.Name, EnclosingAgent!.Name);
+					}
+				}
+
+				// Construct and register the node
+				if (CheckNameIsUnique(Element, Name))
+				{
+					// Add it to the node lookup
+					BgNode NewNode = new BgNode(Name, Inputs.ToArray(), ValidOutputNames.ToArray(), InputDependencies.ToArray(), OrderDependencies.ToArray(), RequiredTokens.ToArray());
+					NewNode.bRunEarly = bRunEarly;
+					NewNode.bNotifyOnWarnings = bNotifyOnWarnings;
+					Graph.NameToNode.Add(Name, NewNode);
+
+					// Register all the output tags in the global name table.
+					foreach (BgNodeOutput Output in NewNode.Outputs)
+					{
+						BgNodeOutput? ExistingOutput;
+						if (Graph.TagNameToNodeOutput.TryGetValue(Output.TagName, out ExistingOutput))
+						{
+							LogError(Element, "Node '{0}' already has an output called '{1}'", ExistingOutput.ProducingNode.Name, Output.TagName);
+						}
+						else
+						{
+							Graph.TagNameToNodeOutput.Add(Output.TagName, Output);
+						}
+					}
+
+					// Add all the tasks
+					EnclosingNode = NewNode;
+					await ReadNodeBodyAsync(Element);
+					EnclosingNode = null;
+
+					// Add it to the current agent
+					EnclosingAgent!.Nodes.Add(NewNode);
+				}
+			}
+		}
+
+		/// <summary>
+		/// Reads a "Switch" element 
+		/// </summary>
+		/// <param name="Element">Xml element to read the definition from</param>
+		/// <param name="ReadContentsAsync">Delegate to read the contents of the element, if the condition evaluates to true</param>
+		protected override async Task ReadSwitchAsync(BgScriptElement Element, Func<BgScriptElement, Task> ReadContentsAsync)
+		{
+			foreach (BgScriptElement ChildElement in Element.ChildNodes.OfType<BgScriptElement>())
+			{
+				if (ChildElement.Name == "Default" || await EvaluateConditionAsync(ChildElement))
+				{
+					await ReadContentsAsync(ChildElement);
+				}
+			}
+		}
+
+		/// <summary>
+		/// Reads a task definition from the given element, and add it to the given list
+		/// </summary>
+		/// <param name="Element">Xml element to read the definition from</param>
+		protected override async Task ReadTaskAsync(BgScriptElement Element)
+		{
+			// If we're running a single node and this element's parent isn't the single node to run, ignore the error and return.
+			if (!string.IsNullOrWhiteSpace(SingleNodeName) && EnclosingNode!.Name != SingleNodeName)
+			{
+				return;
+			}
+
+			if (await EvaluateConditionAsync(Element))
+			{
+				BgTask Info = new BgTask(Element.Location, Element.Name);
+				foreach (XmlAttribute? Attribute in Element.Attributes)
+				{
+					if (String.Compare(Attribute!.Name, "If", StringComparison.InvariantCultureIgnoreCase) != 0)
+					{
+						string ExpandedValue = ExpandProperties(Element, Attribute.Value);
+						Info.Arguments.Add(Attribute.Name, ExpandedValue);
+					}
+				}
+				EnclosingNode!.Tasks.Add(Info);
+			}
+		}
+
+		/// <summary>
+		/// Reads the definition for an email notifier
+		/// </summary>
+		/// <param name="Element">Xml element to read the definition from</param>
+		protected override async Task ReadNotifierAsync(BgScriptElement Element)
+		{
+			if (await EvaluateConditionAsync(Element))
+			{
+				string[] TargetNames = ReadListAttribute(Element, "Targets");
+				string[] ExceptNames = ReadListAttribute(Element, "Except");
+				string[] IndividualNodeNames = ReadListAttribute(Element, "Nodes");
+				string[] ReportNames = ReadListAttribute(Element, "Reports");
+				string[] Users = ReadListAttribute(Element, "Users");
+				string[] Submitters = ReadListAttribute(Element, "Submitters");
+				bool? bWarnings = Element.HasAttribute("Warnings") ? (bool?)ReadBooleanAttribute(Element, "Warnings", true) : null;
+				bool bAbsolute = Element.HasAttribute("Absolute") ? ReadBooleanAttribute(Element, "Absolute", true) : false;
+
+				// Find the list of targets which are included, and recurse through all their dependencies
+				HashSet<BgNode> Nodes = new HashSet<BgNode>();
+				if (TargetNames != null)
+				{
+					HashSet<BgNode> TargetNodes = ResolveReferences(Element, TargetNames);
+					foreach (BgNode Node in TargetNodes)
+					{
+						Nodes.Add(Node);
+						Nodes.UnionWith(Node.InputDependencies);
+					}
+				}
+
+				// Add all the individually referenced nodes
+				if (IndividualNodeNames != null)
+				{
+					HashSet<BgNode> IndividualNodes = ResolveReferences(Element, IndividualNodeNames);
+					Nodes.UnionWith(IndividualNodes);
+				}
+
+				// Exclude all the exceptions
+				if (ExceptNames != null)
+				{
+					HashSet<BgNode> ExceptNodes = ResolveReferences(Element, ExceptNames);
+					Nodes.ExceptWith(ExceptNodes);
+				}
+
+				// Update all the referenced nodes with the settings
+				foreach (BgNode Node in Nodes)
+				{
+					if (Users != null)
+					{
+						if (bAbsolute)
+						{
+							Node.NotifyUsers = new HashSet<string>(Users);
+						}
+						else
+						{
+							Node.NotifyUsers.UnionWith(Users);
+						}
+					}
+					if (Submitters != null)
+					{
+						if (bAbsolute)
+						{
+							Node.NotifySubmitters = new HashSet<string>(Submitters);
+						}
+						else
+						{
+							Node.NotifySubmitters.UnionWith(Submitters);
+						}
+					}
+					if (bWarnings.HasValue)
+					{
+						Node.bNotifyOnWarnings = bWarnings.Value;
+					}
+				}
+
+				// Add the users to the list of reports
+				if (ReportNames != null)
+				{
+					foreach (string ReportName in ReportNames)
+					{
+						BgReport? Report;
+						if (Graph.NameToReport.TryGetValue(ReportName, out Report))
+						{
+							Report.NotifyUsers.UnionWith(Users);
+						}
+						else
+						{
+							LogError(Element, "Report '{0}' has not been defined", ReportName);
+						}
+					}
+				}
+			}
+		}
+
+		/// <summary>
+		/// Reads a warning from the given element, evaluates the condition on it, and writes it to the log if the condition passes.
+		/// </summary>
+		/// <param name="Element">Xml element to read the definition from</param>
+		/// <param name="EventType">The diagnostic event type</param>
+		protected override async Task ReadDiagnosticAsync(BgScriptElement Element, LogEventType EventType)
+		{
+			if (await EvaluateConditionAsync(Element))
+			{
+				string Message = ReadAttribute(Element, "Message");
+
+				BgGraphDiagnostic Diagnostic = new BgGraphDiagnostic(Element.Location, EventType, Message, EnclosingNode, EnclosingAgent);
+				Graph.Diagnostics.Add(Diagnostic);
+			}
+		}
+
+		/// <summary>
+		/// Checks that the given name does not already used to refer to a node, and print an error if it is.
+		/// </summary>
+		/// <param name="Element">Xml element to read from</param>
+		/// <param name="Name">Name of the alias</param>
+		/// <returns>True if the name was registered correctly, false otherwise.</returns>
+		bool CheckNameIsUnique(BgScriptElement Element, string Name)
+		{
+			// Get the nodes that it maps to
+			if (Graph.ContainsName(Name))
+			{
+				LogError(Element, "'{0}' is already defined; cannot add a second time", Name);
+				return false;
+			}
+			return true;
+		}
+
+		/// <summary>
+		/// Resolve a list of references to a set of nodes
+		/// </summary>
+		/// <param name="Element">Element used to locate any errors</param>
+		/// <param name="ReferenceNames">Sequence of names to look up</param>
+		/// <returns>Hashset of all the nodes included by the given names</returns>
+		HashSet<BgNode> ResolveReferences(BgScriptElement Element, IEnumerable<string> ReferenceNames)
+		{
+			HashSet<BgNode> Nodes = new HashSet<BgNode>();
+			foreach (string ReferenceName in ReferenceNames)
+			{
+				BgNode[]? OtherNodes;
+				if (Graph.TryResolveReference(ReferenceName, out OtherNodes))
+				{
+					Nodes.UnionWith(OtherNodes);
+				}
+				else if (!ReferenceName.StartsWith("#") && Graph.TagNameToNodeOutput.ContainsKey("#" + ReferenceName))
+				{
+					LogError(Element, "Reference to '{0}' cannot be resolved; did you mean '#{0}'?", ReferenceName);
+				}
+				else
+				{
+					LogError(Element, "Reference to '{0}' cannot be resolved; check it has been defined.", ReferenceName);
+				}
+			}
+			return Nodes;
+		}
+
+		/// <summary>
+		/// Resolve a list of references to a set of nodes
+		/// </summary>
+		/// <param name="Element">Element used to locate any errors</param>
+		/// <param name="ReferenceNames">Sequence of names to look up</param>
+		/// <returns>Set of all the nodes included by the given names</returns>
+		HashSet<BgNodeOutput> ResolveInputReferences(BgScriptElement Element, IEnumerable<string> ReferenceNames)
+		{
+			HashSet<BgNodeOutput> Inputs = new HashSet<BgNodeOutput>();
+			foreach (string ReferenceName in ReferenceNames)
+			{
+				BgNodeOutput[]? ReferenceInputs;
+				if (Graph.TryResolveInputReference(ReferenceName, out ReferenceInputs))
+				{
+					Inputs.UnionWith(ReferenceInputs);
+				}
+				else if (!ReferenceName.StartsWith("#") && Graph.TagNameToNodeOutput.ContainsKey("#" + ReferenceName))
+				{
+					LogError(Element, "Reference to '{0}' cannot be resolved; did you mean '#{0}'?", ReferenceName);
+				}
+				else
+				{
+					LogError(Element, "Reference to '{0}' cannot be resolved; check it has been defined.", ReferenceName);
+				}
+			}
+			return Inputs;
+		}
+	}
 }
