UBT: Fixed a crash when a platform SDK is only partially installed

diff --git a/Engine/Source/Programs/Shared/EpicGames.Build/UEBuildPlatformSDK.cs b/Engine/Source/Programs/Shared/EpicGames.Build/UEBuildPlatformSDK.cs
--- a/Engine/Source/Programs/Shared/EpicGames.Build/UEBuildPlatformSDK.cs
+++ b/Engine/Source/Programs/Shared/EpicGames.Build/UEBuildPlatformSDK.cs
@@ -368,7 +368,7 @@ namespace EpicGames.Core
 
 				UInt64 Ver;
 
-				if (TryConvertVersionToInt(AutoSDKVersion, out Ver))
+				if (AutoSDKVersion != null && TryConvertVersionToInt(AutoSDKVersion, out Ver))
 				{
 					Log.WriteLine(Verbosity, Options, "{0} using Auto SDK {1} from: {2} 0x{3:X}", PlatformName, AutoSDKVersion, Path.Combine(PlatformSDKRoot, GetAutoSDKDirectoryForMainVersion()), Ver);
 				}
