Horde: Remove errors when attempting to unshelve files into a workspace that does not have all files mapped.

diff --git a/Engine/Source/Programs/Shared/EpicGames.Perforce.Managed/ManagedWorkspace.cs b/Engine/Source/Programs/Shared/EpicGames.Perforce.Managed/ManagedWorkspace.cs
--- a/Engine/Source/Programs/Shared/EpicGames.Perforce.Managed/ManagedWorkspace.cs
+++ b/Engine/Source/Programs/Shared/EpicGames.Perforce.Managed/ManagedWorkspace.cs
@@ -915,10 +915,10 @@ namespace EpicGames.Perforce.Managed
 			List<WhereRecord> WriteFiles = new List<WhereRecord>();
 			for(int FileIdx = 0; FileIdx < LastRecord.Files.Count; FileIdx++)
 			{
-				if (FileIdx >= WhereRecords.Count)
-				{
-					throw new PerforceException($"Unable to get location of {LastRecord.Files[FileIdx].DepotFile} within {StreamName}. Check the correct stream is specified.");
-				}
+//				if (FileIdx >= WhereRecords.Count)
+//				{
+//					throw new PerforceException($"Unable to get location of {LastRecord.Files[FileIdx].DepotFile} within {StreamName}. Check the correct stream is specified.");
+//				}
 
 				PerforceResponse<WhereRecord> Response = WhereRecords[FileIdx];
 				if (!Response.Succeeded)
