Horde: Fix out-of-range error when unshelving a changelist into a workspace that doesn't contain all files.

diff --git a/Engine/Source/Programs/Shared/EpicGames.Perforce.Managed/ManagedWorkspace.cs b/Engine/Source/Programs/Shared/EpicGames.Perforce.Managed/ManagedWorkspace.cs
--- a/Engine/Source/Programs/Shared/EpicGames.Perforce.Managed/ManagedWorkspace.cs
+++ b/Engine/Source/Programs/Shared/EpicGames.Perforce.Managed/ManagedWorkspace.cs
@@ -908,27 +908,22 @@ namespace EpicGames.Perforce.Managed
 			}
 
 			// query the location of each file
-			PerforceResponseList<WhereRecord> WhereRecords = await Perforce.TryWhereAsync(LastRecord.Files.Select(x => x.DepotFile).ToArray(), CancellationToken);
+			PerforceResponseList<WhereRecord> WhereResponseList = await Perforce.TryWhereAsync(LastRecord.Files.Select(x => x.DepotFile).ToArray(), CancellationToken);
+			List<WhereRecord> WhereRecords = WhereResponseList.Where(x => x.Succeeded).Select(x => x.Data).ToList();
 
 			// parse out all the list of deleted and modified files
 			List<WhereRecord> DeleteFiles = new List<WhereRecord>();
 			List<WhereRecord> WriteFiles = new List<WhereRecord>();
-			for(int FileIdx = 0; FileIdx < LastRecord.Files.Count; FileIdx++)
+			foreach(DescribeFileRecord FileRecord in LastRecord.Files)
 			{
-//				if (FileIdx >= WhereRecords.Count)
-//				{
-//					throw new PerforceException($"Unable to get location of {LastRecord.Files[FileIdx].DepotFile} within {StreamName}. Check the correct stream is specified.");
-//				}
-
-				PerforceResponse<WhereRecord> Response = WhereRecords[FileIdx];
-				if (!Response.Succeeded)
+				WhereRecord? WhereRecord = WhereRecords.First(x => x.DepotFile.Equals(FileRecord.DepotFile, StringComparison.OrdinalIgnoreCase));
+				if (WhereRecord == null)
 				{
-					Logger.LogInformation("Unable to get location of {File} in current workspace; ignoring.", LastRecord.Files[FileIdx].DepotFile);
+					Logger.LogInformation("Unable to get location of {File} in current workspace; ignoring.", FileRecord.DepotFile);
 					continue;
 				}
 
-				WhereRecord WhereRecord = Response.Data;
-				switch (LastRecord.Files[FileIdx].Action)
+				switch (FileRecord.Action)
 				{
 					case FileAction.Delete:
 					case FileAction.MoveDelete:
@@ -942,7 +937,7 @@ namespace EpicGames.Perforce.Managed
 						WriteFiles.Add(WhereRecord);
 						break;
 					default:
-						throw new Exception($"Unknown action '{LastRecord.Files[FileIdx].Action}' for shelved file {WhereRecord.DepotFile}");
+						throw new Exception($"Unknown action '{FileRecord.Action}' for shelved file {FileRecord.DepotFile}");
 				}
 			}
 
