Horde: Fix exception when unshelving files.

diff --git a/Engine/Source/Programs/Shared/EpicGames.Perforce.Managed/ManagedWorkspace.cs b/Engine/Source/Programs/Shared/EpicGames.Perforce.Managed/ManagedWorkspace.cs
--- a/Engine/Source/Programs/Shared/EpicGames.Perforce.Managed/ManagedWorkspace.cs
+++ b/Engine/Source/Programs/Shared/EpicGames.Perforce.Managed/ManagedWorkspace.cs
@@ -916,7 +916,7 @@ namespace EpicGames.Perforce.Managed
 			List<WhereRecord> WriteFiles = new List<WhereRecord>();
 			foreach(DescribeFileRecord FileRecord in LastRecord.Files)
 			{
-				WhereRecord? WhereRecord = WhereRecords.First(x => x.DepotFile.Equals(FileRecord.DepotFile, StringComparison.OrdinalIgnoreCase));
+				WhereRecord? WhereRecord = WhereRecords.FirstOrDefault(x => x.DepotFile.Equals(FileRecord.DepotFile, StringComparison.OrdinalIgnoreCase));
 				if (WhereRecord == null)
 				{
 					Logger.LogInformation("Unable to get location of {File} in current workspace; ignoring.", FileRecord.DepotFile);
