Horde: Ensure that any exceptions thrown when performing a Perforce sync are re-thrown on the main thread.

diff --git a/Engine/Source/Programs/Shared/EpicGames.Perforce.Managed/ManagedWorkspace.cs b/Engine/Source/Programs/Shared/EpicGames.Perforce.Managed/ManagedWorkspace.cs
--- a/Engine/Source/Programs/Shared/EpicGames.Perforce.Managed/ManagedWorkspace.cs
+++ b/Engine/Source/Programs/Shared/EpicGames.Perforce.Managed/ManagedWorkspace.cs
@@ -1845,6 +1845,14 @@ namespace EpicGames.Perforce.Managed
 
 						// Wait for anything to complete
 						Task CompleteTask = await Task.WhenAny(Tasks.Keys);
+						try
+						{
+							await CompleteTask; // Make sure we re-throw any exceptions from the task that completed
+						}
+						finally
+						{
+							await Task.WhenAll(Tasks.Keys);
+						}
 						int BatchIdx = Tasks[CompleteTask];
 						Tasks.Remove(CompleteTask);
 
