Fix symbols publishing for Horde.Agent builds.

#rnx

diff --git a/Engine/Source/Programs/Horde/Horde.Agent.Tests/Horde.Agent.Tests.csproj b/Engine/Source/Programs/Horde/Horde.Agent.Tests/Horde.Agent.Tests.csproj
--- a/Engine/Source/Programs/Horde/Horde.Agent.Tests/Horde.Agent.Tests.csproj
+++ b/Engine/Source/Programs/Horde/Horde.Agent.Tests/Horde.Agent.Tests.csproj
@@ -4,6 +4,8 @@
     <TargetFramework>netcoreapp3.1</TargetFramework>
     <Nullable>enable</Nullable>
     <IsPackable>false</IsPackable>
+    <DebugType>full</DebugType>
+    <DebugType>pdbonly</DebugType>
   </PropertyGroup>
 
   <ItemGroup>
diff --git a/Engine/Source/Programs/Shared/EpicGames.IoHash/EpicGames.IoHash.csproj b/Engine/Source/Programs/Shared/EpicGames.IoHash/EpicGames.IoHash.csproj
--- a/Engine/Source/Programs/Shared/EpicGames.IoHash/EpicGames.IoHash.csproj
+++ b/Engine/Source/Programs/Shared/EpicGames.IoHash/EpicGames.IoHash.csproj
@@ -9,6 +9,8 @@
     <GenerateTargetFrameworkAttribute>false</GenerateTargetFrameworkAttribute>
     <Configurations>Debug;Release;Development</Configurations>
     <Nullable>enable</Nullable>
+    <DebugType>full</DebugType>
+    <DebugType>pdbonly</DebugType>
   </PropertyGroup>
   <PropertyGroup Condition=" '$(Configuration)|$(Platform)' == 'Debug|AnyCPU' ">
     <DocumentationFile></DocumentationFile>
diff --git a/Engine/Source/Programs/Shared/EpicGames.Perforce.Managed/EpicGames.Perforce.Managed.csproj b/Engine/Source/Programs/Shared/EpicGames.Perforce.Managed/EpicGames.Perforce.Managed.csproj
--- a/Engine/Source/Programs/Shared/EpicGames.Perforce.Managed/EpicGames.Perforce.Managed.csproj
+++ b/Engine/Source/Programs/Shared/EpicGames.Perforce.Managed/EpicGames.Perforce.Managed.csproj
@@ -6,6 +6,8 @@
     <GenerateAssemblyInfo>false</GenerateAssemblyInfo>
     <GenerateTargetFrameworkAttribute>false</GenerateTargetFrameworkAttribute>
     <Configurations>Debug;Release;Development</Configurations>
+    <DebugType>full</DebugType>
+    <DebugType>pdbonly</DebugType>
   </PropertyGroup>
   <PropertyGroup Condition=" '$(Configuration)|$(Platform)' == 'Debug|AnyCPU' ">
     <DocumentationFile></DocumentationFile>
diff --git a/Engine/Source/Programs/Shared/EpicGames.Perforce/EpicGames.Perforce.csproj b/Engine/Source/Programs/Shared/EpicGames.Perforce/EpicGames.Perforce.csproj
--- a/Engine/Source/Programs/Shared/EpicGames.Perforce/EpicGames.Perforce.csproj
+++ b/Engine/Source/Programs/Shared/EpicGames.Perforce/EpicGames.Perforce.csproj
@@ -9,6 +9,7 @@
     <GenerateTargetFrameworkAttribute>false</GenerateTargetFrameworkAttribute>
     <Configurations>Debug;Release;Development</Configurations>
     <nullable>enable</nullable>
+    <DebugType>pdbonly</DebugType>
   </PropertyGroup>
   <PropertyGroup Condition=" '$(Configuration)|$(Platform)' == 'Debug|AnyCPU' ">
     <DocumentationFile>EpicGames.Perforce.xml</DocumentationFile>
@@ -38,6 +39,7 @@
   <PropertyGroup>
     <WindowsNativeConfig>Release</WindowsNativeConfig>
     <WindowsNativeConfig Condition=" '$(Configuration)|$(Platform)' == 'Debug|AnyCPU' ">Debug</WindowsNativeConfig>
+    <DebugType>full</DebugType>
   </PropertyGroup>
 
   <ItemGroup>
diff --git a/Engine/Source/Programs/Shared/EpicGames.Serialization/EpicGames.Serialization.csproj b/Engine/Source/Programs/Shared/EpicGames.Serialization/EpicGames.Serialization.csproj
--- a/Engine/Source/Programs/Shared/EpicGames.Serialization/EpicGames.Serialization.csproj
+++ b/Engine/Source/Programs/Shared/EpicGames.Serialization/EpicGames.Serialization.csproj
@@ -9,6 +9,7 @@
     <GenerateTargetFrameworkAttribute>false</GenerateTargetFrameworkAttribute>
     <Configurations>Debug;Release;Development</Configurations>
     <Nullable>enable</Nullable>
+    <DebugType>pdbonly</DebugType>
   </PropertyGroup>
   <PropertyGroup Condition=" '$(Configuration)|$(Platform)' == 'Debug|AnyCPU' ">
     <DocumentationFile></DocumentationFile>
