Horde: Fix "Completed in nulls" error message.

diff --git a/Engine/Source/Programs/Shared/EpicGames.Perforce.Managed/ManagedWorkspace.cs b/Engine/Source/Programs/Shared/EpicGames.Perforce.Managed/ManagedWorkspace.cs
--- a/Engine/Source/Programs/Shared/EpicGames.Perforce.Managed/ManagedWorkspace.cs
+++ b/Engine/Source/Programs/Shared/EpicGames.Perforce.Managed/ManagedWorkspace.cs
@@ -424,7 +424,7 @@ namespace EpicGames.Perforce.Managed
 				await CleanInternalAsync(bRemoveUntracked, CancellationToken);
 			}
 
-			Logger.LogInformation("Completed in {ElapsedTime:0.0}s", Timer.Elapsed.TotalSeconds);
+			Logger.LogInformation("Completed in {ElapsedTime}s", $"{Timer.Elapsed.TotalSeconds:0.0}");
 		}
 
 		/// <summary>
