Horde: Include the server address in the p4.ini file.

diff --git a/Engine/Source/Programs/Shared/EpicGames.Perforce.Managed/ManagedWorkspace.cs b/Engine/Source/Programs/Shared/EpicGames.Perforce.Managed/ManagedWorkspace.cs
--- a/Engine/Source/Programs/Shared/EpicGames.Perforce.Managed/ManagedWorkspace.cs
+++ b/Engine/Source/Programs/Shared/EpicGames.Perforce.Managed/ManagedWorkspace.cs
@@ -1185,6 +1185,7 @@ namespace EpicGames.Perforce.Managed
 			FileReference ConfigFile = FileReference.Combine(BaseDir, "p4.ini");
 			using(StreamWriter Writer = new StreamWriter(ConfigFile.FullName))
 			{
+				Writer.WriteLine("P4PORT={0}", PerforceClient.ServerAndPort);
 				Writer.WriteLine("P4CLIENT={0}", PerforceClient.ClientName);
 			}
 		}
