UAT: Add extension methods for the "p4 files" command to IPerforceConnection.

diff --git a/Engine/Source/Programs/Shared/EpicGames.Perforce/PerforceConnection.cs b/Engine/Source/Programs/Shared/EpicGames.Perforce/PerforceConnection.cs
--- a/Engine/Source/Programs/Shared/EpicGames.Perforce/PerforceConnection.cs
+++ b/Engine/Source/Programs/Shared/EpicGames.Perforce/PerforceConnection.cs
@@ -1870,6 +1870,86 @@ namespace EpicGames.Perforce
 
 		#endregion
 
+		#region p4 files
+
+		/// <summary>
+		/// Execute the 'files' command
+		/// </summary>
+		/// <param name="Connection">Connection to the Perforce server</param>
+		/// <param name="Options">Options for the command</param>
+		/// <param name="FileSpecs">List of file specifications to query</param>
+		/// <param name="CancellationToken">Token used to cancel the operation</param>
+		/// <returns>List of response objects</returns>
+		public static async Task<List<FilesRecord>> FilesAsync(this IPerforceConnection Connection, FilesOptions Options, FileSpecList FileSpecs, CancellationToken CancellationToken = default)
+		{
+			return (await TryFilesAsync(Connection, Options, -1, FileSpecs, CancellationToken)).Data;
+		}
+
+		/// <summary>
+		/// Execute the 'files' command
+		/// </summary>
+		/// <param name="Connection">Connection to the Perforce server</param>
+		/// <param name="Options">Options for the command</param>
+		/// <param name="FileSpecs">List of file specifications to query</param>
+		/// <param name="CancellationToken">Token used to cancel the operation</param>
+		/// <returns>List of response objects</returns>
+		public static Task<PerforceResponseList<FilesRecord>> TryFilesAsync(this IPerforceConnection Connection, FilesOptions Options, FileSpecList FileSpecs, CancellationToken CancellationToken = default)
+		{
+			return TryFilesAsync(Connection, Options, -1, FileSpecs, CancellationToken);
+		}
+
+		/// <summary>
+		/// Execute the 'files' command
+		/// </summary>
+		/// <param name="Connection">Connection to the Perforce server</param>
+		/// <param name="Options">Options for the command</param>
+		/// <param name="MaxFiles">Maximum number of results to return. Ignored if less than or equal to zero.</param>
+		/// <param name="FileSpecs">List of file specifications to query</param>
+		/// <param name="CancellationToken">Token used to cancel the operation</param>
+		/// <returns>List of response objects</returns>
+		public static async Task<List<FilesRecord>> FilesAsync(this IPerforceConnection Connection, FilesOptions Options, int MaxFiles, FileSpecList FileSpecs, CancellationToken CancellationToken = default)
+		{
+			return (await TryFilesAsync(Connection, Options, MaxFiles, FileSpecs, CancellationToken)).Data;
+		}
+
+		/// <summary>
+		/// Execute the 'files' command
+		/// </summary>
+		/// <param name="Connection">Connection to the Perforce server</param>
+		/// <param name="Options">Options for the command</param>
+		/// <param name="MaxFiles">Maximum number of results to return. Ignored if less than or equal to zero.</param>
+		/// <param name="FileSpecs">List of file specifications to query</param>
+		/// <param name="CancellationToken">Token used to cancel the operation</param>
+		/// <returns>List of response objects</returns>
+		public static Task<PerforceResponseList<FilesRecord>> TryFilesAsync(this IPerforceConnection Connection, FilesOptions Options, int MaxFiles, FileSpecList FileSpecs, CancellationToken CancellationToken = default)
+		{
+			List<string> Arguments = new List<string>();
+			if ((Options & FilesOptions.AllRevisions) != 0)
+			{
+				Arguments.Add("-a");
+			}
+			if ((Options & FilesOptions.LimitToArchiveDepots) != 0)
+			{
+				Arguments.Add("-A");
+			}
+			if ((Options & FilesOptions.ExcludeDeleted) != 0)
+			{
+				Arguments.Add("-e");
+			}
+			if ((Options & FilesOptions.IgnoreCase) != 0)
+			{
+				Arguments.Add("-i");
+			}
+			if (MaxFiles > 0)
+			{
+				Arguments.Add($"-m{MaxFiles}");
+			}
+
+			return BatchedCommandAsync<FilesRecord>(Connection, "files", Arguments, FileSpecs.List, null, CancellationToken);
+		}
+
+		#endregion
+
 		#region p4 fstat
 
 		/// <summary>
diff --git a/Engine/Source/Programs/Shared/EpicGames.Perforce/Records/FilesOptions.cs b/Engine/Source/Programs/Shared/EpicGames.Perforce/Records/FilesOptions.cs
new file mode 100644
--- /dev/null
+++ b/Engine/Source/Programs/Shared/EpicGames.Perforce/Records/FilesOptions.cs
@@ -0,0 +1,40 @@
+// Copyright Epic Games, Inc. All Rights Reserved.
+
+using System;
+using System.Collections.Generic;
+using System.Text;
+
+namespace EpicGames.Perforce
+{
+	/// <summary>
+	/// Options for the p4 files command
+	/// </summary>
+	[Flags]
+	public enum FilesOptions
+	{
+		/// <summary>
+		/// Default options
+		/// </summary>
+		None = 0,
+
+		/// <summary>
+		/// For each file, list all revisions within a specified revision range, rather than only the highest revision in the range.
+		/// </summary>
+		AllRevisions = 1,
+
+		/// <summary>
+		/// Limit output to files in archive depots.
+		/// </summary>
+		LimitToArchiveDepots = 2,
+
+		/// <summary>
+		/// Exclude deleted, purged, or archived files; the files that remain are those available for syncing or integration.
+		/// </summary>
+		ExcludeDeleted = 4,
+
+		/// <summary>
+		/// Ignore the case of the file argument when listing files in a case sensitive server.
+		/// </summary>
+		IgnoreCase = 8,
+	}
+}
diff --git a/Engine/Source/Programs/Shared/EpicGames.Perforce/Records/FilesRecord.cs b/Engine/Source/Programs/Shared/EpicGames.Perforce/Records/FilesRecord.cs
new file mode 100644
--- /dev/null
+++ b/Engine/Source/Programs/Shared/EpicGames.Perforce/Records/FilesRecord.cs
@@ -0,0 +1,52 @@
+// Copyright Epic Games, Inc. All Rights Reserved.
+
+using System;
+using System.Collections.Generic;
+using System.Diagnostics;
+using System.Text;
+
+namespace EpicGames.Perforce
+{
+	/// <summary>
+	/// Record returned by the p4 files command
+	/// </summary>
+	[DebuggerDisplay("{DepotFile}#{Revision}")]
+	public class FilesRecord
+	{
+		/// <summary>
+		/// Path to the file in the depot
+		/// </summary>
+		[PerforceTag("depotFile")]
+		public string DepotFile = String.Empty;
+
+		/// <summary>
+		/// File revision number
+		/// </summary>
+		[PerforceTag("rev")]
+		public int Revision;
+
+		/// <summary>
+		/// Changelist number that the file was modified in
+		/// </summary>
+		[PerforceTag("change")]
+		public int Change;
+
+		/// <summary>
+		/// The action performed on this file
+		/// </summary>
+		[PerforceTag("action")]
+		public FileAction Action;
+
+		/// <summary>
+		/// The file type after this change
+		/// </summary>
+		[PerforceTag("type")]
+		public string Type = String.Empty;
+
+		/// <summary>
+		/// Modification time (the time that the file was last modified on the client before submit), if in depot.
+		/// </summary>
+		[PerforceTag("time", Optional = true)]
+		public DateTime Time;
+	}
+}
