UAT: Create a background thread to acquire and release the single-instance mutex for running UAT.

diff --git a/Engine/Source/Programs/Shared/EpicGames.Build/Automation/ProcessSingleton.cs b/Engine/Source/Programs/Shared/EpicGames.Build/Automation/ProcessSingleton.cs
--- a/Engine/Source/Programs/Shared/EpicGames.Build/Automation/ProcessSingleton.cs
+++ b/Engine/Source/Programs/Shared/EpicGames.Build/Automation/ProcessSingleton.cs
@@ -32,41 +32,60 @@ namespace UnrealBuildBase
 		/// <param name="bWaitForUATMutex"></param>
 		public static async Task<ExitCode> RunSingleInstanceAsync(Func<Task<ExitCode>> Main, bool bWaitForUATMutex)
 		{
-			bool AllowMultipleInsances = (Environment.GetEnvironmentVariable("uebp_UATMutexNoWait") == "1");
-	
-            string EntryAssemblyLocation = Assembly.GetEntryAssembly()!.GetOriginalLocation();
+			// Need to execute this logic on a background thread, since mutex ownership on Linux has thread affinity (ie. mutexes be released on the
+			// same thread that acquires it, which is not guaranteed through an async continuation)
+			TaskCompletionSource<ExitCode> Result = new TaskCompletionSource<ExitCode>();
+			Thread Thread = new Thread(() => RunSingleInstanceThread(Main, Result, bWaitForUATMutex));
+			Thread.Start();
+			ExitCode ExitCode = await Result.Task;
+			Thread.Join();
+			return ExitCode;
+		}
 
-			string MutexName = GetUniqueMutexForPath(Path.GetFileNameWithoutExtension(EntryAssemblyLocation), EntryAssemblyLocation);
-			using (Mutex SingleInstanceMutex = new Mutex(true, MutexName, out bool bCreatedMutex))
+		public static void RunSingleInstanceThread(Func<Task<ExitCode>> Main, TaskCompletionSource<ExitCode> Result, bool bWaitForUATMutex)
+		{
+			try
 			{
-				IsSoleInstance = bCreatedMutex;
+				bool AllowMultipleInsances = (Environment.GetEnvironmentVariable("uebp_UATMutexNoWait") == "1");
+
+				string EntryAssemblyLocation = Assembly.GetEntryAssembly()!.GetOriginalLocation();
 
-				if (!IsSoleInstance && AllowMultipleInsances == false)
+				string MutexName = GetUniqueMutexForPath(Path.GetFileNameWithoutExtension(EntryAssemblyLocation), EntryAssemblyLocation);
+				using (Mutex SingleInstanceMutex = new Mutex(true, MutexName, out bool bCreatedMutex))
 				{
-					if (bWaitForUATMutex)
+					IsSoleInstance = bCreatedMutex;
+
+					if (!IsSoleInstance && AllowMultipleInsances == false)
 					{
-						Log.TraceWarning("Another instance of UAT at '{0}' is running, and the -WaitForUATMutex parameter has been used. Waiting for other UAT to finish...", EntryAssemblyLocation);
-						int Seconds = 0;
-						while (WaitMutexNoExceptions(SingleInstanceMutex, 15 * 1000) == false)
+						if (bWaitForUATMutex)
+						{
+							Log.TraceWarning("Another instance of UAT at '{0}' is running, and the -WaitForUATMutex parameter has been used. Waiting for other UAT to finish...", EntryAssemblyLocation);
+							int Seconds = 0;
+							while (WaitMutexNoExceptions(SingleInstanceMutex, 15 * 1000) == false)
+							{
+								Seconds += 15;
+								Log.TraceInformation("Still waiting for Mutex. {0} seconds has passed...", Seconds);
+							}
+						}
+						else
 						{
-							Seconds += 15;
-							Log.TraceInformation("Still waiting for Mutex. {0} seconds has passed...", Seconds);
+							throw new Exception($"A conflicting instance of AutomationTool is already running. Current location: {EntryAssemblyLocation}. A process manager may be used to determine the conflicting process and what tool may have launched it");
 						}
 					}
-					else
+
+					ExitCode ExitCode = Task.Run(() => Main()).Result;
+
+					if (IsSoleInstance)
 					{
-						throw new Exception($"A conflicting instance of AutomationTool is already running. Current location: {EntryAssemblyLocation}. A process manager may be used to determine the conflicting process and what tool may have launched it");
+						SingleInstanceMutex.ReleaseMutex();
 					}
-				}
 
-				ExitCode Result = await Main();
-
-				if (IsSoleInstance)
-				{
-					SingleInstanceMutex.ReleaseMutex();
+					Result.SetResult(ExitCode);
 				}
-
-				return Result;
+			}
+			catch (Exception Ex)
+			{
+				Result.TrySetException(Ex);
 			}
 		}
 
