UAT: Fix deadlock on exit.

diff --git a/Engine/Source/Programs/Shared/EpicGames.Build/Automation/ProcessSingleton.cs b/Engine/Source/Programs/Shared/EpicGames.Build/Automation/ProcessSingleton.cs
--- a/Engine/Source/Programs/Shared/EpicGames.Build/Automation/ProcessSingleton.cs
+++ b/Engine/Source/Programs/Shared/EpicGames.Build/Automation/ProcessSingleton.cs
@@ -37,9 +37,7 @@ namespace UnrealBuildBase
 			TaskCompletionSource<ExitCode> Result = new TaskCompletionSource<ExitCode>();
 			Thread Thread = new Thread(() => RunSingleInstanceThread(Main, Result, bWaitForUATMutex));
 			Thread.Start();
-			ExitCode ExitCode = await Result.Task;
-			Thread.Join();
-			return ExitCode;
+			return await Result.Task;
 		}
 
 		public static void RunSingleInstanceThread(Func<Task<ExitCode>> Main, TaskCompletionSource<ExitCode> Result, bool bWaitForUATMutex)
