Undo //UE5/Release-5.0/Engine/Source/Programs/Shared/EpicGames.Build/Automation/... changelist 18443551

diff --git a/Engine/Source/Programs/Shared/EpicGames.Build/Automation/ProcessSingleton.cs b/Engine/Source/Programs/Shared/EpicGames.Build/Automation/ProcessSingleton.cs
--- a/Engine/Source/Programs/Shared/EpicGames.Build/Automation/ProcessSingleton.cs
+++ b/Engine/Source/Programs/Shared/EpicGames.Build/Automation/ProcessSingleton.cs
@@ -7,7 +7,6 @@ using System.IO;
 using System.Reflection;
 using System.Text;
 using System.Threading;
-using System.Threading.Tasks;
 
 namespace UnrealBuildBase
 {
@@ -29,61 +28,44 @@ namespace UnrealBuildBase
 		/// Runs the specified delegate checking if this is the only instance of the application.
 		/// </summary>
 		/// <param name="Main"></param>
-		/// <param name="bWaitForUATMutex"></param>
-		public static async Task<ExitCode> RunSingleInstanceAsync(Func<Task<ExitCode>> Main, bool bWaitForUATMutex)
+		/// <param name="Param"></param>
+		public static ExitCode RunSingleInstance(Func<ExitCode> Main, bool bWaitForUATMutex)
 		{
-			// Need to execute this logic on a background thread, since mutex ownership on Linux has thread affinity (ie. mutexes be released on the
-			// same thread that acquires it, which is not guaranteed through an async continuation)
-			TaskCompletionSource<ExitCode> Result = new TaskCompletionSource<ExitCode>();
-			Thread Thread = new Thread(() => RunSingleInstanceThread(Main, Result, bWaitForUATMutex));
-			Thread.Start();
-			return await Result.Task;
-		}
+			bool AllowMultipleInsances = (Environment.GetEnvironmentVariable("uebp_UATMutexNoWait") == "1");
+	
+            string EntryAssemblyLocation = Assembly.GetEntryAssembly()!.GetOriginalLocation();
 
-		public static void RunSingleInstanceThread(Func<Task<ExitCode>> Main, TaskCompletionSource<ExitCode> Result, bool bWaitForUATMutex)
-		{
-			try
+			string MutexName = GetUniqueMutexForPath(Path.GetFileNameWithoutExtension(EntryAssemblyLocation), EntryAssemblyLocation);
+			using (Mutex SingleInstanceMutex = new Mutex(true, MutexName, out bool bCreatedMutex))
 			{
-				bool AllowMultipleInsances = (Environment.GetEnvironmentVariable("uebp_UATMutexNoWait") == "1");
-
-				string EntryAssemblyLocation = Assembly.GetEntryAssembly()!.GetOriginalLocation();
+				IsSoleInstance = bCreatedMutex;
 
-				string MutexName = GetUniqueMutexForPath(Path.GetFileNameWithoutExtension(EntryAssemblyLocation), EntryAssemblyLocation);
-				using (Mutex SingleInstanceMutex = new Mutex(true, MutexName, out bool bCreatedMutex))
+				if (!IsSoleInstance && AllowMultipleInsances == false)
 				{
-					IsSoleInstance = bCreatedMutex;
-
-					if (!IsSoleInstance && AllowMultipleInsances == false)
+					if (bWaitForUATMutex)
 					{
-						if (bWaitForUATMutex)
-						{
-							Log.TraceWarning("Another instance of UAT at '{0}' is running, and the -WaitForUATMutex parameter has been used. Waiting for other UAT to finish...", EntryAssemblyLocation);
-							int Seconds = 0;
-							while (WaitMutexNoExceptions(SingleInstanceMutex, 15 * 1000) == false)
-							{
-								Seconds += 15;
-								Log.TraceInformation("Still waiting for Mutex. {0} seconds has passed...", Seconds);
-							}
-						}
-						else
+						Log.TraceWarning("Another instance of UAT at '{0}' is running, and the -WaitForUATMutex parameter has been used. Waiting for other UAT to finish...", EntryAssemblyLocation);
+						int Seconds = 0;
+						while (WaitMutexNoExceptions(SingleInstanceMutex, 15 * 1000) == false)
 						{
-							throw new Exception($"A conflicting instance of AutomationTool is already running. Current location: {EntryAssemblyLocation}. A process manager may be used to determine the conflicting process and what tool may have launched it");
+							Seconds += 15;
+							Log.TraceInformation("Still waiting for Mutex. {0} seconds has passed...", Seconds);
 						}
 					}
-
-					ExitCode ExitCode = Task.Run(() => Main()).Result;
-
-					if (IsSoleInstance)
+					else
 					{
-						SingleInstanceMutex.ReleaseMutex();
+						throw new Exception($"A conflicting instance of AutomationTool is already running. Current location: {EntryAssemblyLocation}. A process manager may be used to determine the conflicting process and what tool may have launched it");
 					}
+				}
 
-					Result.SetResult(ExitCode);
+				ExitCode Result = Main();
+
+				if (IsSoleInstance)
+				{
+					SingleInstanceMutex.ReleaseMutex();
 				}
-			}
-			catch (Exception Ex)
-			{
-				Result.TrySetException(Ex);
+
+				return Result;
 			}
 		}
 
