Enable code analysis (and fix warnings) for EpicGames.Perforce.

diff --git a/Engine/Source/Programs/Shared/EpicGames.Perforce/.editorconfig b/Engine/Source/Programs/Shared/EpicGames.Perforce/.editorconfig
new file mode 100644
--- /dev/null
+++ b/Engine/Source/Programs/Shared/EpicGames.Perforce/.editorconfig
@@ -0,0 +1,818 @@
+# NOTE: Requires **VS2019 16.3** or later
+
+# All Rules Enabled as build warnings
+# Description: All Rules are enabled as build warnings. Rules with IsEnabledByDefault = false are force enabled as build warnings.
+
+# Code files
+[*.{cs,vb}]
+
+
+# CA1000: Do not declare static members on generic types
+dotnet_diagnostic.CA1000.severity = warning
+
+# CA1001: Types that own disposable fields should be disposable
+dotnet_diagnostic.CA1001.severity = warning
+
+# CA1002: Do not expose generic lists
+dotnet_diagnostic.CA1002.severity = none
+
+# CA1003: Use generic event handler instances
+dotnet_diagnostic.CA1003.severity = warning
+
+# CA1005: Avoid excessive parameters on generic types
+dotnet_diagnostic.CA1005.severity = warning
+
+# CA1008: Enums should have zero value
+dotnet_diagnostic.CA1008.severity = warning
+
+# CA1010: Generic interface should also be implemented
+dotnet_diagnostic.CA1010.severity = warning
+
+# CA1012: Abstract types should not have public constructors
+dotnet_diagnostic.CA1012.severity = warning
+
+# CA1014: Mark assemblies with CLSCompliant
+dotnet_diagnostic.CA1014.severity = warning
+
+# CA1016: Mark assemblies with assembly version
+dotnet_diagnostic.CA1016.severity = warning
+
+# CA1017: Mark assemblies with ComVisible
+dotnet_diagnostic.CA1017.severity = warning
+
+# CA1018: Mark attributes with AttributeUsageAttribute
+dotnet_diagnostic.CA1018.severity = warning
+
+# CA1019: Define accessors for attribute arguments
+dotnet_diagnostic.CA1019.severity = warning
+
+# CA1021: Avoid out parameters
+dotnet_diagnostic.CA1021.severity = warning
+
+# CA1024: Use properties where appropriate
+dotnet_diagnostic.CA1024.severity = warning
+
+# CA1027: Mark enums with FlagsAttribute
+dotnet_diagnostic.CA1027.severity = warning
+
+# CA1028: Enum Storage should be Int32
+dotnet_diagnostic.CA1028.severity = warning
+
+# CA1030: Use events where appropriate
+dotnet_diagnostic.CA1030.severity = warning
+
+# CA1031: Do not catch general exception types
+dotnet_diagnostic.CA1031.severity = none
+
+# CA1032: Implement standard exception constructors
+dotnet_diagnostic.CA1032.severity = none
+
+# CA1033: Interface methods should be callable by child types
+dotnet_diagnostic.CA1033.severity = warning
+
+# CA1034: Nested types should not be visible
+dotnet_diagnostic.CA1034.severity = warning
+
+# CA1036: Override methods on comparable types
+dotnet_diagnostic.CA1036.severity = warning
+
+# CA1040: Avoid empty interfaces
+dotnet_diagnostic.CA1040.severity = warning
+
+# CA1041: Provide ObsoleteAttribute message
+dotnet_diagnostic.CA1041.severity = warning
+
+# CA1043: Use Integral Or String Argument For Indexers
+dotnet_diagnostic.CA1043.severity = warning
+
+# CA1044: Properties should not be write only
+dotnet_diagnostic.CA1044.severity = warning
+
+# CA1045: Do not pass types by reference
+dotnet_diagnostic.CA1045.severity = warning
+
+# CA1046: Do not overload equality operator on reference types
+dotnet_diagnostic.CA1046.severity = warning
+
+# CA1047: Do not declare protected member in sealed type
+dotnet_diagnostic.CA1047.severity = warning
+
+# CA1050: Declare types in namespaces
+dotnet_diagnostic.CA1050.severity = warning
+
+# CA1051: Do not declare visible instance fields
+dotnet_diagnostic.CA1051.severity = warning
+
+# CA1052: Static holder types should be Static or NotInheritable
+dotnet_diagnostic.CA1052.severity = warning
+
+# CA1054: URI-like parameters should not be strings
+dotnet_diagnostic.CA1054.severity = warning
+
+# CA1055: URI-like return values should not be strings
+dotnet_diagnostic.CA1055.severity = warning
+
+# CA1056: URI-like properties should not be strings
+dotnet_diagnostic.CA1056.severity = warning
+
+# CA1058: Types should not extend certain base types
+dotnet_diagnostic.CA1058.severity = warning
+
+# CA1060: Move pinvokes to native methods class
+dotnet_diagnostic.CA1060.severity = none
+
+# CA1061: Do not hide base class methods
+dotnet_diagnostic.CA1061.severity = warning
+
+# CA1062: Validate arguments of public methods
+dotnet_diagnostic.CA1062.severity = none
+
+# CA1063: Implement IDisposable Correctly
+dotnet_diagnostic.CA1063.severity = warning
+
+# CA1064: Exceptions should be public
+dotnet_diagnostic.CA1064.severity = warning
+
+# CA1065: Do not raise exceptions in unexpected locations
+dotnet_diagnostic.CA1065.severity = none
+
+# CA1066: Implement IEquatable when overriding Object.Equals
+dotnet_diagnostic.CA1066.severity = warning
+
+# CA1067: Override Object.Equals(object) when implementing IEquatable<T>
+dotnet_diagnostic.CA1067.severity = warning
+
+# CA1068: CancellationToken parameters must come last
+dotnet_diagnostic.CA1068.severity = warning
+
+# CA1069: Enums values should not be duplicated
+dotnet_diagnostic.CA1069.severity = warning
+
+# CA1070: Do not declare event fields as virtual
+dotnet_diagnostic.CA1070.severity = warning
+
+# CA1200: Avoid using cref tags with a prefix
+dotnet_diagnostic.CA1200.severity = warning
+
+# CA1303: Do not pass literals as localized parameters
+dotnet_diagnostic.CA1303.severity = warning
+
+# CA1304: Specify CultureInfo
+dotnet_diagnostic.CA1304.severity = none
+
+# CA1305: Specify IFormatProvider
+dotnet_diagnostic.CA1305.severity = none
+
+# CA1307: Specify StringComparison for clarity
+dotnet_diagnostic.CA1307.severity = warning
+
+# CA1308: Normalize strings to uppercase
+dotnet_diagnostic.CA1308.severity = warning
+
+# CA1309: Use ordinal string comparison
+dotnet_diagnostic.CA1309.severity = warning
+
+# CA1310: Specify StringComparison for correctness
+dotnet_diagnostic.CA1310.severity = warning
+
+# CA1401: P/Invokes should not be visible
+dotnet_diagnostic.CA1401.severity = warning
+
+# CA1416: Validate platform compatibility
+dotnet_diagnostic.CA1416.severity = warning
+
+# CA1417: Do not use 'OutAttribute' on string parameters for P/Invokes
+dotnet_diagnostic.CA1417.severity = warning
+
+# CA1418: Use valid platform string
+dotnet_diagnostic.CA1418.severity = warning
+
+# CA1419: Provide a parameterless constructor that is as visible as the containing type for concrete types derived from 'System.Runtime.InteropServices.SafeHandle'
+dotnet_diagnostic.CA1419.severity = warning
+
+# CA1501: Avoid excessive inheritance
+dotnet_diagnostic.CA1501.severity = warning
+
+# CA1502: Avoid excessive complexity
+dotnet_diagnostic.CA1502.severity = none
+
+# CA1505: Avoid unmaintainable code
+dotnet_diagnostic.CA1505.severity = warning
+
+# CA1506: Avoid excessive class coupling
+dotnet_diagnostic.CA1506.severity = none
+
+# CA1507: Use nameof to express symbol names
+dotnet_diagnostic.CA1507.severity = warning
+
+# CA1508: Avoid dead conditional code
+dotnet_diagnostic.CA1508.severity = warning
+
+# CA1509: Invalid entry in code metrics rule specification file
+dotnet_diagnostic.CA1509.severity = warning
+
+# CA1700: Do not name enum values 'Reserved'
+dotnet_diagnostic.CA1700.severity = warning
+
+# CA1707: Identifiers should not contain underscores
+dotnet_diagnostic.CA1707.severity = warning
+
+# CA1708: Identifiers should differ by more than case
+dotnet_diagnostic.CA1708.severity = warning
+
+# CA1710: Identifiers should have correct suffix
+dotnet_diagnostic.CA1710.severity = warning
+
+# CA1711: Identifiers should not have incorrect suffix
+dotnet_diagnostic.CA1711.severity = warning
+
+# CA1712: Do not prefix enum values with type name
+dotnet_diagnostic.CA1712.severity = warning
+
+# CA1713: Events should not have 'Before' or 'After' prefix
+dotnet_diagnostic.CA1713.severity = warning
+
+# CA1715: Identifiers should have correct prefix
+dotnet_diagnostic.CA1715.severity = warning
+
+# CA1716: Identifiers should not match keywords
+dotnet_diagnostic.CA1716.severity = warning
+
+# CA1720: Identifier contains type name
+dotnet_diagnostic.CA1720.severity = none
+
+# CA1721: Property names should not match get methods
+dotnet_diagnostic.CA1721.severity = warning
+
+# CA1724: Type names should not match namespaces
+dotnet_diagnostic.CA1724.severity = warning
+
+# CA1725: Parameter names should match base declaration
+dotnet_diagnostic.CA1725.severity = none
+
+# CA1727: Use PascalCase for named placeholders
+dotnet_diagnostic.CA1727.severity = warning
+
+# CA1802: Use literals where appropriate
+dotnet_diagnostic.CA1802.severity = warning
+
+# CA1805: Do not initialize unnecessarily
+dotnet_diagnostic.CA1805.severity = warning
+
+# CA1806: Do not ignore method results
+dotnet_diagnostic.CA1806.severity = warning
+
+# CA1810: Initialize reference type static fields inline
+dotnet_diagnostic.CA1810.severity = warning
+
+# CA1812: Avoid uninstantiated internal classes
+dotnet_diagnostic.CA1812.severity = warning
+
+# CA1813: Avoid unsealed attributes
+dotnet_diagnostic.CA1813.severity = warning
+
+# CA1814: Prefer jagged arrays over multidimensional
+dotnet_diagnostic.CA1814.severity = warning
+
+# CA1815: Override equals and operator equals on value types
+dotnet_diagnostic.CA1815.severity = none
+
+# CA1816: Dispose methods should call SuppressFinalize
+dotnet_diagnostic.CA1816.severity = warning
+
+# CA1819: Properties should not return arrays
+dotnet_diagnostic.CA1819.severity = warning
+
+# CA1820: Test for empty strings using string length
+dotnet_diagnostic.CA1820.severity = warning
+
+# CA1821: Remove empty Finalizers
+dotnet_diagnostic.CA1821.severity = warning
+
+# CA1822: Mark members as static
+dotnet_diagnostic.CA1822.severity = warning
+
+# CA1823: Avoid unused private fields
+dotnet_diagnostic.CA1823.severity = warning
+
+# CA1824: Mark assemblies with NeutralResourcesLanguageAttribute
+dotnet_diagnostic.CA1824.severity = warning
+
+# CA1825: Avoid zero-length array allocations
+dotnet_diagnostic.CA1825.severity = warning
+
+# CA1826: Do not use Enumerable methods on indexable collections
+dotnet_diagnostic.CA1826.severity = warning
+
+# CA1827: Do not use Count() or LongCount() when Any() can be used
+dotnet_diagnostic.CA1827.severity = warning
+
+# CA1828: Do not use CountAsync() or LongCountAsync() when AnyAsync() can be used
+dotnet_diagnostic.CA1828.severity = warning
+
+# CA1829: Use Length/Count property instead of Count() when available
+dotnet_diagnostic.CA1829.severity = warning
+
+# CA1830: Prefer strongly-typed Append and Insert method overloads on StringBuilder
+dotnet_diagnostic.CA1830.severity = warning
+
+# CA1831: Use AsSpan or AsMemory instead of Range-based indexers when appropriate
+dotnet_diagnostic.CA1831.severity = warning
+
+# CA1832: Use AsSpan or AsMemory instead of Range-based indexers when appropriate
+dotnet_diagnostic.CA1832.severity = warning
+
+# CA1833: Use AsSpan or AsMemory instead of Range-based indexers when appropriate
+dotnet_diagnostic.CA1833.severity = warning
+
+# CA1834: Consider using 'StringBuilder.Append(char)' when applicable
+dotnet_diagnostic.CA1834.severity = warning
+
+# CA1835: Prefer the 'Memory'-based overloads for 'ReadAsync' and 'WriteAsync'
+dotnet_diagnostic.CA1835.severity = warning
+
+# CA1836: Prefer IsEmpty over Count
+dotnet_diagnostic.CA1836.severity = warning
+
+# CA1837: Use 'Environment.ProcessId'
+dotnet_diagnostic.CA1837.severity = warning
+
+# CA1838: Avoid 'StringBuilder' parameters for P/Invokes
+dotnet_diagnostic.CA1838.severity = warning
+
+# CA1839: Use 'Environment.ProcessPath'
+dotnet_diagnostic.CA1839.severity = warning
+
+# CA1840: Use 'Environment.CurrentManagedThreadId'
+dotnet_diagnostic.CA1840.severity = warning
+
+# CA1841: Prefer Dictionary.Contains methods
+dotnet_diagnostic.CA1841.severity = warning
+
+# CA1842: Do not use 'WhenAll' with a single task
+dotnet_diagnostic.CA1842.severity = warning
+
+# CA1843: Do not use 'WaitAll' with a single task
+dotnet_diagnostic.CA1843.severity = warning
+
+# CA1844: Provide memory-based overrides of async methods when subclassing 'Stream'
+dotnet_diagnostic.CA1844.severity = warning
+
+# CA1845: Use span-based 'string.Concat'
+dotnet_diagnostic.CA1845.severity = warning
+
+# CA1846: Prefer 'AsSpan' over 'Substring'
+dotnet_diagnostic.CA1846.severity = warning
+
+# CA1847: Use char literal for a single character lookup
+dotnet_diagnostic.CA1847.severity = warning
+
+# CA1848: Use the LoggerMessage delegates
+dotnet_diagnostic.CA1848.severity = none
+
+# CA1849: Call async methods when in an async method
+dotnet_diagnostic.CA1849.severity = warning
+
+# CA2000: Dispose objects before losing scope
+dotnet_diagnostic.CA2000.severity = warning
+
+# CA2002: Do not lock on objects with weak identity
+dotnet_diagnostic.CA2002.severity = warning
+
+# CA2007: Consider calling ConfigureAwait on the awaited task
+dotnet_diagnostic.CA2007.severity = none
+
+# CA2008: Do not create tasks without passing a TaskScheduler
+dotnet_diagnostic.CA2008.severity = warning
+
+# CA2009: Do not call ToImmutableCollection on an ImmutableCollection value
+dotnet_diagnostic.CA2009.severity = warning
+
+# CA2011: Avoid infinite recursion
+dotnet_diagnostic.CA2011.severity = warning
+
+# CA2012: Use ValueTasks correctly
+dotnet_diagnostic.CA2012.severity = warning
+
+# CA2013: Do not use ReferenceEquals with value types
+dotnet_diagnostic.CA2013.severity = warning
+
+# CA2014: Do not use stackalloc in loops
+dotnet_diagnostic.CA2014.severity = warning
+
+# CA2015: Do not define finalizers for types derived from MemoryManager<T>
+dotnet_diagnostic.CA2015.severity = warning
+
+# CA2016: Forward the 'CancellationToken' parameter to methods
+dotnet_diagnostic.CA2016.severity = warning
+
+# CA2017: Parameter count mismatch
+dotnet_diagnostic.CA2017.severity = warning
+
+# CA2018: 'Buffer.BlockCopy' expects the number of bytes to be copied for the 'count' argument
+dotnet_diagnostic.CA2018.severity = warning
+
+# CA2100: Review SQL queries for security vulnerabilities
+dotnet_diagnostic.CA2100.severity = warning
+
+# CA2101: Specify marshaling for P/Invoke string arguments
+dotnet_diagnostic.CA2101.severity = none
+
+# CA2109: Review visible event handlers
+dotnet_diagnostic.CA2109.severity = warning
+
+# CA2119: Seal methods that satisfy private interfaces
+dotnet_diagnostic.CA2119.severity = warning
+
+# CA2153: Do Not Catch Corrupted State Exceptions
+dotnet_diagnostic.CA2153.severity = warning
+
+# CA2200: Rethrow to preserve stack details
+dotnet_diagnostic.CA2200.severity = warning
+
+# CA2201: Do not raise reserved exception types
+dotnet_diagnostic.CA2201.severity = warning
+
+# CA2207: Initialize value type static fields inline
+dotnet_diagnostic.CA2207.severity = warning
+
+# CA2208: Instantiate argument exceptions correctly
+dotnet_diagnostic.CA2208.severity = warning
+
+# CA2211: Non-constant fields should not be visible
+dotnet_diagnostic.CA2211.severity = warning
+
+# CA2213: Disposable fields should be disposed
+dotnet_diagnostic.CA2213.severity = warning
+
+# CA2214: Do not call overridable methods in constructors
+dotnet_diagnostic.CA2214.severity = warning
+
+# CA2215: Dispose methods should call base class dispose
+dotnet_diagnostic.CA2215.severity = warning
+
+# CA2216: Disposable types should declare finalizer
+dotnet_diagnostic.CA2216.severity = warning
+
+# CA2217: Do not mark enums with FlagsAttribute
+dotnet_diagnostic.CA2217.severity = warning
+
+# CA2218: Override GetHashCode on overriding Equals
+dotnet_diagnostic.CA2218.severity = warning
+
+# CA2219: Do not raise exceptions in finally clauses
+dotnet_diagnostic.CA2219.severity = warning
+
+# CA2224: Override Equals on overloading operator equals
+dotnet_diagnostic.CA2224.severity = warning
+
+# CA2225: Operator overloads have named alternates
+dotnet_diagnostic.CA2225.severity = none
+
+# CA2226: Operators should have symmetrical overloads
+dotnet_diagnostic.CA2226.severity = warning
+
+# CA2227: Collection properties should be read only
+dotnet_diagnostic.CA2227.severity = warning
+
+# CA2229: Implement serialization constructors
+dotnet_diagnostic.CA2229.severity = warning
+
+# CA2231: Overload operator equals on overriding value type Equals
+dotnet_diagnostic.CA2231.severity = warning
+
+# CA2234: Pass system uri objects instead of strings
+dotnet_diagnostic.CA2234.severity = warning
+
+# CA2235: Mark all non-serializable fields
+dotnet_diagnostic.CA2235.severity = warning
+
+# CA2237: Mark ISerializable types with serializable
+dotnet_diagnostic.CA2237.severity = warning
+
+# CA2241: Provide correct arguments to formatting methods
+dotnet_diagnostic.CA2241.severity = warning
+
+# CA2242: Test for NaN correctly
+dotnet_diagnostic.CA2242.severity = warning
+
+# CA2243: Attribute string literals should parse correctly
+dotnet_diagnostic.CA2243.severity = warning
+
+# CA2244: Do not duplicate indexed element initializations
+dotnet_diagnostic.CA2244.severity = warning
+
+# CA2245: Do not assign a property to itself
+dotnet_diagnostic.CA2245.severity = warning
+
+# CA2246: Assigning symbol and its member in the same statement
+dotnet_diagnostic.CA2246.severity = warning
+
+# CA2247: Argument passed to TaskCompletionSource constructor should be TaskCreationOptions enum instead of TaskContinuationOptions enum
+dotnet_diagnostic.CA2247.severity = warning
+
+# CA2248: Provide correct 'enum' argument to 'Enum.HasFlag'
+dotnet_diagnostic.CA2248.severity = warning
+
+# CA2249: Consider using 'string.Contains' instead of 'string.IndexOf'
+dotnet_diagnostic.CA2249.severity = warning
+
+# CA2250: Use 'ThrowIfCancellationRequested'
+dotnet_diagnostic.CA2250.severity = warning
+
+# CA2251: Use 'string.Equals'
+dotnet_diagnostic.CA2251.severity = warning
+
+# CA2252: This API requires opting into preview features
+dotnet_diagnostic.CA2252.severity = warning
+
+# CA2253: Named placeholders should not be numeric values
+dotnet_diagnostic.CA2253.severity = warning
+
+# CA2254: Template should be a static expression
+dotnet_diagnostic.CA2254.severity = warning
+
+# CA2255: The 'ModuleInitializer' attribute should not be used in libraries
+dotnet_diagnostic.CA2255.severity = warning
+
+# CA2256: All members declared in parent interfaces must have an implementation in a DynamicInterfaceCastableImplementation-attributed interface
+dotnet_diagnostic.CA2256.severity = warning
+
+# CA2257: Members defined on an interface with the 'DynamicInterfaceCastableImplementationAttribute' should be 'static'
+dotnet_diagnostic.CA2257.severity = warning
+
+# CA2258: Providing a 'DynamicInterfaceCastableImplementation' interface in Visual Basic is unsupported
+dotnet_diagnostic.CA2258.severity = warning
+
+# CA2300: Do not use insecure deserializer BinaryFormatter
+dotnet_diagnostic.CA2300.severity = warning
+
+# CA2301: Do not call BinaryFormatter.Deserialize without first setting BinaryFormatter.Binder
+dotnet_diagnostic.CA2301.severity = warning
+
+# CA2302: Ensure BinaryFormatter.Binder is set before calling BinaryFormatter.Deserialize
+dotnet_diagnostic.CA2302.severity = warning
+
+# CA2305: Do not use insecure deserializer LosFormatter
+dotnet_diagnostic.CA2305.severity = warning
+
+# CA2310: Do not use insecure deserializer NetDataContractSerializer
+dotnet_diagnostic.CA2310.severity = warning
+
+# CA2311: Do not deserialize without first setting NetDataContractSerializer.Binder
+dotnet_diagnostic.CA2311.severity = warning
+
+# CA2312: Ensure NetDataContractSerializer.Binder is set before deserializing
+dotnet_diagnostic.CA2312.severity = warning
+
+# CA2315: Do not use insecure deserializer ObjectStateFormatter
+dotnet_diagnostic.CA2315.severity = warning
+
+# CA2321: Do not deserialize with JavaScriptSerializer using a SimpleTypeResolver
+dotnet_diagnostic.CA2321.severity = warning
+
+# CA2322: Ensure JavaScriptSerializer is not initialized with SimpleTypeResolver before deserializing
+dotnet_diagnostic.CA2322.severity = warning
+
+# CA2326: Do not use TypeNameHandling values other than None
+dotnet_diagnostic.CA2326.severity = warning
+
+# CA2327: Do not use insecure JsonSerializerSettings
+dotnet_diagnostic.CA2327.severity = warning
+
+# CA2328: Ensure that JsonSerializerSettings are secure
+dotnet_diagnostic.CA2328.severity = warning
+
+# CA2329: Do not deserialize with JsonSerializer using an insecure configuration
+dotnet_diagnostic.CA2329.severity = warning
+
+# CA2330: Ensure that JsonSerializer has a secure configuration when deserializing
+dotnet_diagnostic.CA2330.severity = warning
+
+# CA2350: Do not use DataTable.ReadXml() with untrusted data
+dotnet_diagnostic.CA2350.severity = warning
+
+# CA2351: Do not use DataSet.ReadXml() with untrusted data
+dotnet_diagnostic.CA2351.severity = warning
+
+# CA2352: Unsafe DataSet or DataTable in serializable type can be vulnerable to remote code execution attacks
+dotnet_diagnostic.CA2352.severity = warning
+
+# CA2353: Unsafe DataSet or DataTable in serializable type
+dotnet_diagnostic.CA2353.severity = warning
+
+# CA2354: Unsafe DataSet or DataTable in deserialized object graph can be vulnerable to remote code execution attacks
+dotnet_diagnostic.CA2354.severity = warning
+
+# CA2355: Unsafe DataSet or DataTable type found in deserializable object graph
+dotnet_diagnostic.CA2355.severity = warning
+
+# CA2356: Unsafe DataSet or DataTable type in web deserializable object graph
+dotnet_diagnostic.CA2356.severity = warning
+
+# CA2361: Ensure auto-generated class containing DataSet.ReadXml() is not used with untrusted data
+dotnet_diagnostic.CA2361.severity = warning
+
+# CA2362: Unsafe DataSet or DataTable in auto-generated serializable type can be vulnerable to remote code execution attacks
+dotnet_diagnostic.CA2362.severity = warning
+
+# CA3001: Review code for SQL injection vulnerabilities
+dotnet_diagnostic.CA3001.severity = warning
+
+# CA3002: Review code for XSS vulnerabilities
+dotnet_diagnostic.CA3002.severity = warning
+
+# CA3003: Review code for file path injection vulnerabilities
+dotnet_diagnostic.CA3003.severity = warning
+
+# CA3004: Review code for information disclosure vulnerabilities
+dotnet_diagnostic.CA3004.severity = warning
+
+# CA3005: Review code for LDAP injection vulnerabilities
+dotnet_diagnostic.CA3005.severity = warning
+
+# CA3006: Review code for process command injection vulnerabilities
+dotnet_diagnostic.CA3006.severity = warning
+
+# CA3007: Review code for open redirect vulnerabilities
+dotnet_diagnostic.CA3007.severity = warning
+
+# CA3008: Review code for XPath injection vulnerabilities
+dotnet_diagnostic.CA3008.severity = warning
+
+# CA3009: Review code for XML injection vulnerabilities
+dotnet_diagnostic.CA3009.severity = warning
+
+# CA3010: Review code for XAML injection vulnerabilities
+dotnet_diagnostic.CA3010.severity = warning
+
+# CA3011: Review code for DLL injection vulnerabilities
+dotnet_diagnostic.CA3011.severity = warning
+
+# CA3012: Review code for regex injection vulnerabilities
+dotnet_diagnostic.CA3012.severity = warning
+
+# CA3061: Do Not Add Schema By URL
+dotnet_diagnostic.CA3061.severity = warning
+
+# CA3075: Insecure DTD processing in XML
+dotnet_diagnostic.CA3075.severity = warning
+
+# CA3076: Insecure XSLT script processing.
+dotnet_diagnostic.CA3076.severity = warning
+
+# CA3077: Insecure Processing in API Design, XmlDocument and XmlTextReader
+dotnet_diagnostic.CA3077.severity = warning
+
+# CA3147: Mark Verb Handlers With Validate Antiforgery Token
+dotnet_diagnostic.CA3147.severity = warning
+
+# CA5350: Do Not Use Weak Cryptographic Algorithms
+dotnet_diagnostic.CA5350.severity = warning
+
+# CA5351: Do Not Use Broken Cryptographic Algorithms
+dotnet_diagnostic.CA5351.severity = warning
+
+# CA5358: Review cipher mode usage with cryptography experts
+dotnet_diagnostic.CA5358.severity = warning
+
+# CA5359: Do Not Disable Certificate Validation
+dotnet_diagnostic.CA5359.severity = warning
+
+# CA5360: Do Not Call Dangerous Methods In Deserialization
+dotnet_diagnostic.CA5360.severity = warning
+
+# CA5361: Do Not Disable SChannel Use of Strong Crypto
+dotnet_diagnostic.CA5361.severity = warning
+
+# CA5362: Potential reference cycle in deserialized object graph
+dotnet_diagnostic.CA5362.severity = warning
+
+# CA5363: Do Not Disable Request Validation
+dotnet_diagnostic.CA5363.severity = warning
+
+# CA5364: Do Not Use Deprecated Security Protocols
+dotnet_diagnostic.CA5364.severity = warning
+
+# CA5365: Do Not Disable HTTP Header Checking
+dotnet_diagnostic.CA5365.severity = warning
+
+# CA5366: Use XmlReader for 'DataSet.ReadXml()'
+dotnet_diagnostic.CA5366.severity = warning
+
+# CA5367: Do Not Serialize Types With Pointer Fields
+dotnet_diagnostic.CA5367.severity = warning
+
+# CA5368: Set ViewStateUserKey For Classes Derived From Page
+dotnet_diagnostic.CA5368.severity = warning
+
+# CA5369: Use XmlReader for 'XmlSerializer.Deserialize()'
+dotnet_diagnostic.CA5369.severity = warning
+
+# CA5370: Use XmlReader for XmlValidatingReader constructor
+dotnet_diagnostic.CA5370.severity = warning
+
+# CA5371: Use XmlReader for 'XmlSchema.Read()'
+dotnet_diagnostic.CA5371.severity = warning
+
+# CA5372: Use XmlReader for XPathDocument constructor
+dotnet_diagnostic.CA5372.severity = warning
+
+# CA5373: Do not use obsolete key derivation function
+dotnet_diagnostic.CA5373.severity = warning
+
+# CA5374: Do Not Use XslTransform
+dotnet_diagnostic.CA5374.severity = warning
+
+# CA5375: Do Not Use Account Shared Access Signature
+dotnet_diagnostic.CA5375.severity = warning
+
+# CA5376: Use SharedAccessProtocol HttpsOnly
+dotnet_diagnostic.CA5376.severity = warning
+
+# CA5377: Use Container Level Access Policy
+dotnet_diagnostic.CA5377.severity = warning
+
+# CA5378: Do not disable ServicePointManagerSecurityProtocols
+dotnet_diagnostic.CA5378.severity = warning
+
+# CA5379: Ensure Key Derivation Function algorithm is sufficiently strong
+dotnet_diagnostic.CA5379.severity = warning
+
+# CA5380: Do Not Add Certificates To Root Store
+dotnet_diagnostic.CA5380.severity = warning
+
+# CA5381: Ensure Certificates Are Not Added To Root Store
+dotnet_diagnostic.CA5381.severity = warning
+
+# CA5382: Use Secure Cookies In ASP.NET Core
+dotnet_diagnostic.CA5382.severity = warning
+
+# CA5383: Ensure Use Secure Cookies In ASP.NET Core
+dotnet_diagnostic.CA5383.severity = warning
+
+# CA5384: Do Not Use Digital Signature Algorithm (DSA)
+dotnet_diagnostic.CA5384.severity = warning
+
+# CA5385: Use Rivest–Shamir–Adleman (RSA) Algorithm With Sufficient Key Size
+dotnet_diagnostic.CA5385.severity = warning
+
+# CA5386: Avoid hardcoding SecurityProtocolType value
+dotnet_diagnostic.CA5386.severity = warning
+
+# CA5387: Do Not Use Weak Key Derivation Function With Insufficient Iteration Count
+dotnet_diagnostic.CA5387.severity = warning
+
+# CA5388: Ensure Sufficient Iteration Count When Using Weak Key Derivation Function
+dotnet_diagnostic.CA5388.severity = warning
+
+# CA5389: Do Not Add Archive Item's Path To The Target File System Path
+dotnet_diagnostic.CA5389.severity = warning
+
+# CA5390: Do not hard-code encryption key
+dotnet_diagnostic.CA5390.severity = warning
+
+# CA5391: Use antiforgery tokens in ASP.NET Core MVC controllers
+dotnet_diagnostic.CA5391.severity = warning
+
+# CA5392: Use DefaultDllImportSearchPaths attribute for P/Invokes
+dotnet_diagnostic.CA5392.severity = none
+
+# CA5393: Do not use unsafe DllImportSearchPath value
+dotnet_diagnostic.CA5393.severity = warning
+
+# CA5394: Do not use insecure randomness
+dotnet_diagnostic.CA5394.severity = warning
+
+# CA5395: Miss HttpVerb attribute for action methods
+dotnet_diagnostic.CA5395.severity = warning
+
+# CA5396: Set HttpOnly to true for HttpCookie
+dotnet_diagnostic.CA5396.severity = warning
+
+# CA5397: Do not use deprecated SslProtocols values
+dotnet_diagnostic.CA5397.severity = warning
+
+# CA5398: Avoid hardcoded SslProtocols values
+dotnet_diagnostic.CA5398.severity = warning
+
+# CA5399: HttpClients should enable certificate revocation list checks
+dotnet_diagnostic.CA5399.severity = warning
+
+# CA5400: Ensure HttpClient certificate revocation list check is not disabled
+dotnet_diagnostic.CA5400.severity = warning
+
+# CA5401: Do not use CreateEncryptor with non-default IV
+dotnet_diagnostic.CA5401.severity = warning
+
+# CA5402: Use CreateEncryptor with the default IV 
+dotnet_diagnostic.CA5402.severity = warning
+
+# CA5403: Do not hard-code certificate
+dotnet_diagnostic.CA5403.severity = warning
+
+# CA5404: Do not disable token validation checks
+dotnet_diagnostic.CA5404.severity = warning
+
+# CA5405: Do not always skip token validation in delegates
+dotnet_diagnostic.CA5405.severity = warning
diff --git a/Engine/Source/Programs/Shared/EpicGames.Perforce/EpicGames.Perforce.csproj b/Engine/Source/Programs/Shared/EpicGames.Perforce/EpicGames.Perforce.csproj
--- a/Engine/Source/Programs/Shared/EpicGames.Perforce/EpicGames.Perforce.csproj
+++ b/Engine/Source/Programs/Shared/EpicGames.Perforce/EpicGames.Perforce.csproj
@@ -40,6 +40,7 @@
     <WindowsNativeConfig>Release</WindowsNativeConfig>
     <WindowsNativeConfig Condition=" '$(Configuration)|$(Platform)' == 'Debug|AnyCPU' ">Debug</WindowsNativeConfig>
     <DebugType>full</DebugType>
+    <EnableNETAnalyzers>True</EnableNETAnalyzers>
   </PropertyGroup>
 
   <ItemGroup>
diff --git a/Engine/Source/Programs/Shared/EpicGames.Perforce/IPerforceOutput.cs b/Engine/Source/Programs/Shared/EpicGames.Perforce/IPerforceOutput.cs
--- a/Engine/Source/Programs/Shared/EpicGames.Perforce/IPerforceOutput.cs
+++ b/Engine/Source/Programs/Shared/EpicGames.Perforce/IPerforceOutput.cs
@@ -78,11 +78,11 @@ namespace EpicGames.Perforce
 					{
 						if (Data[Idx] == '\t')
 						{
-							Result.Append("\t");
+							Result.Append('\t');
 						}
 						else if (Data[Idx] == '\\')
 						{
-							Result.Append("\\");
+							Result.Append('\\');
 						}
 						else if (Data[Idx] >= 0x20 && Data[Idx] <= 0x7f)
 						{
@@ -567,10 +567,10 @@ namespace EpicGames.Perforce
 					if (TagInfo != null)
 					{
 						// Get the list field
-						System.Collections.IList? List = (System.Collections.IList?)TagInfo.Field.GetValue(NewRecord);
+						System.Collections.IList? List = (System.Collections.IList?)TagInfo.Property.GetValue(NewRecord);
 						if (List == null)
 						{
-							throw new PerforceException($"Empty list for {TagInfo.Field.Name}");
+							throw new PerforceException($"Empty list for {TagInfo.Property.Name}");
 						}
 
 						// Check the suffix matches the index of the next element
@@ -586,16 +586,16 @@ namespace EpicGames.Perforce
 							return false;
 						}
 					}
-					else if (RecordInfo.SubElementField != null)
+					else if (RecordInfo.SubElementProperty != null)
 					{
 						// Move back to the start of this tag
 						BufferPos = StartBufferPos;
 
 						// Get the list field
-						System.Collections.IList? List = (System.Collections.IList?)RecordInfo.SubElementField.GetValue(NewRecord);
+						System.Collections.IList? List = (System.Collections.IList?)RecordInfo.SubElementProperty.GetValue(NewRecord);
 						if (List == null)
 						{
-							throw new PerforceException($"Invalid field for {RecordInfo.SubElementField.Name}");
+							throw new PerforceException($"Invalid field for {RecordInfo.SubElementProperty.Name}");
 						}
 
 						// Check the suffix matches the index of the next element
diff --git a/Engine/Source/Programs/Shared/EpicGames.Perforce/NativePerforceConnection.cs b/Engine/Source/Programs/Shared/EpicGames.Perforce/NativePerforceConnection.cs
--- a/Engine/Source/Programs/Shared/EpicGames.Perforce/NativePerforceConnection.cs
+++ b/Engine/Source/Programs/Shared/EpicGames.Perforce/NativePerforceConnection.cs
@@ -6,6 +6,7 @@ using System;
 using System.Collections.Concurrent;
 using System.Collections.Generic;
 using System.Diagnostics;
+using System.Diagnostics.CodeAnalysis;
 using System.Linq;
 using System.Runtime.InteropServices;
 using System.Threading;
@@ -17,7 +18,7 @@ namespace EpicGames.Perforce
 	/// <summary>
 	/// Experimental implementation of <see cref="IPerforceConnection"/> which wraps the native C++ API.
 	/// </summary>
-	public class NativePerforceConnection : IPerforceConnection, IDisposable
+	public sealed class NativePerforceConnection : IPerforceConnection, IDisposable
 	{
 		const string NativeDll = "EpicGames.Perforce.Native";
 
@@ -44,6 +45,7 @@ namespace EpicGames.Perforce
 		}
 
 		[StructLayout(LayoutKind.Sequential)]
+		[SuppressMessage("Compiler", "CA1812")]
 		class NativeReadBuffer
 		{
 			public IntPtr Data;
@@ -221,7 +223,7 @@ namespace EpicGames.Perforce
 				// If we've read everything from the next buffer, return it to the write list
 				if (NextBufferPos == NextBufferLen)
 				{
-					Outer.WriteBuffers.Add(NextBuffer);
+					Outer.WriteBuffers.Add(NextBuffer, Token);
 
 					NextBuffer = null;
 					NextBufferPos = 0;
@@ -472,7 +474,7 @@ namespace EpicGames.Perforce
 		public async Task LoginAsync(string Password, CancellationToken CancellationToken = default)
 		{
 			await using Response Response = new Response(this);
-			Requests.Add((() => Client_Login(Client, Password), Response));
+			Requests.Add((() => Client_Login(Client, Password), Response), CancellationToken);
 
 			List<PerforceResponse> Records = await ((IPerforceOutput)Response).ReadResponsesAsync(null, default);
 			if (Records.Count != 1)
diff --git a/Engine/Source/Programs/Shared/EpicGames.Perforce/PerforceAttributes.cs b/Engine/Source/Programs/Shared/EpicGames.Perforce/PerforceAttributes.cs
--- a/Engine/Source/Programs/Shared/EpicGames.Perforce/PerforceAttributes.cs
+++ b/Engine/Source/Programs/Shared/EpicGames.Perforce/PerforceAttributes.cs
@@ -11,18 +11,18 @@ namespace EpicGames.Perforce
 	/// <summary>
 	/// Attributes for fields that should be deserialized from P4 tags
 	/// </summary>
-	[AttributeUsage(AttributeTargets.Field)]
-	public class PerforceTagAttribute : Attribute
+	[AttributeUsage(AttributeTargets.Property)]
+	public sealed class PerforceTagAttribute : Attribute
 	{
 		/// <summary>
 		/// The tag name
 		/// </summary>
-		public string Name;
+		public string Name { get; }
 
 		/// <summary>
 		/// Whether this tag is required for a valid record
 		/// </summary>
-		public bool Optional = false;
+		public bool Optional { get; set; }
 
 		/// <summary>
 		/// Constructor
@@ -38,12 +38,12 @@ namespace EpicGames.Perforce
 	/// Specifies the name of an enum when converted into a P4 string
 	/// </summary>
 	[AttributeUsage(AttributeTargets.Field)]
-	public class PerforceEnumAttribute : Attribute
+	public sealed class PerforceEnumAttribute : Attribute
 	{
 		/// <summary>
 		/// Name of the enum value
 		/// </summary>
-		public string Name;
+		public string Name { get; }
 
 		/// <summary>
 		/// Constructor
@@ -58,8 +58,8 @@ namespace EpicGames.Perforce
 	/// <summary>
 	/// When attached to a list field, indicates that a list of structures can be included in the record
 	/// </summary>
-	[AttributeUsage(AttributeTargets.Field)]
-	public class PerforceRecordListAttribute : Attribute
+	[AttributeUsage(AttributeTargets.Property)]
+	public sealed class PerforceRecordListAttribute : Attribute
 	{
 	}
 }
diff --git a/Engine/Source/Programs/Shared/EpicGames.Perforce/PerforceChildProcess.cs b/Engine/Source/Programs/Shared/EpicGames.Perforce/PerforceChildProcess.cs
--- a/Engine/Source/Programs/Shared/EpicGames.Perforce/PerforceChildProcess.cs
+++ b/Engine/Source/Programs/Shared/EpicGames.Perforce/PerforceChildProcess.cs
@@ -84,7 +84,7 @@ namespace EpicGames.Perforce
 			FullArguments.AddRange(Arguments);
 
 			string FullArgumentList = CommandLineArguments.Join(FullArguments);
-			Logger.LogDebug("Running {0} {1}", PerforceFileName, FullArgumentList);
+			Logger.LogDebug("Running {Executable} {Arguments}", PerforceFileName, FullArgumentList);
 
 			Scope = TraceSpan.Create(Command, Service: "perforce");
 			Scope.AddMetadata("arguments", FullArgumentList);
@@ -185,7 +185,7 @@ namespace EpicGames.Perforce
 		/// <returns>Exit code and output from the process</returns>
 		public async Task<Tuple<bool, string>> TryReadToEndAsync(CancellationToken CancellationToken)
 		{
-			MemoryStream Stream = new MemoryStream();
+			using MemoryStream Stream = new MemoryStream();
 
 			while (await ReadAsync(CancellationToken))
 			{
diff --git a/Engine/Source/Programs/Shared/EpicGames.Perforce/PerforceConnection.cs b/Engine/Source/Programs/Shared/EpicGames.Perforce/PerforceConnection.cs
--- a/Engine/Source/Programs/Shared/EpicGames.Perforce/PerforceConnection.cs
+++ b/Engine/Source/Programs/Shared/EpicGames.Perforce/PerforceConnection.cs
@@ -195,7 +195,7 @@ namespace EpicGames.Perforce
 
 				string FullArgumentList = CommandLineArguments.Join(Arguments);
 
-				Logger.LogDebug("Running {0} {1}", Executable, FullArgumentList);
+				Logger.LogDebug("Running {Executable} {Arguments}", Executable, FullArgumentList);
 				using (ManagedProcess ChildProcess = new ManagedProcess(ChildProcessGroup, Executable, FullArgumentList, null, null, PasswordBytes, ProcessPriorityClass.Normal))
 				{
 					StringBuilder Lines = new StringBuilder("Unable to log in: ");
@@ -255,7 +255,7 @@ namespace EpicGames.Perforce
 		public async Task<Tuple<bool, string>> TrySetAsync(string Name, string Value, CancellationToken CancellationToken = default)
 		{
 			List<string> Arguments = new List<string>();
-			Arguments.Append($"{Name}={Value}");
+			Arguments.Add($"{Name}={Value}");
 
 			using (PerforceChildProcess ChildProcess = new PerforceChildProcess("set", Arguments, null, null, GetGlobalArguments(), Logger))
 			{
@@ -289,7 +289,7 @@ namespace EpicGames.Perforce
 					}
 				}
 
-				Logger.LogDebug($"Unable to get '{Name}' variable: {Response.Item2}");
+				Logger.LogDebug("Unable to get '{Name}' variable: {Response}", Name, Response.Item2);
 				return null;
 			}
 		}
@@ -2396,69 +2396,6 @@ namespace EpicGames.Perforce
 			return SingleResponseCommandAsync<PrintRecord>(Connection, "print", Arguments, null, CancellationToken);
 		}
 
-		class PrintHandler : IDisposable
-		{
-			Dictionary<string, string> DepotFileToLocalFile;
-			FileStream? OutputStream;
-
-			public PrintHandler(Dictionary<string, string> DepotFileToLocalFile)
-			{
-				this.DepotFileToLocalFile = DepotFileToLocalFile;
-			}
-
-			public void Dispose()
-			{
-				CloseStream();
-			}
-
-			private void OpenStream(string FileName)
-			{
-				CloseStream();
-				Directory.CreateDirectory(Path.GetDirectoryName(FileName));
-				OutputStream = File.Open(FileName, FileMode.Create, FileAccess.Write, FileShare.None);
-			}
-
-			private void CloseStream()
-			{
-				if(OutputStream != null)
-				{
-					OutputStream.Dispose();
-					OutputStream = null;
-				}
-			}
-
-			public void HandleRecord(List<KeyValuePair<string, object>> Fields)
-			{
-				if(Fields[0].Key != "code")
-				{
-					throw new Exception("Missing code field");
-				}
-
-				string Value = (string)Fields[0].Value;
-				if(Value == "stat")
-				{
-					string DepotFile = Fields.First(x => x.Key == "depotFile").Value.ToString() ?? String.Empty;
-
-					string? LocalFile;
-					if(!DepotFileToLocalFile.TryGetValue(DepotFile, out LocalFile))
-					{
-						throw new PerforceException("Depot file '{0}' not found in input dictionary", DepotFile);
-					}
-
-					OpenStream(LocalFile);
-				}
-				else if(Value == "binary" || Value == "text")
-				{
-					byte[] Data = (byte[])Fields.First(x => x.Key == "data").Value;
-					OutputStream!.Write(Data, 0, Data.Length);
-				}
-				else
-				{
-					throw new Exception("Unexpected record type");
-				}
-			}
-		}
-
 		#endregion
 
 		#region p4 reconcile
diff --git a/Engine/Source/Programs/Shared/EpicGames.Perforce/PerforceError.cs b/Engine/Source/Programs/Shared/EpicGames.Perforce/PerforceError.cs
--- a/Engine/Source/Programs/Shared/EpicGames.Perforce/PerforceError.cs
+++ b/Engine/Source/Programs/Shared/EpicGames.Perforce/PerforceError.cs
@@ -55,19 +55,19 @@ namespace EpicGames.Perforce
 		/// The severity of this error
 		/// </summary>
 		[PerforceTagAttribute("severity")]
-		public PerforceSeverityCode Severity;
+		public PerforceSeverityCode Severity { get; set; }
 
 		/// <summary>
 		/// The generic error code associated with this message
 		/// </summary>
 		[PerforceTagAttribute("generic")]
-		public PerforceGenericCode Generic;
+		public PerforceGenericCode Generic { get; set; }
 
 		/// <summary>
 		/// The message text
 		/// </summary>
 		[PerforceTagAttribute("data")]
-		public string Data;
+		public string Data { get; set; }
 
 		/// <summary>
 		/// Private constructor for serialization
diff --git a/Engine/Source/Programs/Shared/EpicGames.Perforce/PerforceException.cs b/Engine/Source/Programs/Shared/EpicGames.Perforce/PerforceException.cs
--- a/Engine/Source/Programs/Shared/EpicGames.Perforce/PerforceException.cs
+++ b/Engine/Source/Programs/Shared/EpicGames.Perforce/PerforceException.cs
@@ -16,7 +16,7 @@ namespace EpicGames.Perforce
 		/// <summary>
 		/// For errors returned by the server, contains the error record
 		/// </summary>
-		public PerforceError? Error;
+		public PerforceError? Error { get; set; }
 
 		/// <summary>
 		/// Constructor
diff --git a/Engine/Source/Programs/Shared/EpicGames.Perforce/PerforceInfo.cs b/Engine/Source/Programs/Shared/EpicGames.Perforce/PerforceInfo.cs
--- a/Engine/Source/Programs/Shared/EpicGames.Perforce/PerforceInfo.cs
+++ b/Engine/Source/Programs/Shared/EpicGames.Perforce/PerforceInfo.cs
@@ -17,13 +17,13 @@ namespace EpicGames.Perforce
 		/// The severity level
 		/// </summary>
 		[PerforceTag("level")]
-		public int Level;
+		public int Level { get; set; }
 
 		/// <summary>
 		/// Message data
 		/// </summary>
 		[PerforceTag("data")]
-		public string Data;
+		public string Data { get; set; }
 
 		/// <summary>
 		/// Private constructor for serialization
diff --git a/Engine/Source/Programs/Shared/EpicGames.Perforce/PerforceRecord.cs b/Engine/Source/Programs/Shared/EpicGames.Perforce/PerforceRecord.cs
--- a/Engine/Source/Programs/Shared/EpicGames.Perforce/PerforceRecord.cs
+++ b/Engine/Source/Programs/Shared/EpicGames.Perforce/PerforceRecord.cs
@@ -35,7 +35,7 @@ namespace EpicGames.Perforce
 		/// <summary>
 		/// The raw data for the value, including type, size, and payload
 		/// </summary>
-		public readonly ReadOnlyMemory<byte> Data;
+		public ReadOnlyMemory<byte> Data { get; }
 
 		/// <summary>
 		/// Constructor
diff --git a/Engine/Source/Programs/Shared/EpicGames.Perforce/PerforceReflection.cs b/Engine/Source/Programs/Shared/EpicGames.Perforce/PerforceReflection.cs
--- a/Engine/Source/Programs/Shared/EpicGames.Perforce/PerforceReflection.cs
+++ b/Engine/Source/Programs/Shared/EpicGames.Perforce/PerforceReflection.cs
@@ -42,9 +42,9 @@ namespace EpicGames.Perforce
 		public bool Optional;
 
 		/// <summary>
-		/// The field containing the value of this data.
+		/// The property containing the value of this data.
 		/// </summary>
-		public FieldInfo Field;
+		public PropertyInfo Property;
 
 		/// <summary>
 		/// Parser for this field type
@@ -66,13 +66,13 @@ namespace EpicGames.Perforce
 		/// </summary>
 		/// <param name="Name"></param>
 		/// <param name="Optional"></param>
-		/// <param name="Field"></param>
+		/// <param name="Property"></param>
 		/// <param name="RequiredTagBitMask"></param>
-		public CachedTagInfo(Utf8String Name, bool Optional, FieldInfo Field, ulong RequiredTagBitMask)
+		public CachedTagInfo(Utf8String Name, bool Optional, PropertyInfo Property, ulong RequiredTagBitMask)
 		{
 			this.Name = Name;
 			this.Optional = Optional;
-			this.Field = Field;
+			this.Property = Property;
 			this.RequiredTagBitMask = RequiredTagBitMask;
 			this.SetFromInteger = (Obj, Value) => throw new PerforceException($"Field {Name} was not expecting an integer value.");
 			this.SetFromString = (Obj, String) => throw new PerforceException($"Field {Name} was not expecting a string value.");
@@ -126,9 +126,9 @@ namespace EpicGames.Perforce
 		public CachedRecordInfo? SubElementRecordInfo;
 
 		/// <summary>
-		/// Field containing subelements
+		/// Property containing subelements
 		/// </summary>
-		public FieldInfo? SubElementField;
+		public PropertyInfo? SubElementProperty;
 
 		/// <summary>
 		/// Constructor
@@ -398,15 +398,15 @@ namespace EpicGames.Perforce
 				Record = new CachedRecordInfo(RecordType);
 
 				// Get all the fields for this type
-				FieldInfo[] Fields = RecordType.GetFields(BindingFlags.Public | BindingFlags.NonPublic | BindingFlags.Instance);
+				PropertyInfo[] Properties = RecordType.GetProperties(BindingFlags.Public | BindingFlags.NonPublic | BindingFlags.Instance);
 
 				// Build the map of all tags for this record
-				foreach (FieldInfo Field in Fields)
+				foreach (PropertyInfo Property in Properties)
 				{
-					PerforceTagAttribute? TagAttribute = Field.GetCustomAttribute<PerforceTagAttribute>();
+					PerforceTagAttribute? TagAttribute = Property.GetCustomAttribute<PerforceTagAttribute>();
 					if (TagAttribute != null)
 					{
-						string TagName = TagAttribute.Name ?? Field.Name;
+						string TagName = TagAttribute.Name ?? Property.Name;
 
 						ulong RequiredTagBitMask = 0;
 						if (!TagAttribute.Optional)
@@ -419,53 +419,53 @@ namespace EpicGames.Perforce
 							Record.RequiredTagsBitMask |= RequiredTagBitMask;
 						}
 
-						CachedTagInfo TagInfo = new CachedTagInfo(new Utf8String(TagName), TagAttribute.Optional, Field, RequiredTagBitMask);
+						CachedTagInfo TagInfo = new CachedTagInfo(new Utf8String(TagName), TagAttribute.Optional, Property, RequiredTagBitMask);
 
-						Type FieldType = Field.FieldType;
+						Type FieldType = Property.PropertyType;
 
-						FieldInfo FieldCopy = Field;
+						PropertyInfo PropertyCopy = Property;
 						if (FieldType == typeof(DateTime))
 						{
-							TagInfo.SetFromString = (Obj, String) => FieldCopy.SetValue(Obj, ParseStringAsDateTime(String));
+							TagInfo.SetFromString = (Obj, String) => PropertyCopy.SetValue(Obj, ParseStringAsDateTime(String));
 						}
 						else if (FieldType == typeof(bool))
 						{
-							TagInfo.SetFromString = (Obj, String) => FieldCopy.SetValue(Obj, ParseStringAsBool(String));
+							TagInfo.SetFromString = (Obj, String) => PropertyCopy.SetValue(Obj, ParseStringAsBool(String));
 						}
 						else if (FieldType == typeof(Nullable<bool>))
 						{
-							TagInfo.SetFromString = (Obj, String) => FieldCopy.SetValue(Obj, ParseStringAsNullableBool(String));
+							TagInfo.SetFromString = (Obj, String) => PropertyCopy.SetValue(Obj, ParseStringAsNullableBool(String));
 						}
 						else if (FieldType == typeof(int))
 						{
-							TagInfo.SetFromInteger = (Obj, Int) => FieldCopy.SetValue(Obj, Int);
-							TagInfo.SetFromString = (Obj, String) => FieldCopy.SetValue(Obj, ParseStringAsInt(String));
+							TagInfo.SetFromInteger = (Obj, Int) => PropertyCopy.SetValue(Obj, Int);
+							TagInfo.SetFromString = (Obj, String) => PropertyCopy.SetValue(Obj, ParseStringAsInt(String));
 						}
 						else if (FieldType == typeof(long))
 						{
-							TagInfo.SetFromString = (Obj, String) => FieldCopy.SetValue(Obj, ParseStringAsLong(String));
+							TagInfo.SetFromString = (Obj, String) => PropertyCopy.SetValue(Obj, ParseStringAsLong(String));
 						}
 						else if (FieldType == typeof(string))
 						{
-							TagInfo.SetFromString = (Obj, String) => FieldCopy.SetValue(Obj, ParseString(String));
+							TagInfo.SetFromString = (Obj, String) => PropertyCopy.SetValue(Obj, ParseString(String));
 						}
 						else if (FieldType == typeof(Utf8String))
 						{
-							TagInfo.SetFromString = (Obj, String) => FieldCopy.SetValue(Obj, String.Clone());
+							TagInfo.SetFromString = (Obj, String) => PropertyCopy.SetValue(Obj, String.Clone());
 						}
 						else if (FieldType.IsEnum)
 						{
 							CachedEnumInfo EnumInfo = GetCachedEnumInfo(FieldType);
-							TagInfo.SetFromInteger = (Obj, Int) => FieldCopy.SetValue(Obj, EnumInfo.ParseInteger(Int));
-							TagInfo.SetFromString = (Obj, String) => FieldCopy.SetValue(Obj, EnumInfo.ParseString(String));
+							TagInfo.SetFromInteger = (Obj, Int) => PropertyCopy.SetValue(Obj, EnumInfo.ParseInteger(Int));
+							TagInfo.SetFromString = (Obj, String) => PropertyCopy.SetValue(Obj, EnumInfo.ParseString(String));
 						}
 						else if (FieldType == typeof(DateTimeOffset?))
 						{
-							TagInfo.SetFromString = (Obj, String) => FieldCopy.SetValue(Obj, ParseStringAsNullableDateTimeOffset(String));
+							TagInfo.SetFromString = (Obj, String) => PropertyCopy.SetValue(Obj, ParseStringAsNullableDateTimeOffset(String));
 						}
 						else if (FieldType == typeof(List<string>))
 						{
-							TagInfo.SetFromString = (Obj, String) => ((List<string>)FieldCopy.GetValue(Obj)!).Add(String.ToString());
+							TagInfo.SetFromString = (Obj, String) => ((List<string>)PropertyCopy.GetValue(Obj)!).Add(String.ToString());
 						}
 						else
 						{
@@ -477,11 +477,11 @@ namespace EpicGames.Perforce
 
 					Record.NameToInfo = Record.Fields.ToDictionary(x => x.Name, x => x);
 
-					PerforceRecordListAttribute? SubElementAttribute = Field.GetCustomAttribute<PerforceRecordListAttribute>();
+					PerforceRecordListAttribute? SubElementAttribute = Property.GetCustomAttribute<PerforceRecordListAttribute>();
 					if (SubElementAttribute != null)
 					{
-						Record.SubElementField = Field;
-						Record.SubElementType = Field.FieldType.GenericTypeArguments[0];
+						Record.SubElementProperty = Property;
+						Record.SubElementType = Property.PropertyType.GenericTypeArguments[0];
 						Record.SubElementRecordInfo = GetCachedRecordInfo(Record.SubElementType);
 					}
 				}
diff --git a/Engine/Source/Programs/Shared/EpicGames.Perforce/PerforceResponse.cs b/Engine/Source/Programs/Shared/EpicGames.Perforce/PerforceResponse.cs
--- a/Engine/Source/Programs/Shared/EpicGames.Perforce/PerforceResponse.cs
+++ b/Engine/Source/Programs/Shared/EpicGames.Perforce/PerforceResponse.cs
@@ -16,7 +16,7 @@ namespace EpicGames.Perforce
 		/// <summary>
 		/// Stores the response data
 		/// </summary>
-		protected object InternalData;
+		protected object InternalData { get; }
 
 		/// <summary>
 		/// Constructor
diff --git a/Engine/Source/Programs/Shared/EpicGames.Perforce/PerforceUtils.cs b/Engine/Source/Programs/Shared/EpicGames.Perforce/PerforceUtils.cs
--- a/Engine/Source/Programs/Shared/EpicGames.Perforce/PerforceUtils.cs
+++ b/Engine/Source/Programs/Shared/EpicGames.Perforce/PerforceUtils.cs
@@ -21,10 +21,10 @@ namespace EpicGames.Perforce
 		static public string EscapePath(string Path)
 		{
 			string NewPath = Path;
-			NewPath = NewPath.Replace("%", "%25");
-			NewPath = NewPath.Replace("*", "%2A");
-			NewPath = NewPath.Replace("#", "%23");
-			NewPath = NewPath.Replace("@", "%40");
+			NewPath = NewPath.Replace("%", "%25", StringComparison.Ordinal);
+			NewPath = NewPath.Replace("*", "%2A", StringComparison.Ordinal);
+			NewPath = NewPath.Replace("#", "%23", StringComparison.Ordinal);
+			NewPath = NewPath.Replace("@", "%40", StringComparison.Ordinal);
 			return NewPath;
 		}
 
@@ -34,11 +34,11 @@ namespace EpicGames.Perforce
 		static public string UnescapePath(string Path)
 		{
 			string NewPath = Path;
-			NewPath = NewPath.Replace("%40", "@");
-			NewPath = NewPath.Replace("%23", "#");
-			NewPath = NewPath.Replace("%2A", "*");
-			NewPath = NewPath.Replace("%2a", "*");
-			NewPath = NewPath.Replace("%25", "%");
+			NewPath = NewPath.Replace("%40", "@", StringComparison.Ordinal);
+			NewPath = NewPath.Replace("%23", "#", StringComparison.Ordinal);
+			NewPath = NewPath.Replace("%2A", "*", StringComparison.Ordinal);
+			NewPath = NewPath.Replace("%2a", "*", StringComparison.Ordinal);
+			NewPath = NewPath.Replace("%25", "%", StringComparison.Ordinal);
 			return NewPath;
 		}
 
diff --git a/Engine/Source/Programs/Shared/EpicGames.Perforce/Records/AddRecord.cs b/Engine/Source/Programs/Shared/EpicGames.Perforce/Records/AddRecord.cs
--- a/Engine/Source/Programs/Shared/EpicGames.Perforce/Records/AddRecord.cs
+++ b/Engine/Source/Programs/Shared/EpicGames.Perforce/Records/AddRecord.cs
@@ -17,31 +17,31 @@ namespace EpicGames.Perforce
 		/// Path to the file in the depot
 		/// </summary>
 		[PerforceTag("depotFile")]
-		public string DepotFile;
+		public string DepotFile { get; set; }
 
 		/// <summary>
 		/// Path to the file in the workspace
 		/// </summary>
 		[PerforceTag("clientFile")]
-		public string ClientFile;
+		public string ClientFile { get; set; }
 
 		/// <summary>
 		/// The working revision number of the file that was synced
 		/// </summary>
 		[PerforceTag("workRev")]
-		public int WorkingRevision;
+		public int WorkingRevision { get; set; }
 
 		/// <summary>
 		/// Action taken when syncing the file
 		/// </summary>
 		[PerforceTag("action")]
-		public string Action;
+		public string Action { get; set; }
 
 		/// <summary>
 		/// Type of the file
 		/// </summary>
 		[PerforceTag("type")]
-		public string Type;
+		public string Type { get; set; }
 
 		/// <summary>
 		/// Private constructor for serialization
diff --git a/Engine/Source/Programs/Shared/EpicGames.Perforce/Records/AnnotateLineRecord.cs b/Engine/Source/Programs/Shared/EpicGames.Perforce/Records/AnnotateLineRecord.cs
--- a/Engine/Source/Programs/Shared/EpicGames.Perforce/Records/AnnotateLineRecord.cs
+++ b/Engine/Source/Programs/Shared/EpicGames.Perforce/Records/AnnotateLineRecord.cs
@@ -13,25 +13,25 @@ namespace EpicGames.Perforce
 		/// 
 		/// </summary>
 		[PerforceTag("depotFile")]
-		public readonly string DepotFile;
+		public string DepotFile { get; set; }
 
 		/// <summary>
 		/// The upper changelist 
 		/// </summary>
 		[PerforceTag("upper")]
-		public readonly int UpperRev;
+		public int UpperRev { get; set; }
 
 		/// <summary>
 		/// The lower changelist 
 		/// </summary>
 		[PerforceTag("lower")]
-		public readonly int LowerRev;
+		public int LowerRev { get; set; }
 
 		/// <summary>
 		/// Timestamp of this modification
 		/// </summary>
 		[PerforceTag("time", Optional = true)]
-		public readonly DateTime Time;
+		public DateTime Time { get; set; }
 
 		/// <summary>
 		/// Private constructor for serialization
diff --git a/Engine/Source/Programs/Shared/EpicGames.Perforce/Records/AnnotateRecord.cs b/Engine/Source/Programs/Shared/EpicGames.Perforce/Records/AnnotateRecord.cs
--- a/Engine/Source/Programs/Shared/EpicGames.Perforce/Records/AnnotateRecord.cs
+++ b/Engine/Source/Programs/Shared/EpicGames.Perforce/Records/AnnotateRecord.cs
@@ -14,61 +14,61 @@ namespace EpicGames.Perforce
 		/// The revision number of this file
 		/// </summary>
 		[PerforceTag("rev", Optional = true)]
-		public readonly int RevisionNumber;
+		public int RevisionNumber { get; set; }
 
 		/// <summary>
 		/// The changelist responsible for this revision of the file
 		/// </summary>
 		[PerforceTag("change", Optional = true)]
-		public readonly int ChangeNumber;
+		public int ChangeNumber { get; set; }
 
 		/// <summary>
 		/// Type of the file
 		/// </summary>
 		[PerforceTag("type", Optional = true)]
-		public readonly string Type;
+		public string Type { get; set; }
 
 		/// <summary>
 		/// The upper changelist 
 		/// </summary>
 		[PerforceTag("upper", Optional = true)]
-		public readonly string UpperCL;
+		public string UpperCL { get; set; }
 
 		/// <summary>
 		/// The lower changelist 
 		/// </summary>
 		[PerforceTag("lower", Optional = true)]
-		public readonly string LowerCL;
+		public string LowerCL { get; set; }
 
 		/// <summary>
 		/// Author of the changelist
 		/// </summary>
 		[PerforceTag("user", Optional = true)]
-		public readonly string UserName;
+		public string UserName { get; set; }
 
 		/// <summary>
 		/// Timestamp of this modification
 		/// </summary>
 		[PerforceTag("time", Optional = true)]
-		public readonly DateTime Time;
+		public DateTime Time { get; set; }
 
 		/// <summary>
 		/// Client that submitted this changelist
 		/// </summary>
 		[PerforceTag("client", Optional = true)]
-		public readonly string ClientName;
+		public string ClientName { get; set; }
 
 		/// <summary>
 		/// The actual line
 		/// </summary>
 		[PerforceTag("data", Optional = true)]
-		public readonly string Data;
+		public string Data { get; set; }
 		
 		/// <summary>
 		/// 
 		/// </summary>
 		[PerforceRecordList]
-		public List<AnnotateLineRecord> LineRecord = new List<AnnotateLineRecord>();
+		public List<AnnotateLineRecord> LineRecord { get; } = new List<AnnotateLineRecord>();
 
 		/// <summary>
 		/// Private constructor for serialization
diff --git a/Engine/Source/Programs/Shared/EpicGames.Perforce/Records/ChangeRecord.cs b/Engine/Source/Programs/Shared/EpicGames.Perforce/Records/ChangeRecord.cs
--- a/Engine/Source/Programs/Shared/EpicGames.Perforce/Records/ChangeRecord.cs
+++ b/Engine/Source/Programs/Shared/EpicGames.Perforce/Records/ChangeRecord.cs
@@ -17,48 +17,48 @@ namespace EpicGames.Perforce
 		/// The changelist number. -1 for "new".
 		/// </summary>
 		[PerforceTag("Change")]
-		public int Number = -1;
+		public int Number { get; set; } = -1;
 
 		/// <summary>
 		/// The date the change was last modified
 		/// </summary>
 		[PerforceTag("Date", Optional = true)]
-		public DateTime Date;
+		public DateTime Date { get; set; }
 
 		/// <summary>
 		/// The user that owns or submitted the change
 		/// </summary>
 		[PerforceTag("User")]
-		public string? User;
+		public string? User { get; set; }
 
 		/// <summary>
 		/// The client that owns the change
 		/// </summary>
 		[PerforceTag("Client")]
-		public string? Client;
+		public string? Client { get; set; }
 
 		/// <summary>
 		/// Current changelist status
 		/// </summary>
 		[PerforceTag("Status")]
-		public ChangeStatus Status;
+		public ChangeStatus Status { get; set; }
 
 		/// <summary>
 		/// Whether the change is restricted or not
 		/// </summary>
 		[PerforceTag("Type")]
-		public ChangeType Type;
+		public ChangeType Type { get; set; }
 
 		/// <summary>
 		/// Description for the changelist
 		/// </summary>
 		[PerforceTag("Description")]
-		public string? Description;
+		public string? Description { get; set; }
 
 		/// <summary>
 		/// Files that are open in this changelist
 		/// </summary>
 		[PerforceTag("Files", Optional = true)]
-		public List<string> Files = new List<string>();
+		public List<string> Files { get; } = new List<string>();
 	}
 }
diff --git a/Engine/Source/Programs/Shared/EpicGames.Perforce/Records/ChangesRecord.cs b/Engine/Source/Programs/Shared/EpicGames.Perforce/Records/ChangesRecord.cs
--- a/Engine/Source/Programs/Shared/EpicGames.Perforce/Records/ChangesRecord.cs
+++ b/Engine/Source/Programs/Shared/EpicGames.Perforce/Records/ChangesRecord.cs
@@ -17,48 +17,48 @@ namespace EpicGames.Perforce
 		/// The changelist number
 		/// </summary>
 		[PerforceTag("change")]
-		public int Number;
+		public int Number { get; set; }
 
 		/// <summary>
 		/// The date the change was last modified
 		/// </summary>
 		[PerforceTag("time")]
-		public DateTime Time;
+		public DateTime Time { get; set; }
 
 		/// <summary>
 		/// The user that owns or submitted the change
 		/// </summary>
 		[PerforceTag("user")]
-		public string? User;
+		public string? User { get; set; }
 
 		/// <summary>
 		/// The client that owns the change
 		/// </summary>
 		[PerforceTag("client")]
-		public string? Client;
+		public string? Client { get; set; }
 
 		/// <summary>
 		/// Current changelist status
 		/// </summary>
 		[PerforceTag("status")]
-		public ChangeStatus Status;
+		public ChangeStatus Status { get; set; }
 
 		/// <summary>
 		/// Whether the change is restricted or not
 		/// </summary>
 		[PerforceTag("changeType")]
-		public ChangeType Type;
+		public ChangeType Type { get; set; }
 
 		/// <summary>
 		/// The path affected by this change.
 		/// </summary>
 		[PerforceTag("path", Optional = true)]
-		public string? Path;
+		public string? Path { get; set; }
 
 		/// <summary>
 		/// Description for the changelist
 		/// </summary>
 		[PerforceTag("desc")]
-		public string? Description;
+		public string? Description { get; set; }
 	}
 }
diff --git a/Engine/Source/Programs/Shared/EpicGames.Perforce/Records/CleanOptions.cs b/Engine/Source/Programs/Shared/EpicGames.Perforce/Records/CleanOptions.cs
--- a/Engine/Source/Programs/Shared/EpicGames.Perforce/Records/CleanOptions.cs
+++ b/Engine/Source/Programs/Shared/EpicGames.Perforce/Records/CleanOptions.cs
@@ -15,7 +15,7 @@ namespace EpicGames.Perforce
 		/// <summary>
 		/// User default options
 		/// </summary>
-		Default = 0,
+		None = 0,
 
 		/// <summary>
 		/// Edited files: Find files in the workspace that have been modified and restore them to the last file version that has synced from the depot.
diff --git a/Engine/Source/Programs/Shared/EpicGames.Perforce/Records/CleanRecord.cs b/Engine/Source/Programs/Shared/EpicGames.Perforce/Records/CleanRecord.cs
--- a/Engine/Source/Programs/Shared/EpicGames.Perforce/Records/CleanRecord.cs
+++ b/Engine/Source/Programs/Shared/EpicGames.Perforce/Records/CleanRecord.cs
@@ -15,30 +15,30 @@ namespace EpicGames.Perforce
 		/// Depot path of the file, according to the workspace mapping
 		/// </summary>
 		[PerforceTag("depotFile", Optional = true)]
-		public string? DepotFile;
+		public string? DepotFile { get; set; }
 
 		/// <summary>
 		/// Client path to the file (in local path syntax)
 		/// </summary>
 		[PerforceTag("clientFile", Optional = true)]
-		public string? ClientFile;
+		public string? ClientFile { get; set; }
 
 		/// <summary>
 		/// Path to the file in local syntax
 		/// </summary>
 		[PerforceTag("localFile", Optional = true)]
-		public string? LocalFile;
+		public string? LocalFile { get; set; }
 
 		/// <summary>
 		/// Revision of the file
 		/// </summary>
 		[PerforceTag("rev", Optional = true)]
-		public string? Revision;
+		public string? Revision { get; set; }
 
 		/// <summary>
 		/// Action to be taken to the file
 		/// </summary>
 		[PerforceTag("action", Optional = true)]
-		public FileAction Action;
+		public FileAction Action { get; set; }
 	}
 }
diff --git a/Engine/Source/Programs/Shared/EpicGames.Perforce/Records/ClientRecord.cs b/Engine/Source/Programs/Shared/EpicGames.Perforce/Records/ClientRecord.cs
--- a/Engine/Source/Programs/Shared/EpicGames.Perforce/Records/ClientRecord.cs
+++ b/Engine/Source/Programs/Shared/EpicGames.Perforce/Records/ClientRecord.cs
@@ -17,79 +17,79 @@ namespace EpicGames.Perforce
 		/// The client workspace name, as specified in the P4CLIENT environment variable or its equivalents.
 		/// </summary>
 		[PerforceTag("Client")]
-		public string Name;
+		public string Name { get; set; }
 
 		/// <summary>
 		/// The Perforce user name of the user who owns the workspace.
 		/// </summary>
 		[PerforceTag("Owner")]
-		public string Owner;
+		public string Owner { get; set; }
 
 		/// <summary>
 		/// The date the workspace specification was last modified.
 		/// </summary>
 		[PerforceTag("Update", Optional = true)]
-		public DateTime Update;
+		public DateTime Update { get; set; }
 
 		/// <summary>
 		/// The date and time that the workspace was last used in any way.
 		/// </summary>
 		[PerforceTag("Access", Optional = true)]
-		public DateTime Access;
+		public DateTime Access { get; set; }
 
 		/// <summary>
 		/// The name of the workstation on which this workspace resides.
 		/// </summary>
 		[PerforceTag("Host", Optional = true)]
-		public string? Host;
+		public string? Host { get; set; }
 
 		/// <summary>
 		/// A textual description of the workspace.
 		/// </summary>
 		[PerforceTag("Description", Optional = true)]
-		public string? Description;
+		public string? Description { get; set; }
 
 		/// <summary>
 		/// The directory (on the local host) relative to which all the files in the View: are specified. 
 		/// </summary>
 		[PerforceTag("Root")]
-		public string Root;
+		public string Root { get; set; }
 
 		/// <summary>
 		/// A set of seven switches that control particular workspace options.
 		/// </summary>
 		[PerforceTag("Options")]
-		public ClientOptions Options;
+		public ClientOptions Options { get; set; }
 
 		/// <summary>
 		/// Options to govern the default behavior of p4 submit.
 		/// </summary>
 		[PerforceTag("SubmitOptions")]
-		public ClientSubmitOptions SubmitOptions;
+		public ClientSubmitOptions SubmitOptions { get; set; }
 
 		/// <summary>
 		/// Configure carriage-return/linefeed (CR/LF) conversion. 
 		/// </summary>
 		[PerforceTag("LineEnd")]
-		public ClientLineEndings LineEnd;
+		public ClientLineEndings LineEnd { get; set; }
 
 		/// <summary>
 		/// Associates the workspace with the specified stream.
 		/// </summary>
 		[PerforceTag("Stream", Optional = true)]
-		public string? Stream;
+		public string? Stream { get; set; }
 
 		/// <summary>
 		/// The type of client.
 		/// </summary>
 		[PerforceTag("Type", Optional = true)]
-		public string? Type;
+		public string? Type { get; set; }
 
 		/// <summary>
 		/// Specifies the mappings between files in the depot and files in the workspace.
 		/// </summary>
 		[PerforceTag("View")]
-		public List<string> View = new List<string>();
+		public List<string> View { get; } = new List<string>();
 
 		/// <summary>
 		/// Private constructor for serialization
diff --git a/Engine/Source/Programs/Shared/EpicGames.Perforce/Records/ClientsRecord.cs b/Engine/Source/Programs/Shared/EpicGames.Perforce/Records/ClientsRecord.cs
--- a/Engine/Source/Programs/Shared/EpicGames.Perforce/Records/ClientsRecord.cs
+++ b/Engine/Source/Programs/Shared/EpicGames.Perforce/Records/ClientsRecord.cs
@@ -17,73 +17,73 @@ namespace EpicGames.Perforce
 		/// The client workspace name, as specified in the P4CLIENT environment variable or its equivalents.
 		/// </summary>
 		[PerforceTag("client")]
-		public string Name;
+		public string Name { get; set; }
 
 		/// <summary>
 		/// The Perforce user name of the user who owns the workspace.
 		/// </summary>
 		[PerforceTag("Owner")]
-		public string Owner;
+		public string Owner { get; set; }
 
 		/// <summary>
 		/// The date the workspace specification was last modified.
 		/// </summary>
 		[PerforceTag("Update")]
-		public DateTime Update;
+		public DateTime Update { get; set; }
 
 		/// <summary>
 		/// The date and time that the workspace was last used in any way.
 		/// </summary>
 		[PerforceTag("Access")]
-		public DateTime Access;
+		public DateTime Access { get; set; }
 
 		/// <summary>
 		/// The name of the workstation on which this workspace resides.
 		/// </summary>
 		[PerforceTag("Host", Optional = true)]
-		public string? Host;
+		public string? Host { get; set; }
 
 		/// <summary>
 		/// A textual description of the workspace.
 		/// </summary>
 		[PerforceTag("Description", Optional = true)]
-		public string? Description;
+		public string? Description { get; set; }
 
 		/// <summary>
 		/// The directory (on the local host) relative to which all the files in the View: are specified. 
 		/// </summary>
 		[PerforceTag("Root")]
-		public string Root;
+		public string Root { get; set; }
 
 		/// <summary>
 		/// A set of seven switches that control particular workspace options.
 		/// </summary>
 		[PerforceTag("Options")]
-		public ClientOptions Options;
+		public ClientOptions Options { get; set; }
 
 		/// <summary>
 		/// Options to govern the default behavior of p4 submit.
 		/// </summary>
 		[PerforceTag("SubmitOptions")]
-		public ClientSubmitOptions SubmitOptions;
+		public ClientSubmitOptions SubmitOptions { get; set; }
 
 		/// <summary>
 		/// Configure carriage-return/linefeed (CR/LF) conversion. 
 		/// </summary>
 		[PerforceTag("LineEnd")]
-		public ClientLineEndings LineEnd;
+		public ClientLineEndings LineEnd { get; set; }
 
 		/// <summary>
 		/// Associates the workspace with the specified stream.
 		/// </summary>
 		[PerforceTag("Stream", Optional = true)]
-		public string? Stream;
+		public string? Stream { get; set; }
 
 		/// <summary>
 		/// The edge server ID
 		/// </summary>
 		[PerforceTag("ServerID", Optional = true)]
-		public string? ServerID;
+		public string? ServerID { get; set; }
 
 		/// <summary>
 		/// Private constructor for serialization
diff --git a/Engine/Source/Programs/Shared/EpicGames.Perforce/Records/DeleteRecord.cs b/Engine/Source/Programs/Shared/EpicGames.Perforce/Records/DeleteRecord.cs
--- a/Engine/Source/Programs/Shared/EpicGames.Perforce/Records/DeleteRecord.cs
+++ b/Engine/Source/Programs/Shared/EpicGames.Perforce/Records/DeleteRecord.cs
@@ -17,31 +17,31 @@ namespace EpicGames.Perforce
 		/// Path to the file in the depot
 		/// </summary>
 		[PerforceTag("depotFile")]
-		public string DepotFile;
+		public string DepotFile { get; set; }
 
 		/// <summary>
 		/// Path to the file in the workspace
 		/// </summary>
 		[PerforceTag("clientFile")]
-		public string ClientFile;
+		public string ClientFile { get; set; }
 
 		/// <summary>
 		/// The working revision number of the file that was synced
 		/// </summary>
 		[PerforceTag("workRev")]
-		public int WorkingRevision;
+		public int WorkingRevision { get; set; }
 
 		/// <summary>
 		/// Action taken when syncing the file
 		/// </summary>
 		[PerforceTag("action")]
-		public string Action;
+		public string Action { get; set; }
 
 		/// <summary>
 		/// Type of the file
 		/// </summary>
 		[PerforceTag("type")]
-		public string Type;
+		public string Type { get; set; }
 
 		/// <summary>
 		/// Private constructor for serialization
diff --git a/Engine/Source/Programs/Shared/EpicGames.Perforce/Records/DescribeFileRecord.cs b/Engine/Source/Programs/Shared/EpicGames.Perforce/Records/DescribeFileRecord.cs
--- a/Engine/Source/Programs/Shared/EpicGames.Perforce/Records/DescribeFileRecord.cs
+++ b/Engine/Source/Programs/Shared/EpicGames.Perforce/Records/DescribeFileRecord.cs
@@ -17,37 +17,37 @@ namespace EpicGames.Perforce
 		/// Path to the modified file in depot syntax
 		/// </summary>
 		[PerforceTag("depotFile")]
-		public string DepotFile;
+		public string DepotFile { get; set; }
 
 		/// <summary>
 		/// The action performed on this file
 		/// </summary>
 		[PerforceTag("action")]
-		public FileAction Action;
+		public FileAction Action { get; set; }
 
 		/// <summary>
 		/// The file type after this change
 		/// </summary>
 		[PerforceTag("type")]
-		public string Type;
+		public string Type { get; set; }
 
 		/// <summary>
 		/// The revision number for this file
 		/// </summary>
 		[PerforceTag("rev")]
-		public int Revision;
+		public int Revision { get; set; }
 
 		/// <summary>
 		/// Size of the file, or -1 if not specified
 		/// </summary>
 		[PerforceTag("fileSize", Optional = true)]
-		public readonly long FileSize;
+		public long FileSize { get; set; }
 
 		/// <summary>
 		/// Digest of the file, or null if not specified
 		/// </summary>
 		[PerforceTag("digest", Optional = true)]
-		public readonly string? Digest;
+		public string? Digest { get; set; }
 
 		/// <summary>
 		/// Private constructor for serialization
diff --git a/Engine/Source/Programs/Shared/EpicGames.Perforce/Records/DescribeRecord.cs b/Engine/Source/Programs/Shared/EpicGames.Perforce/Records/DescribeRecord.cs
--- a/Engine/Source/Programs/Shared/EpicGames.Perforce/Records/DescribeRecord.cs
+++ b/Engine/Source/Programs/Shared/EpicGames.Perforce/Records/DescribeRecord.cs
@@ -17,55 +17,55 @@ namespace EpicGames.Perforce
 		/// The changelist number
 		/// </summary>
 		[PerforceTag("change")]
-		public int Number;
+		public int Number { get; set; }
 
 		/// <summary>
 		/// The user that owns or submitted the change
 		/// </summary>
 		[PerforceTag("user")]
-		public string User;
+		public string User { get; set; }
 
 		/// <summary>
 		/// The workspace that contains or submitted the change
 		/// </summary>
 		[PerforceTag("client")]
-		public string Client;
+		public string Client { get; set; }
 
 		/// <summary>
 		/// Time at which the change was submitted
 		/// </summary>
 		[PerforceTag("time")]
-		public DateTime Time;
+		public DateTime Time { get; set; }
 
 		/// <summary>
 		/// The changelist description
 		/// </summary>
 		[PerforceTag("desc")]
-		public string Description;
+		public string Description { get; set; }
 
 		/// <summary>
 		/// The change status (submitted, pending, etc...)
 		/// </summary>
 		[PerforceTag("status")]
-		public ChangeStatus Status;
+		public ChangeStatus Status { get; set; }
 
 		/// <summary>
 		/// Whether the change is restricted or not
 		/// </summary>
 		[PerforceTag("changeType")]
-		public ChangeType Type;
+		public ChangeType Type { get; set; }
 
 		/// <summary>
 		/// Narrowest path which contains all files affected by this change
 		/// </summary>
 		[PerforceTag("path", Optional = true)]
-		public string? Path;
+		public string? Path { get; set; }
 
 		/// <summary>
 		/// The files affected by this change
 		/// </summary>
 		[PerforceRecordList]
-		public List<DescribeFileRecord> Files = new List<DescribeFileRecord>();
+		public List<DescribeFileRecord> Files { get; } = new List<DescribeFileRecord>();
 
 		/// <summary>
 		/// Private constructor for serialization
diff --git a/Engine/Source/Programs/Shared/EpicGames.Perforce/Records/DirsRecord.cs b/Engine/Source/Programs/Shared/EpicGames.Perforce/Records/DirsRecord.cs
--- a/Engine/Source/Programs/Shared/EpicGames.Perforce/Records/DirsRecord.cs
+++ b/Engine/Source/Programs/Shared/EpicGames.Perforce/Records/DirsRecord.cs
@@ -17,7 +17,7 @@ namespace EpicGames.Perforce
 		/// The directory name
 		/// </summary>
 		[PerforceTag("dir")]
-		public string Dir;
+		public string Dir { get; set; }
 
 		/// <summary>
 		/// Private constructor for serialization
diff --git a/Engine/Source/Programs/Shared/EpicGames.Perforce/Records/EditRecord.cs b/Engine/Source/Programs/Shared/EpicGames.Perforce/Records/EditRecord.cs
--- a/Engine/Source/Programs/Shared/EpicGames.Perforce/Records/EditRecord.cs
+++ b/Engine/Source/Programs/Shared/EpicGames.Perforce/Records/EditRecord.cs
@@ -17,31 +17,31 @@ namespace EpicGames.Perforce
 		/// Path to the file in the depot
 		/// </summary>
 		[PerforceTag("depotFile")]
-		public string DepotFile;
+		public string DepotFile { get; set; }
 
 		/// <summary>
 		/// Path to the file in the workspace
 		/// </summary>
 		[PerforceTag("clientFile")]
-		public string ClientFile;
+		public string ClientFile { get; set; }
 
 		/// <summary>
 		/// The working revision number of the file that was synced
 		/// </summary>
 		[PerforceTag("workRev")]
-		public int WorkingRevision;
+		public int WorkingRevision { get; set; }
 
 		/// <summary>
 		/// Action taken when syncing the file
 		/// </summary>
 		[PerforceTag("action")]
-		public string Action;
+		public string Action { get; set; }
 
 		/// <summary>
 		/// Type of the file
 		/// </summary>
 		[PerforceTag("type")]
-		public string Type;
+		public string Type { get; set; }
 
 		/// <summary>
 		/// Private constructor for serialization
diff --git a/Engine/Source/Programs/Shared/EpicGames.Perforce/Records/FStatRecord.cs b/Engine/Source/Programs/Shared/EpicGames.Perforce/Records/FStatRecord.cs
--- a/Engine/Source/Programs/Shared/EpicGames.Perforce/Records/FStatRecord.cs
+++ b/Engine/Source/Programs/Shared/EpicGames.Perforce/Records/FStatRecord.cs
@@ -17,133 +17,133 @@ namespace EpicGames.Perforce
 		/// Depot path to file
 		/// </summary>
 		[PerforceTag("depotFile", Optional = true)]
-		public string? DepotFile;
+		public string? DepotFile { get; set; }
 
 		/// <summary>
 		/// Local path to file (in local syntax by default, or in Perforce syntax with the FStatOptions.ClientFileInPerforceSyntax option)
 		/// </summary>
 		[PerforceTag("clientFile", Optional = true)]
-		public string? ClientFile;
+		public string? ClientFile { get; set; }
 
 		/// <summary>
 		/// For move operations, shows the counterpart file
 		/// </summary>
 		[PerforceTag("movedFile", Optional = true)]
-		public string? MovedFile;
+		public string? MovedFile { get; set; }
 
 		/// <summary>
 		/// Local path to file
 		/// </summary>
 		[PerforceTag("path", Optional = true)]
-		public string? Path;
+		public string? Path { get; set; }
 
 		/// <summary>
 		/// Set if the file is open for add and it is mapped to current client workspace
 		/// </summary>
 		[PerforceTag("isMapped", Optional = true)]
-		public bool IsMapped;
+		public bool IsMapped { get; set; }
 
 		/// <summary>
 		/// Set if file is shelved
 		/// </summary>
 		[PerforceTag("shelved", Optional = true)]
-		public bool Shelved;
+		public bool Shelved { get; set; }
 
 		/// <summary>
 		/// Action taken at head revision, if in depot
 		/// </summary>
 		[PerforceTag("headAction", Optional = true)]
-		public FileAction HeadAction;
+		public FileAction HeadAction { get; set; }
 
 		/// <summary>
 		/// Head revision type, if in depot
 		/// </summary>
 		[PerforceTag("headType", Optional = true)]
-		public string? HeadType;
+		public string? HeadType { get; set; }
 
 		/// <summary>
 		/// Head revision changelist time, if in depot.
 		/// </summary>
 		[PerforceTag("headTime", Optional = true)]
-		public DateTime HeadTime;
+		public DateTime HeadTime { get; set; }
 
 		/// <summary>
 		/// Revision for opened files
 		/// </summary>
 		[PerforceTag("rev", Optional = true)]
-		public int Revision;
+		public int Revision { get; set; }
 
 		/// <summary>
 		/// Head revision number, if in depot
 		/// </summary>
 		[PerforceTag("headRev", Optional = true)]
-		public int HeadRevision;
+		public int HeadRevision { get; set; }
 
 		/// <summary>
 		/// Head revision changelist number, if in depot
 		/// </summary>
 		[PerforceTag("headChange", Optional = true)]
-		public int HeadChange;
+		public int HeadChange { get; set; }
 
 		/// <summary>
 		/// Head revision modification time (the time that the file was last modified on the client before submit), if in depot.
 		/// </summary>
 		[PerforceTag("headModTime", Optional = true)]
-		public DateTime HeadModTime;
+		public DateTime HeadModTime { get; set; }
 
 		/// <summary>
 		/// Head revision of moved file
 		/// </summary>
 		[PerforceTag("movedRev", Optional = true)]
-		public int MovedRevision;
+		public int MovedRevision { get; set; }
 
 		/// <summary>
 		/// Revision last synced to workspace, if on workspace
 		/// </summary>
 		[PerforceTag("haveRev", Optional = true)]
-		public int HaveRevision;
+		public int HaveRevision { get; set; }
 
 		/// <summary>
 		/// Changelist description (if using -e changelist and if the file was part of changelist)
 		/// </summary>
 		[PerforceTag("desc", Optional = true)]
-		public string? Description;
+		public string? Description { get; set; }
 
 		/// <summary>
 		/// MD5 digest of a file (requires -Ol option)
 		/// </summary>
 		[PerforceTag("digest", Optional = true)]
-		public string? Digest;
+		public string? Digest { get; set; }
 
 		/// <summary>
 		/// File length in bytes (requires -Ol option)
 		/// </summary>
 		[PerforceTag("fileSize", Optional = true)]
-		public long FileSize;
+		public long FileSize { get; set; }
 
 		/// <summary>
 		/// Open action, if opened in your workspace
 		/// </summary>
 		[PerforceTag("action", Optional = true)]
-		public FileAction Action;
+		public FileAction Action { get; set; }
 
 		/// <summary>
 		/// Open type, if opened in your workspace
 		/// </summary>
 		[PerforceTag("type", Optional = true)]
-		public string? Type;
+		public string? Type { get; set; }
 
 		/// <summary>
 		/// User who opened the file, if opens
 		/// </summary>
 		[PerforceTag("actionOwner", Optional = true)]
-		public string? ActionOwner;
+		public string? ActionOwner { get; set; }
 
 		/// <summary>
 		/// Open changelist number, if opened in your workspace
 		/// </summary>
 		[PerforceTag("change", Optional = true)]
-		public int ChangeNumber;
+		public int ChangeNumber { get; set; }
 
 		/// <summary>
 		/// Summarize this record for display in the debugger
diff --git a/Engine/Source/Programs/Shared/EpicGames.Perforce/Records/FileLogRecord.cs b/Engine/Source/Programs/Shared/EpicGames.Perforce/Records/FileLogRecord.cs
--- a/Engine/Source/Programs/Shared/EpicGames.Perforce/Records/FileLogRecord.cs
+++ b/Engine/Source/Programs/Shared/EpicGames.Perforce/Records/FileLogRecord.cs
@@ -17,13 +17,13 @@ namespace EpicGames.Perforce
 		/// Path to the file in the depot
 		/// </summary>
 		[PerforceTag("depotFile")]
-		public string DepotPath;
+		public string DepotPath { get; set; }
 
 		/// <summary>
 		/// Revisions of this file
 		/// </summary>
 		[PerforceRecordList]
-		public List<RevisionRecord> Revisions = new List<RevisionRecord>();
+		public List<RevisionRecord> Revisions { get; } = new List<RevisionRecord>();
 
 		/// <summary>
 		/// Private constructor for serialization
diff --git a/Engine/Source/Programs/Shared/EpicGames.Perforce/Records/FileType.cs b/Engine/Source/Programs/Shared/EpicGames.Perforce/Records/FileType.cs
--- a/Engine/Source/Programs/Shared/EpicGames.Perforce/Records/FileType.cs
+++ b/Engine/Source/Programs/Shared/EpicGames.Perforce/Records/FileType.cs
@@ -3,6 +3,7 @@
 using EpicGames.Core;
 using System;
 using System.Collections.Generic;
+using System.Diagnostics.CodeAnalysis;
 using System.Text;
 
 namespace EpicGames.Perforce.Managed
@@ -10,7 +11,7 @@ namespace EpicGames.Perforce.Managed
 	/// <summary>
 	/// Base file type
 	/// </summary>
-	public enum FileTypeBase : uint
+	public enum FileTypeBase : int
 	{
 		/// <summary>
 		/// Text file
@@ -57,7 +58,7 @@ namespace EpicGames.Perforce.Managed
 	/// Modifiers on top of the base filetype
 	/// </summary>
 	[Flags]
-	public enum FileTypeModifiers : uint
+	public enum FileTypeModifiers : int
 	{
 		/// <summary>
 		/// File is always writable on the client
@@ -113,7 +114,7 @@ namespace EpicGames.Perforce.Managed
 	/// <summary>
 	/// Indicates the type of a file
 	/// </summary>
-	public struct FileType
+	public struct FileType : IEquatable<FileType>
 	{
 		/// <summary>
 		/// Size of the object when serialized to memory
@@ -338,9 +339,15 @@ namespace EpicGames.Perforce.Managed
 		}
 
 		/// <inheritdoc/>
-		public override bool Equals(object? obj)
+		public override bool Equals(object? Obj)
 		{
-			return (obj is FileType Type) && Type.Encoded == Encoded;
+			return (Obj is FileType Other) && Equals(Other);
+		}
+
+		/// <inheritdoc/>
+		public bool Equals(FileType Other)
+		{
+			return Encoded == Other.Encoded;
 		}
 
 		/// <inheritdoc/>
diff --git a/Engine/Source/Programs/Shared/EpicGames.Perforce/Records/FilesRecord.cs b/Engine/Source/Programs/Shared/EpicGames.Perforce/Records/FilesRecord.cs
--- a/Engine/Source/Programs/Shared/EpicGames.Perforce/Records/FilesRecord.cs
+++ b/Engine/Source/Programs/Shared/EpicGames.Perforce/Records/FilesRecord.cs
@@ -17,36 +17,36 @@ namespace EpicGames.Perforce
 		/// Path to the file in the depot
 		/// </summary>
 		[PerforceTag("depotFile")]
-		public string DepotFile = String.Empty;
+		public string DepotFile { get; set; } = String.Empty;
 
 		/// <summary>
 		/// File revision number
 		/// </summary>
 		[PerforceTag("rev")]
-		public int Revision;
+		public int Revision { get; set; }
 
 		/// <summary>
 		/// Changelist number that the file was modified in
 		/// </summary>
 		[PerforceTag("change")]
-		public int Change;
+		public int Change { get; set; }
 
 		/// <summary>
 		/// The action performed on this file
 		/// </summary>
 		[PerforceTag("action")]
-		public FileAction Action;
+		public FileAction Action { get; set; }
 
 		/// <summary>
 		/// The file type after this change
 		/// </summary>
 		[PerforceTag("type")]
-		public string Type = String.Empty;
+		public string Type { get; set; } = String.Empty;
 
 		/// <summary>
 		/// Modification time (the time that the file was last modified on the client before submit), if in depot.
 		/// </summary>
 		[PerforceTag("time", Optional = true)]
-		public DateTime Time;
+		public DateTime Time { get; set; }
 	}
 }
diff --git a/Engine/Source/Programs/Shared/EpicGames.Perforce/Records/InfoRecord.cs b/Engine/Source/Programs/Shared/EpicGames.Perforce/Records/InfoRecord.cs
--- a/Engine/Source/Programs/Shared/EpicGames.Perforce/Records/InfoRecord.cs
+++ b/Engine/Source/Programs/Shared/EpicGames.Perforce/Records/InfoRecord.cs
@@ -18,49 +18,49 @@ namespace EpicGames.Perforce
 		/// The current user according to the Perforce environment
 		/// </summary>
 		[PerforceTag("userName", Optional = true)]
-		public string? UserName;
+		public string? UserName { get; set; }
 
 		/// <summary>
 		/// The current client
 		/// </summary>
 		[PerforceTag("clientName", Optional = true)]
-		public string? ClientName;
+		public string? ClientName { get; set; }
 
 		/// <summary>
 		/// The current host
 		/// </summary>
 		[PerforceTag("clientHost", Optional = true)]
-		public string? ClientHost;
+		public string? ClientHost { get; set; }
 
 		/// <summary>
 		/// Root directory for the current client
 		/// </summary>
 		[PerforceTag("clientRoot", Optional = true)]
-		public string? ClientRoot;
+		public string? ClientRoot { get; set; }
 
 		/// <summary>
 		/// Selected stream in the current client
 		/// </summary>
 		[PerforceTag("clientStream", Optional = true)]
-		public string? ClientStream;
+		public string? ClientStream { get; set; }
 
 		/// <summary>
 		/// Address of the Perforce server
 		/// </summary>
 		[PerforceTag("serverAddress", Optional = true)]
-		public string? ServerAddress;
+		public string? ServerAddress { get; set; }
 
 		/// <summary>
 		/// Date and time on the server
 		/// </summary>
 		[PerforceTag("serverDate", Optional = true)]
-		public DateTimeOffset? ServerDate;
+		public DateTimeOffset? ServerDate { get; set; }
 
 		/// <summary>
 		/// Case handling setting on the server
 		/// </summary>
 		[PerforceTag("caseHandling", Optional = true)]
-		public string? CaseHandling;
+		public string? CaseHandling { get; set; }
 
 		/// <summary>
 		/// How to compare paths on this server
@@ -81,12 +81,12 @@ namespace EpicGames.Perforce
 		/// List of services provided by this server
 		/// </summary>
 		[PerforceTag("serverServices", Optional = true)]
-		public string? Services;
+		public string? Services { get; set; }
 
 		/// <summary>
 		/// The server unique id
 		/// </summary>
 		[PerforceTag("ServerID", Optional = true)]
-		public string? ServerID;
+		public string? ServerID { get; set; }
 	}
 }
diff --git a/Engine/Source/Programs/Shared/EpicGames.Perforce/Records/IntegrationRecord.cs b/Engine/Source/Programs/Shared/EpicGames.Perforce/Records/IntegrationRecord.cs
--- a/Engine/Source/Programs/Shared/EpicGames.Perforce/Records/IntegrationRecord.cs
+++ b/Engine/Source/Programs/Shared/EpicGames.Perforce/Records/IntegrationRecord.cs
@@ -17,25 +17,25 @@ namespace EpicGames.Perforce
 		/// The integration action performed for this file
 		/// </summary>
 		[PerforceTag("how")]
-		public IntegrateAction Action;
+		public IntegrateAction Action { get; set; }
 
 		/// <summary>
 		/// The partner file for this integration
 		/// </summary>
 		[PerforceTag("file")]
-		public string OtherFile;
+		public string OtherFile { get; set; }
 
 		/// <summary>
 		/// Min revision of the partner file for this integration
 		/// </summary>
 		[PerforceTag("srev")]
-		public int StartRevisionNumber;
+		public int StartRevisionNumber { get; set; }
 
 		/// <summary>
 		/// Max revision of the partner file for this integration
 		/// </summary>
 		[PerforceTag("erev")]
-		public int EndRevisionNumber;
+		public int EndRevisionNumber { get; set; }
 
 		/// <summary>
 		/// Private constructor for serialization
diff --git a/Engine/Source/Programs/Shared/EpicGames.Perforce/Records/MergeRecord.cs b/Engine/Source/Programs/Shared/EpicGames.Perforce/Records/MergeRecord.cs
--- a/Engine/Source/Programs/Shared/EpicGames.Perforce/Records/MergeRecord.cs
+++ b/Engine/Source/Programs/Shared/EpicGames.Perforce/Records/MergeRecord.cs
@@ -15,23 +15,23 @@ namespace EpicGames.Perforce
 		/// The file being merged from
 		/// </summary>
 		[PerforceTag("fromFile")]
-		public string FromFile = null!;
+		public string FromFile { get; set; } = null!;
 
 		/// <summary>
 		/// The starting revision, as an integer (null if 'none')
 		/// </summary>
-		public int? StartFromRev => String.Equals(StartFromRevText, "none") ? (int?)null : int.Parse(StartFromRevText);
+		public int? StartFromRev => String.Equals(StartFromRevText, "none", StringComparison.Ordinal) ? (int?)null : int.Parse(StartFromRevText);
 
 		/// <summary>
 		/// The first revision being merged (or 'none')
 		/// </summary>
 		[PerforceTag("startFromRev")]
-		public string StartFromRevText = null!;
+		public string StartFromRevText { get; set; } = null!;
 
 		/// <summary>
 		/// The last revision being merged
 		/// </summary>
 		[PerforceTag("endFromRev")]
-		public int EndFromRev;
+		public int EndFromRev { get; set; }
 	}
 }
diff --git a/Engine/Source/Programs/Shared/EpicGames.Perforce/Records/MoveRecord.cs b/Engine/Source/Programs/Shared/EpicGames.Perforce/Records/MoveRecord.cs
--- a/Engine/Source/Programs/Shared/EpicGames.Perforce/Records/MoveRecord.cs
+++ b/Engine/Source/Programs/Shared/EpicGames.Perforce/Records/MoveRecord.cs
@@ -17,31 +17,31 @@ namespace EpicGames.Perforce
 		/// Path to the file in the depot
 		/// </summary>
 		[PerforceTag("depotFile")]
-		public string DepotFile;
+		public string DepotFile { get; set; }
 
 		/// <summary>
 		/// Path to the file in the workspace
 		/// </summary>
 		[PerforceTag("clientFile")]
-		public string ClientFile;
+		public string ClientFile { get; set; }
 
 		/// <summary>
 		/// The working revision number of the file that was synced
 		/// </summary>
 		[PerforceTag("workRev", Optional = true)]
-		public int WorkingRevision;
+		public int WorkingRevision { get; set; }
 
 		/// <summary>
 		/// Action taken when syncing the file
 		/// </summary>
 		[PerforceTag("action")]
-		public string Action;
+		public string Action { get; set; }
 
 		/// <summary>
 		/// Type of the file
 		/// </summary>
 		[PerforceTag("type")]
-		public string Type;
+		public string Type { get; set; }
 
 		/// <summary>
 		/// Private constructor for serialization
diff --git a/Engine/Source/Programs/Shared/EpicGames.Perforce/Records/PrintRecord.cs b/Engine/Source/Programs/Shared/EpicGames.Perforce/Records/PrintRecord.cs
--- a/Engine/Source/Programs/Shared/EpicGames.Perforce/Records/PrintRecord.cs
+++ b/Engine/Source/Programs/Shared/EpicGames.Perforce/Records/PrintRecord.cs
@@ -17,43 +17,43 @@ namespace EpicGames.Perforce
 		/// Path to the file in the depot
 		/// </summary>
 		[PerforceTag("depotFile")]
-		public string DepotFile;
+		public string DepotFile { get; set; }
 
 		/// <summary>
 		/// Revision number of the file
 		/// </summary>
 		[PerforceTag("rev")]
-		public int Revision;
+		public int Revision { get; set; }
 
 		/// <summary>
 		/// Change number of the file
 		/// </summary>
 		[PerforceTag("change")]
-		public int ChangeNumber;
+		public int ChangeNumber { get; set; }
 
 		/// <summary>
 		/// Last action to the file
 		/// </summary>
 		[PerforceTag("action")]
-		public FileAction Action;
+		public FileAction Action { get; set; }
 
 		/// <summary>
 		/// File type
 		/// </summary>
 		[PerforceTag("type")]
-		public string Type;
+		public string Type { get; set; }
 
 		/// <summary>
 		/// Submit time of the file
 		/// </summary>
 		[PerforceTag("time")]
-		public DateTime Time;
+		public DateTime Time { get; set; }
 
 		/// <summary>
 		/// Size of the file in bytes
 		/// </summary>
 		[PerforceTag("fileSize", Optional = true)]
-		public long FileSize;
+		public long FileSize { get; set; }
 
 		/// <summary>
 		/// Private constructor for serialization
diff --git a/Engine/Source/Programs/Shared/EpicGames.Perforce/Records/ReconcileRecord.cs b/Engine/Source/Programs/Shared/EpicGames.Perforce/Records/ReconcileRecord.cs
--- a/Engine/Source/Programs/Shared/EpicGames.Perforce/Records/ReconcileRecord.cs
+++ b/Engine/Source/Programs/Shared/EpicGames.Perforce/Records/ReconcileRecord.cs
@@ -17,31 +17,31 @@ namespace EpicGames.Perforce
 		/// Path to the file in the depot
 		/// </summary>
 		[PerforceTag("depotFile")]
-		public string DepotFile;
+		public string DepotFile { get; set; }
 
 		/// <summary>
 		/// Path to the file in the workspace
 		/// </summary>
 		[PerforceTag("clientFile")]
-		public string ClientFile;
+		public string ClientFile { get; set; }
 
 		/// <summary>
 		/// The working revision number of the file that was synced
 		/// </summary>
 		[PerforceTag("workRev")]
-		public int WorkingRevision;
+		public int WorkingRevision { get; set; }
 
 		/// <summary>
 		/// Action taken when syncing the file
 		/// </summary>
 		[PerforceTag("action")]
-		public string Action;
+		public string Action { get; set; }
 
 		/// <summary>
 		/// Type of the file
 		/// </summary>
 		[PerforceTag("type")]
-		public string Type;
+		public string Type { get; set; }
 
 		/// <summary>
 		/// Private constructor for serialization
diff --git a/Engine/Source/Programs/Shared/EpicGames.Perforce/Records/ResolveRecord.cs b/Engine/Source/Programs/Shared/EpicGames.Perforce/Records/ResolveRecord.cs
--- a/Engine/Source/Programs/Shared/EpicGames.Perforce/Records/ResolveRecord.cs
+++ b/Engine/Source/Programs/Shared/EpicGames.Perforce/Records/ResolveRecord.cs
@@ -17,49 +17,49 @@ namespace EpicGames.Perforce
 		/// Path to the file in the workspace
 		/// </summary>
 		[PerforceTag("clientFile", Optional = true)]
-		public string? ClientFile;
+		public string? ClientFile { get; set; }
 
 		/// <summary>
 		/// Path to the depot file that needs to be resolved
 		/// </summary>
 		[PerforceTag("fromFile", Optional = true)]
-		public string? FromFile;
+		public string? FromFile { get; set; }
 
 		/// <summary>
 		/// Target file for the resolve
 		/// </summary>
 		[PerforceTag("toFile", Optional = true)]
-		public string? ToFile;
+		public string? ToFile { get; set; }
 
 		/// <summary>
 		/// How the file was resolved
 		/// </summary>
 		[PerforceTag("how", Optional = true)]
-		public string? How;
+		public string? How { get; set; }
 
 		/// <summary>
 		/// Start range of changes to be resolved
 		/// </summary>
 		[PerforceTag("startFromRev", Optional = true)]
-		public int FromRevisionStart;
+		public int FromRevisionStart { get; set; }
 
 		/// <summary>
 		/// Ending range of changes to be resolved
 		/// </summary>
 		[PerforceTag("endFromRev", Optional = true)]
-		public int FromRevisionEnd;
+		public int FromRevisionEnd { get; set; }
 
 		/// <summary>
 		/// The type of resolve to perform
 		/// </summary>
 		[PerforceTag("resolveType", Optional = true)]
-		public string? ResolveType;
+		public string? ResolveType { get; set; }
 
 		/// <summary>
 		/// For content resolves, the type of resolve to be performed
 		/// </summary>
 		[PerforceTag("contentResolveType", Optional = true)]
-		public string? ContentResolveType;
+		public string? ContentResolveType { get; set; }
 
 		/// <summary>
 		/// Private constructor for serialization
diff --git a/Engine/Source/Programs/Shared/EpicGames.Perforce/Records/RevertRecord.cs b/Engine/Source/Programs/Shared/EpicGames.Perforce/Records/RevertRecord.cs
--- a/Engine/Source/Programs/Shared/EpicGames.Perforce/Records/RevertRecord.cs
+++ b/Engine/Source/Programs/Shared/EpicGames.Perforce/Records/RevertRecord.cs
@@ -17,31 +17,31 @@ namespace EpicGames.Perforce
 		/// Path to the file in the depot
 		/// </summary>
 		[PerforceTag("depotFile")]
-		public string DepotFile;
+		public string DepotFile { get; set; }
 
 		/// <summary>
 		/// Path to the file in the workspace
 		/// </summary>
 		[PerforceTag("clientFile")]
-		public string ClientFile;
+		public string ClientFile { get; set; }
 
 		/// <summary>
 		/// The revision number that is now in the workspace
 		/// </summary>
 		[PerforceTag("haveRev")]
-		public int HaveRevision;
+		public int HaveRevision { get; set; }
 
 		/// <summary>
 		/// The previous action that the file was opened as
 		/// </summary>
 		[PerforceTag("oldAction")]
-		public string OldAction;
+		public string OldAction { get; set; }
 
 		/// <summary>
 		/// Action taken to revert the file
 		/// </summary>
 		[PerforceTag("action")]
-		public string Action;
+		public string Action { get; set; }
 
 		/// <summary>
 		/// Private constructor for serialization
diff --git a/Engine/Source/Programs/Shared/EpicGames.Perforce/Records/RevisionRecord.cs b/Engine/Source/Programs/Shared/EpicGames.Perforce/Records/RevisionRecord.cs
--- a/Engine/Source/Programs/Shared/EpicGames.Perforce/Records/RevisionRecord.cs
+++ b/Engine/Source/Programs/Shared/EpicGames.Perforce/Records/RevisionRecord.cs
@@ -17,67 +17,67 @@ namespace EpicGames.Perforce
 		/// The revision number of this file
 		/// </summary>
 		[PerforceTag("rev")]
-		public readonly int RevisionNumber;
+		public int RevisionNumber { get; set; }
 
 		/// <summary>
 		/// The changelist responsible for this revision of the file
 		/// </summary>
 		[PerforceTag("change")]
-		public readonly int ChangeNumber;
+		public int ChangeNumber { get; set; }
 
 		/// <summary>
 		/// Action performed to the file in this revision
 		/// </summary>
 		[PerforceTag("action")]
-		public readonly FileAction Action;
+		public FileAction Action { get; set; }
 
 		/// <summary>
 		/// Type of the file
 		/// </summary>
 		[PerforceTag("type")]
-		public readonly string Type;
+		public string Type { get; set; }
 
 		/// <summary>
 		/// Timestamp of this modification
 		/// </summary>
 		[PerforceTag("time")]
-		public readonly DateTime Time;
+		public DateTime Time { get; set; }
 
 		/// <summary>
 		/// Author of the changelist
 		/// </summary>
 		[PerforceTag("user")]
-		public readonly string UserName;
+		public string UserName { get; set; }
 
 		/// <summary>
 		/// Client that submitted this changelist
 		/// </summary>
 		[PerforceTag("client")]
-		public readonly string ClientName;
+		public string ClientName { get; set; }
 
 		/// <summary>
 		/// Size of the file, or -1 if not specified
 		/// </summary>
 		[PerforceTag("fileSize", Optional = true)]
-		public readonly long FileSize;
+		public long FileSize { get; set; }
 
 		/// <summary>
 		/// Digest of the file, or null if not specified
 		/// </summary>
 		[PerforceTag("digest", Optional = true)]
-		public readonly string? Digest;
+		public string? Digest { get; set; }
 
 		/// <summary>
 		/// Description of this changelist
 		/// </summary>
 		[PerforceTag("desc")]
-		public string Description;
+		public string Description { get; set; }
 
 		/// <summary>
 		/// Integration records for this revision
 		/// </summary>
 		[PerforceRecordList]
-		public List<IntegrationRecord> Integrations = new List<IntegrationRecord>();
+		public List<IntegrationRecord> Integrations { get; } = new List<IntegrationRecord>();
 
 		/// <summary>
 		/// Private constructor for serialization
diff --git a/Engine/Source/Programs/Shared/EpicGames.Perforce/Records/ShelveRecord.cs b/Engine/Source/Programs/Shared/EpicGames.Perforce/Records/ShelveRecord.cs
--- a/Engine/Source/Programs/Shared/EpicGames.Perforce/Records/ShelveRecord.cs
+++ b/Engine/Source/Programs/Shared/EpicGames.Perforce/Records/ShelveRecord.cs
@@ -17,25 +17,25 @@ namespace EpicGames.Perforce
 		/// The changelist containing the shelved file
 		/// </summary>
 		[PerforceTag("change", Optional = true)]
-		public int Change;
+		public int Change { get; set; }
 
 		/// <summary>
 		/// Path to the file in the depot
 		/// </summary>
 		[PerforceTag("depotFile")]
-		public string DepotFile;
+		public string DepotFile { get; set; }
 
 		/// <summary>
 		/// The revision number of the file that was shelved
 		/// </summary>
 		[PerforceTag("rev")]
-		public int Revision;
+		public int Revision { get; set; }
 
 		/// <summary>
 		/// The action to be applied to the file
 		/// </summary>
 		[PerforceTag("action")]
-		public string Action;
+		public string Action { get; set; }
 
 		/// <summary>
 		/// Private constructor for serialization
diff --git a/Engine/Source/Programs/Shared/EpicGames.Perforce/Records/StreamRecord.cs b/Engine/Source/Programs/Shared/EpicGames.Perforce/Records/StreamRecord.cs
--- a/Engine/Source/Programs/Shared/EpicGames.Perforce/Records/StreamRecord.cs
+++ b/Engine/Source/Programs/Shared/EpicGames.Perforce/Records/StreamRecord.cs
@@ -16,67 +16,67 @@ namespace EpicGames.Perforce
 		/// The stream path
 		/// </summary>
 		[PerforceTag("Stream")]
-		public string Stream;
+		public string Stream { get; set; }
 
 		/// <summary>
 		/// The stream name
 		/// </summary>
 		[PerforceTag("Name")]
-		public string Name;
+		public string Name { get; set; }
 
 		/// <summary>
 		/// The Perforce user name of the user who owns the stream.
 		/// </summary>
 		[PerforceTag("Owner")]
-		public string Owner;
+		public string Owner { get; set; }
 
 		/// <summary>
 		/// The date the stream specification was last modified.
 		/// </summary>
 		[PerforceTag("Update", Optional = true)]
-		public DateTime Update;
+		public DateTime Update { get; set; }
 
 		/// <summary>
 		/// The date and time that the stream was last used in any way.
 		/// </summary>
 		[PerforceTag("Access", Optional = true)]
-		public DateTime Access;
+		public DateTime Access { get; set; }
 
 		/// <summary>
 		/// The parent stream
 		/// </summary>
 		[PerforceTag("Parent", Optional = true)]
-		public string? Parent;
+		public string? Parent { get; set; }
 
 		/// <summary>
 		/// The stream type
 		/// </summary>
 		[PerforceTag("Type", Optional = true)]
-		public string? Type;
+		public string? Type { get; set; }
 
 		/// <summary>
 		/// A textual description of the stream.
 		/// </summary>
 		[PerforceTag("Description", Optional = true)]
-		public string? Description;
+		public string? Description { get; set; }
 
 		/// <summary>
 		/// Options for this stream
 		/// </summary>
 		[PerforceTag("Options")]
-		public StreamOptions Options;
+		public StreamOptions Options { get; set; }
 
 		/// <summary>
 		/// List of paths in the stream spec
 		/// </summary>
 		[PerforceTag("Paths")]
-		public List<string> Paths = new List<string>();
+		public List<string> Paths { get; } = new List<string>();
 
 		/// <summary>
 		/// Computed view for the stream
 		/// </summary>
 		[PerforceTag("View", Optional = true)]
-		public List<string> View = new List<string>();
+		public List<string> View { get; } = new List<string>();
 
 		/// <summary>
 		/// Private constructor for serialization
diff --git a/Engine/Source/Programs/Shared/EpicGames.Perforce/Records/StreamsRecord.cs b/Engine/Source/Programs/Shared/EpicGames.Perforce/Records/StreamsRecord.cs
--- a/Engine/Source/Programs/Shared/EpicGames.Perforce/Records/StreamsRecord.cs
+++ b/Engine/Source/Programs/Shared/EpicGames.Perforce/Records/StreamsRecord.cs
@@ -18,79 +18,79 @@ namespace EpicGames.Perforce
 		/// Path to the stream
 		/// </summary>
 		[PerforceTag("Stream")]
-		public string Stream;
+		public string Stream { get; set; }
 
 		/// <summary>
 		/// Last time the stream definition was updated
 		/// </summary>
 		[PerforceTag("Update")]
-		public DateTime Update;
+		public DateTime Update { get; set; }
 
 		/// <summary>
 		/// Last time the stream definition was accessed
 		/// </summary>
 		[PerforceTag("Access")]
-		public DateTime Access;
+		public DateTime Access { get; set; }
 
 		/// <summary>
 		/// Owner of this stream
 		/// </summary>
 		[PerforceTag("Owner")]
-		public string Owner;
+		public string Owner { get; set; }
 
 		/// <summary>
 		/// Name of the stream. This may be modified after the stream is initially created, but it's underlying depot path will not change.
 		/// </summary>
 		[PerforceTag("Name")]
-		public string Name;
+		public string Name { get; set; }
 
 		/// <summary>
 		/// The parent stream
 		/// </summary>
 		[PerforceTag("Parent", Optional = true)]
-		public string? Parent;
+		public string? Parent { get; set; }
 
 		/// <summary>
 		/// Type of the stream
 		/// </summary>
 		[PerforceTag("Type")]
-		public StreamType Type;
+		public StreamType Type { get; set; }
 
 		/// <summary>
 		/// User supplied description of the stream
 		/// </summary>
 		[PerforceTag("desc", Optional = true)]
-		public string? Description;
+		public string? Description { get; set; }
 
 		/// <summary>
 		/// Options for the stream definition
 		/// </summary>
 		[PerforceTag("Options")]
-		public StreamsOptions Options;
+		public StreamsOptions Options { get; set; }
 
 		/// <summary>
 		/// Whether this stream is more stable than the parent stream
 		/// </summary>
 		[PerforceTag("firmerThanParent")]
-		public Nullable<bool> FirmerThanParent;
+		public bool? FirmerThanParent { get; set; }
 
 		/// <summary>
 		/// Whether changes from this stream flow to the parent stream
 		/// </summary>
 		[PerforceTag("changeFlowsToParent")]
-		public bool ChangeFlowsToParent;
+		public bool ChangeFlowsToParent { get; set; }
 
 		/// <summary>
 		/// Whether changes from this stream flow from the parent stream
 		/// </summary>
 		[PerforceTag("changeFlowsFromParent")]
-		public bool ChangeFlowsFromParent;
+		public bool ChangeFlowsFromParent { get; set; }
 
 		/// <summary>
 		/// The mainline branch associated with this stream
 		/// </summary>
 		[PerforceTag("baseParent", Optional = true)]
-		public string? BaseParent;
+		public string? BaseParent { get; set; }
 
 		/// <summary>
 		/// Private constructor for serialization
diff --git a/Engine/Source/Programs/Shared/EpicGames.Perforce/Records/SubmitRecord.cs b/Engine/Source/Programs/Shared/EpicGames.Perforce/Records/SubmitRecord.cs
--- a/Engine/Source/Programs/Shared/EpicGames.Perforce/Records/SubmitRecord.cs
+++ b/Engine/Source/Programs/Shared/EpicGames.Perforce/Records/SubmitRecord.cs
@@ -17,7 +17,7 @@ namespace EpicGames.Perforce
 		/// Submitted changelist number
 		/// </summary>
 		[PerforceTag("change", Optional = true)]
-		public int ChangeNumber = -1;
+		public int ChangeNumber { get; set; } = -1;
 
 		/// <summary>
 		/// Format this record for display in the debugger
diff --git a/Engine/Source/Programs/Shared/EpicGames.Perforce/Records/SyncRecord.cs b/Engine/Source/Programs/Shared/EpicGames.Perforce/Records/SyncRecord.cs
--- a/Engine/Source/Programs/Shared/EpicGames.Perforce/Records/SyncRecord.cs
+++ b/Engine/Source/Programs/Shared/EpicGames.Perforce/Records/SyncRecord.cs
@@ -17,49 +17,49 @@ namespace EpicGames.Perforce
 		/// Path to the file in the depot
 		/// </summary>
 		[PerforceTag("depotFile")]
-		public string DepotFile;
+		public string DepotFile { get; set; }
 
 		/// <summary>
 		/// Path to the file in the workspace
 		/// </summary>
 		[PerforceTag("clientFile")]
-		public string ClientFile;
+		public string ClientFile { get; set; }
 
 		/// <summary>
 		/// The revision number of the file that was synced
 		/// </summary>
 		[PerforceTag("rev")]
-		public int Revision;
+		public int Revision { get; set; }
 
 		/// <summary>
 		/// Action taken when syncing the file
 		/// </summary>
 		[PerforceTag("action")]
-		public string Action;
+		public string Action { get; set; }
 
 		/// <summary>
 		/// Size of the file
 		/// </summary>
 		[PerforceTag("fileSize", Optional = true)]
-		public long FileSize;
+		public long FileSize { get; set; }
 
 		/// <summary>
 		/// Stores the total size of all files that are being synced (only set for the first sync record)
 		/// </summary>
 		[PerforceTag("totalFileSize", Optional = true)]
-		public long TotalFileSize;
+		public long TotalFileSize { get; set; }
 
 		/// <summary>
 		/// Stores the total number of files that will be synced (only set for the first sync record)
 		/// </summary>
 		[PerforceTag("totalFileCount", Optional = true)]
-		public long TotalFileCount;
+		public long TotalFileCount { get; set; }
 
 		/// <summary>
 		/// Change that modified the file
 		/// </summary>
 		[PerforceTag("change", Optional = true)]
-		public int Change;
+		public int Change { get; set; }
 
 		/// <summary>
 		/// Private constructor for serialization
diff --git a/Engine/Source/Programs/Shared/EpicGames.Perforce/Records/SyncSummaryRecord.cs b/Engine/Source/Programs/Shared/EpicGames.Perforce/Records/SyncSummaryRecord.cs
--- a/Engine/Source/Programs/Shared/EpicGames.Perforce/Records/SyncSummaryRecord.cs
+++ b/Engine/Source/Programs/Shared/EpicGames.Perforce/Records/SyncSummaryRecord.cs
@@ -17,18 +17,18 @@ namespace EpicGames.Perforce
 		/// The total size of all files synced
 		/// </summary>
 		[PerforceTag("totalFileSize")]
-		public long TotalFileSize;
+		public long TotalFileSize { get; set; }
 
 		/// <summary>
 		/// The total number of files synced
 		/// </summary>
 		[PerforceTag("totalFileCount")]
-		public long TotalFileCount;
+		public long TotalFileCount { get; set; }
 
 		/// <summary>
 		/// The changelist that was synced to
 		/// </summary>
 		[PerforceTag("change")]
-		public int Change;
+		public int Change { get; set; }
 	}
 }
diff --git a/Engine/Source/Programs/Shared/EpicGames.Perforce/Records/UndoRecord.cs b/Engine/Source/Programs/Shared/EpicGames.Perforce/Records/UndoRecord.cs
--- a/Engine/Source/Programs/Shared/EpicGames.Perforce/Records/UndoRecord.cs
+++ b/Engine/Source/Programs/Shared/EpicGames.Perforce/Records/UndoRecord.cs
@@ -17,31 +17,31 @@ namespace EpicGames.Perforce
 		/// Path to the file in the depot
 		/// </summary>
 		[PerforceTag("depotFile")]
-		public string DepotFile;
+		public string DepotFile { get; set; }
 
 		/// <summary>
 		/// Path to the file in the workspace
 		/// </summary>
 		[PerforceTag("clientFile")]
-		public string ClientFile;
+		public string ClientFile { get; set; }
 
 		/// <summary>
 		/// The working revision number of the file that was synced
 		/// </summary>
 		[PerforceTag("workRev")]
-		public int WorkingRevision;
+		public int WorkingRevision { get; set; }
 
 		/// <summary>
 		/// Action taken when syncing the file
 		/// </summary>
 		[PerforceTag("action")]
-		public string Action;
+		public string Action { get; set; }
 
 		/// <summary>
 		/// Type of the file
 		/// </summary>
 		[PerforceTag("type")]
-		public string Type;
+		public string Type { get; set; }
 
 		/// <summary>
 		/// Private constructor for serialization
diff --git a/Engine/Source/Programs/Shared/EpicGames.Perforce/Records/UnshelveRecord.cs b/Engine/Source/Programs/Shared/EpicGames.Perforce/Records/UnshelveRecord.cs
--- a/Engine/Source/Programs/Shared/EpicGames.Perforce/Records/UnshelveRecord.cs
+++ b/Engine/Source/Programs/Shared/EpicGames.Perforce/Records/UnshelveRecord.cs
@@ -17,19 +17,19 @@ namespace EpicGames.Perforce
 		/// Path to the file in the depot
 		/// </summary>
 		[PerforceTag("depotFile")]
-		public string DepotFile;
+		public string DepotFile { get; set; }
 
 		/// <summary>
 		/// The revision number of the file
 		/// </summary>
 		[PerforceTag("rev")]
-		public int Revision;
+		public int Revision { get; set; }
 
 		/// <summary>
 		/// Open action for the file
 		/// </summary>
 		[PerforceTag("action")]
-		public string Action;
+		public string Action { get; set; }
 
 		/// <summary>
 		/// Private constructor for serialization
diff --git a/Engine/Source/Programs/Shared/EpicGames.Perforce/Records/UpdateChangeOptions.cs b/Engine/Source/Programs/Shared/EpicGames.Perforce/Records/UpdateChangeOptions.cs
--- a/Engine/Source/Programs/Shared/EpicGames.Perforce/Records/UpdateChangeOptions.cs
+++ b/Engine/Source/Programs/Shared/EpicGames.Perforce/Records/UpdateChangeOptions.cs
@@ -11,6 +11,7 @@ namespace EpicGames.Perforce
 	/// <summary>
 	/// Options for the UpdateChange command
 	/// </summary>
+	[Flags]
 	public enum UpdateChangeOptions
 	{
 		/// <summary>
diff --git a/Engine/Source/Programs/Shared/EpicGames.Perforce/Records/UserRecord.cs b/Engine/Source/Programs/Shared/EpicGames.Perforce/Records/UserRecord.cs
--- a/Engine/Source/Programs/Shared/EpicGames.Perforce/Records/UserRecord.cs
+++ b/Engine/Source/Programs/Shared/EpicGames.Perforce/Records/UserRecord.cs
@@ -17,49 +17,49 @@ namespace EpicGames.Perforce
 		/// The name for the user
 		/// </summary>
 		[PerforceTag("User")]
-		public string UserName;
+		public string UserName { get; set; }
 
 		/// <summary>
 		/// Registered email address for reviews
 		/// </summary>
 		[PerforceTag("Email")]
-		public string Email;
+		public string Email { get; set; }
 
 		/// <summary>
 		/// Last time the user's information was updated
 		/// </summary>
 		[PerforceTag("Update", Optional = true)]
-		public DateTime Update;
+		public DateTime Update { get; set; }
 
 		/// <summary>
 		/// Last time the user's information was accessed
 		/// </summary>
 		[PerforceTag("Access", Optional = true)]
-		public DateTime Access;
+		public DateTime Access { get; set; }
 
 		/// <summary>
 		/// The user's full name
 		/// </summary>
 		[PerforceTag("FullName")]
-		public string FullName;
+		public string FullName { get; set; }
 
 		/// <summary>
 		/// Paths which the user is watching
 		/// </summary>
 		[PerforceTag("Reviews", Optional = true)]
-		public List<string> Reviews = new List<string>();
+		public List<string> Reviews { get; } = new List<string>();
 
 		/// <summary>
 		/// The type of user
 		/// </summary>
 		[PerforceTag("Type", Optional = true)]
-		public string? Type;
+		public string? Type { get; set; }
 
 		/// <summary>
 		/// Method used to authenticate
 		/// </summary>
 		[PerforceTag("AuthMethod")]
-		public string? AuthMethod;
+		public string? AuthMethod { get; set; }
 
 		/// <summary>
 		/// Private constructor for serialization
diff --git a/Engine/Source/Programs/Shared/EpicGames.Perforce/Records/WhereRecord.cs b/Engine/Source/Programs/Shared/EpicGames.Perforce/Records/WhereRecord.cs
--- a/Engine/Source/Programs/Shared/EpicGames.Perforce/Records/WhereRecord.cs
+++ b/Engine/Source/Programs/Shared/EpicGames.Perforce/Records/WhereRecord.cs
@@ -17,25 +17,25 @@ namespace EpicGames.Perforce
 		/// Path to the file in the depot
 		/// </summary>
 		[PerforceTag("depotFile")]
-		public string DepotFile;
+		public string DepotFile { get; set; }
 
 		/// <summary>
 		/// Path to the file in client syntax
 		/// </summary>
 		[PerforceTag("clientFile")]
-		public string ClientFile;
+		public string ClientFile { get; set; }
 
 		/// <summary>
 		/// Path to the file on dist
 		/// </summary>
 		[PerforceTag("path")]
-		public string Path;
+		public string Path { get; set; }
 
 		/// <summary>
 		/// Indicates that the given file or path is being unmapped from the workspace
 		/// </summary>
 		[PerforceTag("unmap", Optional = true)]
-		public bool Unmap;
+		public bool Unmap { get; set; }
 
 		/// <summary>
 		/// Private constructor for serialization
