Fix code analysis warning.

diff --git a/Engine/Source/Programs/Shared/EpicGames.Perforce/IPerforceOutput.cs b/Engine/Source/Programs/Shared/EpicGames.Perforce/IPerforceOutput.cs
--- a/Engine/Source/Programs/Shared/EpicGames.Perforce/IPerforceOutput.cs
+++ b/Engine/Source/Programs/Shared/EpicGames.Perforce/IPerforceOutput.cs
@@ -284,9 +284,6 @@ namespace EpicGames.Perforce
 		public static async Task ReadRecordsAsync(this IPerforceOutput Perforce, Action<PerforceRecord> HandleRecord, CancellationToken CancellationToken)
 		{
 			PerforceRecord Record = new PerforceRecord();
-			Record.Rows = new List<KeyValuePair<Utf8String, PerforceValue>>();
-
-			// Read all the records into a list
 			while (await Perforce.ReadAsync(CancellationToken))
 			{
 				// Start a read to add more data
diff --git a/Engine/Source/Programs/Shared/EpicGames.Perforce/PerforceRecord.cs b/Engine/Source/Programs/Shared/EpicGames.Perforce/PerforceRecord.cs
--- a/Engine/Source/Programs/Shared/EpicGames.Perforce/PerforceRecord.cs
+++ b/Engine/Source/Programs/Shared/EpicGames.Perforce/PerforceRecord.cs
@@ -215,12 +215,12 @@ namespace EpicGames.Perforce
 	/// <summary>
 	/// Low-overhead record type for generic responses
 	/// </summary>
-	public struct PerforceRecord
+	public class PerforceRecord
 	{
 		/// <summary>
 		/// The rows in this record
 		/// </summary>
-		public List<KeyValuePair<Utf8String, PerforceValue>> Rows;
+		public List<KeyValuePair<Utf8String, PerforceValue>> Rows { get; } = new List<KeyValuePair<Utf8String, PerforceValue>>();
 
 		/// <summary>
 		/// Copy this record into the given array of values. This method is O(n) if every record key being in the list of keys in the same order, but O(n^2) if not.
