Horde: Fix multi-line errors creating multiple single-line events, rather than one event that covers the whole thing.

diff --git a/Engine/Source/Programs/Horde/Horde.Agent.Tests/LogParserTests.cs b/Engine/Source/Programs/Horde/Horde.Agent.Tests/LogParserTests.cs
--- a/Engine/Source/Programs/Horde/Horde.Agent.Tests/LogParserTests.cs
+++ b/Engine/Source/Programs/Horde/Horde.Agent.Tests/LogParserTests.cs
@@ -62,7 +62,7 @@ namespace HordeAgentTests
 			{
 			}
 
-			protected override void WriteFormattedEvent(LogLevel Level, byte[] Line)
+			protected override void WriteFormattedEvent(LogLevel Level, int LineIndex, int LineCount, byte[] Line)
 			{
 				Lines.Add(LogEvent.Read(Line));
 			}
diff --git a/Engine/Source/Programs/Horde/Horde.Agent/Utility/JsonLogger.cs b/Engine/Source/Programs/Horde/Horde.Agent/Utility/JsonLogger.cs
--- a/Engine/Source/Programs/Horde/Horde.Agent/Utility/JsonLogger.cs
+++ b/Engine/Source/Programs/Horde/Horde.Agent/Utility/JsonLogger.cs
@@ -63,7 +63,7 @@ namespace HordeAgent.Utility
 
 			if(State is JsonLogEvent JsonEvent)
 			{
-				WriteFormattedEvent(LogLevel, JsonEvent.Data.ToArray());
+				WriteFormattedEvent(LogLevel, JsonEvent.LineIndex, JsonEvent.LineCount, JsonEvent.Data.ToArray());
 				return;
 			}
 
@@ -72,15 +72,17 @@ namespace HordeAgent.Utility
 			{
 				Event = LogEvent.FromState(LogLevel, EventId, State, Exception, Formatter);
 			}
-			WriteFormattedEvent(Event.Level, Event.ToJsonBytes());
+			WriteFormattedEvent(Event.Level, Event.LineIndex, Event.LineCount, Event.ToJsonBytes());
 		}
 
 		/// <summary>
 		/// Writes a formatted event
 		/// </summary>
 		/// <param name="Level">The log level</param>
+		/// <param name="LineIndex">Index of the current line within this event</param>
+		/// <param name="LineCount">Number of lines in this event</param>
 		/// <param name="Line">Utf-8 encoded JSON line data</param>
-		protected abstract void WriteFormattedEvent(LogLevel Level, byte[] Line);
+		protected abstract void WriteFormattedEvent(LogLevel Level, int LineIndex, int LineCount, byte[] Line);
 
 		/// <summary>
 		/// Callback to write a systemic event
diff --git a/Engine/Source/Programs/Horde/Horde.Agent/Utility/JsonRpcLogger.cs b/Engine/Source/Programs/Horde/Horde.Agent/Utility/JsonRpcLogger.cs
--- a/Engine/Source/Programs/Horde/Horde.Agent/Utility/JsonRpcLogger.cs
+++ b/Engine/Source/Programs/Horde/Horde.Agent/Utility/JsonRpcLogger.cs
@@ -82,7 +82,7 @@ namespace HordeAgent.Parser
 			this.Outcome = JobStepOutcome.Success;
 		}
 
-		protected override void WriteFormattedEvent(LogLevel Level, byte[] Line)
+		protected override void WriteFormattedEvent(LogLevel Level, int LineIndex, int LineCount, byte[] Line)
 		{
 			// Update the state of this job if this is an error status
 			if (Level == LogLevel.Error || Level == LogLevel.Critical)
@@ -96,9 +96,12 @@ namespace HordeAgent.Parser
 
 			// If we want an event for this log event, create one now
 			CreateEventRequest? Event = null;
-			if (Level == LogLevel.Warning || Level == LogLevel.Error || Level == LogLevel.Critical)
+			if (LineIndex == 0)
 			{
-				Event = CreateEvent(Level, 1);
+				if (Level == LogLevel.Warning || Level == LogLevel.Error || Level == LogLevel.Critical)
+				{
+					Event = CreateEvent(Level, LineCount);
+				}
 			}
 
 			// Write the data to the output channel
diff --git a/Engine/Source/Programs/Horde/Horde.Build.Tests/IssueServiceTests.cs b/Engine/Source/Programs/Horde/Horde.Build.Tests/IssueServiceTests.cs
--- a/Engine/Source/Programs/Horde/Horde.Build.Tests/IssueServiceTests.cs
+++ b/Engine/Source/Programs/Horde/Horde.Build.Tests/IssueServiceTests.cs
@@ -71,9 +71,12 @@ namespace HordeServerTests
 				}
 			}
 
-			protected override void WriteFormattedEvent(LogLevel Level, byte[] Line)
+			protected override void WriteFormattedEvent(LogLevel Level, int LineIndex, int LineCount, byte[] Line)
 			{
-				Events.Add((Level, Line));
+				if (LineIndex == 0)
+				{
+					Events.Add((Level, Line));
+				}
 			}
 
 			private async Task WriteAsync(LogLevel Level, byte[] Line)
diff --git a/Engine/Source/Programs/Shared/EpicGames.Core/JsonLogEvent.cs b/Engine/Source/Programs/Shared/EpicGames.Core/JsonLogEvent.cs
--- a/Engine/Source/Programs/Shared/EpicGames.Core/JsonLogEvent.cs
+++ b/Engine/Source/Programs/Shared/EpicGames.Core/JsonLogEvent.cs
@@ -15,6 +15,8 @@ namespace EpicGames.Core
 	{
 		static Utf8String LevelString { get; } = new Utf8String("Level");
 		static Utf8String PropertiesString { get; } = new Utf8String("Properties");
+		static Utf8String LineIndexString { get; } = new Utf8String("Line");
+		static Utf8String LineCountString { get; } = new Utf8String("LineCount");
 		static Utf8String EventIdString { get; } = new Utf8String("EventId");
 		static Utf8String IdString { get; } = new Utf8String("Id");
 		static Utf8String RenderedMessageString { get; } = new Utf8String("RenderedMessage");
@@ -29,6 +31,16 @@ namespace EpicGames.Core
 		/// </summary>
 		public EventId EventId { get; }
 
+		/// <summary>
+		/// Index of this line within the current event
+		/// </summary>
+		public int LineIndex { get; }
+
+		/// <summary>
+		/// Number of lines in the current event
+		/// </summary>
+		public int LineCount { get; }
+
 		/// <summary>
 		/// The utf-8 encoded JSON event
 		/// </summary>
@@ -37,10 +49,12 @@ namespace EpicGames.Core
 		/// <summary>
 		/// Constructor
 		/// </summary>
-		public JsonLogEvent(LogLevel Level, EventId EventId, ReadOnlyMemory<byte> Data)
+		public JsonLogEvent(LogLevel Level, EventId EventId, int LineIndex, int LineCount, ReadOnlyMemory<byte> Data)
 		{
 			this.Level = Level;
 			this.EventId = EventId;
+			this.LineIndex = LineIndex;
+			this.LineCount = LineCount;
 			this.Data = Data;
 		}
 
@@ -52,10 +66,13 @@ namespace EpicGames.Core
 		/// <returns></returns>
 		public static bool TryParse(ReadOnlyMemory<byte> Data, out JsonLogEvent Event)
 		{
+			Utf8StringComparer StringComparer = Utf8StringComparer.OrdinalIgnoreCase;
 			try
 			{
 				LogLevel Level = LogLevel.None;
 				int EventId = 0;
+				int LineIndex = 0;
+				int LineCount = 1;
 
 				Utf8JsonReader Reader = new Utf8JsonReader(Data.Span);
 				if (Reader.Read() && Reader.TokenType == JsonTokenType.StartObject)
@@ -63,22 +80,30 @@ namespace EpicGames.Core
 					ReadOnlySpan<byte> PropertyName;
 					for (; ReadNextPropertyName(ref Reader, out PropertyName); Reader.Skip())
 					{
-						if (Utf8StringComparer.OrdinalIgnoreCase.Equals(PropertyName, LevelString.Span))
+						if (StringComparer.Equals(PropertyName, LevelString.Span))
 						{
 							if (!Enum.TryParse(Reader.GetString(), out Level))
 							{
 								break;
 							}
 						}
-						else if (Utf8StringComparer.OrdinalIgnoreCase.Equals(PropertyName, PropertiesString.Span) && Reader.TokenType == JsonTokenType.StartObject)
+						else if (StringComparer.Equals(PropertyName, LineIndexString.Span))
+						{
+							LineIndex = Reader.GetInt32();
+						}
+						else if (StringComparer.Equals(PropertyName, LineCountString.Span))
+						{
+							LineCount = Reader.GetInt32();
+						}
+						else if (StringComparer.Equals(PropertyName, PropertiesString.Span) && Reader.TokenType == JsonTokenType.StartObject)
 						{
 							for (; ReadNextPropertyName(ref Reader, out PropertyName); Reader.Skip())
 							{
-								if (Utf8StringComparer.OrdinalIgnoreCase.Equals(EventIdString.Span, PropertyName) && Reader.TokenType == JsonTokenType.StartObject)
+								if (StringComparer.Equals(EventIdString.Span, PropertyName) && Reader.TokenType == JsonTokenType.StartObject)
 								{
 									for (; ReadNextPropertyName(ref Reader, out PropertyName); Reader.Skip())
 									{
-										if (Utf8StringComparer.OrdinalIgnoreCase.Equals(IdString.Span, PropertyName) && Reader.TryGetInt32(out int NewEventId))
+										if (StringComparer.Equals(IdString.Span, PropertyName) && Reader.TryGetInt32(out int NewEventId))
 										{
 											EventId = NewEventId;
 										}
@@ -91,7 +116,7 @@ namespace EpicGames.Core
 
 				if (Reader.TokenType == JsonTokenType.EndObject && Level != LogLevel.None && Reader.BytesConsumed == Data.Length)
 				{
-					Event = new JsonLogEvent(Level, new EventId(EventId), Data);
+					Event = new JsonLogEvent(Level, new EventId(EventId), LineIndex, LineCount, Data);
 					return true;
 				}
 			}
