Horde: Fix race condition in call to AgentService.GetRate().

diff --git a/Engine/Source/Programs/Horde/Horde.Build/Agents/Impl/AgentService.cs b/Engine/Source/Programs/Horde/Horde.Build/Agents/Impl/AgentService.cs
--- a/Engine/Source/Programs/Horde/Horde.Build/Agents/Impl/AgentService.cs
+++ b/Engine/Source/Programs/Horde/Horde.Build/Agents/Impl/AgentService.cs
@@ -79,13 +79,13 @@ namespace HordeServer.Services
 		IClock Clock;
 		ILogger Logger;
 		ITicker Ticker;
-
+		
 		RedisString<AgentRateTable> AgentRateTableData;
 		RedisService RedisService;
 
 		// Lazily updated costs for different agent types
-		AsyncCachedValue<AgentRateTable> AgentRateTable;
-
+		AsyncCachedValue<AgentRateTable?> AgentRateTable;
+		
 		// Lazily updated list of current pools
 		AsyncCachedValue<List<IPool>> PoolsList;
 
@@ -106,7 +106,7 @@ namespace HordeServer.Services
 			this.AclService = AclService;
 			this.DowntimeService = DowntimeService;
 			this.AgentRateTableData = new RedisString<AgentRateTable>(RedisService.Database, "agent-rates");
-			this.AgentRateTable = new AsyncCachedValue<AgentRateTable>(() => AgentRateTableData.GetAsync(), TimeSpan.FromSeconds(2.0));//.FromMinutes(5.0));
+			this.AgentRateTable = new AsyncCachedValue<AgentRateTable?>(() => AgentRateTableData.GetAsync(), TimeSpan.FromSeconds(2.0));//.FromMinutes(5.0));
 			this.PoolsList = new AsyncCachedValue<List<IPool>>(() => PoolCollection.GetAsync(), TimeSpan.FromSeconds(30.0));
 			this.DogStatsd = DogStatsd;
 			this.TaskSources = TaskSources.ToArray();
@@ -825,8 +825,8 @@ namespace HordeServer.Services
 				double Rate = 0.0;
 
 				// Get the rate table
-				AgentRateTable RateTable = await AgentRateTable.GetAsync();
-				if (RateTable.Entries.Count > 0)
+				AgentRateTable? RateTable = await AgentRateTable.GetAsync();
+				if (RateTable != null && RateTable.Entries.Count > 0)
 				{
 					IAgent? Agent = await GetAgentAsync(AgentId);
 					if (Agent != null)
diff --git a/Engine/Source/Programs/Shared/EpicGames.Redis/RedisString.cs b/Engine/Source/Programs/Shared/EpicGames.Redis/RedisString.cs
--- a/Engine/Source/Programs/Shared/EpicGames.Redis/RedisString.cs
+++ b/Engine/Source/Programs/Shared/EpicGames.Redis/RedisString.cs
@@ -32,13 +32,6 @@ namespace EpicGames.Redis
 			this.Key = Key;
 		}
 
-		/// <inheritdoc cref="IDatabaseAsync.StringGetAsync(RedisKey, CommandFlags)"/>
-		public async Task<TElement> GetAsync(CommandFlags Flags = CommandFlags.None)
-		{
-			RedisValue Value = await Database.StringGetAsync(Key, Flags);
-			return RedisSerializer.Deserialize<TElement>(Value);
-		}
-
 		/// <inheritdoc cref="IDatabaseAsync.StringLengthAsync(RedisKey, CommandFlags)"/>
 		public Task<long> LengthAsync(CommandFlags Flags = CommandFlags.None)
 		{
@@ -57,6 +50,28 @@ namespace EpicGames.Redis
 	/// </summary>
 	public static class RedisStringExtensions
 	{
+		/// <inheritdoc cref="IDatabaseAsync.StringGetAsync(RedisKey, CommandFlags)"/>
+		public static async Task<TElement?> GetAsync<TElement>(this RedisString<TElement> String, CommandFlags Flags = CommandFlags.None) where TElement : class
+		{
+			RedisValue Value = await String.Database.StringGetAsync(String.Key, Flags);
+			if (Value.IsNull)
+			{
+				return null;
+			}
+			return RedisSerializer.Deserialize<TElement>(RedisValue.Null);
+		}
+
+		/// <inheritdoc cref="IDatabaseAsync.StringGetAsync(RedisKey, CommandFlags)"/>
+		public static async Task<TElement?> GetValueAsync<TElement>(this RedisString<TElement> String, CommandFlags Flags = CommandFlags.None) where TElement : struct
+		{
+			RedisValue Value = await String.Database.StringGetAsync(String.Key, Flags);
+			if (Value.IsNull)
+			{
+				return default(TElement);
+			}
+			return RedisSerializer.Deserialize<TElement>(RedisValue.Null);
+		}
+
 		/// <inheritdoc cref="IDatabaseAsync.StringDecrementAsync(RedisKey, long, CommandFlags)"/>
 		public static Task<long> DecrementAsync(this RedisString<long> String, long Value = 1L, CommandFlags Flags = CommandFlags.None)
 		{
