Horde: Remove debug code.

diff --git a/Engine/Source/Programs/Shared/EpicGames.Redis/RedisString.cs b/Engine/Source/Programs/Shared/EpicGames.Redis/RedisString.cs
--- a/Engine/Source/Programs/Shared/EpicGames.Redis/RedisString.cs
+++ b/Engine/Source/Programs/Shared/EpicGames.Redis/RedisString.cs
@@ -58,7 +58,7 @@ namespace EpicGames.Redis
 			{
 				return null;
 			}
-			return RedisSerializer.Deserialize<TElement>(RedisValue.Null);
+			return RedisSerializer.Deserialize<TElement>(Value);
 		}
 
 		/// <inheritdoc cref="IDatabaseAsync.StringGetAsync(RedisKey, CommandFlags)"/>
@@ -69,7 +69,7 @@ namespace EpicGames.Redis
 			{
 				return default(TElement);
 			}
-			return RedisSerializer.Deserialize<TElement>(RedisValue.Null);
+			return RedisSerializer.Deserialize<TElement>(Value);
 		}
 
 		/// <inheritdoc cref="IDatabaseAsync.StringDecrementAsync(RedisKey, long, CommandFlags)"/>
