Horde: Handle edge cases with values for RedisString objects.

diff --git a/Engine/Source/Programs/Shared/EpicGames.Redis/RedisString.cs b/Engine/Source/Programs/Shared/EpicGames.Redis/RedisString.cs
--- a/Engine/Source/Programs/Shared/EpicGames.Redis/RedisString.cs
+++ b/Engine/Source/Programs/Shared/EpicGames.Redis/RedisString.cs
@@ -54,7 +54,7 @@ namespace EpicGames.Redis
 		public static async Task<TElement?> GetAsync<TElement>(this RedisString<TElement> String, CommandFlags Flags = CommandFlags.None) where TElement : class
 		{
 			RedisValue Value = await String.Database.StringGetAsync(String.Key, Flags);
-			if (Value.IsNull)
+			if (Value.IsNullOrEmpty)
 			{
 				return null;
 			}
@@ -65,7 +65,7 @@ namespace EpicGames.Redis
 		public static async Task<TElement?> GetValueAsync<TElement>(this RedisString<TElement> String, CommandFlags Flags = CommandFlags.None) where TElement : struct
 		{
 			RedisValue Value = await String.Database.StringGetAsync(String.Key, Flags);
-			if (Value.IsNull)
+			if (Value.IsNullOrEmpty)
 			{
 				return default(TElement);
 			}
