Horde: Strip VT100 escape codes from log output.

diff --git a/Engine/Source/Programs/Shared/EpicGames.Core/LogEventParser.cs b/Engine/Source/Programs/Shared/EpicGames.Core/LogEventParser.cs
--- a/Engine/Source/Programs/Shared/EpicGames.Core/LogEventParser.cs
+++ b/Engine/Source/Programs/Shared/EpicGames.Core/LogEventParser.cs
@@ -197,7 +197,7 @@ namespace EpicGames.Core
 					return;
 				}
 			}
-			Buffer.AddLine(Encoding.UTF8.GetString(Data.Span));
+			Buffer.AddLine(StringUtils.ParseEscapeCodes(Encoding.UTF8.GetString(Data.Span)));
 		}
 
 		/// <summary>
diff --git a/Engine/Source/Programs/Shared/EpicGames.Core/StringUtils.cs b/Engine/Source/Programs/Shared/EpicGames.Core/StringUtils.cs
--- a/Engine/Source/Programs/Shared/EpicGames.Core/StringUtils.cs
+++ b/Engine/Source/Programs/Shared/EpicGames.Core/StringUtils.cs
@@ -503,5 +503,50 @@ namespace EpicGames.Core
 			Bytes = null;
 			return false;
 		}
+
+		/// <summary>
+		/// Parses a string to remove VT100 escape codes
+		/// </summary>
+		/// <returns></returns>
+		public static string ParseEscapeCodes(string Line)
+		{
+			char EscapeChar = '\u001b';
+
+			int Index = Line.IndexOf(EscapeChar);
+			if (Index != -1)
+			{
+				int LastIndex = 0;
+
+				StringBuilder Result = new StringBuilder();
+				for (; ; )
+				{
+					Result.Append(Line, LastIndex, Index - LastIndex);
+
+					while (Index < Line.Length)
+					{
+						char Char = Line[Index];
+						if ((Char >= 'a' && Char <= 'z') || (Char >= 'A' && Char <= 'Z'))
+						{
+							Index++;
+							break;
+						}
+						Index++;
+					}
+
+					LastIndex = Index;
+
+					Index = Line.IndexOf(EscapeChar, Index);
+					if (Index == -1)
+					{
+						break;
+					}
+				}
+				Result.Append(Line, LastIndex, Line.Length - LastIndex);
+
+				Line = Result.ToString();
+			}
+
+			return Line;
+		}
 	}
 }
