BuildGraph: Add a task for querying all modified files since a certain changelist. Can be used in conjuction with Horde UE_HORDE_LAST_SUCCESS_CL to isolate files to rebuild.

diff --git a/Engine/Source/Programs/AutomationTool/AutomationUtils/AutomationUtils.Automation.csproj b/Engine/Source/Programs/AutomationTool/AutomationUtils/AutomationUtils.Automation.csproj
--- a/Engine/Source/Programs/AutomationTool/AutomationUtils/AutomationUtils.Automation.csproj
+++ b/Engine/Source/Programs/AutomationTool/AutomationUtils/AutomationUtils.Automation.csproj
@@ -51,6 +51,7 @@
   <ItemGroup>
     <ProjectReference Include="..\..\Shared\EpicGames.Core\EpicGames.Core.csproj" PrivateAssets="All"><Private>true</Private></ProjectReference>
     <ProjectReference Include="..\..\Shared\EpicGames.Build\EpicGames.Build.csproj" PrivateAssets="All"> <Private>true</Private> </ProjectReference>
+    <ProjectReference Include="..\..\Shared\EpicGames.Perforce\EpicGames.Perforce.csproj" />
     <ProjectReference Include="..\..\UnrealBuildTool\UnrealBuildTool.csproj" PrivateAssets="All"><Private>true</Private></ProjectReference>
   </ItemGroup>
   <ItemGroup>
diff --git a/Engine/Source/Programs/AutomationTool/AutomationUtils/P4Utils.cs b/Engine/Source/Programs/AutomationTool/AutomationUtils/P4Utils.cs
--- a/Engine/Source/Programs/AutomationTool/AutomationUtils/P4Utils.cs
+++ b/Engine/Source/Programs/AutomationTool/AutomationUtils/P4Utils.cs
@@ -1,4 +1,6 @@
 // Copyright Epic Games, Inc. All Rights Reserved.
+
+using EpicGames.Perforce;
 using System;
 using System.Collections.Generic;
 using System.Text;
@@ -940,6 +942,7 @@ namespace AutomationTool
 	{
 		static private P4Connection PerforceConnection;
 		static private P4Environment PerforceEnvironment;
+		static private IPerforceSettings PerforceSettings;
 
 		static public P4Connection P4
 		{
@@ -965,6 +968,18 @@ namespace AutomationTool
 			}
 		}
 
+		static public IPerforceSettings P4Settings
+		{
+			get
+			{
+				if (PerforceSettings == null)
+				{
+					throw new AutomationException("Attempt to use P4Settings before it was initialized or P4 support is disabled.");
+				}
+				return PerforceSettings;
+			}
+		}
+
 		/// <summary>
 		/// Initializes build environment. If the build command needs a specific env-var mapping or
 		/// has an extended BuildEnvironment, it must implement this method accordingly.
@@ -973,6 +988,11 @@ namespace AutomationTool
 		{
 			// Temporary connection - will use only the currently set env vars to connect to P4
 			PerforceEnvironment = new P4Environment(CmdEnv);
+
+			PerforceSettings Settings = new PerforceSettings(PerforceEnvironment.ServerAndPort, PerforceEnvironment.User);
+			Settings.PreferNativeClient = true;
+			Settings.ClientName = PerforceEnvironment.Client;
+			PerforceSettings = Settings;
 		}
 
 		/// <summary>
diff --git a/Engine/Source/Programs/AutomationTool/BuildGraph/Tasks/FindModifiedFiles.cs b/Engine/Source/Programs/AutomationTool/BuildGraph/Tasks/FindModifiedFiles.cs
new file mode 100644
--- /dev/null
+++ b/Engine/Source/Programs/AutomationTool/BuildGraph/Tasks/FindModifiedFiles.cs
@@ -0,0 +1,173 @@
+// Copyright Epic Games, Inc. All Rights Reserved.
+
+using EpicGames.BuildGraph;
+using EpicGames.Perforce;
+using Microsoft.Extensions.Logging.Abstractions;
+using System;
+using System.Collections.Generic;
+using System.Linq;
+using System.Text;
+using System.Threading.Tasks;
+using System.Xml;
+using EpicGames.Core;
+using UnrealBuildBase;
+using UnrealBuildTool;
+
+namespace AutomationTool.Tasks
+{
+	/// <summary>
+	/// Parameters for a task that compiles a C# project
+	/// </summary>
+	public class FindModifiedFilesTaskParameters
+	{
+		/// <summary>
+		/// The C# project file to compile. Using semicolons, more than one project file can be specified.
+		/// </summary>
+		[TaskParameter]
+		public string Paths;
+
+		/// <summary>
+		/// The configuration to compile.
+		/// </summary>
+		[TaskParameter(Optional = true)]
+		public int Change;
+
+		/// <summary>
+		/// The configuration to compile.
+		/// </summary>
+		[TaskParameter(Optional = true)]
+		public int MinChange;
+
+		/// <summary>
+		/// The configuration to compile.
+		/// </summary>
+		[TaskParameter(Optional = true)]
+		public int MaxChange;
+
+		/// <summary>
+		/// The file to write to
+		/// </summary>
+		[TaskParameter(Optional = true)]
+		public FileReference Output;
+	}
+
+	/// <summary>
+	/// Compiles C# project files, and their dependencies.
+	/// </summary>
+	[TaskElement("FindModifiedFiles", typeof(FindModifiedFilesTaskParameters))]
+	public class FindModifiedFilesTask : BgTaskImpl
+	{
+		/// <summary>
+		/// Parameters for the task
+		/// </summary>
+		FindModifiedFilesTaskParameters Parameters;
+
+		/// <summary>
+		/// Constructor.
+		/// </summary>
+		/// <param name="InParameters">Parameters for this task</param>
+		public FindModifiedFilesTask(FindModifiedFilesTaskParameters InParameters)
+		{
+			Parameters = InParameters;
+		}
+
+		/// <summary>
+		/// Execute the task.
+		/// </summary>
+		/// <param name="Job">Information about the current job</param>
+		/// <param name="BuildProducts">Set of build products produced by this node.</param>
+		/// <param name="TagNameToFileSet">Mapping from tag names to the set of files they include</param>
+		public override async Task ExecuteAsync(JobContext Job, HashSet<FileReference> BuildProducts, Dictionary<string, HashSet<FileReference>> TagNameToFileSet)
+		{
+			using IPerforceConnection Connection = await PerforceConnection.CreateAsync(CommandUtils.P4Settings, Log.Logger);
+
+			StringBuilder Filter = new StringBuilder($"//{Connection.Settings.ClientName}/...");
+			if (Parameters.Change > 0)
+			{
+				Filter.Append($"@={Parameters.Change}");
+			}
+			else if (Parameters.MinChange > 0)
+			{
+				if (Parameters.MaxChange > 0)
+				{
+					Filter.Append($"@{Parameters.MinChange},{Parameters.MaxChange}");
+				}
+				else
+				{
+					Filter.Append($"@>={Parameters.MinChange}");
+				}
+			}
+			else
+			{
+				throw new AutomationException("Change or MinChange must be specified to FindModifiedFiles task");
+			}
+
+			StreamRecord StreamRecord = await Connection.GetStreamAsync(CommandUtils.P4Env.Branch, true);
+			PerforceViewMap ViewMap = PerforceViewMap.Parse(StreamRecord.View);
+
+			List<string> LocalFiles = new List<string>();
+
+			List<FilesRecord> Files = await Connection.FilesAsync(FilesOptions.None, Filter.ToString());
+			foreach (FilesRecord File in Files)
+			{
+				string LocalFile;
+				if (ViewMap.TryMapFile(File.DepotFile, StringComparison.OrdinalIgnoreCase, out LocalFile))
+				{
+					LocalFiles.Add(LocalFile);
+				}
+				else
+				{
+					Log.TraceInformation("Unable to map {0} to workspace; skipping.", File.DepotFile);
+				}
+			}
+
+			Log.TraceInformation("Found {0} modified files matching {1}", LocalFiles.Count, Filter.ToString());
+			foreach (string LocalFile in LocalFiles)
+			{
+				Log.TraceInformation("  {0}", LocalFile);
+			}
+
+			if (Parameters.Output != null)
+			{
+				await FileReference.WriteAllLinesAsync(Parameters.Output, LocalFiles);
+				Log.TraceInformation("Written {0}", Parameters.Output);
+			}
+		}
+
+		/// <summary>
+		/// Output this task out to an XML writer.
+		/// </summary>
+		public override void Write(XmlWriter Writer)
+		{
+			Write(Writer, Parameters);
+		}
+
+		/// <summary>
+		/// Find all the tags which are used as inputs to this task
+		/// </summary>
+		/// <returns>The tag names which are read by this task</returns>
+		public override IEnumerable<string> FindConsumedTagNames()
+		{
+			//			return FindTagNamesFromFilespec(Parameters.Project);
+			return Enumerable.Empty<string>();
+		}
+
+		/// <summary>
+		/// Find all the tags which are modified by this task
+		/// </summary>
+		/// <returns>The tag names which are modified by this task</returns>
+		public override IEnumerable<string> FindProducedTagNames()
+		{
+			return Enumerable.Empty<string>();
+/*			foreach (string TagName in FindTagNamesFromList(Parameters.Tag))
+			{
+				yield return TagName;
+			}
+
+			foreach (string TagName in FindTagNamesFromList(Parameters.TagReferences))
+			{
+				yield return TagName;
+			}*/
+		}
+	}
+}
diff --git a/Engine/Source/Programs/Shared/EpicGames.Perforce/PerforceViewMap.cs b/Engine/Source/Programs/Shared/EpicGames.Perforce/PerforceViewMap.cs
new file mode 100644
--- /dev/null
+++ b/Engine/Source/Programs/Shared/EpicGames.Perforce/PerforceViewMap.cs
@@ -0,0 +1,333 @@
+// Copyright Epic Games, Inc. All Rights Reserved.
+
+using EpicGames.Core;
+using System;
+using System.Collections.Generic;
+using System.Diagnostics;
+using System.Diagnostics.CodeAnalysis;
+using System.Globalization;
+using System.Linq;
+using System.Text;
+using System.Text.RegularExpressions;
+
+namespace EpicGames.Perforce
+{
+	/// <summary>
+	/// Stores a mapping from one set of paths to another
+	/// </summary>
+	public class PerforceViewMap
+	{
+		/// <summary>
+		/// List of entries making up the view
+		/// </summary>
+		public List<PerforceViewMapEntry> Entries { get; }
+
+		/// <summary>
+		/// Default constructor
+		/// </summary>
+		public PerforceViewMap()
+		{
+			Entries = new List<PerforceViewMapEntry>();
+		}
+
+		/// <summary>
+		/// Construct from an existing set of entries
+		/// </summary>
+		/// <param name="Entries"></param>
+		public PerforceViewMap(IEnumerable<PerforceViewMapEntry> Entries)
+		{
+			this.Entries = Entries.ToList();
+		}
+
+		/// <summary>
+		/// Constructor
+		/// </summary>
+		/// <param name="Other"></param>
+		public PerforceViewMap(PerforceViewMap Other)
+		{
+			Entries = new List<PerforceViewMapEntry>(Other.Entries);
+		}
+
+		/// <summary>
+		/// Construct a view map from a set of entries
+		/// </summary>
+		/// <param name="Entries"></param>
+		public static PerforceViewMap Parse(IEnumerable<string> Entries)
+		{
+			return new PerforceViewMap(Entries.Select(x => PerforceViewMapEntry.Parse(x)));
+		}
+
+		/// <summary>
+		/// Determines if a file is included in the view
+		/// </summary>
+		/// <param name="File">The file to test</param>
+		/// <param name="Comparison">The comparison type</param>
+		/// <returns>True if the file is included in the view</returns>
+		public bool MatchFile(string File, StringComparison Comparison)
+		{
+			bool Included = false;
+			foreach (PerforceViewMapEntry Entry in Entries)
+			{
+				if (Entry.MatchFile(File, Comparison))
+				{
+					Included = Entry.Include;
+				}
+			}
+			return Included;
+		}
+
+		/// <summary>
+		/// Attempts to convert a source file to its target path
+		/// </summary>
+		/// <param name="SourceFile"></param>
+		/// <param name="Comparison">The comparison type</param>
+		/// <param name="TargetFile"></param>
+		/// <returns></returns>
+		public bool TryMapFile(string SourceFile, StringComparison Comparison, out string TargetFile)
+		{
+			PerforceViewMapEntry? MapEntry = null;
+			foreach (PerforceViewMapEntry Entry in Entries)
+			{
+				if (Entry.MatchFile(SourceFile, Comparison))
+				{
+					MapEntry = Entry;
+				}
+			}
+
+			if (MapEntry != null && MapEntry.Include)
+			{
+				TargetFile = MapEntry.MapFile(SourceFile);
+				return true;
+			}
+			else
+			{
+				TargetFile = String.Empty;
+				return false;
+			}
+		}
+
+		/// <summary>
+		/// Gets the root paths from the view entries
+		/// </summary>
+		/// <returns></returns>
+		public List<string> GetRootPaths(StringComparison Comparison)
+		{
+			List<string> RootPaths = new List<string>();
+			foreach (PerforceViewMapEntry Entry in Entries)
+			{
+				if (Entry.Include)
+				{
+					int LastSlashIdx = Entry.SourcePrefix.LastIndexOf('/');
+					ReadOnlySpan<char> RootPath = Entry.SourcePrefix.AsSpan(0, LastSlashIdx + 1);
+
+					for (int Idx = 0; ; Idx++)
+					{
+						if (Idx == RootPaths.Count)
+						{
+							RootPaths.Add(RootPath.ToString());
+							break;
+						}
+						else if (RootPaths[Idx].AsSpan().StartsWith(RootPath, Comparison))
+						{
+							RootPaths[Idx] = RootPath.ToString();
+							break;
+						}
+						else if (RootPath.StartsWith(RootPaths[Idx], Comparison))
+						{
+							break;
+						}
+					}
+				}
+			}
+			return RootPaths;
+		}
+	}
+
+	/// <summary>
+	/// Entry within a ViewMap
+	/// </summary>
+	public class PerforceViewMapEntry
+	{
+		/// <summary>
+		/// Whether to include files matching this pattern
+		/// </summary>
+		public bool Include { get; }
+
+		/// <summary>
+		/// The wildcard string - either '*' or '...'
+		/// </summary>
+		public string Wildcard { get; }
+
+		/// <summary>
+		/// The source part of the pattern before the wildcard
+		/// </summary>
+		public string SourcePrefix { get; }
+
+		/// <summary>
+		/// The source part of the pattern after the wildcard. Perforce does not permit a slash to be in this part of the mapping.
+		/// </summary>
+		public string SourceSuffix { get; }
+
+		/// <summary>
+		/// The target mapping for the pattern before the wildcard
+		/// </summary>
+		public string TargetPrefix { get; }
+
+		/// <summary>
+		/// The target mapping for the pattern after the wildcard
+		/// </summary>
+		public string TargetSuffix { get; }
+
+		/// <summary>
+		/// The full source pattern
+		/// </summary>
+		public string Source => $"{SourcePrefix}{Wildcard}{SourceSuffix}";
+
+		/// <summary>
+		/// The full target pattern
+		/// </summary>
+		public string Target => $"{TargetPrefix}{Wildcard}{TargetSuffix}";
+
+		/// <summary>
+		/// Tests if the entry has a file wildcard ('*')
+		/// </summary>
+		/// <returns>True if the entry has a file wildcard</returns>
+		public bool IsFileWildcard() => Wildcard.Length == 1;
+
+		/// <summary>
+		/// Tests if the entry has a path wildcard ('...')
+		/// </summary>
+		/// <returns>True if the entry has a path wildcard</returns>
+		public bool IsPathWildcard() => Wildcard.Length == 3;
+
+		/// <summary>
+		/// Constructor
+		/// </summary>
+		/// <param name="Other"></param>
+		public PerforceViewMapEntry(PerforceViewMapEntry Other)
+			: this(Other.Include, Other.Wildcard, Other.SourcePrefix, Other.SourceSuffix, Other.TargetPrefix, Other.TargetSuffix)
+		{
+		}
+
+		/// <summary>
+		/// Constructor
+		/// </summary>
+		/// <param name="Include"></param>
+		/// <param name="Source"></param>
+		/// <param name="Target"></param>
+		public PerforceViewMapEntry(bool Include, string Source, string Target)
+		{
+			this.Include = Include;
+
+			Match Match = Regex.Match(Source, @"^(.*)(\*|\.\.\.|%%1)(.*)$");
+			if (Match.Success)
+			{
+				string WildcardStr = Match.Groups[2].Value;
+
+				SourcePrefix = Match.Groups[1].Value;
+				SourceSuffix = Match.Groups[3].Value;
+				Wildcard = Match.Groups[2].Value;
+
+				int OtherIdx = Target.IndexOf(WildcardStr, StringComparison.Ordinal);
+				TargetPrefix = Target.Substring(0, OtherIdx);
+				TargetSuffix = Target.Substring(OtherIdx + Wildcard.Length);
+
+				if (WildcardStr.Equals("%%1", StringComparison.Ordinal))
+				{
+					Wildcard = "*";
+				}
+			}
+			else
+			{
+				SourcePrefix = Source;
+				SourceSuffix = String.Empty;
+				TargetPrefix = Target;
+				TargetSuffix = String.Empty;
+				Wildcard = String.Empty;
+			}
+		}
+
+		/// <summary>
+		/// Constructor
+		/// </summary>
+		/// <param name="Include"></param>
+		/// <param name="Wildcard"></param>
+		/// <param name="SourcePrefix"></param>
+		/// <param name="SourceSuffix"></param>
+		/// <param name="TargetPrefix"></param>
+		/// <param name="TargetSuffix"></param>
+		public PerforceViewMapEntry(bool Include, string Wildcard, string SourcePrefix, string SourceSuffix, string TargetPrefix, string TargetSuffix)
+		{
+			this.Include = Include;
+			this.Wildcard = Wildcard;
+			this.SourcePrefix = SourcePrefix;
+			this.SourceSuffix = SourceSuffix;
+			this.TargetPrefix = TargetPrefix;
+			this.TargetSuffix = TargetSuffix;
+		}
+
+		/// <summary>
+		/// Parse a view map entry from a string, as returned by spec documents
+		/// </summary>
+		/// <param name="Entry"></param>
+		/// <returns></returns>
+		public static PerforceViewMapEntry Parse(string Entry)
+		{
+			Match Match = Regex.Match(Entry, @"^\s*(-?)\s*([^ ]+)\s+([^ ]+)\s*$");
+			if(!Match.Success)
+			{
+				throw new PerforceException($"Unable to parse view map entry: {Entry}");
+			}
+			return new PerforceViewMapEntry(Match.Groups[1].Length == 0, Match.Groups[2].Value, Match.Groups[3].Value);
+		}
+
+		/// <summary>
+		/// Maps a file to the target path
+		/// </summary>
+		/// <param name="SourceFile"></param>
+		/// <returns></returns>
+		public string MapFile(string SourceFile)
+		{
+			int Count = SourceFile.Length - SourceSuffix.Length - SourcePrefix.Length;
+			return String.Concat(TargetPrefix, SourceFile.AsSpan(SourcePrefix.Length, Count), TargetSuffix);
+		}
+
+		/// <summary>
+		/// Determine if a file matches the current entry
+		/// </summary>
+		/// <param name="Path">Path to the file</param>
+		/// <param name="Comparison">The comparison type</param>
+		/// <returns>True if the path matches the entry</returns>
+		public bool MatchFile(string Path, StringComparison Comparison)
+		{
+			if (Wildcard.Length == 0)
+			{
+				return String.Equals(Path, SourcePrefix, Comparison);
+			}
+			else
+			{
+				if (!Path.StartsWith(SourcePrefix, Comparison) || !Path.EndsWith(SourceSuffix, Comparison))
+				{
+					return false;
+				}
+				if (IsFileWildcard() && Path.AsSpan(SourcePrefix.Length, Path.Length - SourceSuffix.Length - SourcePrefix.Length).IndexOf('/') != -1)
+				{
+					return false;
+				}
+				return true;
+			}
+		}
+
+		/// <inheritdoc/>
+		public override string ToString()
+		{
+			StringBuilder Builder = new StringBuilder();
+			if (!Include)
+			{
+				Builder.Append('-');
+			}
+			Builder.Append($"{SourcePrefix}{Wildcard}{SourceSuffix} {TargetPrefix}{Wildcard}{TargetSuffix}");
+			return Builder.ToString();
+		}
+	}
+}
diff --git a/Engine/Source/Programs/UnrealBuildTool/Configuration/TargetDescriptor.cs b/Engine/Source/Programs/UnrealBuildTool/Configuration/TargetDescriptor.cs
--- a/Engine/Source/Programs/UnrealBuildTool/Configuration/TargetDescriptor.cs
+++ b/Engine/Source/Programs/UnrealBuildTool/Configuration/TargetDescriptor.cs
@@ -123,7 +123,7 @@ namespace UnrealBuildTool
 					{
 						if (!String.IsNullOrWhiteSpace(File))
 						{
-							SpecificFilesToCompile.Add(new FileReference(File));
+							SpecificFilesToCompile.Add(FileReference.Combine(Unreal.RootDirectory, File));
 						}
 					}
 				}
