[CS] Add utility methods for writing word-wrapped text to the console.

diff --git a/Engine/Source/Programs/Shared/EpicGames.Core/ConsoleUtils.cs b/Engine/Source/Programs/Shared/EpicGames.Core/ConsoleUtils.cs
new file mode 100644
--- /dev/null
+++ b/Engine/Source/Programs/Shared/EpicGames.Core/ConsoleUtils.cs
@@ -0,0 +1,64 @@
+// Copyright Epic Games, Inc. All Rights Reserved.
+
+using System;
+using System.Collections.Generic;
+using System.Text;
+
+namespace EpicGames.Core
+{
+	/// <summary>
+	/// Utility methods for writing to the console
+	/// </summary>
+	public static class ConsoleUtils
+	{
+		/// <summary>
+		/// Gets the width of the window for formatting purposes
+		/// </summary>
+		public static int WindowWidth
+		{
+			get
+			{
+				// Get the window width, using a default value if there's no console attached to this process.
+				int NewWindowWidth;
+				try
+				{
+					NewWindowWidth = Console.WindowWidth;
+				}
+				catch
+				{
+					NewWindowWidth = 240;
+				}
+
+				if (NewWindowWidth <= 0)
+				{
+					NewWindowWidth = 240;
+				}
+
+				return NewWindowWidth;
+			}
+		}
+
+		/// <summary>
+		/// Writes the given text to the console as a sequence of word-wrapped lines
+		/// </summary>
+		/// <param name="Text">The text to write to the console</param>
+		public static void WriteLineWithWordWrap(string Text)
+		{
+			WriteLineWithWordWrap(Text, 0, 0);
+		}
+
+		/// <summary>
+		/// Writes the given text to the console as a sequence of word-wrapped lines
+		/// </summary>
+		/// <param name="Text">The text to write to the console</param>
+		/// <param name="InitialIndent">Indent for the first line</param>
+		/// <param name="HangingIndent">Indent for lines after the first</param>
+		public static void WriteLineWithWordWrap(string Text, int InitialIndent, int HangingIndent)
+		{
+			foreach (string Line in StringUtils.WordWrap(Text, InitialIndent, HangingIndent, WindowWidth))
+			{
+				Console.WriteLine(Line);
+			}
+		}
+	}
+}
diff --git a/Engine/Source/Programs/Shared/EpicGames.Core/HelpUtils.cs b/Engine/Source/Programs/Shared/EpicGames.Core/HelpUtils.cs
--- a/Engine/Source/Programs/Shared/EpicGames.Core/HelpUtils.cs
+++ b/Engine/Source/Programs/Shared/EpicGames.Core/HelpUtils.cs
@@ -18,29 +18,7 @@ namespace EpicGames.Core
 		/// <summary>
 		/// Gets the width of the window for formatting purposes
 		/// </summary>
-		public static int WindowWidth
-		{
-			get
-			{
-				// Get the window width, using a default value if there's no console attached to this process.
-				int NewWindowWidth;
-				try
-				{
-					NewWindowWidth = Console.WindowWidth;
-				}
-				catch
-				{
-					NewWindowWidth = 240;
-				}
-
-				if (NewWindowWidth <= 0)
-				{
-					NewWindowWidth = 240;
-				}
-
-				return NewWindowWidth;
-			}
-		}
+		public static int WindowWidth => ConsoleUtils.WindowWidth;
 
 		/// <summary>
 		/// Prints help for the given object type
@@ -127,7 +105,7 @@ namespace EpicGames.Core
 		/// <param name="Logger">Logger implementation to write to</param>
 		public static void PrintParagraph(string Text, int MaxWidth, ILogger Logger)
 		{
-			List<string> Lines = StringUtils.WordWrap(Text, MaxWidth);
+			IEnumerable<string> Lines = StringUtils.WordWrap(Text, MaxWidth);
 			foreach (string Line in Lines)
 			{
 				Logger.LogInformation(Line);
@@ -184,7 +162,7 @@ namespace EpicGames.Core
 					string ParamString = IndentString + Item.Key.PadRight(RightPadding);
 
 					// Build the description line by line, adding the same amount of intending each time. 
-					List<string> DescriptionLines = StringUtils.WordWrap(Item.Value, MaxWidth - ParamString.Length);
+					IEnumerable<string> DescriptionLines = StringUtils.WordWrap(Item.Value, MaxWidth - ParamString.Length);
 
 					foreach(string DescriptionLine in DescriptionLines)
 					{
diff --git a/Engine/Source/Programs/Shared/EpicGames.Core/StringUtils.cs b/Engine/Source/Programs/Shared/EpicGames.Core/StringUtils.cs
--- a/Engine/Source/Programs/Shared/EpicGames.Core/StringUtils.cs
+++ b/Engine/Source/Programs/Shared/EpicGames.Core/StringUtils.cs
@@ -117,73 +117,85 @@ namespace EpicGames.Core
 			return Result;
 		}
 
+		/// <inheritdoc cref="WordWrap(string, int, int, int)">
+		public static IEnumerable<string> WordWrap(string Text, int MaxWidth)
+		{
+			return WordWrap(Text, 0, 0, MaxWidth);
+		}
+
 		/// <summary>
-		/// Takes a given sentence and wraps it on a word by word basis so that no line exceeds the
-		/// set maximum line length. Words longer than a line are broken up. Returns the sentence as
-		/// a list of individual lines.
+		/// Takes a given sentence and wraps it on a word by word basis so that no line exceeds the set maximum line length. Words longer than a line 
+		/// are broken up. Returns the sentence as a list of individual lines.
 		/// </summary>
 		/// <param name="Text">The text to be wrapped</param>
+		/// <param name="InitialIndent">Indent for the first line</param>
+		/// <param name="HangingIndent">Indent for subsequent lines</param>
 		/// <param name="MaxWidth">The maximum (non negative) length of the returned sentences</param>
-		/// <param name="Lines">Receives a list of word-wrapped lines</param>
-		public static List<string> WordWrap(string Text, int MaxWidth)
+		public static IEnumerable<string> WordWrap(string Text, int InitialIndent, int HangingIndent, int MaxWidth)
 		{
-			// Early out
-			if (Text.Length == 0)
+			StringBuilder Builder = new StringBuilder();
+
+			int MinIdx = 0;
+			for (int LineIdx = 0; MinIdx < Text.Length; LineIdx++)
 			{
-				return new List<string>();
+				int Indent = (LineIdx == 0) ? InitialIndent : HangingIndent;
+				int MaxWidthForLine = MaxWidth - Indent;
+				int MaxIdx = GetWordWrapLineEnd(Text, MinIdx, MaxWidthForLine);
+
+				int PrintMaxIdx = MaxIdx;
+				while (PrintMaxIdx > MinIdx && Text[PrintMaxIdx - 1] == ' ')
+				{
+					PrintMaxIdx--;
+				}
+
+				Builder.Clear();
+				Builder.Append(' ', Indent);
+				Builder.Append(Text, MinIdx, PrintMaxIdx - MinIdx);
+				yield return Builder.ToString();
+
+				MinIdx = MaxIdx;
 			}
+		}
 
-			string[] Words = Text.Split(' ');
-			List<string> WrappedWords = new List<string>();
+		/// <summary>
+		/// Gets the next character index to end a word-wrapped line on
+		/// </summary>
+		static int GetWordWrapLineEnd(string Text, int MinIdx, int MaxWidth)
+		{
+			int MaxIdx = MinIdx + MaxWidth;
+			if (MaxIdx >= Text.Length)
+			{
+				return Text.Length;
+			}
 
-			string CurrentSentence = string.Empty;
-			foreach (var Word in Words)
+			if (Text[MaxIdx] == ' ')
 			{
-				// if this is a very large word, split it
-				if (Word.Length > MaxWidth)
+				for (; ; MaxIdx++)
 				{
-					// If the current sentence is ready to be written, do that.
-					if (CurrentSentence.Length >= MaxWidth)
+					if (MaxIdx == Text.Length)
 					{
-						// next line and reset sentence
-						WrappedWords.Add(CurrentSentence);
-						CurrentSentence = string.Empty;
+						return MaxIdx;
 					}
-
-					// Top up the current line
-					WrappedWords.Add(CurrentSentence + Word.Substring(0, MaxWidth - CurrentSentence.Length));
-
-					int length = MaxWidth - CurrentSentence.Length;
-					while (length + MaxWidth < Word.Length)
+					if (Text[MaxIdx] != ' ')
 					{
-						// Place the starting lengths into their own lines
-						WrappedWords.Add(Word.Substring(length, Math.Min(MaxWidth, Word.Length - length)));
-						length += MaxWidth;
+						return MaxIdx;
 					}
-
-					// then the trailing end into the next line
-					CurrentSentence += Word.Substring(length, Math.Min(MaxWidth, Word.Length - length)) + " ";
 				}
-				else
+			}
+			else
+			{
+				for(int TryMaxIdx = MaxIdx; ; TryMaxIdx--)
 				{
-					if (CurrentSentence.Length + Word.Length > MaxWidth)
+					if(TryMaxIdx == MinIdx)
 					{
-						// next line and reset sentence
-						WrappedWords.Add(CurrentSentence);
-						CurrentSentence = string.Empty;
+						return MaxIdx;
+					}
+					if (Text[TryMaxIdx - 1] == ' ')
+					{
+						return TryMaxIdx;
 					}
-
-					// Add the word to the current sentence.
-					CurrentSentence += Word + " ";
 				}
 			}
-
-			if (CurrentSentence.Length > 0)
-			{
-				WrappedWords.Add(CurrentSentence);
-			}
-
-			return WrappedWords;
 		}
 
 		/// <summary>
@@ -237,7 +249,6 @@ namespace EpicGames.Core
 			return FormatList(Arguments.ToArray(), Conjunction);
 		}
 
-
 		/// <summary>
 		/// Formats a list of items
 		/// </summary>
