CS: Add support for setting properties via the CommandLineArguments class.

diff --git a/Engine/Source/Programs/Shared/EpicGames.Core/CommandLineArguments.cs b/Engine/Source/Programs/Shared/EpicGames.Core/CommandLineArguments.cs
--- a/Engine/Source/Programs/Shared/EpicGames.Core/CommandLineArguments.cs
+++ b/Engine/Source/Programs/Shared/EpicGames.Core/CommandLineArguments.cs
@@ -79,6 +79,27 @@ namespace EpicGames.Core
 	[DebuggerTypeProxy(typeof(CommandLineArgumentListView))]
 	public class CommandLineArguments : IReadOnlyList<string>, IReadOnlyCollection<string>, IEnumerable<string>, IEnumerable
 	{
+		/// <summary>
+		/// Information about a property or field that can receive an argument
+		/// </summary>
+		class ArgumentTarget
+		{
+			public MemberInfo Member { get; }
+			public Type ValueType { get; }
+			public Action<object?, object?> SetValue { get; }
+			public Func<object?, object?> GetValue { get; }
+			public CommandLineAttribute[] Attributes { get; }
+
+			public ArgumentTarget(MemberInfo Member, Type ValueType, Action<object?, object?> SetValue, Func<object?, object?> GetValue, CommandLineAttribute[] Attributes)
+			{
+				this.Member = Member;
+				this.ValueType = ValueType;
+				this.SetValue = SetValue;
+				this.GetValue = GetValue;
+				this.Attributes = Attributes;
+			}
+		}
+
 		/// <summary>
 		/// The raw array of arguments
 		/// </summary>
@@ -661,13 +682,13 @@ namespace EpicGames.Core
 		/// <summary>
 		/// Gets the prefix for a particular argument
 		/// </summary>
-		/// <param name="FieldInfo">The field hosting the attribute</param>
+		/// <param name="Target">The target hosting the attribute</param>
 		/// <param name="Attribute">The attribute instance</param>
 		/// <returns>Prefix for this argument</returns>
-		private static string GetArgumentPrefix(FieldInfo FieldInfo, CommandLineAttribute Attribute)
+		private static string GetArgumentPrefix(ArgumentTarget Target, CommandLineAttribute Attribute)
 		{
 			// Get the inner field type, unwrapping nullable types
-			Type ValueType = FieldInfo.FieldType;
+			Type ValueType = Target.ValueType;
 			if (ValueType.IsGenericType && ValueType.GetGenericTypeDefinition() == typeof(Nullable<>))
 			{
 				ValueType = ValueType.GetGenericArguments()[0];
@@ -678,11 +699,11 @@ namespace EpicGames.Core
 			{
 				if (ValueType == typeof(bool))
 				{
-					Prefix = String.Format("-{0}", FieldInfo.Name);
+					Prefix = String.Format("-{0}", Target.Member.Name);
 				}
 				else
 				{
-					Prefix = String.Format("-{0}=", FieldInfo.Name);
+					Prefix = String.Format("-{0}=", Target.Member.Name);
 				}
 			}
 			else
@@ -705,6 +726,32 @@ namespace EpicGames.Core
 			ApplyTo(TargetObject, Log.Logger);
 		}
 
+		static List<ArgumentTarget> GetArgumentTargetsForType(Type? TargetType)
+		{
+			List<ArgumentTarget> Targets = new List<ArgumentTarget>();
+			while (TargetType != null && TargetType != typeof(object))
+			{
+				foreach (FieldInfo FieldInfo in TargetType.GetFields(BindingFlags.Instance | BindingFlags.GetField | BindingFlags.Public | BindingFlags.NonPublic | BindingFlags.DeclaredOnly))
+				{
+					IEnumerable<CommandLineAttribute> Attributes = FieldInfo.GetCustomAttributes<CommandLineAttribute>();
+					if (Attributes.Any())
+					{
+						Targets.Add(new ArgumentTarget(FieldInfo, FieldInfo.FieldType, FieldInfo.SetValue, FieldInfo.GetValue, Attributes.ToArray()));
+					}
+				}
+				foreach (PropertyInfo PropertyInfo in TargetType.GetProperties(BindingFlags.Instance | BindingFlags.GetProperty | BindingFlags.Public | BindingFlags.NonPublic | BindingFlags.DeclaredOnly))
+				{
+					IEnumerable<CommandLineAttribute> Attributes = PropertyInfo.GetCustomAttributes<CommandLineAttribute>();
+					if (Attributes.Any())
+					{
+						Targets.Add(new ArgumentTarget(PropertyInfo, PropertyInfo.PropertyType, PropertyInfo.SetValue, PropertyInfo.GetValue, Attributes.ToArray()));
+					}
+				}
+				TargetType = TargetType.BaseType;
+			}
+			return Targets;
+		}
+
 		/// <summary>
 		/// Applies these arguments to fields with the [CommandLine] attribute in the given object.
 		/// </summary>
@@ -712,86 +759,84 @@ namespace EpicGames.Core
 		/// <param name="Logger">Sink for error/warning messages</param>
 		public void ApplyTo(object TargetObject, ILogger Logger)
 		{
-			// Build a mapping from name to field and attribute for this object
 			List<string> MissingArguments = new List<string>();
-			for(Type TargetType = TargetObject.GetType(); TargetType != typeof(object); TargetType = TargetType.BaseType!)
+
+			// Build a mapping from name to field and attribute for this object
+			List<ArgumentTarget> Targets = GetArgumentTargetsForType(TargetObject.GetType());
+			foreach (ArgumentTarget Target in Targets)
 			{
-				foreach(FieldInfo FieldInfo in TargetType.GetFields(BindingFlags.Instance | BindingFlags.GetField | BindingFlags.GetProperty | BindingFlags.Public | BindingFlags.NonPublic | BindingFlags.DeclaredOnly))
-				{
-					// If any attribute is required, keep track of it so we can include an error for it
-					string? RequiredPrefix = null;
+				// If any attribute is required, keep track of it so we can include an error for it
+				string? RequiredPrefix = null;
 
-					// Keep track of whether a value has already been assigned to this field
-					string? AssignedArgument = null;
+				// Keep track of whether a value has already been assigned to this field
+				string? AssignedArgument = null;
 
-					// Loop through all the attributes for different command line options that can modify it
-					IEnumerable<CommandLineAttribute> Attributes = FieldInfo.GetCustomAttributes<CommandLineAttribute>();
-					foreach(CommandLineAttribute Attribute in Attributes)
-					{
-						// Get the appropriate prefix for this attribute
-						string Prefix = GetArgumentPrefix(FieldInfo, Attribute);
+				// Loop through all the attributes for different command line options that can modify it
+				foreach(CommandLineAttribute Attribute in Target.Attributes)
+				{
+					// Get the appropriate prefix for this attribute
+					string Prefix = GetArgumentPrefix(Target, Attribute);
 
-						// Get the value with the correct prefix
-						int FirstIndex;
-						if(ArgumentToFirstIndex.TryGetValue(Prefix, out FirstIndex))
+					// Get the value with the correct prefix
+					int FirstIndex;
+					if(ArgumentToFirstIndex.TryGetValue(Prefix, out FirstIndex))
+					{
+						for(int Index = FirstIndex; Index != -1; Index = NextArgumentIndex[Index])
 						{
-							for(int Index = FirstIndex; Index != -1; Index = NextArgumentIndex[Index])
+							// Get the argument text
+							string Argument = Arguments[Index];
+
+							// Get the text for this value
+							string ValueText;
+							if(Attribute.Value != null)
 							{
-								// Get the argument text
-								string Argument = Arguments[Index];
+								ValueText = Attribute.Value;
+							}
+							else if(FlagArguments.Get(Index))
+							{
+								ValueText = "true";
+							}
+							else
+							{
+								ValueText = Argument.Substring(Prefix.Length);
+							}
 
-								// Get the text for this value
-								string ValueText;
-								if(Attribute.Value != null)
-								{
-									ValueText = Attribute.Value;
-								}
-								else if(FlagArguments.Get(Index))
-								{
-									ValueText = "true";
-								}
-								else
+							// Apply the value to the field
+							if(Attribute.ListSeparator == 0)
+							{
+								if(ApplyArgument(TargetObject, Target, Argument, ValueText, AssignedArgument, Logger))
 								{
-									ValueText = Argument.Substring(Prefix.Length);
+									AssignedArgument = Argument;
 								}
-
-								// Apply the value to the field
-								if(Attribute.ListSeparator == 0)
+							}
+							else
+							{
+								foreach(string ItemValueText in ValueText.Split(Attribute.ListSeparator))
 								{
-									if(ApplyArgument(TargetObject, FieldInfo, Argument, ValueText, AssignedArgument, Logger))
+									if(ApplyArgument(TargetObject, Target, Argument, ItemValueText, AssignedArgument, Logger))
 									{
 										AssignedArgument = Argument;
 									}
 								}
-								else
-								{
-									foreach(string ItemValueText in ValueText.Split(Attribute.ListSeparator))
-									{
-										if(ApplyArgument(TargetObject, FieldInfo, Argument, ItemValueText, AssignedArgument, Logger))
-										{
-											AssignedArgument = Argument;
-										}
-									}
-								}
-
-								// Mark this argument as used
-								UsedArguments.Set(Index, true);
 							}
-						}
 
-						// If this attribute is marked as required, keep track of it so we can warn if the field is not assigned to
-						if(Attribute.Required && RequiredPrefix == null)
-						{
-							RequiredPrefix = Prefix;
+							// Mark this argument as used
+							UsedArguments.Set(Index, true);
 						}
 					}
 
-					// Make sure that this field has been assigned to
-					if(AssignedArgument == null && RequiredPrefix != null)
+					// If this attribute is marked as required, keep track of it so we can warn if the field is not assigned to
+					if(Attribute.Required && RequiredPrefix == null)
 					{
-						MissingArguments.Add(RequiredPrefix);
+						RequiredPrefix = Prefix;
 					}
 				}
+
+				// Make sure that this field has been assigned to
+				if(AssignedArgument == null && RequiredPrefix != null)
+				{
+					MissingArguments.Add(RequiredPrefix);
+				}
 			}
 
 			// If any arguments were missing, print an error about them
@@ -816,39 +861,34 @@ namespace EpicGames.Core
 		public static List<KeyValuePair<string, string>> GetParameters(Type Type)
 		{
 			List<KeyValuePair<string, string>> Parameters = new List<KeyValuePair<string, string>>();
-			for (Type? TargetType = Type; TargetType != typeof(object) && TargetType != null; TargetType = TargetType.BaseType)
+
+			List<ArgumentTarget> Targets = GetArgumentTargetsForType(Type);
+			foreach (ArgumentTarget Target in Targets)
 			{
-				foreach (FieldInfo FieldInfo in TargetType.GetFields(BindingFlags.Instance | BindingFlags.GetField | BindingFlags.GetProperty | BindingFlags.Public | BindingFlags.NonPublic | BindingFlags.DeclaredOnly))
+				StringBuilder DescriptionBuilder = new StringBuilder();
+				foreach (DescriptionAttribute Attribute in Target.Member.GetCustomAttributes<DescriptionAttribute>())
 				{
-					List<CommandLineAttribute> Attributes = FieldInfo.GetCustomAttributes<CommandLineAttribute>().ToList();
-					if(Attributes.Count > 0)
+					if(DescriptionBuilder.Length > 0)
 					{
-						StringBuilder DescriptionBuilder = new StringBuilder();
-						foreach (DescriptionAttribute Attribute in FieldInfo.GetCustomAttributes<DescriptionAttribute>())
-						{
-							if(DescriptionBuilder.Length > 0)
-							{
-								DescriptionBuilder.Append("\n");
-							}
-							DescriptionBuilder.Append(Attribute.Description);
-						}
+						DescriptionBuilder.Append("\n");
+					}
+					DescriptionBuilder.Append(Attribute.Description);
+				}
 
-						string Description = DescriptionBuilder.ToString();
-						if (Description.Length == 0)
-						{
-							Description = "No description available.";
-						}
+				string Description = DescriptionBuilder.ToString();
+				if (Description.Length == 0)
+				{
+					Description = "No description available.";
+				}
 
-						foreach (CommandLineAttribute Attribute in Attributes)
-						{
-							string Prefix = GetArgumentPrefix(FieldInfo, Attribute);
-							if(Prefix.EndsWith("=", StringComparison.Ordinal))
-							{
-								Prefix += "...";
-							}
-							Parameters.Add(new KeyValuePair<string, string>(Prefix, Description));
-						}
+				foreach (CommandLineAttribute Attribute in Target.Attributes)
+				{
+					string Prefix = GetArgumentPrefix(Target, Attribute);
+					if(Prefix.EndsWith("=", StringComparison.Ordinal))
+					{
+						Prefix += "...";
 					}
+					Parameters.Add(new KeyValuePair<string, string>(Prefix, Description));
 				}
 			}
 			return Parameters;
@@ -1053,7 +1093,22 @@ namespace EpicGames.Core
 		{
 			return UsedArguments.Cast<bool>().Count(B => B) != 0;
 		}
-		
+
+		/// <summary>
+		/// Checks to see if any arguments are used
+		/// </summary>
+		/// <returns></returns>
+		public IEnumerable<string> GetUnusedArguments()
+		{
+			for(int Idx = 0; Idx < Arguments.Length; Idx++)
+			{
+				if (!UsedArguments[Idx])
+				{
+					yield return Arguments[Idx];
+				}
+			}
+		}
+
 		/// <summary>
 		/// Count the number of value (non-flag) arguments on the command line
 		/// </summary>
@@ -1087,20 +1142,19 @@ namespace EpicGames.Core
 		/// Parses and assigns a value to a field
 		/// </summary>
 		/// <param name="TargetObject">The target object to assign values to</param>
-		/// <param name="Field">The field to assign the value to</param>
+		/// <param name="Target">The target to assign the value to</param>
 		/// <param name="ArgumentText">The full argument text</param>
 		/// <param name="ValueText">Argument text</param>
 		/// <param name="PreviousArgumentText">The previous text used to configure this field</param>
 		/// <param name="Logger">Logger for error/warning messages</param>
 		/// <returns>True if the value was assigned to the field, false otherwise</returns>
-		private static bool ApplyArgument(object TargetObject, FieldInfo Field, string ArgumentText, string ValueText, string? PreviousArgumentText, ILogger Logger)
+		private static bool ApplyArgument(object TargetObject, ArgumentTarget Target, string ArgumentText, string ValueText, string? PreviousArgumentText, ILogger Logger)
 		{
-			// The value type for items of this field
-			Type ValueType = Field.FieldType;
+			Type ValueType = Target.ValueType;
 
 			// Check if the field type implements ICollection<>. If so, we can take multiple values.
 			Type? CollectionType = null;
-			foreach (Type InterfaceType in Field.FieldType.GetInterfaces())
+			foreach (Type InterfaceType in ValueType.GetInterfaces())
 			{
 				if (InterfaceType.IsGenericType && InterfaceType.GetGenericTypeDefinition() == typeof(ICollection<>))
 				{
@@ -1124,7 +1178,7 @@ namespace EpicGames.Core
 				// Check if this field has already been assigned to. Output a warning if the previous value is in conflict with the new one.
 				if(PreviousArgumentText != null)
 				{
-					object? PreviousValue = Field.GetValue(TargetObject);
+					object? PreviousValue = Target.GetValue(TargetObject);
 					if(!Object.Equals(PreviousValue, Value))
 					{
 						Logger.LogWarning("Argument '{0}' conflicts with '{1}'; ignoring.", ArgumentText, PreviousArgumentText);
@@ -1133,13 +1187,13 @@ namespace EpicGames.Core
 				}
 
 				// Set the value on the target object
-				Field.SetValue(TargetObject, Value);
+				Target.SetValue(TargetObject, Value);
 				return true;
 			}
 			else
 			{
 				// Call the 'Add' method on the collection
-				CollectionType.InvokeMember("Add", BindingFlags.InvokeMethod, null, Field.GetValue(TargetObject), new object[] { Value });
+				CollectionType.InvokeMember("Add", BindingFlags.InvokeMethod, null, Target.GetValue(TargetObject), new object[] { Value });
 				return true;
 			}
 		}
diff --git a/Engine/Source/Programs/Shared/EpicGames.Core/CommandLineAttribute.cs b/Engine/Source/Programs/Shared/EpicGames.Core/CommandLineAttribute.cs
--- a/Engine/Source/Programs/Shared/EpicGames.Core/CommandLineAttribute.cs
+++ b/Engine/Source/Programs/Shared/EpicGames.Core/CommandLineAttribute.cs
@@ -12,7 +12,7 @@ namespace EpicGames.Core
 	/// <summary>
 	/// Attribute to indicate the name of a command line argument
 	/// </summary>
-	[AttributeUsage(AttributeTargets.Field, AllowMultiple = true)]
+	[AttributeUsage(AttributeTargets.Field | AttributeTargets.Property, AllowMultiple = true)]
 	public class CommandLineAttribute : Attribute
 	{
 		/// <summary>
