UE5.0 - Only check for platform SDKs once in UBT when compiling multiple targets

diff --git a/Engine/Source/Programs/Shared/EpicGames.Build/UEBuildPlatformSDK.cs b/Engine/Source/Programs/Shared/EpicGames.Build/UEBuildPlatformSDK.cs
--- a/Engine/Source/Programs/Shared/EpicGames.Build/UEBuildPlatformSDK.cs
+++ b/Engine/Source/Programs/Shared/EpicGames.Build/UEBuildPlatformSDK.cs
@@ -352,9 +352,16 @@ namespace EpicGames.Core
 
 		#region Print SDK Info
 		private static bool bHasShownTurnkey = false;
+		private SDKStatus? SDKInfoValidity = null;
+
 		public virtual SDKStatus PrintSDKInfoAndReturnValidity(LogEventType Verbosity = LogEventType.Console, LogFormatOptions Options = LogFormatOptions.None,
 			LogEventType ErrorVerbosity = LogEventType.Error, LogFormatOptions ErrorOptions = LogFormatOptions.None)
 		{
+			if (SDKInfoValidity != null)
+			{
+				return SDKInfoValidity.Value;
+			}
+
 			string? ManualSDKVersion, AutoSDKVersion;
 			GetInstalledVersions(out ManualSDKVersion, out AutoSDKVersion);
 
@@ -423,6 +430,7 @@ namespace EpicGames.Core
 				}
 			}
 
+			SDKInfoValidity = Validity;
 			return Validity;
 		}
 
