CS: Changes to ManagedProcess to better support cancellation on Mac/Linux when using merged output pipes.

- Data read from stdout/stderr is now writen to buffers posted in a Channel<T>, which correctly supports await and cancellation (unlike AsynchronousPipe).
- Public StdOut/StdErr streams are now a custom stream implementation that reads from the channel and copies data into the output buffer.
- Buffer objects are sized at 1024 bytes and pooled, which should be enough to accomodate most output lines without splitting.
- When the process is disposed, we now terminate the entire process tree. The tasks left reading from the child process terminate asynchronously.

Tested in a command line and WinForms app (by forcing SupportsJobObjects to return false) with asynchronous and synchronous reads to ensure code does not deadlock.

diff --git a/Engine/Source/Programs/Shared/EpicGames.Core/EpicGames.Core.csproj b/Engine/Source/Programs/Shared/EpicGames.Core/EpicGames.Core.csproj
--- a/Engine/Source/Programs/Shared/EpicGames.Core/EpicGames.Core.csproj
+++ b/Engine/Source/Programs/Shared/EpicGames.Core/EpicGames.Core.csproj
@@ -44,6 +44,7 @@
     <PackageReference Include="JetBrains.Annotations" Version="2021.2.0" />
     <PackageReference Include="Microsoft.Extensions.Logging" Version="5.0.0" />
     <PackageReference Include="Microsoft.Extensions.Logging.Abstractions" Version="5.0.0" />
+    <PackageReference Include="Microsoft.Extensions.ObjectPool" Version="6.0.1" />
     <PackageReference Include="OpenTracing" Version="0.12.1" />
     <PackageReference Include="System.Memory" Version="4.5.4" />
   </ItemGroup>
diff --git a/Engine/Source/Programs/Shared/EpicGames.Core/ManagedProcess.cs b/Engine/Source/Programs/Shared/EpicGames.Core/ManagedProcess.cs
--- a/Engine/Source/Programs/Shared/EpicGames.Core/ManagedProcess.cs
+++ b/Engine/Source/Programs/Shared/EpicGames.Core/ManagedProcess.cs
@@ -1,5 +1,6 @@
 // Copyright Epic Games, Inc. All Rights Reserved.
 
+using Microsoft.Extensions.ObjectPool;
 using Microsoft.Win32.SafeHandles;
 using System;
 using System.Collections.Generic;
@@ -11,6 +12,7 @@ using System.IO.Pipes;
 using System.Runtime.InteropServices;
 using System.Text;
 using System.Threading;
+using System.Threading.Channels;
 using System.Threading.Tasks;
 
 namespace EpicGames.Core
@@ -406,24 +408,24 @@ namespace EpicGames.Core
 		Process? FrameworkProcess;
 
 		/// <summary>
-		/// Thread to read from stdout
+		/// Task to read from stdout and stderr
 		/// </summary>
-		Thread? FrameworkStdOutThread;
+		Task? FrameworkOutputTask;
 
 		/// <summary>
-		/// Thread to read from stderr
+		/// Merged output & error write stream for the framework child process. We use a channel rather than a pipe here to support cancellation, which is not supported by async pipes on Windows.
 		/// </summary>
-		Thread? FrameworkStdErrThread;
+		Channel<ChannelBuffer>? FrameworkChannel;
 
 		/// <summary>
-		/// Merged output & error write stream for the framework child process.
+		/// Pool of buffer objects output from the framework
 		/// </summary>
-		AnonymousPipeServerStream? FrameworkMergedStdWriter;
+		ObjectPool<ChannelBuffer>? FrameworkBufferPool;
 
 		/// <summary>
-		/// Tracks how many pipes have been copied to the merged writer.
+		/// Cancellation token for background threads
 		/// </summary>
-		int FrameworkMergedStdWriterThreadCount;
+		CancellationTokenSource? FrameworkCancellationSource;
 
 		/// <summary>
 		/// Static lock object. This is used to synchronize the creation of child processes - in particular, the inheritance of stdout/stderr write pipes. If processes
@@ -441,6 +443,70 @@ namespace EpicGames.Core
 		int BufferLen = 0;
 		bool bCarriageReturn = false;
 
+		class ChannelBuffer
+		{
+			public byte[] Data = new byte[1024];
+			public int Length = 0;
+		}
+
+		class ChannelReadStream : Stream
+		{
+			ChannelReader<ChannelBuffer> Reader;
+			ObjectPool<ChannelBuffer> BufferPool;
+
+			ChannelBuffer? CurrentBuf = null;
+			int CurrentPos = 0;
+
+			public ChannelReadStream(ChannelReader<ChannelBuffer> Reader, ObjectPool<ChannelBuffer> BufferPool)
+			{
+				this.Reader = Reader;
+				this.BufferPool = BufferPool;
+			}
+
+			public override bool CanRead => true;
+			public override bool CanSeek => false;
+			public override bool CanWrite => false;
+			public override long Length => throw new NotSupportedException();
+			public override long Position { get => throw new NotSupportedException(); set => throw new NotSupportedException(); }
+			public override void Flush() { }
+			public override long Seek(long offset, SeekOrigin origin) => throw new NotSupportedException();
+			public override void SetLength(long value) => throw new NotSupportedException();
+			public override void Write(byte[] buffer, int offset, int count) => throw new NotSupportedException();
+
+			public override int Read(byte[] Buffer, int Offset, int Count)
+			{
+				// Perform the read on the thread pool to avoid deadlocking the continuation pathway on WinForms apps
+				Task<int> ReadTask = Task.Run(async () => await ReadAsync(Buffer.AsMemory(Offset, Count), CancellationToken.None));
+				return ReadTask.GetAwaiter().GetResult();
+			}
+
+			public override async ValueTask<int> ReadAsync(Memory<byte> Buffer, CancellationToken CancellationToken)
+			{
+				if (CurrentBuf != null && CurrentPos == CurrentBuf.Length)
+				{
+					BufferPool.Return(CurrentBuf);
+					CurrentBuf = null;
+					CurrentPos = 0;
+				}
+
+				while (CurrentBuf == null)
+				{
+					if (!await Reader.WaitToReadAsync(CancellationToken))
+					{
+						return 0;
+					}
+					if (Reader.TryRead(out CurrentBuf))
+					{
+						break;
+					}
+				}
+
+				int ReadLen = Math.Min(CurrentBuf.Length - CurrentPos, Buffer.Length);
+				CurrentBuf.Data.AsSpan(CurrentPos, ReadLen).CopyTo(Buffer.Span);
+				return ReadLen;
+			}
+		}
+
 		/// <summary>
 		/// Spawns a new managed process.
 		/// </summary>
@@ -787,45 +853,8 @@ namespace EpicGames.Core
 				}
 			}
 
-			// Merge StdOut with StdErr
-			if ((ManagedFlags & ManagedProcessFlags.MergeOutputPipes) != 0)
-			{
-				// AnonymousPipes block reading even if the stream has been fully read, until the writer pipe handle is closed.
-				FrameworkMergedStdWriter = new AnonymousPipeServerStream(PipeDirection.Out);
-				StdOut = new AnonymousPipeClientStream(PipeDirection.In, FrameworkMergedStdWriter.ClientSafePipeHandle);
-				StdOutText = new StreamReader(StdOut, Console.OutputEncoding);
-
-				StdErr = StdOut;
-				StdErrText = StdOutText;
-				FrameworkMergedStdWriterThreadCount = 2;
-
-				FrameworkStdOutThread = new Thread(() => {
-					CopyPipe(FrameworkProcess.StandardOutput.BaseStream, FrameworkMergedStdWriter!);
-					if (Interlocked.Decrement(ref FrameworkMergedStdWriterThreadCount) == 0)
-					{
-						// Dispose AnonymousPipe to unblock readers.
-						FrameworkMergedStdWriter?.Dispose();
-					}
-				});
-				FrameworkStdOutThread.Name = $"ManagedProcess Merge StdOut";
-				FrameworkStdOutThread.IsBackground = true;
-
-				FrameworkStdErrThread = new Thread(() => {
-					CopyPipe(FrameworkProcess.StandardError.BaseStream, FrameworkMergedStdWriter!);
-					if (Interlocked.Decrement(ref FrameworkMergedStdWriterThreadCount) == 0)
-					{
-						// Dispose AnonymousPipe to unblock readers.
-						FrameworkMergedStdWriter?.Dispose();
-					}
-				});
-				FrameworkStdErrThread.Name = $"ManagedProcess Merge StdErr";
-				FrameworkStdErrThread.IsBackground = true;
-			}
-
 			FrameworkStartTime = DateTime.Now;
 			FrameworkProcess.Start();
-			FrameworkStdOutThread?.Start();
-			FrameworkStdErrThread?.Start();
 
 			try
 			{
@@ -845,7 +874,26 @@ namespace EpicGames.Core
 
 			Id = FrameworkProcess.Id;
 			StdIn = FrameworkProcess.StandardInput.BaseStream;
-			if ((ManagedFlags & ManagedProcessFlags.MergeOutputPipes) == 0)
+
+			if ((ManagedFlags & ManagedProcessFlags.MergeOutputPipes) != 0)
+			{
+				// AnonymousPipes block reading even if the stream has been fully read, until the writer pipe handle is closed.
+				FrameworkBufferPool = ObjectPool.Create<ChannelBuffer>();
+				FrameworkChannel = Channel.CreateBounded<ChannelBuffer>(new BoundedChannelOptions(128) { FullMode = BoundedChannelFullMode.Wait, SingleReader = true, SingleWriter = false });
+				FrameworkCancellationSource = new CancellationTokenSource();
+
+				StdOut = new ChannelReadStream(FrameworkChannel.Reader, FrameworkBufferPool);
+				StdOutText = new StreamReader(StdOut, Console.OutputEncoding);
+
+				StdErr = StdOut;
+				StdErrText = StdOutText;
+
+				CancellationToken CancellationToken = FrameworkCancellationSource.Token;
+				Task StdOutTask = Task.Run(() => CopyPipeAsync(FrameworkProcess.StandardOutput.BaseStream, FrameworkChannel, CancellationToken));
+				Task StdErrTask = Task.Run(() => CopyPipeAsync(FrameworkProcess.StandardError.BaseStream, FrameworkChannel, CancellationToken));
+				FrameworkOutputTask = Task.WhenAll(StdOutTask, StdErrTask).ContinueWith(x => FrameworkChannel.Writer.Complete());
+			}
+			else
 			{
 				StdOut = FrameworkProcess.StandardOutput.BaseStream;
 				StdOutText = FrameworkProcess.StandardOutput;
@@ -859,33 +907,20 @@ namespace EpicGames.Core
 		/// </summary>
 		/// <param name="Source"></param>
 		/// <param name="Target"></param>
-		void CopyPipe(Stream Source, Stream Target)
+		async Task CopyPipeAsync(Stream Source, ChannelWriter<ChannelBuffer> Target, CancellationToken CancellationToken)
 		{
-			try
-			{
-				Source.CopyTo(Target);
-				Target.Flush();
-			}
-			catch
+			for (; ; )
 			{
-			}
-		}
+				ChannelBuffer Buffer = FrameworkBufferPool!.Get();
+				int ReadLen = await Source.ReadAsync(Buffer.Data, 0, Buffer.Data.Length, CancellationToken);
 
-		/// <summary>
-		/// Copy data from one pipe to another.
-		/// </summary>
-		/// <param name="Source"></param>
-		/// <param name="Target"></param>
-		/// <param name="CancellationToken"></param>
-		async Task CopyPipeAsync(Stream Source, Stream Target, CancellationToken CancellationToken)
-		{
-			try
-			{
-				await Source.CopyToAsync(Target, CancellationToken);
-				await Target.FlushAsync(CancellationToken);
-			}
-			catch
-			{
+				if (ReadLen == 0)
+				{
+					break;
+				}
+
+				Buffer.Length = ReadLen;
+				await Target.WriteAsync(Buffer, CancellationToken);
 			}
 		}
 
@@ -925,28 +960,18 @@ namespace EpicGames.Core
 			}
 			if(FrameworkProcess != null)
 			{
+				FrameworkCancellationSource!.Cancel();
+
 				if (!FrameworkProcess.HasExited)
 				{
-					try
-					{
-						FrameworkProcess.Kill();
-						FrameworkProcess.WaitForExit();
-					}
-					catch
-					{
-					}
+					try { FrameworkProcess.Kill(true); } catch { } // Kill entire process tree
+					try { FrameworkProcess.WaitForExit(); } catch { }
 				}
-				
-				FrameworkMergedStdWriter?.Dispose();
-				FrameworkMergedStdWriter = null;
-
-				FrameworkStdOutThread?.Join();
-				FrameworkStdOutThread = null;
-				FrameworkStdErrThread?.Join();
-				FrameworkStdErrThread = null;
 
 				FrameworkProcess.Dispose();
 				FrameworkProcess = null;
+
+				FrameworkOutputTask!.ContinueWith(x => FrameworkCancellationSource.Dispose());
 			}
 		}
 
