UGS-CLI: Various improvements/fixes.

* Support for creating workspaces using native P4 client, so "ugs init" now works correctly.
* Output a list of files which need to be clobbered when syncing, and add a -Clobber option to confirm intent.
* Fix exception accessing folders that user does not have permissions to when detecting a new workspace.

diff --git a/Engine/Source/Programs/Shared/EpicGames.Perforce.Native/Client.cpp b/Engine/Source/Programs/Shared/EpicGames.Perforce.Native/Client.cpp
--- a/Engine/Source/Programs/Shared/EpicGames.Perforce.Native/Client.cpp
+++ b/Engine/Source/Programs/Shared/EpicGames.Perforce.Native/Client.cpp
@@ -49,6 +49,9 @@ private:
 	unsigned int MaxLength;
 	FOnBufferReadyFn* OnBufferReady;
 
+	const char* InputBuffer = nullptr;
+	int InputLength = 0;
+
 public:
 	const char* Func = nullptr;
 	const char* PromptResponse = nullptr;
@@ -59,6 +62,18 @@ public:
 		this->OnBufferReady = OnBufferReady;
 	}
 
+	virtual void InputData(StrBuf* strbuf, Error* e) override
+	{
+		if (InputBuffer != nullptr)
+		{
+			strbuf->Set(InputBuffer, InputLength);
+		}
+		else
+		{
+			ClientUser::InputData(strbuf, e);
+		}
+	}
+
 	virtual void Prompt(Error* err, StrBuf& rsp, int noEcho, Error* e) override
 	{
 		rsp.Set(PromptResponse);
@@ -79,6 +94,12 @@ public:
 		rsp.Set(PromptResponse);
 	}
 
+	void SetInputBuffer(const char* InputBuffer, int InputLength)
+	{
+		this->InputBuffer = InputBuffer;
+		this->InputLength = InputLength;
+	}
+
 	void SetWriteBuffer(FWriteBuffer* WriteBuffer)
 	{
 		this->Data = WriteBuffer->Data;
@@ -430,15 +451,17 @@ extern "C" NATIVE_API FClient* Client_Create(const FSettings* Settings, FWriteBu
 extern "C" NATIVE_API void Client_Login(FClient * Client, const char* Password)
 {
 	Client->User.PromptResponse = Password;
+	Client->User.SetInputBuffer(nullptr, 0);
 	Client->ClientApi.SetArgv(0, nullptr);
 	Client->ClientApi.Run("login", &Client->User);
 	Client->User.Flush();
 	Client->User.PromptResponse = nullptr;
 }
 
-extern "C" NATIVE_API void Client_Command(FClient* Client, const char* Func, int ArgCount, const char** Args)
+extern "C" NATIVE_API void Client_Command(FClient* Client, const char* Func, int ArgCount, const char** Args, const char* InputData, int InputLength)
 {
 	Client->User.Func = Func;
+	Client->User.SetInputBuffer(InputData, InputLength);
 	Client->ClientApi.SetArgv(ArgCount, (char* const*)Args);
 	Client->ClientApi.Run(Func, &Client->User);
 	Client->User.Flush();
diff --git a/Engine/Source/Programs/Shared/EpicGames.Perforce/IPerforceOutput.cs b/Engine/Source/Programs/Shared/EpicGames.Perforce/IPerforceOutput.cs
--- a/Engine/Source/Programs/Shared/EpicGames.Perforce/IPerforceOutput.cs
+++ b/Engine/Source/Programs/Shared/EpicGames.Perforce/IPerforceOutput.cs
@@ -301,7 +301,7 @@ namespace EpicGames.Perforce
 				for (; ; )
 				{
 					int InitialBufferPos = BufferPos;
-					if (!ReadRecord(Data, ref BufferPos, Record.Rows))
+					if (!ParseRecord(Data, ref BufferPos, Record.Rows))
 					{
 						BufferPos = InitialBufferPos;
 						break;
@@ -325,7 +325,8 @@ namespace EpicGames.Perforce
 		/// <param name="BufferPos">Current read position within the buffer</param>
 		/// <param name="Rows">List of rows to read into</param>
 		/// <returns>True if a record could be read; false if more data is required</returns>
-		static bool ReadRecord(ReadOnlyMemory<byte> Buffer, ref int BufferPos, List<KeyValuePair<Utf8String, PerforceValue>> Rows)
+		[SuppressMessage("Design", "CA1045:Do not pass types by reference", Justification = "<Pending>")]
+		public static bool ParseRecord(ReadOnlyMemory<byte> Buffer, ref int BufferPos, List<KeyValuePair<Utf8String, PerforceValue>> Rows)
 		{
 			Rows.Clear();
 			ReadOnlySpan<byte> BufferSpan = Buffer.Span;
diff --git a/Engine/Source/Programs/Shared/EpicGames.Perforce/NativePerforceConnection.cs b/Engine/Source/Programs/Shared/EpicGames.Perforce/NativePerforceConnection.cs
--- a/Engine/Source/Programs/Shared/EpicGames.Perforce/NativePerforceConnection.cs
+++ b/Engine/Source/Programs/Shared/EpicGames.Perforce/NativePerforceConnection.cs
@@ -1,5 +1,6 @@
 // Copyright Epic Games, Inc. All Rights Reserved.
 
+using EpicGames.Core;
 using EpicGames.Perforce;
 using Microsoft.Extensions.Logging;
 using System;
@@ -9,6 +10,7 @@ using System.Diagnostics;
 using System.Diagnostics.CodeAnalysis;
 using System.Linq;
 using System.Runtime.InteropServices;
+using System.Text;
 using System.Threading;
 using System.Threading.Channels;
 using System.Threading.Tasks;
@@ -73,7 +75,7 @@ namespace EpicGames.Perforce
 		static extern void Client_Login(IntPtr Client, [MarshalAs(UnmanagedType.LPStr)] string Password);
 
 		[DllImport(NativeDll, CallingConvention = CallingConvention.Cdecl, CharSet = CharSet.Ansi, BestFitMapping = false, ThrowOnUnmappableChar = true)]
-		static extern void Client_Command(IntPtr Client, [MarshalAs(UnmanagedType.LPStr)] string Command, int NumArgs, string[] Args);
+		static extern void Client_Command(IntPtr Client, [MarshalAs(UnmanagedType.LPStr)] string Command, int NumArgs, string[] Args, byte[]? InputData, int InputLength);
 
 		[DllImport(NativeDll, CallingConvention = CallingConvention.Cdecl)]
 		static extern void Client_Destroy(IntPtr Client);
@@ -485,9 +487,10 @@ namespace EpicGames.Perforce
 		/// <inheritdoc/>
 		public Task<IPerforceOutput> CommandAsync(string Command, IReadOnlyList<string> Arguments, IReadOnlyList<string>? FileArguments, byte[]? InputData)
 		{
+			byte[]? SpecData = null;
 			if (InputData != null)
 			{
-				throw new NotImplementedException();
+				SpecData = Encoding.UTF8.GetBytes(FormatSpec(InputData));
 			}
 
 			List<string> AllArguments = new List<string>(Arguments);
@@ -497,7 +500,7 @@ namespace EpicGames.Perforce
 			}
 
 			Response Response = new Response(this);
-			Requests.Add((() => Client_Command(Client, Command, AllArguments.Count, AllArguments.ToArray()), Response));
+			Requests.Add((() => Client_Command(Client, Command, AllArguments.Count, AllArguments.ToArray(), SpecData, SpecData?.Length ?? 0), Response));
 			return Task.FromResult<IPerforceOutput>(Response);
 		}
 
@@ -508,5 +511,44 @@ namespace EpicGames.Perforce
 			Requests.Add((() => Client_Login(Client, Password), Response), CancellationToken);
 			return Task.FromResult<IPerforceOutput>(Response);
 		}
+
+		/// <summary>
+		/// Converts a Python marshalled data blob to a spec definition
+		/// </summary>
+		static string FormatSpec(byte[] InputData)
+		{
+			int Pos = 0;
+
+			List<KeyValuePair<Utf8String, PerforceValue>> Rows = new List<KeyValuePair<Utf8String, PerforceValue>>();
+			if (!PerforceOutputExtensions.ParseRecord(InputData, ref Pos, Rows))
+			{
+				throw new PerforceException("Unable to parse input data as record");
+			}
+			if (Pos != InputData.Length)
+			{
+				throw new PerforceException("Garbage after end of spec data");
+			}
+
+			StringBuilder Result = new StringBuilder();
+			foreach ((Utf8String Key, PerforceValue Value) in Rows)
+			{
+				string[] ValueLines = Value.ToString().Split('\n');
+				if (ValueLines.Length == 1)
+				{
+					Result.AppendLine($"{Key}: {ValueLines[0]}");
+				}
+				else
+				{
+					Result.AppendLine($"{Key}:");
+					foreach (string ValueLine in ValueLines)
+					{
+						Result.AppendLine($"\t{ValueLine}");
+					}
+				}
+				Result.AppendLine();
+			}
+
+			return Result.ToString();
+		}
 	}
 }
diff --git a/Engine/Source/Programs/UnrealGameSync/UnrealGameSyncCmd/Program.cs b/Engine/Source/Programs/UnrealGameSync/UnrealGameSyncCmd/Program.cs
--- a/Engine/Source/Programs/UnrealGameSync/UnrealGameSyncCmd/Program.cs
+++ b/Engine/Source/Programs/UnrealGameSync/UnrealGameSyncCmd/Program.cs
@@ -262,10 +262,17 @@ namespace UnrealGameSyncCmd
 			DirectoryReference? Dir = DirectoryReference.GetCurrentDirectory();
 			for (; Dir != null; Dir = Dir.ParentDirectory)
 			{
-				UserWorkspaceSettings? Settings;
-				if (UserWorkspaceSettings.TryLoad(Dir, out Settings))
+				try
 				{
-					return Settings;
+					UserWorkspaceSettings? Settings;
+					if (UserWorkspaceSettings.TryLoad(Dir, out Settings))
+					{
+						return Settings;
+					}
+				}
+				catch
+				{
+					// Guard against directories we can't access, eg. [REDACTED_USER_HOME]
 				}
 			}
 			return null;
@@ -558,6 +565,9 @@ namespace UnrealGameSyncCmd
 				[CommandLine("-NoGPF", Value = "false")]
 				[CommandLine("-NoProjectFiles", Value = "false")]
 				public bool ProjectFiles { get; set; } = true;
+
+				[CommandLine("-Clobber")]
+				public bool Clobber { get; set; }
 			}
 
 			public override async Task ExecuteAsync(CommandContext Context)
@@ -599,6 +609,10 @@ namespace UnrealGameSyncCmd
 				{
 					Options |= WorkspaceUpdateOptions.GenerateProjectFiles;
 				}
+				if (SyncOptions.Clobber)
+				{
+					Options |= WorkspaceUpdateOptions.Clobber;
+				}
 				Options |= WorkspaceUpdateOptions.RemoveFilteredFiles;
 
 				ProjectInfo ProjectInfo = State.CreateProjectInfo();
@@ -611,7 +625,16 @@ namespace UnrealGameSyncCmd
 
 				WorkspaceUpdate Update = new WorkspaceUpdate(UpdateContext);
 				(WorkspaceUpdateResult Result, string Message) = await Update.ExecuteAsync(PerforceClient.Settings, ProjectInfo, State, Context.Logger, CancellationToken.None);
-				if(Result != WorkspaceUpdateResult.Success)
+				if (Result == WorkspaceUpdateResult.FilesToClobber)
+				{
+					Logger.LogWarning("The following files are modified in your workspace:");
+					foreach (string File in UpdateContext.ClobberFiles.Keys.OrderBy(x => x))
+					{
+						Logger.LogWarning("  {File}", File);
+					}
+					Logger.LogWarning("Use -Clobber to overwrite");
+				}
+				else if (Result != WorkspaceUpdateResult.Success)
 				{
 					Logger.LogError("{Message} (Result: {Result})", Message, Result);
 				}
diff --git a/Engine/Source/Programs/UnrealGameSync/UnrealGameSyncShared/WorkspaceUpdate.cs b/Engine/Source/Programs/UnrealGameSync/UnrealGameSyncShared/WorkspaceUpdate.cs
--- a/Engine/Source/Programs/UnrealGameSync/UnrealGameSyncShared/WorkspaceUpdate.cs
+++ b/Engine/Source/Programs/UnrealGameSync/UnrealGameSyncShared/WorkspaceUpdate.cs
@@ -42,6 +42,7 @@ namespace UnrealGameSync
 		SyncAllProjects = 0x1000,
 		IncludeAllProjectsInSolution = 0x2000,
 		RemoveFilteredFiles = 0x4000,
+		Clobber = 0x8000,
 	}
 
 	public enum WorkspaceUpdateResult
@@ -1664,23 +1665,26 @@ namespace UnrealGameSync
 			// If any files need to be clobbered, defer to the main thread to figure out which ones
 			if (TamperedFiles.Count > 0)
 			{
-				int NumNewFilesToClobber = 0;
-				foreach (string TamperedFile in TamperedFiles)
+				if ((Context.Options & WorkspaceUpdateOptions.Clobber) == 0)
 				{
-					if (!Context.ClobberFiles.ContainsKey(TamperedFile))
+					int NumNewFilesToClobber = 0;
+					foreach (string TamperedFile in TamperedFiles)
 					{
-						Context.ClobberFiles[TamperedFile] = true;
-						if (TamperedFile.EndsWith(LocalObjectVersionFileName, StringComparison.OrdinalIgnoreCase) || TamperedFile.EndsWith(LocalVersionHeaderFileName, StringComparison.OrdinalIgnoreCase))
+						if (!Context.ClobberFiles.ContainsKey(TamperedFile))
 						{
-							// Hack for UseFastModularVersioningV2; we don't need to update these files any more.
-							continue;
+							Context.ClobberFiles[TamperedFile] = true;
+							if (TamperedFile.EndsWith(LocalObjectVersionFileName, StringComparison.OrdinalIgnoreCase) || TamperedFile.EndsWith(LocalVersionHeaderFileName, StringComparison.OrdinalIgnoreCase))
+							{
+								// Hack for UseFastModularVersioningV2; we don't need to update these files any more.
+								continue;
+							}
+							NumNewFilesToClobber++;
 						}
-						NumNewFilesToClobber++;
 					}
-				}
-				if (NumNewFilesToClobber > 0)
-				{
-					return (WorkspaceUpdateResult.FilesToClobber, $"Cancelled sync after checking files to clobber ({NumNewFilesToClobber} new files).");
+					if (NumNewFilesToClobber > 0)
+					{
+						return (WorkspaceUpdateResult.FilesToClobber, $"Cancelled sync after checking files to clobber ({NumNewFilesToClobber} new files).");
+					}
 				}
 				foreach (string TamperedFile in TamperedFiles)
 				{
