Horde: If a call to p4 login doesn't return a proper record, call p4 login -s to get one.

diff --git a/Engine/Source/Programs/Shared/EpicGames.Perforce/PerforceConnection.cs b/Engine/Source/Programs/Shared/EpicGames.Perforce/PerforceConnection.cs
--- a/Engine/Source/Programs/Shared/EpicGames.Perforce/PerforceConnection.cs
+++ b/Engine/Source/Programs/Shared/EpicGames.Perforce/PerforceConnection.cs
@@ -2353,7 +2353,15 @@ namespace EpicGames.Perforce
 				DiscardPasswordPrompt(Response);
 
 				List<PerforceResponse> ParsedResponses = await Response.ReadResponsesAsync(typeof(LoginRecord), CancellationToken);
-				return new PerforceResponse<LoginRecord>(ParsedResponses.First());
+
+				PerforceResponse FirstResponse = ParsedResponses.First();
+				if (FirstRespone.Info != null)
+				{
+					// Older versions of P4.EXE do not return a login record for succesful login, instread just returning a string. Call p4 login -s to get the login state instead.
+					return await TryGetLoginStateAsync(Connection, CancellationToken);
+				}
+
+				return new PerforceResponse<LoginRecord>(FirstResponse);
 			}
 		}
 
