Horde: Fix typo.

diff --git a/Engine/Source/Programs/Shared/EpicGames.Perforce/PerforceConnection.cs b/Engine/Source/Programs/Shared/EpicGames.Perforce/PerforceConnection.cs
--- a/Engine/Source/Programs/Shared/EpicGames.Perforce/PerforceConnection.cs
+++ b/Engine/Source/Programs/Shared/EpicGames.Perforce/PerforceConnection.cs
@@ -2355,7 +2355,7 @@ namespace EpicGames.Perforce
 				List<PerforceResponse> ParsedResponses = await Response.ReadResponsesAsync(typeof(LoginRecord), CancellationToken);
 
 				PerforceResponse FirstResponse = ParsedResponses.First();
-				if (FirstRespone.Info != null)
+				if (FirstResponse.Info != null)
 				{
 					// Older versions of P4.EXE do not return a login record for succesful login, instread just returning a string. Call p4 login -s to get the login state instead.
 					return await TryGetLoginStateAsync(Connection, CancellationToken);
