Horde: Fix revert of files in workspaces not working correctly (missing path specification), and exception parsing OpenedRecord entries with ShortOutput option specified.

diff --git a/Engine/Source/Programs/Shared/EpicGames.Perforce.Managed/ManagedWorkspace.cs b/Engine/Source/Programs/Shared/EpicGames.Perforce.Managed/ManagedWorkspace.cs
--- a/Engine/Source/Programs/Shared/EpicGames.Perforce.Managed/ManagedWorkspace.cs
+++ b/Engine/Source/Programs/Shared/EpicGames.Perforce.Managed/ManagedWorkspace.cs
@@ -1241,7 +1241,7 @@ namespace EpicGames.Perforce.Managed
 				Stopwatch Timer = Stopwatch.StartNew();
 
 				// Get a list of open files
-				List<OpenedRecord> OpenedFilesResponse = await PerforceClient.OpenedAsync(OpenedOptions.ShortOutput, -1, PerforceClient.Settings.ClientName!, null, 1, new string[0], CancellationToken).ToListAsync(CancellationToken);
+				List<OpenedRecord> OpenedFilesResponse = await PerforceClient.OpenedAsync(OpenedOptions.ShortOutput, -1, PerforceClient.Settings.ClientName!, null, 1, FileSpecList.Any, CancellationToken).ToListAsync(CancellationToken);
 
 				// If there are any files, revert them
 				if(OpenedFilesResponse.Any())
diff --git a/Engine/Source/Programs/Shared/EpicGames.Perforce/Records/OpenedRecord.cs b/Engine/Source/Programs/Shared/EpicGames.Perforce/Records/OpenedRecord.cs
--- a/Engine/Source/Programs/Shared/EpicGames.Perforce/Records/OpenedRecord.cs
+++ b/Engine/Source/Programs/Shared/EpicGames.Perforce/Records/OpenedRecord.cs
@@ -22,21 +22,21 @@ namespace EpicGames.Perforce
 		public string DepotFile { get; set; } = String.Empty;
 		
 		/// <summary>
-		/// Client path to file
+		/// Client path to file (not returned if OpenedOptions.ShortOutput is set)
 		/// </summary>
-		[PerforceTag("clientFile")]
+		[PerforceTag("clientFile", Optional = true)]
 		public string ClientFile { get; set; } = String.Empty;
 
 		/// <summary>
-		/// The revision of the file 
+		/// The revision of the file (not returned if OpenedOptions.ShortOutput is set)
 		/// </summary>
-		[PerforceTag("rev")]
+		[PerforceTag("rev", Optional = true)]
 		public int Revision { get; set; }
 
 		/// <summary>
-		/// The synced revision of the file (may be 'none' for adds)
+		/// The synced revision of the file (may be 'none' for adds) (not returned if OpenedOptions.ShortOutput is set)
 		/// </summary>
-		[PerforceTag("haveRev")]
+		[PerforceTag("haveRev", Optional = true)]
 		public int HaveRevision { get; set; }
 
 		/// <summary>
