CS: Rename compact binary files to match the classes they contain.

diff --git a/Engine/Source/Programs/Shared/EpicGames.Serialization/CbConverter.cs b/Engine/Source/Programs/Shared/EpicGames.Serialization/CbConverter.cs
new file mode 100644
--- /dev/null
+++ b/Engine/Source/Programs/Shared/EpicGames.Serialization/CbConverter.cs
@@ -0,0 +1,108 @@
+// Copyright Epic Games, Inc. All Rights Reserved.
+
+using EpicGames.Core;
+using System;
+using System.Collections.Generic;
+using System.Text;
+
+namespace EpicGames.Serialization
+{
+	/// <summary>
+	/// Attribute declaring the converter to use for a type
+	/// </summary>
+	[AttributeUsage(AttributeTargets.Class | AttributeTargets.Struct)]
+	public class CbConverterAttribute : Attribute
+	{
+		/// <summary>
+		/// The converter type
+		/// </summary>
+		public Type ConverterType { get; }
+
+		/// <summary>
+		/// Constructor
+		/// </summary>
+		/// <param name="ConverterType"></param>
+		public CbConverterAttribute(Type ConverterType)
+		{
+			this.ConverterType = ConverterType;
+		}
+	}
+
+	/// <summary>
+	/// Base class for all converters. Deriving from <see cref="CbConverter{T}"/> is more efficient
+	/// </summary>
+	public abstract class CbConverter
+	{
+		/// <summary>
+		/// Reads an object from a field
+		/// </summary>
+		/// <param name="Field"></param>
+		/// <returns></returns>
+		public abstract object? ReadObject(CbField Field);
+
+		/// <summary>
+		/// Writes an object to a field
+		/// </summary>
+		/// <param name="Writer"></param>
+		/// <param name="Value"></param>
+		public abstract void WriteObject(CbWriter Writer, object? Value);
+
+		/// <summary>
+		/// Writes an object to a named field, if not equal to the default value
+		/// </summary>
+		/// <param name="Writer"></param>
+		/// <param name="Name"></param>
+		/// <param name="Value"></param>
+		public abstract void WriteNamedObject(CbWriter Writer, Utf8String Name, object? Value);
+	}
+
+	/// <summary>
+	/// Converter for a particular type
+	/// </summary>
+	public abstract class CbConverter<T> : CbConverter
+	{
+		/// <summary>
+		/// Reads an object from a field
+		/// </summary>
+		/// <param name="Field"></param>
+		/// <returns></returns>
+		public abstract T Read(CbField Field);
+
+		/// <summary>
+		/// Writes an object to a field
+		/// </summary>
+		/// <param name="Writer"></param>
+		/// <param name="Value"></param>
+		public abstract void Write(CbWriter Writer, T Value);
+
+		/// <summary>
+		/// Writes an object to a named field, if not equal to the default value
+		/// </summary>
+		/// <param name="Writer"></param>
+		/// <param name="Name"></param>
+		/// <param name="Value"></param>
+		public abstract void WriteNamed(CbWriter Writer, Utf8String Name, T Value);
+
+		/// <inheritdoc/>
+		public sealed override object? ReadObject(CbField Field) => Read(Field);
+
+		/// <inheritdoc/>
+		public sealed override void WriteObject(CbWriter Writer, object? Value) => Write(Writer, (T)Value!);
+
+		/// <inheritdoc/>
+		public sealed override void WriteNamedObject(CbWriter Writer, Utf8String Name, object? Value) => WriteNamed(Writer, Name, (T)Value!);
+	}
+
+	/// <summary>
+	/// Base class for converter factories
+	/// </summary>
+	public abstract class CbConverterFactory
+	{
+		/// <summary>
+		/// Create a converter for the given type
+		/// </summary>
+		/// <param name="Type">The type to create a converter for</param>
+		/// <returns>The converter instance, or null if this factory does not support the given type</returns>
+		public abstract CbConverter? CreateConverter(Type Type);
+	}
+}
diff --git a/Engine/Source/Programs/Shared/EpicGames.Serialization/CbField.cs b/Engine/Source/Programs/Shared/EpicGames.Serialization/CbField.cs
new file mode 100644
--- /dev/null
+++ b/Engine/Source/Programs/Shared/EpicGames.Serialization/CbField.cs
@@ -0,0 +1,2646 @@
+// Copyright Epic Games, Inc. All Rights Reserved.
+
+using EpicGames.Core;
+using System;
+using System.Buffers;
+using System.Buffers.Binary;
+using System.Collections;
+using System.Collections.Generic;
+using System.ComponentModel;
+using System.Diagnostics;
+using System.Globalization;
+using System.Linq;
+using System.Numerics;
+using System.Text;
+using System.Text.Json;
+using System.Text.Json.Serialization;
+
+namespace EpicGames.Serialization
+{
+	/// <summary>
+	/// Field types and flags for FCbField[View].
+	///
+	/// DO NOT CHANGE THE VALUE OF ANY MEMBERS OF THIS ENUM!
+	/// BACKWARD COMPATIBILITY REQUIRES THAT THESE VALUES BE FIXED!
+	/// SERIALIZATION USES HARD-CODED CONSTANTS BASED ON THESE VALUES!
+	/// </summary>
+	[Flags]
+	public enum CbFieldType : byte
+	{
+		/// <summary>
+		/// A field type that does not occur in a valid object.
+		/// </summary>
+		None = 0x00,
+
+		/// <summary>
+		/// Null. Payload is empty.
+		/// </summary>
+		Null = 0x01,
+
+		/// <summary>
+		/// Object is an array of fields with unique non-empty names.
+		///
+		/// Payload is a VarUInt byte count for the encoded fields followed by the fields.
+		/// </summary>
+		Object = 0x02,
+
+		/// <summary>
+		/// UniformObject is an array of fields with the same field types and unique non-empty names.
+		/// 
+		/// Payload is a VarUInt byte count for the encoded fields followed by the fields.
+		/// </summary>
+		UniformObject = 0x03,
+
+		/// <summary>
+		/// Array is an array of fields with no name that may be of different types.
+		/// 
+		/// Payload is a VarUInt byte count, followed by a VarUInt item count, followed by the fields.
+		/// </summary>
+		Array = 0x04,
+
+		/// <summary>
+		/// UniformArray is an array of fields with no name and with the same field type.
+		///
+		/// Payload is a VarUInt byte count, followed by a VarUInt item count, followed by field type,
+		/// followed by the fields without their field type.
+		/// </summary>
+		UniformArray = 0x05,
+
+		/// <summary> 
+		/// Binary. Payload is a VarUInt byte count followed by the data. 
+		/// /// </summary>
+		Binary = 0x06,
+
+		/// <summary>
+		/// String in UTF-8. Payload is a VarUInt byte count then an unterminated UTF-8 string. 
+		/// </summary>
+		String = 0x07,
+
+		/// <summary>
+		/// Non-negative integer with the range of a 64-bit unsigned integer.
+		/// 
+		/// Payload is the value encoded as a VarUInt.
+		/// </summary>
+		IntegerPositive = 0x08,
+
+		/// <summary>
+		/// Negative integer with the range of a 64-bit signed integer.
+		///
+		/// Payload is the ones' complement of the value encoded as a VarUInt.
+		/// </summary>
+		IntegerNegative = 0x09,
+
+		/// <summary>
+		/// Single precision float. Payload is one big endian IEEE 754 binary32 float.
+		/// /// </summary>
+		Float32 = 0x0a,
+
+		/// <summary>
+		/// Double precision float. Payload is one big endian IEEE 754 binary64 float. 
+		/// </summary>
+		Float64 = 0x0b,
+
+		/// <summary>
+		/// Boolean false value. Payload is empty. 
+		/// </summary>
+		BoolFalse = 0x0c,
+
+		/// <summary>
+		/// Boolean true value. Payload is empty. 
+		/// </summary>
+		BoolTrue = 0x0d,
+
+		/// <summary>
+		/// CompactBinaryAttachment is a reference to a compact binary attachment stored externally.
+		///
+		/// Payload is a 160-bit hash digest of the referenced compact binary data.
+		/// </summary>
+		ObjectAttachment = 0x0e,
+
+		/// <summary>
+		/// BinaryAttachment is a reference to a binary attachment stored externally.
+		///
+		/// Payload is a 160-bit hash digest of the referenced binary data.
+		/// </summary>
+		BinaryAttachment = 0x0f,
+
+		/// <summary>
+		/// Hash. Payload is a 160-bit hash digest. 
+		/// </summary>
+		Hash = 0x10,
+
+		/// <summary>
+		/// UUID/GUID. Payload is a 128-bit UUID as defined by RFC 4122. 
+		/// </summary>
+		Uuid = 0x11,
+
+		/// <summary>
+		/// Date and time between 0001-01-01 00:00:00.0000000 and 9999-12-31 23:59:59.9999999.
+		///
+		/// Payload is a big endian int64 count of 100ns ticks since 0001-01-01 00:00:00.0000000.
+		/// </summary>
+		DateTime = 0x12,
+
+		/// <summary>
+		/// Difference between two date/time values.
+		/// 
+		/// Payload is a big endian int64 count of 100ns ticks in the span, and may be negative.
+		/// </summary>
+		TimeSpan = 0x13,
+
+		/// <summary>
+		/// ObjectId is an opaque object identifier. See FCbObjectId.
+		///
+		/// Payload is a 12-byte object identifier.
+		/// </summary>
+		ObjectId = 0x14,
+
+		/// <summary>
+		/// CustomById identifies the sub-type of its payload by an integer identifier.
+		///
+		/// Payload is a VarUInt byte count of the sub-type identifier and the sub-type payload, followed
+		/// by a VarUInt of the sub-type identifier then the payload of the sub-type.
+		/// </summary>
+		CustomById = 0x1e,
+
+		/// <summary>
+		/// CustomByType identifies the sub-type of its payload by a string identifier.
+		///
+		/// Payload is a VarUInt byte count of the sub-type identifier and the sub-type payload, followed
+		/// by a VarUInt byte count of the unterminated sub-type identifier, then the sub-type identifier
+		/// without termination, then the payload of the sub-type.
+		/// </summary>
+		CustomByName = 0x1f,
+
+		/// <summary>
+		/// Reserved for future use as a flag. Do not add types in this range. 
+		/// </summary>
+		Reserved = 0x20,
+
+		/// <summary>
+		/// A transient flag which indicates that the object or array containing this field has stored
+		/// the field type before the payload and name. Non-uniform objects and fields will set this.
+		/// 
+		/// Note: Since the flag must never be serialized, this bit may be repurposed in the future.
+		/// </summary>
+		HasFieldType = 0x40,
+
+		/// <summary>
+		/// A persisted flag which indicates that the field has a name stored before the payload. 
+		/// </summary>
+		HasFieldName = 0x80,
+	}
+
+	/// <summary>
+	/// A binary attachment, referenced by <see cref="IoHash"/>
+	/// </summary>
+	[DebuggerDisplay("{Hash}")]
+	[JsonConverter(typeof(CbBinaryAttachmentJsonConverter))]
+	[TypeConverter(typeof(CbBinaryAttachmentTypeConverter))]
+	public struct CbBinaryAttachment
+	{
+		/// <summary>
+		/// Attachment with a hash of zero
+		/// </summary>
+		public static CbBinaryAttachment Zero { get; } = new CbBinaryAttachment(IoHash.Zero);
+
+		/// <summary>
+		/// Hash of the referenced object
+		/// </summary>
+		public IoHash Hash;
+
+		/// <summary>
+		/// Constructor
+		/// </summary>
+		/// <param name="Hash">Hash of the referenced object</param>
+		public CbBinaryAttachment(IoHash Hash)
+		{
+			this.Hash = Hash;
+		}
+
+		/// <inheritdoc/>
+		public override string ToString() => Hash.ToString();
+
+		/// <inheritdoc/>
+		public override bool Equals(object? Obj) => Obj is CbBinaryAttachment Other && Other.Hash == Hash;
+
+		/// <inheritdoc/>
+		public override int GetHashCode() => Hash.GetHashCode();
+
+		/// <inheritdoc/>
+		public static bool operator ==(CbBinaryAttachment Lhs, CbBinaryAttachment Rhs) => Lhs.Hash == Rhs.Hash;
+
+		/// <inheritdoc/>
+		public static bool operator !=(CbBinaryAttachment Lhs, CbBinaryAttachment Rhs) => Lhs.Hash != Rhs.Hash;
+
+		/// <summary>
+		/// Convert a hash to a binary attachment 
+		/// </summary>
+		/// <param name="Hash">The attachment to convert</param>
+		public static implicit operator CbBinaryAttachment(IoHash Hash) => new CbBinaryAttachment(Hash);
+
+		/// <summary>
+		/// Use a binary attachment as a hash
+		/// </summary>
+		/// <param name="Attachment">The attachment to convert</param>
+		public static implicit operator IoHash(CbBinaryAttachment Attachment) => Attachment.Hash;
+	}
+
+	/// <summary>
+	/// Type converter for IoHash to and from JSON
+	/// </summary>
+	sealed class CbBinaryAttachmentJsonConverter : JsonConverter<CbBinaryAttachment>
+	{
+		/// <inheritdoc/>
+		public override CbBinaryAttachment Read(ref Utf8JsonReader Reader, Type TypeToConvert, JsonSerializerOptions Options) => IoHash.Parse(Reader.ValueSpan);
+
+		/// <inheritdoc/>
+		public override void Write(Utf8JsonWriter Writer, CbBinaryAttachment Value, JsonSerializerOptions Options) => Writer.WriteStringValue(Value.Hash.ToUtf8String().Span);
+	}
+
+	/// <summary>
+	/// Type converter from strings to IoHash objects
+	/// </summary>
+	sealed class CbBinaryAttachmentTypeConverter : TypeConverter
+	{
+		/// <inheritdoc/>
+		public override bool CanConvertFrom(ITypeDescriptorContext Context, Type SourceType) => SourceType == typeof(string);
+
+		/// <inheritdoc/>
+		public override object ConvertFrom(ITypeDescriptorContext Context, CultureInfo Culture, object Value) => new CbBinaryAttachment(IoHash.Parse((string)Value));
+	}
+
+	/// <summary>
+	/// An object attachment, referenced by <see cref="IoHash"/>
+	/// </summary>
+	[DebuggerDisplay("{Hash}")]
+	[JsonConverter(typeof(CbObjectAttachmentJsonConverter))]
+	[TypeConverter(typeof(CbObjectAttachmentTypeConverter))]
+	public struct CbObjectAttachment
+	{
+		/// <summary>
+		/// Attachment with a hash of zero
+		/// </summary>
+		public static CbObjectAttachment Zero { get; } = new CbObjectAttachment(IoHash.Zero);
+
+		/// <summary>
+		/// Hash of the referenced object
+		/// </summary>
+		public IoHash Hash;
+
+		/// <summary>
+		/// Constructor
+		/// </summary>
+		/// <param name="Hash">Hash of the referenced object</param>
+		public CbObjectAttachment(IoHash Hash)
+		{
+			this.Hash = Hash;
+		}
+
+		/// <inheritdoc/>
+		public override string ToString() => Hash.ToString();
+
+		/// <inheritdoc/>
+		public override bool Equals(object? Obj) => Obj is CbObjectAttachment Other && Other.Hash == Hash;
+
+		/// <inheritdoc/>
+		public override int GetHashCode() => Hash.GetHashCode();
+
+		/// <inheritdoc/>
+		public static bool operator ==(CbObjectAttachment Lhs, CbObjectAttachment Rhs) => Lhs.Hash == Rhs.Hash;
+
+		/// <inheritdoc/>
+		public static bool operator !=(CbObjectAttachment Lhs, CbObjectAttachment Rhs) => Lhs.Hash != Rhs.Hash;
+
+		/// <summary>
+		/// Use an object attachment as a hash
+		/// </summary>
+		/// <param name="Attachment">The attachment to convert</param>
+		public static implicit operator CbObjectAttachment(IoHash Hash) => new CbObjectAttachment(Hash);
+
+		/// <summary>
+		/// Use an object attachment as a hash
+		/// </summary>
+		/// <param name="Attachment">The attachment to convert</param>
+		public static implicit operator IoHash(CbObjectAttachment Attachment) => Attachment.Hash;
+	}
+
+	/// <summary>
+	/// Type converter for IoHash to and from JSON
+	/// </summary>
+	sealed class CbObjectAttachmentJsonConverter : JsonConverter<CbObjectAttachment>
+	{
+		/// <inheritdoc/>
+		public override CbObjectAttachment Read(ref Utf8JsonReader Reader, Type TypeToConvert, JsonSerializerOptions Options) => IoHash.Parse(Reader.ValueSpan);
+
+		/// <inheritdoc/>
+		public override void Write(Utf8JsonWriter Writer, CbObjectAttachment Value, JsonSerializerOptions Options) => Writer.WriteStringValue(Value.Hash.ToUtf8String().Span);
+	}
+
+	/// <summary>
+	/// Type converter from strings to IoHash objects
+	/// </summary>
+	sealed class CbObjectAttachmentTypeConverter : TypeConverter
+	{
+		/// <inheritdoc/>
+		public override bool CanConvertFrom(ITypeDescriptorContext Context, Type SourceType) => SourceType == typeof(string);
+
+		/// <inheritdoc/>
+		public override object ConvertFrom(ITypeDescriptorContext Context, CultureInfo Culture, object Value) => new CbObjectAttachment(IoHash.Parse((string)Value));
+	}
+
+	/// <summary>
+	/// Methods that operate on <see cref="CbFieldType"/>.
+	/// </summary>
+	public static class CbFieldUtils
+	{
+		private const CbFieldType SerializedTypeMask = (CbFieldType)0b_1001_1111;
+		private const CbFieldType TypeMask = (CbFieldType)0b_0001_1111;
+
+		private const CbFieldType ObjectMask = (CbFieldType)0b_0001_1110;
+		private const CbFieldType ObjectBase = (CbFieldType)0b_0000_0010;
+
+		private const CbFieldType ArrayMask = (CbFieldType)0b_0001_1110;
+		private const CbFieldType ArrayBase = (CbFieldType)0b_0000_0100;
+
+		private const CbFieldType IntegerMask = (CbFieldType)0b_0011_1110;
+		private const CbFieldType IntegerBase = (CbFieldType)0b_0000_1000;
+
+		private const CbFieldType FloatMask = (CbFieldType)0b_0001_1100;
+		private const CbFieldType FloatBase = (CbFieldType)0b_0000_1000;
+
+		private const CbFieldType BoolMask = (CbFieldType)0b_0001_1110;
+		private const CbFieldType BoolBase = (CbFieldType)0b_0000_1100;
+
+		private const CbFieldType AttachmentMask = (CbFieldType)0b_0001_1110;
+		private const CbFieldType AttachmentBase = (CbFieldType)0b_0000_1110;
+
+		/// <summary>
+		/// Removes flags from the given type
+		/// </summary>
+		/// <param name="Type">Type to check</param>
+		/// <returns>Type without flag fields</returns>
+		public static CbFieldType GetType(CbFieldType Type)
+		{
+			return Type & TypeMask;
+		}
+
+		/// <summary>
+		/// Gets the serialized type
+		/// </summary>
+		/// <param name="Type">Type to check</param>
+		/// <returns>Type without flag fields</returns>
+		public static CbFieldType GetSerializedType(CbFieldType Type)
+		{
+			return Type & SerializedTypeMask;
+		}
+
+		/// <summary>
+		/// Tests if the given field has a type
+		/// </summary>
+		/// <param name="Type">Type to check</param>
+		/// <returns>True if the field has a type</returns>
+		public static bool HasFieldType(CbFieldType Type)
+		{
+			return (Type & CbFieldType.HasFieldType) != 0;
+		}
+
+		/// <summary>
+		/// Tests if the given field has a name
+		/// </summary>
+		/// <param name="Type">Type to check</param>
+		/// <returns>True if the field has a name</returns>
+		public static bool HasFieldName(CbFieldType Type)
+		{
+			return (Type & CbFieldType.HasFieldName) != 0;
+		}
+
+		/// <summary>
+		/// Tests if the given field type is none
+		/// </summary>
+		/// <param name="Type">Type to check</param>
+		/// <returns>True if the field is none</returns>
+		public static bool IsNone(CbFieldType Type)
+		{
+			return GetType(Type) == CbFieldType.None;
+		}
+
+		/// <summary>
+		/// Tests if the given field type is a null value
+		/// </summary>
+		/// <param name="Type">Type to check</param>
+		/// <returns>True if the field is a null</returns>
+		public static bool IsNull(CbFieldType Type)
+		{
+			return GetType(Type) == CbFieldType.Null;
+		}
+
+		/// <summary>
+		/// Tests if the given field type is an object
+		/// </summary>
+		/// <param name="Type">Type to check</param>
+		/// <returns>True if the field is an object type</returns>
+		public static bool IsObject(CbFieldType Type)
+		{
+			return (Type & ObjectMask) == ObjectBase;
+		}
+
+		/// <summary>
+		/// Tests if the given field type is an array
+		/// </summary>
+		/// <param name="Type">Type to check</param>
+		/// <returns>True if the field is an array type</returns>
+		public static bool IsArray(CbFieldType Type)
+		{
+			return (Type & ArrayMask) == ArrayBase;
+		}
+
+		/// <summary>
+		/// Tests if the given field type is binary
+		/// </summary>
+		/// <param name="Type">Type to check</param>
+		/// <returns>True if the field is binary</returns>
+		public static bool IsBinary(CbFieldType type)
+		{
+			return GetType(type) == CbFieldType.Binary;
+		}
+
+		/// <summary>
+		/// Tests if the given field type is a string
+		/// </summary>
+		/// <param name="Type">Type to check</param>
+		/// <returns>True if the field is an array type</returns>
+		public static bool IsString(CbFieldType Type)
+		{
+			return GetType(Type) == CbFieldType.String;
+		}
+
+		/// <summary>
+		/// Tests if the given field type is an integer
+		/// </summary>
+		/// <param name="Type">Type to check</param>
+		/// <returns>True if the field is an integer type</returns>
+		public static bool IsInteger(CbFieldType Type)
+		{
+			return (Type & IntegerMask) == IntegerBase;
+		}
+
+		/// <summary>
+		/// Tests if the given field type is a float (or integer, due to implicit conversion)
+		/// </summary>
+		/// <param name="Type">Type to check</param>
+		/// <returns>True if the field is a float type</returns>
+		public static bool IsFloat(CbFieldType Type)
+		{
+			return (Type & FloatMask) == FloatBase;
+		}
+
+		/// <summary>
+		/// Tests if the given field type is a boolean
+		/// </summary>
+		/// <param name="Type">Type to check</param>
+		/// <returns>True if the field is an bool type</returns>
+		public static bool IsBool(CbFieldType Type)
+		{
+			return (Type & BoolMask) == BoolBase;
+		}
+
+		/// <summary>
+		/// Tests if the given field type is a compact binary attachment
+		/// </summary>
+		/// <param name="Type">Type to check</param>
+		/// <returns>True if the field is a compact binary attachment</returns>
+		public static bool IsObjectAttachment(CbFieldType type)
+		{
+			return GetType(type) == CbFieldType.ObjectAttachment;
+		}
+
+		/// <summary>
+		/// Tests if the given field type is a binary attachment
+		/// </summary>
+		/// <param name="Type">Type to check</param>
+		/// <returns>True if the field is a binary attachment</returns>
+		public static bool IsBinaryAttachment(CbFieldType type)
+		{
+			return GetType(type) == CbFieldType.BinaryAttachment;
+		}
+
+		/// <summary>
+		/// Tests if the given field type is an attachment
+		/// </summary>
+		/// <param name="Type">Type to check</param>
+		/// <returns>True if the field is an attachment type</returns>
+		public static bool IsAttachment(CbFieldType Type)
+		{
+			return (Type & AttachmentMask) == AttachmentBase;
+		}
+
+		/// <summary>
+		/// Tests if the given field type is a hash
+		/// </summary>
+		/// <param name="Type">Type to check</param>
+		/// <returns>True if the field is a hash</returns>
+		public static bool IsHash(CbFieldType Type)
+		{
+			return GetType(Type) == CbFieldType.Hash || IsAttachment(Type);
+		}
+
+		/// <summary>
+		/// Tests if the given field type is a UUID
+		/// </summary>
+		/// <param name="Type">Type to check</param>
+		/// <returns>True if the field is a UUID</returns>
+		public static bool IsUuid(CbFieldType type)
+		{
+			return GetType(type) == CbFieldType.Uuid;
+		}
+
+		/// <summary>
+		/// Tests if the given field type is a date/time
+		/// </summary>
+		/// <param name="Type">Type to check</param>
+		/// <returns>True if the field is a date/time</returns>
+		public static bool IsDateTime(CbFieldType type)
+		{
+			return GetType(type) == CbFieldType.DateTime;
+		}
+
+		/// <summary>
+		/// Tests if the given field type is a timespan
+		/// </summary>
+		/// <param name="Type">Type to check</param>
+		/// <returns>True if the field is a timespan</returns>
+		public static bool IsTimeSpan(CbFieldType type)
+		{
+			return GetType(type) == CbFieldType.TimeSpan;
+		}
+
+		/// <summary>
+		/// Tests if the given field type has fields
+		/// </summary>
+		/// <param name="Type">Type to check</param>
+		/// <returns>True if the field has fields</returns>
+		public static bool HasFields(CbFieldType Type)
+		{
+			CbFieldType NoFlags = GetType(Type);
+			return NoFlags >= CbFieldType.Object && NoFlags <= CbFieldType.UniformArray;
+		}
+
+		/// <summary>
+		/// Tests if the given field type has uniform fields (array/object)
+		/// </summary>
+		/// <param name="Type">Type to check</param>
+		/// <returns>True if the field has uniform fields</returns>
+		public static bool HasUniformFields(CbFieldType Type)
+		{
+			CbFieldType LocalType = GetType(Type);
+			return LocalType == CbFieldType.UniformObject || LocalType == CbFieldType.UniformArray;
+		}
+
+		/// <summary>
+		/// Tests if the type is or may contain fields of any attachment type.
+		/// </summary>
+		public static bool MayContainAttachments(CbFieldType Type)
+		{
+			return IsObject(Type) | IsArray(Type) | IsAttachment(Type);
+		}
+	}
+
+	/// <summary>
+	/// Errors that can occur when accessing a field. */
+	/// </summary>
+	public enum CbFieldError : byte
+	{
+		/// <summary>
+		/// The field is not in an error state.
+		/// </summary>
+		None,
+
+		/// <summary>
+		/// The value type does not match the requested type.
+		/// </summary>
+		TypeError,
+
+		/// <summary>
+		/// The value is out of range for the requested type.
+		/// </summary>
+		RangeError,
+	}
+
+	/// <summary>
+	/// Simplified view of <see cref="CbField"/> in the debugger, for fields with a name
+	/// </summary>
+	class CbFieldWithNameDebugView
+	{
+		public string? Name;
+		public object? Value;
+	}
+
+	/// <summary>
+	/// Simplified view of <see cref="CbField"/> for the debugger
+	/// </summary>
+	class CbFieldDebugView
+	{
+		public CbFieldDebugView(CbField Field) => Value = Field.HasName()
+				? new CbFieldWithNameDebugView { Name = Field.Name.ToString(), Value = Field.Value } 
+				: Field.Value;
+
+		[DebuggerBrowsable(DebuggerBrowsableState.RootHidden)]
+		public object? Value { get; }
+	}
+
+	/// <summary>
+	/// An atom of data in the compact binary format.
+	///
+	/// Accessing the value of a field is always a safe operation, even if accessed as the wrong type.
+	/// An invalid access will return a default value for the requested type, and set an error code on
+	/// the field that can be checked with GetLastError and HasLastError. A valid access will clear an
+	/// error from a previous invalid access.
+	///
+	/// A field is encoded in one or more bytes, depending on its type and the type of object or array
+	/// that contains it. A field of an object or array which is non-uniform encodes its field type in
+	/// the first byte, and includes the HasFieldName flag for a field in an object. The field name is
+	/// encoded in a variable-length unsigned integer of its size in bytes, for named fields, followed
+	/// by that many bytes of the UTF-8 encoding of the name with no null terminator.The remainder of
+	/// the field is the payload and is described in the field type enum. Every field must be uniquely
+	/// addressable when encoded, which means a zero-byte field is not permitted, and only arises in a
+	/// uniform array of fields with no payload, where the answer is to encode as a non-uniform array.
+	/// </summary>
+	[DebuggerDisplay("{DebugValue,nq}")]
+	[DebuggerTypeProxy(typeof(CbFieldDebugView))]
+	[JsonConverter(typeof(CbFieldJsonConverter))]
+	public class CbField : IEquatable<CbField>, IEnumerable<CbField>
+	{
+		/// <summary>
+		/// Type returned for none values
+		/// </summary>
+		[DebuggerDisplay("<none>")]
+		class NoneValueType
+		{
+		}
+
+		/// <summary>
+		/// Special value returned for "none" fields.
+		/// </summary>
+		static NoneValueType None { get; } = new NoneValueType();
+
+		/// <summary>
+		/// Formatter for the debug string
+		/// </summary>
+		object? DebugValue => HasName() ? $"{Name} = {Value}" : Value;
+
+		/// <summary>
+		/// Default empty field
+		/// </summary>
+		public static CbField Empty { get; } = new CbField();
+
+		/// <summary>
+		/// The field type, with the transient HasFieldType flag if the field contains its type
+		/// </summary>
+		public CbFieldType TypeWithFlags { get; }
+
+		/// <summary>
+		/// Data for this field
+		/// </summary>
+		public ReadOnlyMemory<byte> Memory { get; }
+
+		/// <summary>
+		/// Offset of the name with the memory
+		/// </summary>
+		public int NameLen;
+
+		/// <summary>
+		/// Offset of the payload within the memory
+		/// </summary>
+		public int PayloadOffset;
+
+		/// <summary>
+		/// Error for parsing the current field type
+		/// </summary>
+		public CbFieldError Error { get; private set; }
+
+		/// <summary>
+		/// Default constructor
+		/// </summary>
+		public CbField()
+			: this(ReadOnlyMemory<byte>.Empty, CbFieldType.None)
+		{
+		}
+
+		/// <summary>
+		/// Copy constructor
+		/// </summary>
+		/// <param name="Other"></param>
+		public CbField(CbField Other)
+		{
+			this.TypeWithFlags = Other.TypeWithFlags;
+			this.Memory = Other.Memory;
+			this.NameLen = Other.NameLen;
+			this.PayloadOffset = Other.PayloadOffset;
+			this.Error = Other.Error;
+		}
+
+		/// <summary>
+		/// Construct a field from a pointer to its data and an optional externally-provided type.
+		/// </summary>
+		/// <param>Data Pointer to the start of the field data.</param>
+		/// <param>Type HasFieldType means that Data contains the type. Otherwise, use the given type.</param>
+		public CbField(ReadOnlyMemory<byte> Data, CbFieldType Type = CbFieldType.HasFieldType)
+		{
+			int Offset = 0;
+			if (CbFieldUtils.HasFieldType(Type))
+			{
+				Type = (CbFieldType)Data.Span[Offset] | CbFieldType.HasFieldType;
+				Offset++;
+			}
+
+			if (CbFieldUtils.HasFieldName(Type))
+			{
+				NameLen = (int)VarInt.Read(Data.Slice(Offset).Span, out int NameLenByteCount);
+				Offset += NameLenByteCount + NameLen;
+			}
+
+			this.Memory = Data;
+			this.TypeWithFlags = Type;
+			this.PayloadOffset = Offset;
+
+			Memory = Memory.Slice(0, (int)Math.Min((ulong)Memory.Length, (ulong)PayloadOffset + GetPayloadSize()));
+		}
+
+		/// <summary>
+		/// Returns the name of the field if it has a name, otherwise an empty view.
+		/// </summary>
+		public Utf8String Name => new Utf8String(Memory.Slice(PayloadOffset - NameLen, NameLen));
+
+		/// <summary>
+		/// Gets the value of this field
+		/// </summary>
+		public object? Value
+		{
+			get
+			{
+				CbFieldType FieldType = CbFieldUtils.GetType(TypeWithFlags);
+				switch (FieldType)
+				{
+					case CbFieldType.None:
+						return None;
+					case CbFieldType.Null:
+						return null;
+					case CbFieldType.Object:
+					case CbFieldType.UniformObject:
+						return AsObject();
+					case CbFieldType.Array:
+					case CbFieldType.UniformArray:
+						return AsArray();
+					case CbFieldType.Binary:
+						return AsBinary();
+					case CbFieldType.String:
+						return AsString();
+					case CbFieldType.IntegerPositive:
+						return AsUInt64();
+					case CbFieldType.IntegerNegative:
+						return AsInt64();
+					case CbFieldType.Float32:
+						return AsFloat();
+					case CbFieldType.Float64:
+						return AsDouble();
+					case CbFieldType.BoolFalse:
+						return false;
+					case CbFieldType.BoolTrue:
+						return true;
+					case CbFieldType.ObjectAttachment:
+						return AsObjectAttachment();
+					case CbFieldType.BinaryAttachment:
+						return AsBinaryAttachment();
+					case CbFieldType.Hash:
+						return AsHash();
+					case CbFieldType.Uuid:
+						return AsUuid();
+					case CbFieldType.DateTime:
+						return AsDateTime();
+					case CbFieldType.TimeSpan:
+						return AsTimeSpan();
+					default:
+						throw new NotImplementedException($"Unknown field type ({FieldType})");
+				}
+			}
+		}
+
+		/// <inheritdoc cref="Name"/>
+		public Utf8String GetName() => Name;
+
+		/// <summary>
+		/// Access the field as an object. Defaults to an empty object on error. 
+		/// </summary>
+		/// <returns></returns>
+		public CbObject AsObject()
+		{
+			if (CbFieldUtils.IsObject(TypeWithFlags))
+			{
+				Error = CbFieldError.None;
+				return CbObject.FromFieldNoCheck(this);
+			}
+			else
+			{
+				Error = CbFieldError.TypeError;
+				return CbObject.Empty;
+			}
+		}
+
+		/// <summary>
+		/// Access the field as an array. Defaults to an empty array on error. 
+		/// </summary>
+		/// <returns></returns>
+		public CbArray AsArray()
+		{
+			if (CbFieldUtils.IsArray(TypeWithFlags))
+			{
+				Error = CbFieldError.None;
+				return CbArray.FromFieldNoCheck(this);
+			}
+			else
+			{
+				Error = CbFieldError.TypeError;
+				return CbArray.Empty;
+			}
+		}
+
+		/// <summary>
+		/// Access the field as binary data.
+		/// </summary>
+		/// <returns></returns>
+		public ReadOnlyMemory<byte> AsBinary(ReadOnlyMemory<byte> Default = default)
+		{
+			if (CbFieldUtils.IsBinary(TypeWithFlags))
+			{
+				Error = CbFieldError.None;
+
+				ulong Length = VarInt.Read(Payload.Span, out int BytesRead);
+				return Payload.Slice(BytesRead, (int)Length);
+			}
+			else
+			{
+				Error = CbFieldError.TypeError;
+				return Default;
+			}
+		}
+
+		/// <summary>
+		/// Access the field as a UTF-8 string.
+		/// </summary>
+		/// <returns></returns>
+		public Utf8String AsString()
+		{
+			return AsString(default);
+		}
+
+		/// <summary>
+		/// Access the field as a UTF-8 string. Returns the provided default on error.
+		/// </summary>
+		/// <param name="Default">Default value to return</param>
+		/// <returns></returns>
+		public Utf8String AsString(Utf8String Default)
+		{
+			if (CbFieldUtils.IsString(TypeWithFlags))
+			{
+				ulong ValueSize = VarInt.Read(Payload.Span, out int ValueSizeByteCount);
+				if (ValueSize >= (1UL << 31))
+				{
+					Error = CbFieldError.RangeError;
+					return Default;
+				}
+				else
+				{
+					Error = CbFieldError.None;
+					return new Utf8String(Payload.Slice(ValueSizeByteCount, (int)ValueSize));
+				}
+			}
+			else
+			{
+				Error = CbFieldError.TypeError;
+				return Default;
+			}
+		}
+
+		/// <summary>
+		/// Access the field as an int8. Returns the provided default on error.
+		/// </summary>
+		public sbyte AsInt8(sbyte Default = 0)
+		{
+			return (sbyte)AsInteger((ulong)Default, 7, true);
+		}
+
+		/// <summary>
+		/// Access the field as an int16. Returns the provided default on error.
+		/// </summary>
+		public short AsInt16(short Default = 0)
+		{
+			return (short)AsInteger((ulong)Default, 15, true);
+		}
+
+		/// <summary>
+		/// Access the field as an int32. Returns the provided default on error.
+		/// </summary>
+		public int AsInt32()
+		{
+			return AsInt32(0);
+		}
+
+		/// <summary>
+		/// Access the field as an int32. Returns the provided default on error.
+		/// </summary>
+		public int AsInt32(int Default)
+		{
+			return (int)AsInteger((ulong)Default, 31, true);
+		}
+
+		/// <summary>
+		/// Access the field as an int64. Returns the provided default on error.
+		/// </summary>
+		public long AsInt64()
+		{
+			return AsInt64(0);
+		}
+
+		/// <summary>
+		/// Access the field as an int64. Returns the provided default on error.
+		/// </summary>
+		public long AsInt64(long Default)
+		{
+			return (long)AsInteger((ulong)Default, 63, true);
+		}
+
+		/// <summary>
+		/// Access the field as an int8. Returns the provided default on error.
+		/// </summary>
+		public byte AsUInt8(byte Default = 0)
+		{
+			return (byte)AsInteger(Default, 8, false);
+		}
+
+		/// <summary>
+		/// Access the field as an int16. Returns the provided default on error.
+		/// </summary>
+		public ushort AsUInt16(ushort Default = 0)
+		{
+			return (ushort)AsInteger(Default, 16, false);
+		}
+
+		/// <summary>
+		/// Access the field as an int32. Returns the provided default on error.
+		/// </summary>
+		public uint AsUInt32(uint Default = 0)
+		{
+			return (uint)AsInteger(Default, 32, false);
+		}
+
+		/// <summary>
+		/// Access the field as an int64. Returns the provided default on error.
+		/// </summary>
+		public ulong AsUInt64(ulong Default = 0)
+		{
+			return (ulong)AsInteger(Default, 64, false);
+		}
+
+		/// <summary>
+		/// Access the field as an integer, checking that it's in the correct range
+		/// </summary>
+		/// <param name="Default"></param>
+		/// <param name="MagnitudeBits"></param>
+		/// <param name="IsSigned"></param>
+		/// <returns></returns>
+		private ulong AsInteger(ulong Default, int MagnitudeBits, bool IsSigned)
+		{
+			if (CbFieldUtils.IsInteger(TypeWithFlags))
+			{
+				// A shift of a 64-bit value by 64 is undefined so shift by one less because magnitude is never zero.
+				ulong OutOfRangeMask = ~(ulong)1 << (MagnitudeBits - 1);
+				ulong IsNegative = (ulong)(byte)(TypeWithFlags) & 1;
+
+				int MagnitudeByteCount;
+				ulong Magnitude = VarInt.Read(Payload.Span, out MagnitudeByteCount);
+				ulong Value = Magnitude ^ (ulong)-(long)(IsNegative);
+
+				if ((Magnitude & OutOfRangeMask) == 0 && (IsNegative == 0 || IsSigned))
+				{
+					Error = CbFieldError.None;
+					return Value;
+				}
+				else
+				{
+					Error = CbFieldError.RangeError;
+					return Default;
+				}
+			}
+			else
+			{
+				Error = CbFieldError.TypeError;
+				return Default;
+			}
+		}
+
+		/// <summary>
+		/// Access the field as a float. Returns the provided default on error.
+		/// </summary>
+		/// <param name="Default">Default value</param>
+		/// <returns>Value of the field</returns>
+		public float AsFloat(float Default = 0.0f)
+		{
+			switch (GetType())
+			{
+				case CbFieldType.IntegerPositive:
+				case CbFieldType.IntegerNegative:
+					{
+						ulong IsNegative = (ulong)TypeWithFlags & 1;
+						ulong OutOfRangeMask = ~((1UL << /*FLT_MANT_DIG*/ 24) - 1);
+
+						int MagnitudeByteCount;
+						ulong Magnitude = VarInt.Read(Payload.Span, out MagnitudeByteCount) + IsNegative;
+						bool IsInRange = (Magnitude & OutOfRangeMask) == 0;
+						Error = IsInRange ? CbFieldError.None : CbFieldError.RangeError;
+						return IsInRange ? (float)((IsNegative != 0) ? (float)-(long)Magnitude : (float)Magnitude) : Default;
+					}
+				case CbFieldType.Float32:
+					Error = CbFieldError.None;
+					return BitConverter.Int32BitsToSingle(BinaryPrimitives.ReadInt32BigEndian(Payload.Span));
+				case CbFieldType.Float64:
+					Error = CbFieldError.RangeError;
+					return Default;
+				default:
+					Error = CbFieldError.TypeError;
+					return Default;
+			}
+		}
+
+		/// <summary>
+		/// Access the field as a double.
+		/// </summary>
+		/// <returns>Value of the field</returns>
+		public double AsDouble() => AsDouble(0.0);
+
+		/// <summary>
+		/// Access the field as a double. Returns the provided default on error.
+		/// </summary>
+		/// <param name="Default">Default value</param>
+		/// <returns>Value of the field</returns>
+		public double AsDouble(double Default)
+		{
+			switch (GetType())
+			{
+				case CbFieldType.IntegerPositive:
+				case CbFieldType.IntegerNegative:
+					{
+						ulong IsNegative = (ulong)TypeWithFlags & 1;
+						ulong OutOfRangeMask = ~((1UL << /*DBL_MANT_DIG*/ 53) - 1);
+
+						int MagnitudeByteCount;
+						ulong Magnitude = VarInt.Read(Payload.Span, out MagnitudeByteCount) + IsNegative;
+						bool IsInRange = (Magnitude & OutOfRangeMask) == 0;
+						Error = IsInRange ? CbFieldError.None : CbFieldError.RangeError;
+						return IsInRange ? (double)((IsNegative != 0) ? (double)-(long)Magnitude : (double)Magnitude) : Default;
+					}
+				case CbFieldType.Float32:
+					{
+						Error = CbFieldError.None;
+						return BitConverter.Int32BitsToSingle(BinaryPrimitives.ReadInt32BigEndian(Payload.Span));
+					}
+				case CbFieldType.Float64:
+					{
+						Error = CbFieldError.None;
+						return BitConverter.Int64BitsToDouble(BinaryPrimitives.ReadInt64BigEndian(Payload.Span));
+					}
+				default:
+					Error = CbFieldError.TypeError;
+					return Default;
+			}
+		}
+
+		/// <summary>
+		/// Access the field as a bool. Returns the provided default on error.
+		/// </summary>
+		/// <returns>Value of the field</returns>
+		public bool AsBool() => AsBool(false);
+
+		/// <summary>
+		/// Access the field as a bool. Returns the provided default on error.
+		/// </summary>
+		/// <param name="Default">Default value</param>
+		/// <returns>Value of the field</returns>
+		public bool AsBool(bool Default)
+		{
+			switch (GetType())
+			{
+				case CbFieldType.BoolTrue:
+					Error = CbFieldError.None;
+					return true;
+				case CbFieldType.BoolFalse:
+					Error = CbFieldError.None;
+					return false;
+				default:
+					Error = CbFieldError.TypeError;
+					return Default;
+			}
+		}
+
+		/// <summary>
+		/// Access the field as a hash referencing an object attachment. Returns the provided default on error.
+		/// </summary>
+		/// <returns>Value of the field</returns>
+		public CbObjectAttachment AsObjectAttachment() => AsObjectAttachment(CbObjectAttachment.Zero);
+
+		/// <summary>
+		/// Access the field as a hash referencing an object attachment. Returns the provided default on error.
+		/// </summary>
+		/// <param name="Default">Default value</param>
+		/// <returns>Value of the field</returns>
+		public CbObjectAttachment AsObjectAttachment(CbObjectAttachment Default)
+		{
+			if (CbFieldUtils.IsObjectAttachment(TypeWithFlags))
+			{
+				Error = CbFieldError.None;
+				return new IoHash(Payload);
+			}
+			else
+			{
+				Error = CbFieldError.TypeError;
+				return Default;
+			}
+		}
+
+		/// <summary>
+		/// Access the field as a hash referencing a binary attachment. Returns the provided default on error.
+		/// </summary>
+		/// <param name="Default">Default value</param>
+		/// <returns>Value of the field</returns>
+		public CbBinaryAttachment AsBinaryAttachment() => AsBinaryAttachment(CbBinaryAttachment.Zero);
+
+		/// <summary>
+		/// Access the field as a hash referencing a binary attachment. Returns the provided default on error.
+		/// </summary>
+		/// <param name="Default">Default value</param>
+		/// <returns>Value of the field</returns>
+		public CbBinaryAttachment AsBinaryAttachment(CbBinaryAttachment Default)
+		{
+			if (CbFieldUtils.IsBinaryAttachment(TypeWithFlags))
+			{
+				Error = CbFieldError.None;
+				return new IoHash(Payload);
+			}
+			else
+			{
+				Error = CbFieldError.TypeError;
+				return Default;
+			}
+		}
+
+		/// <summary>
+		/// Access the field as a hash referencing an attachment. Returns the provided default on error.
+		/// </summary>
+		/// <param name="Default">Default value</param>
+		/// <returns>Value of the field</returns>
+		public IoHash AsAttachment() => AsAttachment(IoHash.Zero);
+
+		/// <summary>
+		/// Access the field as a hash referencing an attachment. Returns the provided default on error.
+		/// </summary>
+		/// <param name="Default">Default value</param>
+		/// <returns>Value of the field</returns>
+		public IoHash AsAttachment(IoHash Default)
+		{
+			if (CbFieldUtils.IsAttachment(TypeWithFlags))
+			{
+				Error = CbFieldError.None;
+				return new IoHash(Payload);
+			}
+			else
+			{
+				Error = CbFieldError.TypeError;
+				return Default;
+			}
+		}
+
+		/// <summary>
+		/// Access the field as a hash referencing an attachment. Returns the provided default on error.
+		/// </summary>
+		/// <returns>Value of the field</returns>
+		public IoHash AsHash() => AsHash(IoHash.Zero);
+
+		/// <summary>
+		/// Access the field as a hash referencing an attachment. Returns the provided default on error.
+		/// </summary>
+		/// <param name="Default">Default value</param>
+		/// <returns>Value of the field</returns>
+		public IoHash AsHash(IoHash Default)
+		{
+			if (CbFieldUtils.IsHash(TypeWithFlags))
+			{
+				Error = CbFieldError.None;
+				return new IoHash(Payload);
+			}
+			else
+			{
+				Error = CbFieldError.TypeError;
+				return Default;
+			}
+		}
+
+		/// <summary>
+		/// Access the field as a UUID. Returns a nil UUID on error.
+		/// </summary>
+		/// <param name="Default">Default value</param>
+		/// <returns>Value of the field</returns>
+		public Guid AsUuid(Guid Default = default(Guid))
+		{
+			if (CbFieldUtils.IsUuid(TypeWithFlags))
+			{
+				Error = CbFieldError.None;
+
+				ReadOnlySpan<byte> Span = Payload.Span;
+				uint A = BinaryPrimitives.ReadUInt32BigEndian(Span);
+				ushort B = BinaryPrimitives.ReadUInt16BigEndian(Span.Slice(4));
+				ushort C = BinaryPrimitives.ReadUInt16BigEndian(Span.Slice(6));
+
+				return new Guid(A, B, C, Span[8], Span[9], Span[10], Span[11], Span[12], Span[13], Span[14], Span[15]);
+			}
+			else
+			{
+				Error = CbFieldError.TypeError;
+				return Default;
+			}
+		}
+
+		/// <summary>
+		/// Reads a date time as number of ticks from the stream
+		/// </summary>
+		/// <param name="Default"></param>
+		/// <returns></returns>
+		public long AsDateTimeTicks(long Default = 0)
+		{
+			if (CbFieldUtils.IsDateTime(TypeWithFlags))
+			{
+				Error = CbFieldError.None;
+				return BinaryPrimitives.ReadInt64BigEndian(Payload.Span);
+			}
+			else
+			{
+				Error = CbFieldError.TypeError;
+				return Default;
+			}
+		}
+
+		/// <summary>
+		/// Access the field as a DateTime.
+		/// </summary>
+		/// <returns></returns>
+		public DateTime AsDateTime()
+		{
+			return AsDateTime(new DateTime(0, DateTimeKind.Utc));
+		}
+
+		/// <summary>
+		/// Access the field as a DateTime.
+		/// </summary>
+		/// <param name="Default"></param>
+		/// <returns></returns>
+		public DateTime AsDateTime(DateTime Default)
+		{
+			return new DateTime(AsDateTimeTicks(Default.ToUniversalTime().Ticks), DateTimeKind.Utc);
+		}
+
+		/// <summary>
+		/// Reads a timespan as number of ticks from the stream
+		/// </summary>
+		/// <param name="Default"></param>
+		/// <returns></returns>
+		public long AsTimeSpanTicks(long Default = 0)
+		{
+			if (CbFieldUtils.IsTimeSpan(TypeWithFlags))
+			{
+				Error = CbFieldError.None;
+				return BinaryPrimitives.ReadInt64BigEndian(Payload.Span);
+			}
+			else
+			{
+				Error = CbFieldError.TypeError;
+				return Default;
+			}
+		}
+
+		/// <summary>
+		/// Reads a timespan as number of ticks from the stream
+		/// </summary>
+		/// <param name="Default"></param>
+		/// <returns></returns>
+		public TimeSpan AsTimeSpan(TimeSpan Default = default) => new TimeSpan(AsTimeSpanTicks(Default.Ticks));
+
+		/// <inheritdoc cref="CbFieldUtils.HasFieldName(CbFieldType)"/>
+		public bool HasName() => CbFieldUtils.HasFieldName(TypeWithFlags);
+
+		/// <inheritdoc cref="CbFieldUtils.IsNull(CbFieldType)"/>
+		public bool IsNull() => CbFieldUtils.IsNull(TypeWithFlags);
+
+		/// <inheritdoc cref="CbFieldUtils.IsObject(CbFieldType)"/>
+		public bool IsObject() => CbFieldUtils.IsObject(TypeWithFlags);
+
+		/// <inheritdoc cref="CbFieldUtils.IsArray(CbFieldType)"/>
+		public bool IsArray() => CbFieldUtils.IsArray(TypeWithFlags);
+
+		/// <inheritdoc cref="CbFieldUtils.IsBinary(CbFieldType)"/>
+		public bool IsBinary() => CbFieldUtils.IsBinary(TypeWithFlags);
+
+		/// <inheritdoc cref="CbFieldUtils.IsString(CbFieldType)"/>
+		public bool IsString() => CbFieldUtils.IsString(TypeWithFlags);
+
+		/// <inheritdoc cref="CbFieldUtils.IsInteger(CbFieldType)"/>
+		public bool IsInteger() => CbFieldUtils.IsInteger(TypeWithFlags);
+
+		/// <inheritdoc cref="CbFieldUtils.IsFloat(CbFieldType)"/>
+		public bool IsFloat() => CbFieldUtils.IsFloat(TypeWithFlags);
+
+		/// <inheritdoc cref="CbFieldUtils.IsBool(CbFieldType)"/>
+		public bool IsBool() => CbFieldUtils.IsBool(TypeWithFlags);
+
+		/// <inheritdoc cref="CbFieldUtils.IsObjectAttachment(CbFieldType)"/>
+		public bool IsObjectAttachment() => CbFieldUtils.IsObjectAttachment(TypeWithFlags);
+
+		/// <inheritdoc cref="CbFieldUtils.IsBinaryAttachment(CbFieldType)"/>
+		public bool IsBinaryAttachment() => CbFieldUtils.IsBinaryAttachment(TypeWithFlags);
+
+		/// <inheritdoc cref="CbFieldUtils.IsAttachment(CbFieldType)"/>
+		public bool IsAttachment() => CbFieldUtils.IsAttachment(TypeWithFlags);
+
+		/// <inheritdoc cref="CbFieldUtils.IsHash(CbFieldType)"/>
+		public bool IsHash() => CbFieldUtils.IsHash(TypeWithFlags);
+
+		/// <inheritdoc cref="CbFieldUtils.IsUuid(CbFieldType)"/>
+		public bool IsUuid() => CbFieldUtils.IsUuid(TypeWithFlags);
+
+		/// <inheritdoc cref="CbFieldUtils.IsDateTime(CbFieldType)"/>
+		public bool IsDateTime() => CbFieldUtils.IsDateTime(TypeWithFlags);
+
+		/// <inheritdoc cref="CbFieldUtils.IsTimeSpan(CbFieldType)"/>
+		public bool IsTimeSpan() => CbFieldUtils.IsTimeSpan(TypeWithFlags);
+
+		/// <summary>
+		/// Whether the field has a value
+		/// </summary>
+		/// <param name="Field"></param>
+		public static explicit operator bool(CbField Field) => Field.HasValue();
+
+		/// <summary>
+		/// Whether the field has a value.
+		///
+		/// All fields in a valid object or array have a value. A field with no value is returned when
+		/// finding a field by name fails or when accessing an iterator past the end.
+		/// </summary>
+		public bool HasValue() => !CbFieldUtils.IsNone(TypeWithFlags);
+
+		/// <summary>
+		/// Whether the last field access encountered an error.
+		/// </summary>
+		public bool HasError() => Error != CbFieldError.None;
+
+		/// <inheritdoc cref="Error"/>
+		public CbFieldError GetError() => Error;
+
+		/// <summary>
+		/// Returns the size of the field in bytes, including the type and name
+		/// </summary>
+		/// <returns></returns>
+		public int GetSize() => sizeof(CbFieldType) + GetViewNoType().Length;
+
+		/// <summary>
+		/// Calculate the hash of the field, including the type and name.
+		/// </summary>
+		/// <returns></returns>
+		public Blake3Hash GetHash()
+		{
+			using (Blake3.Hasher Hasher = Blake3.Hasher.New())
+			{
+				AppendHash(Hasher);
+
+				byte[] Hash = new byte[32];
+				Hasher.Finalize(Hash);
+
+				return new Blake3Hash(Hash);
+			}
+		}
+
+		/// <summary>
+		/// Append the hash of the field, including the type and name
+		/// </summary>
+		/// <param name="Hasher"></param>
+		void AppendHash(Blake3.Hasher Hasher)
+		{
+			Span<byte> Data = stackalloc byte[1];
+			Data[0] = (byte)CbFieldUtils.GetSerializedType(TypeWithFlags);
+			Hasher.Update(Data);
+
+			Hasher.Update(GetViewNoType().Span);
+		}
+
+		/// <summary>
+		/// Whether this field is identical to the other field.
+		/// 
+		/// Performs a deep comparison of any contained arrays or objects and their fields. Comparison
+		/// assumes that both fields are valid and are written in the canonical format. Fields must be
+		/// written in the same order in arrays and objects, and name comparison is case sensitive. If
+		/// these assumptions do not hold, this may return false for equivalent inputs. Validation can
+		/// be performed with ValidateCompactBinary, except for field order and field name case.
+		/// </summary>
+		/// <param name="Other"></param>
+		/// <returns></returns>
+		public bool Equals(CbField? Other)
+		{
+			return Other != null && CbFieldUtils.GetSerializedType(TypeWithFlags) == CbFieldUtils.GetSerializedType(Other.TypeWithFlags) && GetViewNoType().Span.SequenceEqual(Other.GetViewNoType().Span);
+		}
+
+		/// <summary>
+		/// Copy the field into a buffer of exactly GetSize() bytes, including the type and name.
+		/// </summary>
+		/// <param name="Buffer"></param>
+		public void CopyTo(Span<byte> Buffer)
+		{
+			Buffer[0] = (byte)CbFieldUtils.GetSerializedType(TypeWithFlags);
+			GetViewNoType().Span.CopyTo(Buffer.Slice(1));
+		}
+
+		/// <summary>
+		/// Invoke the visitor for every attachment in the field.
+		/// </summary>
+		/// <param name="Visitor"></param>
+		public void IterateAttachments(Action<CbField> Visitor)
+		{
+			switch (GetType())
+			{
+				case CbFieldType.Object:
+				case CbFieldType.UniformObject:
+					CbObject.FromFieldNoCheck(this).IterateAttachments(Visitor);
+					break;
+				case CbFieldType.Array:
+				case CbFieldType.UniformArray:
+					CbArray.FromFieldNoCheck(this).IterateAttachments(Visitor);
+					break;
+				case CbFieldType.ObjectAttachment:
+				case CbFieldType.BinaryAttachment:
+					Visitor(this);
+					break;
+			}
+		}
+
+		/// <summary>
+		/// Try to get a view of the field as it would be serialized, such as by CopyTo.
+		///
+		/// A view is available if the field contains its type. Access the equivalent for other fields
+		/// through FCbField::GetBuffer, FCbField::Clone, or CopyTo.
+		/// </summary>
+		/// <param name="OutView"></param>
+		/// <returns></returns>
+		public bool TryGetView(out ReadOnlyMemory<byte> OutView)
+		{
+			if (CbFieldUtils.HasFieldType(TypeWithFlags))
+			{
+				OutView = Memory;
+				return true;
+			}
+			else
+			{
+				OutView = ReadOnlyMemory<byte>.Empty;
+				return false;
+			}
+		}
+
+		/// <summary>
+		/// Find a field of an object by case-sensitive name comparison, otherwise a field with no value.
+		/// </summary>
+		/// <param name="Name"></param>
+		/// <returns></returns>
+		public CbField this[Utf8String Name]
+		{
+			get { return this.FirstOrDefault(Field => Field.Name == Name) ?? CbField.Empty; }
+		}
+
+		/// <summary>
+		/// Create an iterator for the fields of an array or object, otherwise an empty iterator.
+		/// </summary>
+		/// <returns></returns>
+		public CbFieldIterator CreateIterator()
+		{
+			CbFieldType LocalTypeWithFlags = TypeWithFlags;
+			if (CbFieldUtils.HasFields(LocalTypeWithFlags))
+			{
+				ReadOnlyMemory<byte> PayloadBytes = Payload;
+				int PayloadSizeByteCount;
+				int PayloadSize = (int)VarInt.Read(PayloadBytes.Span, out PayloadSizeByteCount);
+				PayloadBytes = PayloadBytes.Slice(PayloadSizeByteCount);
+				int NumByteCount = CbFieldUtils.IsArray(LocalTypeWithFlags) ? (int)VarInt.Measure(PayloadBytes.Span) : 0;
+				if (PayloadSize > NumByteCount)
+				{
+					PayloadBytes = PayloadBytes.Slice(NumByteCount);
+
+					CbFieldType UniformType = CbFieldType.HasFieldType;
+					if (CbFieldUtils.HasUniformFields(TypeWithFlags))
+					{
+						UniformType = (CbFieldType)PayloadBytes.Span[0];
+						PayloadBytes = PayloadBytes.Slice(1);
+					}
+
+					return new CbFieldIterator(PayloadBytes, UniformType);
+				}
+			}
+			return new CbFieldIterator(ReadOnlyMemory<byte>.Empty, CbFieldType.HasFieldType);
+		}
+
+		/// <inheritdoc/>
+		public IEnumerator<CbField> GetEnumerator()
+		{
+			for (CbFieldIterator Iter = CreateIterator(); Iter; Iter.MoveNext())
+			{
+				yield return Iter.Current;
+			}
+		}
+
+		/// <inheritdoc/>
+		IEnumerator IEnumerable.GetEnumerator() => GetEnumerator();
+
+		/// <summary>
+		/// Returns a view of the name and value payload, which excludes the type.
+		/// </summary>
+		/// <returns></returns>
+		private ReadOnlyMemory<byte> GetViewNoType()
+		{
+			int NameSize = CbFieldUtils.HasFieldName(TypeWithFlags) ? NameLen + (int)VarInt.Measure((uint)NameLen) : 0;
+			return Memory.Slice(PayloadOffset - NameSize);
+		}
+
+		/// <summary>
+		/// Accessor for the payload
+		/// </summary>
+		internal ReadOnlyMemory<byte> Payload => Memory.Slice(PayloadOffset);
+
+		/// <summary>
+		/// Returns a view of the value payload, which excludes the type and name.
+		/// </summary>
+		/// <returns></returns>
+		internal ReadOnlyMemory<byte> GetPayloadView() => Memory.Slice(PayloadOffset);
+
+		/// <summary>
+		/// Returns the type of the field excluding flags.
+		/// </summary>
+		internal new CbFieldType GetType() => CbFieldUtils.GetType(TypeWithFlags);
+
+		/// <summary>
+		/// Returns the type of the field excluding flags.
+		/// </summary>
+		internal CbFieldType GetTypeWithFlags() => TypeWithFlags;
+
+		/// <summary>
+		/// Returns the size of the value payload in bytes, which is the field excluding the type and name.
+		/// </summary>
+		/// <param name="Type"></param>
+		/// <param name="Payload"></param>
+		/// <returns></returns>
+		public ulong GetPayloadSize()
+		{
+			switch (GetType())
+			{
+				case CbFieldType.None:
+				case CbFieldType.Null:
+					return 0;
+				case CbFieldType.Object:
+				case CbFieldType.UniformObject:
+				case CbFieldType.Array:
+				case CbFieldType.UniformArray:
+				case CbFieldType.Binary:
+				case CbFieldType.String:
+					{
+						ulong PayloadSize = VarInt.Read(Payload.Span, out int BytesRead);
+						return PayloadSize + (ulong)BytesRead;
+					}
+				case CbFieldType.IntegerPositive:
+				case CbFieldType.IntegerNegative:
+					{
+						return (ulong)VarInt.Measure(Payload.Span);
+					}
+				case CbFieldType.Float32:
+					return 4;
+				case CbFieldType.Float64:
+					return 8;
+				case CbFieldType.BoolFalse:
+				case CbFieldType.BoolTrue:
+					return 0;
+				case CbFieldType.ObjectAttachment:
+				case CbFieldType.BinaryAttachment:
+				case CbFieldType.Hash:
+					return 20;
+				case CbFieldType.Uuid:
+					return 16;
+				case CbFieldType.DateTime:
+				case CbFieldType.TimeSpan:
+					return 8;
+				case CbFieldType.ObjectId:
+					return 12;
+				default:
+					return 0;
+			}
+		}
+
+		#region Mimic inheritance from TCbBufferFactory
+
+		public static CbField Clone(ReadOnlyMemory<byte> Data) => Clone(new CbField(Data));
+		public static CbField Clone(CbField Other) => Other;
+		public static CbField MakeView(ReadOnlyMemory<byte> Data) => new CbField(Data);
+		public static CbField MakeView(CbField Other) => Other;
+
+		#endregion
+	}
+
+	/// <summary>
+	/// Converter to and from JSON objects
+	/// </summary>
+	public class CbFieldJsonConverter : JsonConverter<CbField>
+	{
+		/// <inheritdoc/>
+		public override CbField Read(ref Utf8JsonReader reader, Type typeToConvert, JsonSerializerOptions options)
+		{
+			throw new NotImplementedException();
+		}
+
+		/// <inheritdoc/>
+		public override void Write(Utf8JsonWriter Writer, CbField Field, JsonSerializerOptions Options)
+		{
+			switch (Field.GetType())
+			{
+				case CbFieldType.Null:
+					if (Field.HasName())
+					{
+						Writer.WriteNull(Field.Name.Span);
+					}
+					else
+					{
+						Writer.WriteNullValue();
+					}
+					break;
+				case CbFieldType.Object:
+				case CbFieldType.UniformObject:
+					if (Field.HasName())
+					{
+						Writer.WriteStartObject(Field.Name.Span);
+					}
+					else
+					{
+						Writer.WriteStartObject();
+					}
+
+					foreach (CbField Member in Field.AsObject())
+					{
+						Write(Writer, Member, Options);
+					}
+					Writer.WriteEndObject();
+					break;
+				case CbFieldType.Array:
+				case CbFieldType.UniformArray:
+					if (Field.HasName())
+					{
+						Writer.WriteStartArray(Field.Name.Span);
+					}
+					else
+					{
+						Writer.WriteStartArray();
+					}
+
+					foreach (CbField Element in Field.AsArray())
+					{
+						Write(Writer, Element, Options);
+					}
+					Writer.WriteEndArray();
+					break;
+				case CbFieldType.Binary:
+					if (Field.HasName())
+					{
+						Writer.WriteBase64String(Field.Name.Span, Field.AsBinary().Span);
+					}
+					else
+					{
+						Writer.WriteBase64StringValue(Field.AsBinary().Span);
+					}
+					break;
+				case CbFieldType.String:
+					if (Field.HasName())
+					{
+						Writer.WriteString(Field.Name.Span, Field.AsString().Span);
+					}
+					else
+					{
+						Writer.WriteStringValue(Field.AsString().Span);
+					}
+					break;
+				case CbFieldType.IntegerPositive:
+					if (Field.HasName())
+					{
+						Writer.WriteNumber(Field.Name.Span, Field.AsUInt64());
+					}
+					else
+					{
+						Writer.WriteNumberValue(Field.AsUInt64());
+					}
+					break;
+				case CbFieldType.IntegerNegative:
+					if (Field.HasName())
+					{
+						Writer.WriteNumber(Field.Name.Span, Field.AsInt64());
+					}
+					else
+					{
+						Writer.WriteNumberValue(Field.AsInt64());
+					}
+					break;
+				case CbFieldType.Float32:
+				case CbFieldType.Float64:
+					if (Field.HasName())
+					{
+						Writer.WriteNumber(Field.Name.Span, Field.AsDouble());
+					}
+					else
+					{
+						Writer.WriteNumberValue(Field.AsDouble());
+					}
+					break;
+				case CbFieldType.BoolFalse:
+				case CbFieldType.BoolTrue:
+					if (Field.HasName())
+					{
+						Writer.WriteBoolean(Field.Name.Span, Field.AsBool());
+					}
+					else
+					{
+						Writer.WriteBooleanValue(Field.AsBool());
+					}
+					break;
+				case CbFieldType.ObjectAttachment:
+				case CbFieldType.BinaryAttachment:
+				case CbFieldType.Hash:
+					if (Field.HasName())
+					{
+						Writer.WriteString(Field.Name.Span, Field.AsHash().ToString());
+					}
+					else
+					{
+						Writer.WriteStringValue(Field.AsHash().ToString());
+					}
+					break;
+				case CbFieldType.Uuid:
+					if (Field.HasName())
+					{
+						Writer.WriteString(Field.Name.Span, Field.AsUuid().ToString());
+					}
+					else
+					{
+						Writer.WriteStringValue(Field.AsUuid().ToString());
+					}
+					break;
+				case CbFieldType.DateTime:
+					if (Field.HasName())
+					{
+						Writer.WriteNumber(Field.Name.Span, Field.AsDateTimeTicks());
+					}
+					else
+					{
+						Writer.WriteNumberValue(Field.AsDateTimeTicks());
+					}
+					break;
+				case CbFieldType.TimeSpan:
+					if (Field.HasName())
+					{
+						Writer.WriteNumber(Field.Name.Span, Field.AsTimeSpanTicks());
+					}
+					else
+					{
+						Writer.WriteNumberValue(Field.AsTimeSpanTicks());
+					}
+					break;
+				default:
+					throw new NotImplementedException($"Unhandled type in cb-json converter");
+			}
+		}
+	}
+
+	public class CbFieldEnumerator : IEnumerator<CbField>
+	{
+		/// <summary>
+		/// The underlying buffer
+		/// </summary>
+		ReadOnlyMemory<byte> Data;
+
+		/// <summary>
+		/// Type for all fields
+		/// </summary>
+		CbFieldType UniformType { get; }
+
+		/// <inheritdoc/>
+		public CbField Current { get; private set; } = null!;
+
+		/// <inheritdoc/>
+		object? IEnumerator.Current => Current;
+
+		/// <summary>
+		/// Constructor
+		/// </summary>
+		/// <param name="Data"></param>
+		/// <param name="UniformType"></param>
+		public CbFieldEnumerator(ReadOnlyMemory<byte> Data, CbFieldType UniformType)
+		{
+			this.Data = Data;
+			this.UniformType = UniformType;
+		}
+
+		/// <inheritdoc/>
+		public void Dispose()
+		{
+		}
+
+		/// <inheritdoc/>
+		public void Reset()
+		{
+			throw new InvalidOperationException();
+		}
+
+		/// <inheritdoc/>
+		public bool MoveNext()
+		{
+			if (Data.Length > 0)
+			{
+				Current = new CbField(Data, UniformType);
+				return true;
+			}
+			else
+			{
+				Current = null!;
+				return false;
+			}
+		}
+
+		/// <summary>
+		/// Clone this enumerator
+		/// </summary>
+		/// <returns></returns>
+		public CbFieldEnumerator Clone()
+		{
+			return new CbFieldEnumerator(Data, UniformType);
+		}
+	}
+
+	/// <summary>
+	/// Iterator for fields
+	/// </summary>
+	public class CbFieldIterator
+	{
+		/// <summary>
+		/// The underlying buffer
+		/// </summary>
+		ReadOnlyMemory<byte> NextData;
+
+		/// <summary>
+		/// Type for all fields
+		/// </summary>
+		CbFieldType UniformType;
+
+		/// <summary>
+		/// The current iterator
+		/// </summary>
+		public CbField Current { get; private set; } = null!;
+
+		/// <summary>
+		/// Default constructor
+		/// </summary>
+		public CbFieldIterator()
+			: this(ReadOnlyMemory<byte>.Empty, CbFieldType.HasFieldType)
+		{
+		}
+
+		/// <summary>
+		/// Constructor for single field iterator
+		/// </summary>
+		/// <param name="Field"></param>
+		private CbFieldIterator(CbField Field)
+		{
+			NextData = ReadOnlyMemory<byte>.Empty;
+			Current = Field;
+		}
+
+		/// <summary>
+		/// Constructor
+		/// </summary>
+		/// <param name="Data"></param>
+		/// <param name="UniformType"></param>
+		public CbFieldIterator(ReadOnlyMemory<byte> Data, CbFieldType UniformType)
+		{
+			this.NextData = Data;
+			this.UniformType = UniformType;
+
+			MoveNext();
+		}
+
+		/// <summary>
+		/// Copy constructor
+		/// </summary>
+		/// <param name="Other"></param>
+		public CbFieldIterator(CbFieldIterator Other)
+		{
+			this.NextData = Other.NextData;
+			this.UniformType = Other.UniformType;
+			this.Current = Other.Current;
+		}
+
+		/// <summary>
+		/// Construct a field range that contains exactly one field.
+		/// </summary>
+		/// <param name="Field"></param>
+		/// <returns></returns>
+		public static CbFieldIterator MakeSingle(CbField Field)
+		{
+			return new CbFieldIterator(Field);
+		}
+
+		/// <summary>
+		/// Construct a field range from a buffer containing zero or more valid fields.
+		/// </summary>
+		/// <param name="View">A buffer containing zero or more valid fields.</param>
+		/// <param name="Type">HasFieldType means that View contains the type.Otherwise, use the given type.</param>
+		/// <returns></returns>
+		public static CbFieldIterator MakeRange(ReadOnlyMemory<byte> View, CbFieldType Type = CbFieldType.HasFieldType)
+		{
+			return new CbFieldIterator(View, Type);
+		}
+
+		/// <summary>
+		/// Check if the current value is valid
+		/// </summary>
+		/// <returns></returns>
+		public bool IsValid()
+		{
+			return Current.GetType() != CbFieldType.None;
+		}
+
+		/// <summary>
+		/// Accessor for the current value
+		/// </summary>
+		/// <returns></returns>
+		public CbField GetCurrent()
+		{
+			return Current;
+		}
+
+		/// <summary>
+		/// Copy the field range into a buffer of exactly GetRangeSize() bytes.
+		/// </summary>
+		/// <param name="Buffer"></param>
+		public void CopyRangeTo(Span<byte> Buffer)
+		{
+			ReadOnlyMemory<byte> Source;
+			if (TryGetRangeView(out Source))
+			{
+				Source.Span.CopyTo(Buffer);
+			}
+			else
+			{
+				for (CbFieldIterator It = new CbFieldIterator(this); It; It.MoveNext())
+				{
+					int Size = It.Current.GetSize();
+					It.Current.CopyTo(Buffer);
+					Buffer = Buffer.Slice(Size);
+				}
+			}
+		}
+
+		/// <summary>
+		/// Invoke the visitor for every attachment in the field range.
+		/// </summary>
+		/// <param name="Visitor"></param>
+		public void IterateRangeAttachments(Action<CbField> Visitor)
+		{
+			// Always iterate over non-uniform ranges because we do not know if they contain an attachment.
+			if (CbFieldUtils.HasFieldType(Current.GetTypeWithFlags()))
+			{
+				for (CbFieldIterator It = new CbFieldIterator(this); It; ++It)
+				{
+					if (CbFieldUtils.MayContainAttachments(It.Current.GetTypeWithFlags()))
+					{
+						It.Current.IterateAttachments(Visitor);
+					}
+				}
+			}
+			// Only iterate over uniform ranges if the uniform type may contain an attachment.
+			else
+			{
+				if (CbFieldUtils.MayContainAttachments(Current.GetTypeWithFlags()))
+				{
+					for (CbFieldIterator It = new CbFieldIterator(this); It; ++It)
+					{
+						It.Current.IterateAttachments(Visitor);
+					}
+				}
+			}
+		}
+
+		/// <summary>
+		/// Try to get a view of every field in the range as they would be serialized.
+		///
+		/// A view is available if each field contains its type. Access the equivalent for other field
+		/// ranges through FCbFieldIterator::CloneRange or CopyRangeTo.
+		/// </summary>
+		/// <returns></returns>
+		bool TryGetRangeView(out ReadOnlyMemory<byte> OutView)
+		{
+			throw new NotImplementedException();
+/*			FMemoryView View;
+			if (FieldType::TryGetView(View))
+			{
+				OutView = MakeMemoryView(View.GetData(), FieldsEnd);
+				return true;
+			}
+			return false;*/
+		}
+
+		/// <summary>
+		/// Move to the next element
+		/// </summary>
+		/// <returns></returns>
+		public bool MoveNext()
+		{
+			if (NextData.Length > 0)
+			{
+				Current = new CbField(NextData, UniformType);
+				NextData = NextData.Slice(Current.Memory.Length);
+				return true;
+			}
+			else
+			{
+				Current = CbField.Empty;
+				return false;
+			}
+		}
+
+		/// <summary>
+		/// Test whether the iterator is valid
+		/// </summary>
+		/// <param name="Iterator"></param>
+		public static implicit operator bool(CbFieldIterator Iterator)
+		{
+			return Iterator.IsValid();
+		}
+
+		/// <summary>
+		/// Move to the next item
+		/// </summary>
+		/// <param name="Iterator"></param>
+		/// <returns></returns>
+		public static CbFieldIterator operator ++(CbFieldIterator Iterator)
+		{
+			return new CbFieldIterator(Iterator.NextData, Iterator.UniformType);
+		}
+
+		public override bool Equals(object? obj)
+		{
+			throw new NotImplementedException();
+		}
+
+		public override int GetHashCode()
+		{
+			throw new NotImplementedException();
+		}
+
+		public static bool operator ==(CbFieldIterator A, CbFieldIterator B)
+		{
+			return A.Current.Equals(B.Current);
+		}
+		public static bool operator !=(CbFieldIterator A, CbFieldIterator B)
+		{
+			return !A.Current.Equals(B.Current);
+		}
+	}
+
+	/// <summary>
+	/// Simplified view of <see cref="CbArray"/> for display in the debugger
+	/// </summary>
+	class CbArrayDebugView
+	{
+		CbArray Array;
+
+		public CbArrayDebugView(CbArray Array) => this.Array = Array;
+
+		[DebuggerBrowsable(DebuggerBrowsableState.RootHidden)]
+		public object?[] Value => Array.Select(x => x.Value).ToArray();
+	}
+
+	/// <summary>
+	/// Array of CbField that have no names.
+	///
+	/// Accessing a field of the array requires iteration. Access by index is not provided because the
+	/// cost of accessing an item by index scales linearly with the index.
+	/// </summary>
+	[DebuggerDisplay("Count = {Count}")]
+	[DebuggerTypeProxy(typeof(CbArrayDebugView))]
+	public class CbArray : IEnumerable<CbField>
+	{
+		/// <summary>
+		/// The field containing this array
+		/// </summary>
+		readonly CbField InnerField;
+
+		/// <summary>
+		/// Empty array constant
+		/// </summary>
+		public static CbArray Empty { get; } = new CbArray(new byte[] { (byte)CbFieldType.Array, 1, 0 });
+
+		/// <summary>
+		/// Construct an array with no fields
+		/// </summary>
+		public CbArray()
+		{
+			InnerField = Empty.InnerField;
+		}
+
+		/// <summary>
+		/// Constructor
+		/// </summary>
+		/// <param name="Field"></param>
+		private CbArray(CbField Field)
+		{
+			InnerField = Field;
+		}
+
+		/// <summary>
+		/// Constructor
+		/// </summary>
+		/// <param name="Data"></param>
+		/// <param name="Type"></param>
+		public CbArray(ReadOnlyMemory<byte> Data, CbFieldType Type = CbFieldType.HasFieldType)
+		{
+			InnerField = new CbField(Data, Type);
+		}
+
+		/// <summary>
+		/// Number of items in this array
+		/// </summary>
+		public int Count
+		{
+			get
+			{
+				ReadOnlyMemory<byte> PayloadBytes = InnerField.Payload;
+				PayloadBytes = PayloadBytes.Slice(VarInt.Measure(PayloadBytes.Span));
+				return (int)VarInt.Read(PayloadBytes.Span, out int NumByteCount);
+			}
+		}
+
+		/// <summary>
+		/// Access the array as an array field.
+		/// </summary>
+		/// <returns></returns>
+		public CbField AsField() => InnerField;
+
+		/// <summary>
+		/// Construct an array from an array field. No type check is performed!
+		/// </summary>
+		/// <param name="Field"></param>
+		/// <returns></returns>
+		public static CbArray FromFieldNoCheck(CbField Field) => new CbArray(Field);
+
+		/// <summary>
+		/// Returns the size of the array in bytes if serialized by itself with no name.
+		/// </summary>
+		/// <returns></returns>
+		public int GetSize()
+		{
+			return (int)Math.Min((ulong)sizeof(CbFieldType) + InnerField.GetPayloadSize(), int.MaxValue);
+		}
+
+		/// <summary>
+		/// Calculate the hash of the array if serialized by itself with no name.
+		/// </summary>
+		/// <returns></returns>
+		public Blake3Hash GetHash()
+		{
+			using (Blake3.Hasher Hasher = Blake3.Hasher.New())
+			{
+				AppendHash(Hasher);
+
+				byte[] Result = new byte[Blake3Hash.NumBytes];
+				Hasher.Finalize(Result);
+
+				return new Blake3Hash(Result);
+			}
+		}
+
+		/// <summary>
+		/// Append the hash of the array if serialized by itself with no name.
+		/// </summary>
+		public void AppendHash(Blake3.Hasher Hasher)
+		{
+			byte[] SerializedType = new byte[] { (byte)InnerField.GetType() };
+			Hasher.Update(SerializedType);
+			Hasher.Update(InnerField.Payload.Span);
+		}
+
+		/// <inheritdoc/>
+		public override bool Equals(object? Obj) => Equals(Obj as CbArray);
+
+		/// <inheritdoc/>
+		public override int GetHashCode() => BinaryPrimitives.ReadInt32BigEndian(GetHash().Span);
+
+		/// <summary>
+		/// Whether this array is identical to the other array.
+		///
+		/// Performs a deep comparison of any contained arrays or objects and their fields. Comparison
+		/// assumes that both fields are valid and are written in the canonical format.Fields must be
+		/// written in the same order in arrays and objects, and name comparison is case sensitive.If
+		/// these assumptions do not hold, this may return false for equivalent inputs. Validation can
+		/// be done with the All mode to check these assumptions about the format of the inputs.
+		/// </summary>
+		/// <param name="Other"></param>
+		/// <returns></returns>
+		public bool Equals(CbArray? Other)
+		{
+			return Other != null && GetType() == Other.GetType() && GetPayloadView().Span.SequenceEqual(Other.GetPayloadView().Span);
+		}
+
+		/// <summary>
+		/// Copy the array into a buffer of exactly GetSize() bytes, with no name.
+		/// </summary>
+		/// <param name="Buffer"></param>
+		public void CopyTo(Span<byte> Buffer)
+		{
+			Buffer[0] = (byte)GetType();
+			GetPayloadView().Span.CopyTo(Buffer.Slice(1));
+		}
+
+		/** Invoke the visitor for every attachment in the array. */
+		public void IterateAttachments(Action<CbField> Visitor) => CreateIterator().IterateRangeAttachments(Visitor);
+
+		/// <summary>
+		/// Try to get a view of the array as it would be serialized, such as by CopyTo.
+		/// 
+		/// A view is available if the array contains its type and has no name. Access the equivalent
+		/// for other arrays through FCbArray::GetBuffer, FCbArray::Clone, or CopyTo.
+		/// </summary>
+		public bool TryGetView(out ReadOnlyMemory<byte> OutView)
+		{
+			if(InnerField.HasName())
+			{
+				OutView = ReadOnlyMemory<byte>.Empty;
+				return false;
+			}
+			return InnerField.TryGetView(out OutView);
+		}
+
+		/// <inheritdoc cref="CbField.CreateIterator"/>
+		public CbFieldIterator CreateIterator() => InnerField.CreateIterator();
+
+		/// <inheritdoc/>
+		public IEnumerator<CbField> GetEnumerator() => InnerField.GetEnumerator();
+
+		/// <inheritdoc/>
+		IEnumerator IEnumerable.GetEnumerator() => GetEnumerator();
+
+		#region Mimic inheritance from CbField
+
+		/// <inheritdoc cref="CbField.GetType"/>
+		internal new CbFieldType GetType() => InnerField.GetType();
+
+		/// <inheritdoc cref="CbField.GetPayloadView"/>
+		internal ReadOnlyMemory<byte> GetPayloadView() => InnerField.GetPayloadView();
+
+		#endregion
+
+		#region Mimic inheritance from TCbBufferFactory
+
+		public static CbArray Clone(ReadOnlyMemory<byte> Data) => Clone(new CbArray(Data));
+		public static CbArray Clone(CbArray Other) => Other;
+		public static CbArray MakeView(ReadOnlyMemory<byte> Data) => new CbArray(Data);
+		public static CbArray MakeView(CbArray Other) => Other;
+
+		#endregion
+	}
+
+	/// <summary>
+	/// Simplified view of <see cref="CbObject"/> for display in the debugger
+	/// </summary>
+	class CbObjectDebugView
+	{
+		[DebuggerDisplay("{Name}: {Value}")]
+		public class Property
+		{
+			public string? Name { get; set; }
+			public object? Value { get; set; }
+		}
+
+		CbObject Object;
+
+		public CbObjectDebugView(CbObject Object) => this.Object = Object;
+
+		[DebuggerBrowsable(DebuggerBrowsableState.RootHidden)]
+		public Property[] Properties => Object.Select(x => new Property { Name = x.Name.ToString(), Value = x.Value }).ToArray();
+	}
+
+	/// <summary>
+	/// Array of CbField that have unique names.
+	///
+	/// Accessing the fields of an object is always a safe operation, even if the requested field does
+	/// not exist. Fields may be accessed by name or through iteration. When a field is requested that
+	/// is not found in the object, the field that it returns has no value (evaluates to false) though
+	/// attempting to access the empty field is also safe, as described by FCbFieldView.
+	/// </summary>
+	[DebuggerTypeProxy(typeof(CbObjectDebugView))]
+	public class CbObject : IEnumerable<CbField>
+	{
+		/// <summary>
+		/// Empty array constant
+		/// </summary>
+		public static CbObject Empty = CbObject.FromFieldNoCheck(new CbField(new byte[] { (byte)CbFieldType.Object, 0 }));
+
+		/// <summary>
+		/// The inner field object
+		/// </summary>
+		private CbField InnerField;
+
+		/// <summary>
+		/// Constructor
+		/// </summary>
+		/// <param name="Buffer"></param>
+		private CbObject(CbField Field)
+		{
+			InnerField = new CbField(Field.Memory, Field.TypeWithFlags);
+		}
+
+		/// <summary>
+		/// Constructor
+		/// </summary>
+		/// <param name="Buffer"></param>
+		public CbObject(ReadOnlyMemory<byte> Buffer, CbFieldType FieldType = CbFieldType.HasFieldType)
+		{
+			InnerField = new CbField(Buffer, FieldType);
+		}
+
+		/// <summary>
+		/// Builds an object by calling a delegate with a writer
+		/// </summary>
+		/// <param name="Build"></param>
+		/// <returns></returns>
+		public static CbObject Build(Action<CbWriter> Build)
+		{
+			CbWriter Writer = new CbWriter();
+			Writer.BeginObject();
+			Build(Writer);
+			Writer.EndObject();
+			return new CbObject(Writer.ToByteArray());
+		}
+
+		/// <summary>
+		/// Find a field by case-sensitive name comparison.
+		///
+		/// The cost of this operation scales linearly with the number of fields in the object. Prefer to
+		/// iterate over the fields only once when consuming an object.
+		/// </summary>
+		/// <param name="Name">The name of the field.</param>
+		/// <returns>The matching field if found, otherwise a field with no value.</returns>
+		public CbField Find(Utf8String Name) => InnerField[Name];
+
+		/// <summary>
+		/// Find a field by case-insensitive name comparison.
+		/// </summary>
+		/// <param name="Name">The name of the field.</param>
+		/// <returns>The matching field if found, otherwise a field with no value.</returns>
+		public CbField FindIgnoreCase(Utf8String Name) => InnerField.FirstOrDefault(Field => Utf8StringComparer.OrdinalIgnoreCase.Equals(Field.Name, Name)) ?? new CbField();
+
+		/// <summary>
+		/// Find a field by case-sensitive name comparison.
+		/// </summary>
+		/// <param name="Name">The name of the field.</param>
+		/// <returns>The matching field if found, otherwise a field with no value.</returns>
+		public CbField this[Utf8String Name] => InnerField[Name];
+
+		/// <inheritdoc cref="AsFieldView"/>
+		public CbField AsField() => InnerField;
+
+		/// <summary>
+		/// Construct an object from an object field. No type check is performed!
+		/// </summary>
+		/// <param name="Field"></param>
+		/// <returns></returns>
+		public static CbObject FromFieldNoCheck(CbField Field) => new CbObject(Field);
+
+		/// <summary>
+		/// Returns the size of the object in bytes if serialized by itself with no name.
+		/// </summary>
+		/// <returns></returns>
+		public int GetSize()
+		{
+			return sizeof(CbFieldType) + InnerField.Payload.Length;
+		}
+
+		/// <summary>
+		/// Calculate the hash of the object if serialized by itself with no name.
+		/// </summary>
+		/// <returns></returns>
+		public Blake3Hash GetHash()
+		{
+			using (Blake3.Hasher Hasher = Blake3.Hasher.New())
+			{
+				AppendHash(Hasher);
+
+				byte[] Data = new byte[Blake3Hash.NumBytes];
+				Hasher.Finalize(Data);
+
+				return new Blake3Hash(Data);
+			}
+		}
+
+		/// <summary>
+		/// Append the hash of the object if serialized by itself with no name.
+		/// </summary>
+		/// <param name="Hasher"></param>
+		public void AppendHash(Blake3.Hasher Hasher)
+		{
+			byte[] Temp = new byte[] { (byte)InnerField.GetType() };
+			Hasher.Update(Temp);
+			Hasher.Update(InnerField.Payload.Span);
+		}
+
+		/// <inheritdoc/>
+		public override bool Equals(object? Obj) => Equals(Obj as CbObject);
+
+		/// <inheritdoc/>
+		public override int GetHashCode() => BinaryPrimitives.ReadInt32BigEndian(GetHash().Span);
+
+		/// <summary>
+		/// Whether this object is identical to the other object.
+		/// 
+		/// Performs a deep comparison of any contained arrays or objects and their fields. Comparison
+		/// assumes that both fields are valid and are written in the canonical format. Fields must be
+		/// written in the same order in arrays and objects, and name comparison is case sensitive. If
+		/// these assumptions do not hold, this may return false for equivalent inputs. Validation can
+		/// be done with the All mode to check these assumptions about the format of the inputs.
+		/// </summary>
+		/// <param name="Other"></param>
+		/// <returns></returns>
+		public bool Equals(CbObject? Other)
+		{
+			return Other != null && InnerField.GetType() == Other.InnerField.GetType() && InnerField.Payload.Span.SequenceEqual(Other.InnerField.Payload.Span);
+		}
+
+		/// <summary>
+		/// Copy the object into a buffer of exactly GetSize() bytes, with no name.
+		/// </summary>
+		/// <param name="Buffer"></param>
+		public void CopyTo(Span<byte> Buffer)
+		{
+			Buffer[0] = (byte)InnerField.GetType();
+			InnerField.Payload.Span.CopyTo(Buffer.Slice(1));
+		}
+
+		/// <summary>
+		/// Invoke the visitor for every attachment in the object.
+		/// </summary>
+		/// <param name="Visitor"></param>
+		public void IterateAttachments(Action<CbField> Visitor) => CreateIterator().IterateRangeAttachments(Visitor);
+
+		/// <summary>
+		/// Creates a view of the object, excluding the name
+		/// </summary>
+		/// <returns></returns>
+		public ReadOnlyMemory<byte> GetView()
+		{
+			ReadOnlyMemory<byte> Memory;
+			if (!TryGetView(out Memory))
+			{
+				byte[] Data = new byte[GetSize()];
+				CopyTo(Data);
+				Memory = Data;
+			}
+			return Memory;
+		}
+
+		/// <summary>
+		/// Try to get a view of the object as it would be serialized, such as by CopyTo.
+		/// 
+		/// A view is available if the object contains its type and has no name. Access the equivalent
+		/// for other objects through FCbObject::GetBuffer, FCbObject::Clone, or CopyTo.
+		/// </summary>
+		/// <param name="OutView"></param>
+		/// <returns></returns>
+		public bool TryGetView(out ReadOnlyMemory<byte> OutView)
+		{
+			if (InnerField.HasName())
+			{
+				OutView = ReadOnlyMemory<byte>.Empty;
+				return false;
+			}
+			return InnerField.TryGetView(out OutView);
+		}
+
+		/// <inheritdoc cref="CbField.CreateIterator"/>
+		public CbFieldIterator CreateIterator() => InnerField.CreateIterator();
+
+		/// <inheritdoc/>
+		public IEnumerator<CbField> GetEnumerator() => InnerField.GetEnumerator();
+
+		/// <inheritdoc/>
+		IEnumerator IEnumerable.GetEnumerator() => InnerField.GetEnumerator();
+
+		/// <summary>
+		/// Clone this object
+		/// </summary>
+		/// <param name="Object"></param>
+		/// <returns></returns>
+		public static CbObject Clone(CbObject Object) => Object;
+
+		#region Conversion to Json
+		/// <summary>
+		/// Convert this object to JSON
+		/// </summary>
+		/// <returns></returns>
+		public string ToJson()
+		{
+			ArrayBufferWriter<byte> Buffer = new ArrayBufferWriter<byte>();
+			using (Utf8JsonWriter JsonWriter = new Utf8JsonWriter(Buffer))
+			{
+				ToJson(JsonWriter);
+			}
+			return Encoding.UTF8.GetString(Buffer.WrittenMemory.Span);
+		}
+
+		/// <summary>
+		/// Write this object to JSON
+		/// </summary>
+		/// <param name="Writer"></param>
+		public void ToJson(Utf8JsonWriter Writer)
+		{
+			Writer.WriteStartObject();
+			foreach (CbField Field in InnerField)
+			{
+				WriteField(Field, Writer);
+			}
+			Writer.WriteEndObject();
+		}
+
+		/// <summary>
+		/// Write a single field to a writer
+		/// </summary>
+		/// <param name="Field"></param>
+		/// <param name="Writer"></param>
+		private static void WriteField(CbField Field, Utf8JsonWriter Writer)
+		{
+			if (Field.IsObject())
+			{
+				Writer.WriteStartObject();
+				CbObject Object = Field.AsObject();
+				foreach (CbField ObjectField in Object.InnerField)
+				{
+					WriteField(ObjectField, Writer);
+				}
+				Writer.WriteEndObject();
+			}
+			else if (Field.IsArray())
+			{
+				Writer.WriteStartArray();
+				Writer.WriteEndArray();
+			}
+			else if (Field.IsInteger())
+			{
+				if (Field.GetType() == CbFieldType.IntegerNegative)
+				{
+					Writer.WriteNumber(Field.Name.Span, -Field.AsInt64());
+				}
+				else
+				{
+					Writer.WriteNumber(Field.Name.Span, Field.AsUInt64());
+				}
+			}
+			else if (Field.IsBool())
+			{
+				Writer.WriteBoolean(Field.Name.Span, Field.AsBool());
+			}
+			else if (Field.IsNull())
+			{
+				Writer.WriteNullValue();
+			}
+			else if (Field.IsDateTime())
+			{
+				Writer.WriteString(Field.Name.Span, Field.AsDateTime());
+			}
+			else if (Field.IsHash())
+			{
+				Writer.WriteString(Field.Name.Span, StringUtils.FormatUtf8HexString(Field.AsHash().Span).Span);
+			}
+			else if (Field.IsString())
+			{
+				Writer.WriteString(Field.Name.Span, Field.AsString().Span);
+			}
+			else
+			{
+				throw new NotImplementedException($"Unhandled type {Field.GetType()} when attempting to convert to json");
+			}
+		}
+		#endregion
+	}
+}
diff --git a/Engine/Source/Programs/Shared/EpicGames.Serialization/CbSerializer.cs b/Engine/Source/Programs/Shared/EpicGames.Serialization/CbSerializer.cs
new file mode 100644
--- /dev/null
+++ b/Engine/Source/Programs/Shared/EpicGames.Serialization/CbSerializer.cs
@@ -0,0 +1,269 @@
+// Copyright Epic Games, Inc. All Rights Reserved.
+
+using EpicGames.Core;
+using EpicGames.Serialization.Converters;
+using System;
+using System.Collections;
+using System.Collections.Concurrent;
+using System.Collections.Generic;
+using System.Diagnostics.CodeAnalysis;
+using System.Linq;
+using System.Linq.Expressions;
+using System.Reflection;
+using System.Reflection.Emit;
+
+namespace EpicGames.Serialization
+{
+	/// <summary>
+	/// Attribute used to mark a property that should be serialized to compact binary
+	/// </summary>
+	[AttributeUsage(AttributeTargets.Property)]
+	public class CbFieldAttribute : Attribute
+	{
+		/// <summary>
+		/// Name of the serialized field
+		/// </summary>
+		public string? Name { get; set; }
+
+		/// <summary>
+		/// Default constructor
+		/// </summary>
+		public CbFieldAttribute()
+		{
+		}
+
+		/// <summary>
+		/// Constructor
+		/// </summary>
+		/// <param name="Name"></param>
+		public CbFieldAttribute(string Name)
+		{
+			this.Name = Name;
+		}
+	}
+
+	/// <summary>
+	/// Exception thrown when serializing cb objects
+	/// </summary>
+	public class CbException : Exception
+	{
+		/// <inheritdoc cref="Exception(string?)"/>
+		public CbException(string Message) : base(Message)
+		{
+		}
+
+		/// <inheritdoc cref="Exception(string?, Exception)"/>
+		public CbException(string Message, Exception Inner) : base(Message, Inner)
+		{
+		}
+	}
+
+	/// <summary>
+	/// Attribute-driven compact binary serializer
+	/// </summary>
+	public static class CbSerializer
+	{
+		/// <summary>
+		/// Object used for locking access to shared objects
+		/// </summary>
+		static object LockObject = new object();
+
+		/// <summary>
+		/// Cache of type to converter
+		/// </summary>
+		public static Dictionary<Type, CbConverter> TypeToConverter = new Dictionary<Type, CbConverter>()
+		{
+			[typeof(CbField)] = new CbFieldConverter(),
+			[typeof(CbObject)] = new CbObjectConverter()
+		};
+
+		/// <summary>
+		/// List of converter factories. Must be 
+		/// </summary>
+		public static List<CbConverterFactory> ConverterFactories = new List<CbConverterFactory>
+		{
+			new CbDefaultConverterFactory(),
+			new CbListConverterFactory(),
+			new CbDictionaryConverterFactory(),
+			new CbNullableConverterFactory()
+		};
+	
+		/// <summary>
+		/// Class used to cache values returned by CreateConverterInfo without having to do a dictionary lookup.
+		/// </summary>
+		/// <typeparam name="T">The type to be converted</typeparam>
+		class CbConverterCache<T>
+		{
+			/// <summary>
+			/// The converter instance
+			/// </summary>
+			public static CbConverter<T> Instance { get; } = CreateConverter();
+
+			/// <summary>
+			/// Create a typed converter
+			/// </summary>
+			/// <returns></returns>
+			static CbConverter<T> CreateConverter()
+			{
+				CbConverter Converter = GetConverter(typeof(T));
+				return (Converter as CbConverter<T>) ?? new CbConverterWrapper(Converter);
+			}
+
+			/// <summary>
+			/// Wrapper class to convert an untyped converter into a typed one
+			/// </summary>
+			class CbConverterWrapper : CbConverter<T>
+			{
+				CbConverter Inner;
+
+				public CbConverterWrapper(CbConverter Inner) => this.Inner = Inner;
+
+				/// <inheritdoc/>
+				public override T Read(CbField Field) => (T)Inner.ReadObject(Field)!;
+
+				/// <inheritdoc/>
+				public override void Write(CbWriter Writer, T Value) => Inner.WriteObject(Writer, Value);
+
+				/// <inheritdoc/>
+				public override void WriteNamed(CbWriter Writer, Utf8String Name, T Value) => Inner.WriteNamedObject(Writer, Name, Value);
+			}
+		}
+
+		/// <summary>
+		/// Gets the converter for a particular type
+		/// </summary>
+		/// <param name="Type"></param>
+		/// <returns></returns>
+		public static CbConverter GetConverter(Type Type)
+		{
+			CbConverter? Converter;
+			lock (LockObject)
+			{
+				if (!TypeToConverter.TryGetValue(Type, out Converter))
+				{
+					CbConverterAttribute? ConverterAttribute = Type.GetCustomAttribute<CbConverterAttribute>();
+					if (ConverterAttribute != null)
+					{
+						Type ConverterType = ConverterAttribute.ConverterType;
+						if (Type.IsGenericType && ConverterType.IsGenericTypeDefinition)
+						{
+							ConverterType = ConverterType.MakeGenericType(Type.GetGenericArguments());
+						}
+						Converter = (CbConverter?)Activator.CreateInstance(ConverterType)!;
+					}
+					else
+					{
+						for (int Idx = ConverterFactories.Count - 1; Idx >= 0 && Converter == null; Idx--)
+						{
+							Converter = ConverterFactories[Idx].CreateConverter(Type);
+						}
+
+						if (Converter == null)
+						{
+							throw new CbException($"Unable to create converter for {Type.Name}");
+						}
+					}
+					TypeToConverter.Add(Type, Converter!);
+				}
+			}
+			return Converter;
+		}
+
+		/// <summary>
+		/// Gets the converter for a given type
+		/// </summary>
+		/// <typeparam name="T">Type to retrieve the converter for</typeparam>
+		/// <returns></returns>
+		public static CbConverter<T> GetConverter<T>()
+		{
+			return CbConverterCache<T>.Instance;
+		}
+
+		/// <summary>
+		/// Serialize an object
+		/// </summary>
+		/// <param name="Type">Type of the object to serialize</param>
+		/// <param name="Value"></param>
+		/// <returns></returns>
+		public static CbObject Serialize(Type Type, object Value)
+		{
+			CbWriter Writer = new CbWriter();
+			GetConverter(Type).WriteObject(Writer, Value);
+			return Writer.ToObject();
+		}
+
+		/// <summary>
+		/// Serialize an object
+		/// </summary>
+		/// <typeparam name="T"></typeparam>
+		/// <param name="Value"></param>
+		/// <returns></returns>
+		public static CbObject Serialize<T>(T Value)
+		{
+			CbWriter Writer = new CbWriter();
+			GetConverter<T>().Write(Writer, Value);
+			return Writer.ToObject();
+		}
+
+		/// <summary>
+		/// Serialize a property to a given writer
+		/// </summary>
+		/// <typeparam name="T"></typeparam>
+		/// <param name="Writer"></param>
+		/// <param name="Value"></param>
+		public static void Serialize<T>(CbWriter Writer, T Value)
+		{
+			GetConverter<T>().Write(Writer, Value);
+		}
+
+		/// <summary>
+		/// Serialize a named property to the given writer
+		/// </summary>
+		/// <typeparam name="T"></typeparam>
+		/// <param name="Writer"></param>
+		/// <param name="Name"></param>
+		/// <param name="Value"></param>
+		public static void Serialize<T>(CbWriter Writer, Utf8String Name, T Value)
+		{
+			GetConverter<T>().WriteNamed(Writer, Name, Value);
+		}
+
+		/// <summary>
+		/// Deserialize an object from a <see cref="CbObject"/>
+		/// </summary>
+		/// <param name="Field"></param>
+		/// <param name="Type">Type of the object to read</param>
+		/// <returns></returns>
+		public static object? Deserialize(CbField Field, Type Type)
+		{
+			return GetConverter(Type).ReadObject(Field);
+		}
+
+		/// <summary>
+		/// Deserialize an object from a <see cref="CbField"/>
+		/// </summary>
+		/// <typeparam name="T"></typeparam>
+		/// <param name="Field"></param>
+		/// <returns></returns>
+		public static T Deserialize<T>(CbField Field)
+		{
+			return GetConverter<T>().Read(Field);
+		}
+
+		/// <summary>
+		/// Deserialize an object from a <see cref="CbObject"/>
+		/// </summary>
+		/// <typeparam name="T"></typeparam>
+		/// <param name="Object"></param>
+		/// <returns></returns>
+		public static T Deserialize<T>(CbObject Object) => Deserialize<T>(Object.AsField());
+
+		/// <summary>
+		/// Deserialize an object from a block of memory
+		/// </summary>
+		/// <typeparam name="T"></typeparam>
+		/// <param name="Data"></param>
+		/// <returns></returns>
+		public static T Deserialize<T>(ReadOnlyMemory<byte> Data) => Deserialize<T>(new CbField(Data));
+	}
+}
diff --git a/Engine/Source/Programs/Shared/EpicGames.Serialization/CbWriter.cs b/Engine/Source/Programs/Shared/EpicGames.Serialization/CbWriter.cs
new file mode 100644
--- /dev/null
+++ b/Engine/Source/Programs/Shared/EpicGames.Serialization/CbWriter.cs
@@ -0,0 +1,874 @@
+// Copyright Epic Games, Inc. All Rights Reserved.
+
+using EpicGames.Core;
+using System;
+using System.Buffers.Binary;
+using System.Collections.Generic;
+using System.Linq;
+using System.Text;
+using System.Text.Json;
+
+namespace EpicGames.Serialization
+{
+	/// <summary>
+	/// Exception for <see cref="CbWriter"/>
+	/// </summary>
+	public class CbWriterException : Exception
+	{
+		/// <summary>
+		/// Constructor
+		/// </summary>
+		/// <param name="Message"></param>
+		public CbWriterException(string Message)
+			: base(Message)
+		{
+		}
+
+		/// <summary>
+		/// Constructor
+		/// </summary>
+		/// <param name="Message"></param>
+		/// <param name="Ex"></param>
+		public CbWriterException(string Message, Exception? Ex)
+			: base(Message, Ex)
+		{
+		}
+	}
+
+	/// <summary>
+	/// Forward-only writer for compact binary objects
+	/// </summary>
+	public class CbWriter
+	{
+		/// <summary>
+		/// Stores information about an object or array scope within the written buffer which requires a header to be inserted containing
+		/// the size or number of elements when copied to an output buffer
+		/// </summary>
+		class Scope
+		{
+			public CbFieldType FieldType;
+			public CbFieldType UniformFieldType;
+			public int Offset; // Offset to insert the length/count
+			public int Length; // Excludes the size of this field's headers, and child fields' headers.
+			public int Count;
+			public List<Scope> Children = new List<Scope>();
+			public int SizeOfChildHeaders; // Sum of additional headers for child items, recursively.
+
+			public Scope(CbFieldType FieldType, CbFieldType UniformFieldType, int Offset)
+			{
+				Reset(FieldType, UniformFieldType, Offset);
+			}
+
+			public void Reset(CbFieldType FieldType, CbFieldType UniformFieldType, int Offset)
+			{
+				this.FieldType = FieldType;
+				this.UniformFieldType = UniformFieldType;
+				this.Offset = Offset;
+				this.Length = 0;
+				this.Count = 0;
+				this.Children.Clear();
+				this.SizeOfChildHeaders = 0;
+			}
+		}
+
+		/// <summary>
+		/// Chunk of written data. Chunks are allocated as needed and chained together with scope annotations to produce the output data.
+		/// </summary>
+		class Chunk
+		{
+			public int Offset;
+			public int Length;
+			public byte[] Data;
+			public List<Scope> Scopes = new List<Scope>();
+
+			public Chunk(int Offset, int MaxLength)
+			{
+				this.Data = new byte[MaxLength];
+				Reset(Offset);
+			}
+
+			public void Reset(int Offset)
+			{
+				this.Offset = Offset;
+				this.Length = 0;
+				this.Scopes.Clear();
+			}
+		}
+
+		const int DefaultChunkSize = 1024;
+
+		List<Chunk> Chunks = new List<Chunk>();
+		Stack<Scope> OpenScopes = new Stack<Scope>();
+		Chunk CurrentChunk => Chunks[Chunks.Count - 1];
+		Scope CurrentScope => OpenScopes.Peek();
+		int CurrentOffset;
+		List<Chunk> FreeChunks = new List<Chunk>();
+		List<Scope> FreeScopes = new List<Scope>();
+
+		/// <summary>
+		/// Constructor
+		/// </summary>
+		public CbWriter()
+			: this(DefaultChunkSize)
+		{
+		}
+
+		/// <summary>
+		/// Constructor
+		/// </summary>
+		/// <param name="Reserve">Amount of data to reserve for output</param>
+		public CbWriter(int Reserve)
+		{
+			Chunks.Add(new Chunk(0, Reserve));
+			OpenScopes.Push(new Scope(CbFieldType.Array, CbFieldType.None, 0));
+		}
+
+		/// <summary>
+		/// 
+		/// </summary>
+		public void Clear()
+		{
+			foreach (Chunk Chunk in Chunks)
+			{
+				FreeChunk(Chunk);
+			}
+
+			CurrentOffset = 0;
+
+			Chunks.Clear();
+			Chunks.Add(AllocChunk(0, DefaultChunkSize));
+
+			OpenScopes.Clear();
+			OpenScopes.Push(AllocScope(CbFieldType.Array, CbFieldType.None, 0));
+		}
+
+		/// <summary>
+		/// Allocate a new chunk object
+		/// </summary>
+		/// <param name="Offset">Offset of the chunk</param>
+		/// <param name="MaxLength">Maximum length of the chunk</param>
+		/// <returns>New chunk object</returns>
+		Chunk AllocChunk(int Offset, int MaxLength)
+		{
+			for(int Idx = FreeChunks.Count - 1; Idx >= 0; Idx--)
+			{
+				Chunk Chunk = FreeChunks[Idx];
+				if (Chunk.Data.Length >= MaxLength)
+				{
+					FreeChunks.RemoveAt(Idx);
+					Chunk.Reset(Offset);
+					return Chunk;
+				}
+			}
+			return new Chunk(Offset, MaxLength);
+		}
+
+		/// <summary>
+		/// Adds a chunk to the free list
+		/// </summary>
+		/// <param name="Chunk"></param>
+		void FreeChunk(Chunk Chunk)
+		{
+			// Add the scopes to the free list
+			FreeScopes.AddRange(Chunk.Scopes);
+			Chunk.Scopes.Clear();
+
+			// Insert it into the free list, sorted by descending size
+			for (int Idx = 0; ; Idx++)
+			{
+				if (Idx == FreeChunks.Count || Chunk.Data.Length >= FreeChunks[Idx].Data.Length)
+				{
+					FreeChunks.Insert(Idx, Chunk);
+					break;
+				}
+			}
+		}
+
+		/// <summary>
+		/// Allocate a scope object
+		/// </summary>
+		/// <param name="FieldType"></param>
+		/// <param name="UniformFieldType"></param>
+		/// <param name="Offset"></param>
+		/// <returns></returns>
+		Scope AllocScope(CbFieldType FieldType, CbFieldType UniformFieldType, int Offset)
+		{
+			if (FreeScopes.Count > 0)
+			{
+				Scope Scope = FreeScopes[FreeScopes.Count - 1];
+				Scope.Reset(FieldType, UniformFieldType, Offset);
+				FreeScopes.RemoveAt(FreeScopes.Count - 1);
+				return Scope;
+			}
+			return new Scope(FieldType, UniformFieldType, Offset);
+		}
+
+		/// <summary>
+		/// Ensure that a block of contiguous memory of the given length is available in the output buffer
+		/// </summary>
+		/// <param name="Length"></param>
+		/// <returns>The allocated memory</returns>
+		Memory<byte> Allocate(int Length)
+		{
+			Chunk LastChunk = CurrentChunk;
+			if (LastChunk.Length + Length > LastChunk.Data.Length)
+			{
+				int ChunkSize = Math.Max(Length, DefaultChunkSize);
+				LastChunk = AllocChunk(CurrentOffset, ChunkSize);
+				Chunks.Add(LastChunk);
+			}
+
+			Memory<byte> Buffer = LastChunk.Data.AsMemory(LastChunk.Length, Length);
+			LastChunk.Length += Length;
+			CurrentOffset += Length;
+			return Buffer;
+		}
+
+		/// <summary>
+		/// Insert a new scope
+		/// </summary>
+		/// <param name="FieldType"></param>
+		/// <param name="UniformFieldType"></param>
+		void PushScope(CbFieldType FieldType, CbFieldType UniformFieldType)
+		{
+			Scope NewScope = AllocScope(FieldType, UniformFieldType, CurrentOffset);
+			CurrentScope.Children.Add(NewScope);
+			OpenScopes.Push(NewScope);
+
+			CurrentChunk.Scopes.Add(NewScope);
+		}
+
+		/// <summary>
+		/// Pop a scope from the current open list
+		/// </summary>
+		void PopScope()
+		{
+			Scope Scope = CurrentScope;
+			Scope.Length = CurrentOffset - Scope.Offset;
+			Scope.SizeOfChildHeaders = ComputeSizeOfChildHeaders(Scope);
+			OpenScopes.Pop();
+		}
+
+		/// <summary>
+		/// Writes the header for an unnamed field
+		/// </summary>
+		/// <param name="Type"></param>
+		void WriteFieldHeader(CbFieldType Type)
+		{
+			Scope Scope = CurrentScope;
+			if (!CbFieldUtils.IsArray(Scope.FieldType))
+			{
+				throw new CbWriterException($"Anonymous fields are not allowed within fields of type {Scope.FieldType}");
+			}
+			
+			if (Scope.UniformFieldType == CbFieldType.None)
+			{
+				Allocate(1).Span[0] = (byte)Type;
+			}
+			else if (Scope.UniformFieldType != Type)
+			{
+				throw new CbWriterException($"Mismatched type for uniform array - expected {Scope.UniformFieldType}, not {Type}");
+			}
+			Scope.Count++;
+		}
+
+		/// <summary>
+		/// Writes the header for a named field
+		/// </summary>
+		/// <param name="Type"></param>
+		/// <param name="Name"></param>
+		void WriteFieldHeader(CbFieldType Type, Utf8String Name)
+		{
+			Scope Scope = CurrentScope;
+			if (!CbFieldUtils.IsObject(Scope.FieldType))
+			{
+				throw new CbWriterException($"Named fields are not allowed within fields of type {Scope.FieldType}");
+			}
+
+			int NameVarIntLength = VarInt.Measure(Name.Length);
+			if (Scope.UniformFieldType == CbFieldType.None)
+			{
+				Span<byte> Buffer = Allocate(1 + NameVarIntLength + Name.Length).Span;
+				Buffer[0] = (byte)(Type | CbFieldType.HasFieldName);
+				WriteBinaryPayload(Buffer[1..], Name.Span);
+			}
+			else
+			{
+				if (Scope.UniformFieldType != Type)
+				{
+					throw new CbWriterException($"Mismatched type for uniform object - expected {Scope.UniformFieldType}, not {Type}");
+				}
+				WriteBinaryPayload(Name.Span);
+			}
+			Scope.Count++;
+		}
+
+		/// <summary>
+		/// Copies an entire field value to the output
+		/// </summary>
+		/// <param name="Field"></param>
+		public void WriteFieldValue(CbField Field)
+		{
+			WriteFieldHeader(Field.GetType());
+			int Size = (int)Field.GetPayloadSize();
+			Field.GetPayloadView().CopyTo(Allocate(Size));
+		}
+
+		/// <summary>
+		/// Copies an entire field value to the output, using the name from the field
+		/// </summary>
+		/// <param name="Field"></param>
+		public void WriteField(CbField Field) => WriteField(Field.GetName(), Field);
+
+		/// <summary>
+		/// Copies an entire field value to the output
+		/// </summary>
+		/// <param name="Field"></param>
+		public void WriteField(Utf8String Name, CbField Field)
+		{
+			WriteFieldHeader(Field.GetType(), Name);
+			int Size = (int)Field.GetPayloadSize();
+			Field.GetPayloadView().CopyTo(Allocate(Size));
+		}
+
+		/// <summary>
+		/// Begin writing an object field
+		/// </summary>
+		public void BeginObject()
+		{
+			WriteFieldHeader(CbFieldType.Object);
+			PushScope(CbFieldType.Object, CbFieldType.None);
+		}
+
+		/// <summary>
+		/// Begin writing an object field
+		/// </summary>
+		/// <param name="Name">Name of the field</param>
+		public void BeginObject(Utf8String Name)
+		{
+			WriteFieldHeader(CbFieldType.Object, Name);
+			PushScope(CbFieldType.Object, CbFieldType.None);
+		}
+
+		/// <summary>
+		/// End the current object
+		/// </summary>
+		public void EndObject()
+		{
+			PopScope();
+		}
+
+		/// <summary>
+		/// Begin writing an array field
+		/// </summary>
+		public void BeginArray()
+		{
+			WriteFieldHeader(CbFieldType.Array);
+			PushScope(CbFieldType.Array, CbFieldType.None);
+		}
+
+		/// <summary>
+		/// Begin writing a named array field
+		/// </summary>
+		/// <param name="Name"></param>
+		public void BeginArray(Utf8String Name)
+		{
+			WriteFieldHeader(CbFieldType.Array, Name);
+			PushScope(CbFieldType.Array, CbFieldType.None);
+		}
+
+		/// <summary>
+		/// End the current array
+		/// </summary>
+		public void EndArray()
+		{
+			PopScope();
+		}
+
+		/// <summary>
+		/// Begin writing a uniform array field
+		/// </summary>
+		/// <param name="FieldType">The field type for elements in the array</param>
+		public void BeginUniformArray(CbFieldType FieldType)
+		{
+			WriteFieldHeader(CbFieldType.UniformArray);
+			PushScope(CbFieldType.UniformArray, FieldType);
+			Allocate(1).Span[0] = (byte)FieldType;
+		}
+
+		/// <summary>
+		/// Begin writing a named uniform array field
+		/// </summary>
+		/// <param name="Name">Name of the field</param>
+		/// <param name="FieldType">The field type for elements in the array</param>
+		public void BeginUniformArray(Utf8String Name, CbFieldType FieldType)
+		{
+			WriteFieldHeader(CbFieldType.UniformArray, Name);
+			PushScope(CbFieldType.UniformArray, FieldType);
+			Allocate(1).Span[0] = (byte)FieldType;
+		}
+
+		/// <summary>
+		/// End the current array
+		/// </summary>
+		public void EndUniformArray()
+		{
+			PopScope();
+		}
+
+		/// <summary>
+		/// Write a null field
+		/// </summary>
+		public void WriteNullValue()
+		{
+			WriteFieldHeader(CbFieldType.Null);
+		}
+
+		/// <summary>
+		/// Write a named null field
+		/// </summary>
+		/// <param name="Name">Name of the field</param>
+		public void WriteNull(Utf8String Name)
+		{
+			WriteFieldHeader(CbFieldType.Null, Name);
+		}
+
+		/// <summary>
+		/// Writes a boolean value
+		/// </summary>
+		/// <param name="Value"></param>
+		public void WriteBoolValue(bool Value)
+		{
+			WriteFieldHeader(Value? CbFieldType.BoolTrue : CbFieldType.BoolFalse);
+		}
+
+		/// <summary>
+		/// Writes a boolean value
+		/// </summary>
+		/// <param name="Name">Name of the field</param>
+		/// <param name="Value"></param>
+		public void WriteBool(Utf8String Name, bool Value)
+		{
+			WriteFieldHeader(Value ? CbFieldType.BoolTrue : CbFieldType.BoolFalse, Name);
+		}
+
+		/// <summary>
+		/// Writes the payload for an integer
+		/// </summary>
+		/// <param name="Value">Value to write</param>
+		void WriteIntegerPayload(ulong Value)
+		{
+			int Length = VarInt.Measure(Value);
+			Span<byte> Buffer = Allocate(Length).Span;
+			VarInt.Write(Buffer, Value);
+		}
+
+		/// <summary>
+		/// Writes an unnamed integer field
+		/// </summary>
+		/// <param name="Value">Value to be written</param>
+		public void WriteIntegerValue(long Value)
+		{
+			if (Value >= 0)
+			{
+				WriteFieldHeader(CbFieldType.IntegerPositive);
+				WriteIntegerPayload((ulong)Value);
+			}
+			else
+			{
+				WriteFieldHeader(CbFieldType.IntegerNegative);
+				WriteIntegerPayload((ulong)-Value);
+			}
+		}
+
+		/// <summary>
+		/// Writes an named integer field
+		/// </summary>
+		/// <param name="Name">Name of the field</param>
+		/// <param name="Value">Value to be written</param>
+		public void WriteInteger(Utf8String Name, long Value)
+		{
+			if (Value >= 0)
+			{
+				WriteFieldHeader(CbFieldType.IntegerPositive, Name);
+				WriteIntegerPayload((ulong)Value);
+			}
+			else
+			{
+				WriteFieldHeader(CbFieldType.IntegerNegative, Name);
+				WriteIntegerPayload((ulong)-Value);
+			}
+		}
+
+		/// <summary>
+		/// Writes an unnamed integer field
+		/// </summary>
+		/// <param name="Value">Value to be written</param>
+		public void WriteIntegerValue(ulong Value)
+		{
+			WriteFieldHeader(CbFieldType.IntegerPositive);
+			WriteIntegerPayload(Value);
+		}
+
+		/// <summary>
+		/// Writes a named integer field
+		/// </summary>
+		/// <param name="Name">Name of the field</param>
+		/// <param name="Value">Value to be written</param>
+		public void WriteInteger(Utf8String Name, ulong Value)
+		{
+			WriteFieldHeader(CbFieldType.IntegerPositive, Name);
+			WriteIntegerPayload(Value);
+		}
+
+		/// <summary>
+		/// Writes an unnamed double field
+		/// </summary>
+		/// <param name="Value">Value to be written</param>
+		public void WriteDoubleValue(double Value)
+		{
+			WriteFieldHeader(CbFieldType.Float64);
+			BinaryPrimitives.WriteInt64BigEndian(Allocate(sizeof(double)).Span, BitConverter.DoubleToInt64Bits(Value));
+		}
+
+		/// <summary>
+		/// Writes a named double field
+		/// </summary>
+		/// <param name="Name">Name of the field</param>
+		/// <param name="Value">Value to be written</param>
+		public void WriteDouble(Utf8String Name, double Value)
+		{
+			WriteFieldHeader(CbFieldType.Float64, Name);
+			BinaryPrimitives.WriteInt64BigEndian(Allocate(sizeof(double)).Span, BitConverter.DoubleToInt64Bits(Value));
+		}
+
+		/// <summary>
+		/// Writes the payload for a <see cref="DateTime"/> value
+		/// </summary>
+		/// <param name="DateTime">The value to write</param>
+		void WriteDateTimePayload(DateTime DateTime)
+		{
+			Span<byte> Buffer = Allocate(sizeof(long)).Span;
+			BinaryPrimitives.WriteInt64BigEndian(Buffer, DateTime.Ticks);
+		}
+
+		/// <summary>
+		/// Writes an unnamed <see cref="DateTime"/> field
+		/// </summary>
+		/// <param name="Value">Value to be written</param>
+		public void WriteDateTimeValue(DateTime Value)
+		{
+			WriteFieldHeader(CbFieldType.DateTime);
+			WriteDateTimePayload(Value);
+		}
+
+		/// <summary>
+		/// Writes a named <see cref="DateTime"/> field
+		/// </summary>
+		/// <param name="Name">Name of the field</param>
+		/// <param name="Value">Value to be written</param>
+		public void WriteDateTime(Utf8String Name, DateTime Value)
+		{
+			WriteFieldHeader(CbFieldType.DateTime, Name);
+			WriteDateTimePayload(Value);
+		}
+
+		/// <summary>
+		/// Writes the payload for a hash
+		/// </summary>
+		/// <param name="Hash"></param>
+		void WriteHashPayload(IoHash Hash)
+		{
+			Span<byte> Buffer = Allocate(IoHash.NumBytes).Span;
+			Hash.Span.CopyTo(Buffer);
+		}
+
+		/// <summary>
+		/// Writes an unnamed <see cref="IoHash"/> field
+		/// </summary>
+		/// <param name="Hash"></param>
+		public void WriteHashValue(IoHash Hash)
+		{
+			WriteFieldHeader(CbFieldType.Hash);
+			WriteHashPayload(Hash);
+		}
+
+		/// <summary>
+		/// Writes a named <see cref="IoHash"/> field
+		/// </summary>
+		/// <param name="Name">Name of the field</param>
+		/// <param name="Value">Value to be written</param>
+		public void WriteHash(Utf8String Name, IoHash Value)
+		{
+			WriteFieldHeader(CbFieldType.Hash, Name);
+			WriteHashPayload(Value);
+		}
+
+		/// <summary>
+		/// Writes an unnamed reference to a binary attachment
+		/// </summary>
+		/// <param name="Hash">Hash of the attachment</param>
+		public void WriteBinaryAttachmentValue(IoHash Hash)
+		{
+			WriteFieldHeader(CbFieldType.BinaryAttachment);
+			WriteHashPayload(Hash);
+		}
+
+		/// <summary>
+		/// Writes a named reference to a binary attachment
+		/// </summary>
+		/// <param name="Name">Name of the field</param>
+		/// <param name="Hash">Hash of the attachment</param>
+		public void WriteBinaryAttachment(Utf8String Name, IoHash Hash)
+		{
+			WriteFieldHeader(CbFieldType.BinaryAttachment, Name);
+			WriteHashPayload(Hash);
+		}
+
+		/// <summary>
+		/// Writes the payload for an object to the buffer
+		/// </summary>
+		/// <param name="Object"></param>
+		void WriteObjectPayload(CbObject Object)
+		{
+			CbField Field = Object.AsField();
+			Memory<byte> Buffer = Allocate(Field.Payload.Length);
+			Field.Payload.CopyTo(Buffer);
+		}
+
+		/// <summary>
+		/// Writes an object directly into the writer
+		/// </summary>
+		/// <param name="Object">Object to write</param>
+		public void WriteObject(CbObject Object)
+		{
+			WriteFieldHeader(CbFieldType.Object);
+			WriteObjectPayload(Object);
+		}
+
+		/// <summary>
+		/// Writes an object directly into the writer
+		/// </summary>
+		/// <param name="Name">Name of the object</param>
+		/// <param name="Object">Object to write</param>
+		public void WriteObject(Utf8String Name, CbObject Object)
+		{
+			WriteFieldHeader(CbFieldType.Object, Name);
+			WriteObjectPayload(Object);
+		}
+
+		/// <summary>
+		/// Writes an unnamed reference to an object attachment
+		/// </summary>
+		/// <param name="Hash">Hash of the attachment</param>
+		public void WriteObjectAttachmentValue(IoHash Hash)
+		{
+			WriteFieldHeader(CbFieldType.ObjectAttachment);
+			WriteHashPayload(Hash);
+		}
+
+		/// <summary>
+		/// Writes a named reference to an object attachment
+		/// </summary>
+		/// <param name="Name">Name of the field</param>
+		/// <param name="Hash">Hash of the attachment</param>
+		public void WriteObjectAttachment(Utf8String Name, IoHash Hash)
+		{
+			WriteFieldHeader(CbFieldType.ObjectAttachment, Name);
+			WriteHashPayload(Hash);
+		}
+
+		/// <summary>
+		/// Writes the payload for a binary value
+		/// </summary>
+		/// <param name="Output">Output buffer</param>
+		/// <param name="Value">Value to be written</param>
+		static void WriteBinaryPayload(Span<byte> Output, ReadOnlySpan<byte> Value)
+		{
+			int VarIntLength = VarInt.Write(Output, Value.Length);
+			Output = Output[VarIntLength..];
+
+			Value.CopyTo(Output);
+			CheckSize(Output, Value.Length);
+		}
+
+		/// <summary>
+		/// Writes the payload for a binary value
+		/// </summary>
+		/// <param name="Value">Value to be written</param>
+		void WriteBinaryPayload(ReadOnlySpan<byte> Value)
+		{
+			int ValueVarIntLength = VarInt.Measure(Value.Length);
+			Span<byte> Buffer = Allocate(ValueVarIntLength + Value.Length).Span;
+			WriteBinaryPayload(Buffer, Value);
+		}
+
+		/// <summary>
+		/// Writes an unnamed string value
+		/// </summary>
+		/// <param name="Value">Value to be written</param>
+		public void WriteStringValue(Utf8String Value)
+		{
+			WriteFieldHeader(CbFieldType.String);
+			WriteBinaryPayload(Value.Span);
+		}
+
+		/// <summary>
+		/// Writes a named string value
+		/// </summary>
+		/// <param name="Name">Name of the field</param>
+		/// <param name="Value">Value to be written</param>
+		public void WriteString(Utf8String Name, Utf8String Value)
+		{
+			WriteFieldHeader(CbFieldType.String, Name);
+			WriteBinaryPayload(Value.Span);
+		}
+
+		/// <summary>
+		/// Writes an unnamed binary value
+		/// </summary>
+		/// <param name="Value">Value to be written</param>
+		public void WriteBinaryValue(ReadOnlySpan<byte> Value)
+		{
+			WriteFieldHeader(CbFieldType.Binary);
+			WriteBinaryPayload(Value);
+		}
+
+		/// <summary>
+		/// Writes a named binary value
+		/// </summary>
+		/// <param name="Name">Name of the field</param>
+		/// <param name="Value">Value to be written</param>
+		public void WriteBinary(Utf8String Name, ReadOnlySpan<byte> Value)
+		{
+			WriteFieldHeader(CbFieldType.Binary, Name);
+			WriteBinaryPayload(Value);
+		}
+
+		/// <summary>
+		/// Check that the given span is the required size
+		/// </summary>
+		/// <param name="Span"></param>
+		/// <param name="ExpectedSize"></param>
+		static void CheckSize(Span<byte> Span, int ExpectedSize)
+		{
+			if (Span.Length != ExpectedSize)
+			{
+				throw new Exception("Size of buffer is not correct");
+			}
+		}
+
+		/// <summary>
+		/// Gets the size of the serialized data
+		/// </summary>
+		/// <returns></returns>
+		public int GetSize()
+		{
+			if (OpenScopes.Count != 1)
+			{
+				throw new CbWriterException("Unfinished scope in writer");
+			}
+
+			return CurrentOffset + ComputeSizeOfChildHeaders(CurrentScope);
+		}
+
+		/// <summary>
+		/// Copy the data from this writer to a buffer
+		/// </summary>
+		/// <param name="Buffer"></param>
+		public void CopyTo(Span<byte> Buffer)
+		{
+			int BufferOffset = 0;
+
+			int SourceOffset = 0;
+			foreach (Chunk Chunk in Chunks)
+			{
+				foreach (Scope Scope in Chunk.Scopes)
+				{
+					ReadOnlySpan<byte> SourceData = Chunk.Data.AsSpan(SourceOffset - Chunk.Offset, Scope.Offset - SourceOffset);
+					SourceData.CopyTo(Buffer.Slice(BufferOffset));
+
+					BufferOffset += SourceData.Length;
+					SourceOffset += SourceData.Length;
+
+					BufferOffset += WriteScopeHeader(Buffer.Slice(BufferOffset), Scope);
+				}
+
+				ReadOnlySpan<byte> LastSourceData = Chunk.Data.AsSpan(SourceOffset - Chunk.Offset, (Chunk.Offset + Chunk.Length) - SourceOffset);
+				LastSourceData.CopyTo(Buffer.Slice(BufferOffset));
+				BufferOffset += LastSourceData.Length;
+				SourceOffset += LastSourceData.Length;
+			}
+		}
+
+		/// <summary>
+		/// Convert the data into a compact binary object
+		/// </summary>
+		/// <returns></returns>
+		public CbObject ToObject()
+		{
+			return new CbObject(ToByteArray());
+		}
+
+		/// <summary>
+		/// Convert the data into a flat array
+		/// </summary>
+		/// <returns></returns>
+		public byte[] ToByteArray()
+		{
+			byte[] Buffer = new byte[GetSize()];
+			CopyTo(Buffer);
+			return Buffer;
+		}
+
+		/// <summary>
+		/// Comptues the size of any child headers
+		/// </summary>
+		/// <param name="Scope"></param>
+		static int ComputeSizeOfChildHeaders(Scope Scope)
+		{
+			int SizeOfChildHeaders = 0;
+			foreach (Scope ChildScope in Scope.Children)
+			{
+				switch (ChildScope.FieldType)
+				{
+					case CbFieldType.Object:
+					case CbFieldType.UniformObject:
+						SizeOfChildHeaders += ChildScope.SizeOfChildHeaders + VarInt.Measure(ChildScope.Length + ChildScope.SizeOfChildHeaders);
+						break;
+					case CbFieldType.Array:
+					case CbFieldType.UniformArray:
+						int ArrayCountLength = VarInt.Measure(ChildScope.Count);
+						SizeOfChildHeaders += ChildScope.SizeOfChildHeaders + VarInt.Measure(ChildScope.Length + ChildScope.SizeOfChildHeaders + ArrayCountLength) + ArrayCountLength;
+						break;
+					default:
+						throw new InvalidOperationException();
+				}
+			}
+			return SizeOfChildHeaders;
+		}
+
+		/// <summary>
+		/// Writes the header for a particular scope
+		/// </summary>
+		/// <param name="Span"></param>
+		/// <param name="Scope"></param>
+		/// <returns></returns>
+		static int WriteScopeHeader(Span<byte> Span, Scope Scope)
+		{
+			switch (Scope.FieldType)
+			{
+				case CbFieldType.Object:
+				case CbFieldType.UniformObject:
+					return VarInt.Write(Span, Scope.Length + Scope.SizeOfChildHeaders);
+				case CbFieldType.Array:
+				case CbFieldType.UniformArray:
+					int NumItemsLength = VarInt.Measure(Scope.Count);
+					int Offset = VarInt.Write(Span, Scope.Length + Scope.SizeOfChildHeaders + NumItemsLength);
+					return Offset + VarInt.Write(Span.Slice(Offset), Scope.Count);
+				default:
+					throw new InvalidOperationException();
+			}
+		}
+	}
+}
\ No newline at end of file
diff --git a/Engine/Source/Programs/Shared/EpicGames.Serialization/CompactBinary.cs b/Engine/Source/Programs/Shared/EpicGames.Serialization/CompactBinary.cs
deleted file mode 100644
diff --git a/Engine/Source/Programs/Shared/EpicGames.Serialization/CompactBinarySerializer.cs b/Engine/Source/Programs/Shared/EpicGames.Serialization/CompactBinarySerializer.cs
deleted file mode 100644
diff --git a/Engine/Source/Programs/Shared/EpicGames.Serialization/CompactBinaryWriter.cs b/Engine/Source/Programs/Shared/EpicGames.Serialization/CompactBinaryWriter.cs
deleted file mode 100644
diff --git a/Engine/Source/Programs/Shared/EpicGames.Serialization/Converters/CbConverter.cs b/Engine/Source/Programs/Shared/EpicGames.Serialization/Converters/CbConverter.cs
deleted file mode 100644
