CS: Add overloads for commonly serialized compact binary types.

diff --git a/Engine/Source/Programs/Shared/EpicGames.Horde/Common/Condition.cs b/Engine/Source/Programs/Shared/EpicGames.Horde/Common/Condition.cs
--- a/Engine/Source/Programs/Shared/EpicGames.Horde/Common/Condition.cs
+++ b/Engine/Source/Programs/Shared/EpicGames.Horde/Common/Condition.cs
@@ -603,13 +603,13 @@ namespace EpicGames.Horde.Common
 	public class ConditionCbConverter : CbConverter<Condition>
 	{
 		/// <inheritdoc/>
-		public override Condition Read(CbField Field) => Condition.TryParse(Field.AsString().ToString());
+		public override Condition Read(CbField Field) => Condition.TryParse(Field.AsUtf8String().ToString());
 
 		/// <inheritdoc/>
-		public override void Write(CbWriter Writer, Condition Value) => Writer.WriteStringValue(Value.Text);
+		public override void Write(CbWriter Writer, Condition Value) => Writer.WriteUtf8StringValue(Value.Text);
 
 		/// <inheritdoc/>
-		public override void WriteNamed(CbWriter Writer, Utf8String Name, Condition Value) => Writer.WriteString(Name, Value.Text);
+		public override void WriteNamed(CbWriter Writer, Utf8String Name, Condition Value) => Writer.WriteUtf8String(Name, Value.Text);
 	}
 
 	/// <summary>
diff --git a/Engine/Source/Programs/Shared/EpicGames.Perforce.Managed/StreamSnapshotFromMemory.cs b/Engine/Source/Programs/Shared/EpicGames.Perforce.Managed/StreamSnapshotFromMemory.cs
--- a/Engine/Source/Programs/Shared/EpicGames.Perforce.Managed/StreamSnapshotFromMemory.cs
+++ b/Engine/Source/Programs/Shared/EpicGames.Perforce.Managed/StreamSnapshotFromMemory.cs
@@ -92,7 +92,7 @@ namespace EpicGames.Perforce.Managed
 			CbObject RootObj = new CbObject(Data.AsMemory(CurrentSignature.Length));
 
 			CbObject RootObj2 = RootObj["root"].AsObject();
-			Utf8String RootPath = RootObj2["path"].AsString(BasePath);
+			Utf8String RootPath = RootObj2["path"].AsUtf8String(BasePath);
 			StreamTreeRef Root = new StreamTreeRef(RootPath, RootObj2);
 
 			CbArray Array = RootObj["items"].AsArray();
@@ -121,7 +121,7 @@ namespace EpicGames.Perforce.Managed
 			Writer.BeginObject("root");
 			if (Root.Path != BasePath)
 			{
-				Writer.WriteString("path", Root.Path);
+				Writer.WriteUtf8String("path", Root.Path);
 			}
 			Root.Write(Writer);
 			Writer.EndObject();
diff --git a/Engine/Source/Programs/Shared/EpicGames.Perforce.Managed/StreamTree.cs b/Engine/Source/Programs/Shared/EpicGames.Perforce.Managed/StreamTree.cs
--- a/Engine/Source/Programs/Shared/EpicGames.Perforce.Managed/StreamTree.cs
+++ b/Engine/Source/Programs/Shared/EpicGames.Perforce.Managed/StreamTree.cs
@@ -69,7 +69,7 @@ namespace EpicGames.Perforce.Managed
 			Length = Field[LengthField].AsInt64();
 
 			Md5Hash Digest = new Md5Hash(Field[DigestField].AsBinary());
-			Utf8String Type = Field[TypeField].AsString();
+			Utf8String Type = Field[TypeField].AsUtf8String();
 			ContentId = new FileContentId(Digest, Type);
 
 			Revision = Field[RevisionField].AsInt32();
@@ -82,8 +82,8 @@ namespace EpicGames.Perforce.Managed
 		public void Write(CbWriter Writer)
 		{
 			Writer.WriteInteger(LengthField, Length);
-			Writer.WriteBinary(DigestField, ContentId.Digest.Span);
-			Writer.WriteString(TypeField, ContentId.Type);
+			Writer.WriteBinarySpan(DigestField, ContentId.Digest.Span);
+			Writer.WriteUtf8String(TypeField, ContentId.Type);
 			Writer.WriteInteger(RevisionField, Revision);
 		}
 	}
@@ -211,7 +211,7 @@ namespace EpicGames.Perforce.Managed
 			{
 				CbObject FileObject = FileField.AsObject();
 
-				Utf8String Name = FileObject[NameField].AsString();
+				Utf8String Name = FileObject[NameField].AsUtf8String();
 				Utf8String FilePath = ReadPath(FileObject, Path, Name);
 				StreamFile File = new StreamFile(FilePath, FileObject);
 
@@ -223,7 +223,7 @@ namespace EpicGames.Perforce.Managed
 			{
 				CbObject TreeObject = TreeField.AsObject();
 
-				Utf8String Name = TreeObject[NameField].AsString();
+				Utf8String Name = TreeObject[NameField].AsUtf8String();
 				Utf8String TreePath = ReadPath(TreeObject, Path, Name);
 				StreamTreeRef Tree = new StreamTreeRef(TreePath, TreeObject);
 
@@ -243,7 +243,7 @@ namespace EpicGames.Perforce.Managed
 				foreach ((Utf8String Name, StreamFile File) in NameToFile.OrderBy(x => x.Key))
 				{
 					Writer.BeginObject();
-					Writer.WriteString(NameField, Name);
+					Writer.WriteUtf8String(NameField, Name);
 					WritePath(Writer, File.Path, Path, Name);
 					File.Write(Writer);
 					Writer.EndObject();
@@ -257,7 +257,7 @@ namespace EpicGames.Perforce.Managed
 				foreach ((Utf8String Name, StreamTreeRef Tree) in NameToTree.OrderBy(x => x.Key))
 				{
 					Writer.BeginObject();
-					Writer.WriteString(NameField, Name);
+					Writer.WriteUtf8String(NameField, Name);
 					WritePath(Writer, Tree.Path, Path, Name);
 					Tree.Write(Writer);
 					Writer.EndObject();
@@ -275,7 +275,7 @@ namespace EpicGames.Perforce.Managed
 		/// <returns></returns>
 		static Utf8String ReadPath(CbObject Object, Utf8String BasePath, Utf8String Name)
 		{
-			Utf8String Path = Object[PathField].AsString();
+			Utf8String Path = Object[PathField].AsUtf8String();
 			if (Path.IsEmpty)
 			{
 				byte[] Data = new byte[BasePath.Length + 1 + Name.Length];
@@ -298,7 +298,7 @@ namespace EpicGames.Perforce.Managed
 		{
 			if (Path.Length != ParentPath.Length + Name.Length + 1 || !Path.StartsWith(ParentPath) || Path[ParentPath.Length] != '/' || !Path.EndsWith(Name))
 			{
-				Writer.WriteString(PathField, Path);
+				Writer.WriteUtf8String(PathField, Path);
 			}
 		}
 
diff --git a/Engine/Source/Programs/Shared/EpicGames.Serialization.Tests/EngineTests.cs b/Engine/Source/Programs/Shared/EpicGames.Serialization.Tests/EngineTests.cs
--- a/Engine/Source/Programs/Shared/EpicGames.Serialization.Tests/EngineTests.cs
+++ b/Engine/Source/Programs/Shared/EpicGames.Serialization.Tests/EngineTests.cs
@@ -51,7 +51,7 @@ namespace EpicGames.Serialization.Tests
 			[CbFieldType.Array] = new CbFieldAccessors(CbArray.Empty, x => x.IsArray(), x => x.AsArray(), (x, y) => x.AsArray()),
 			[CbFieldType.UniformArray] = new CbFieldAccessors(CbArray.Empty, x => x.IsArray(), x => x.AsArray(), (x, y) => x.AsArray()),
 			[CbFieldType.Binary] = CbFieldAccessors.FromStruct<ReadOnlyMemory<byte>>(x => x.IsBinary(), (x, y) => x.AsBinary(y), (x, y) => x.Span.SequenceEqual(y.Span)),
-			[CbFieldType.String] = new CbFieldAccessors(Utf8String.Empty, x => x.IsString(), x => x.AsString(), (x, Default) => x.AsString((Utf8String)Default)),
+			[CbFieldType.String] = new CbFieldAccessors(Utf8String.Empty, x => x.IsString(), x => x.AsUtf8String(), (x, Default) => x.AsUtf8String((Utf8String)Default)),
 			[CbFieldType.IntegerPositive] = CbFieldAccessors.FromStruct<ulong>(x => x.IsInteger(), (x, y) => x.AsUInt64(y)),
 			[CbFieldType.IntegerNegative] = CbFieldAccessors.FromStruct<long>(x => x.IsInteger(), (x, y) => x.AsInt64(y)),
 			[CbFieldType.Float32] = CbFieldAccessors.FromStruct<float>(x => x.IsFloat(), (x, y) => x.AsFloat(y)),
diff --git a/Engine/Source/Programs/Shared/EpicGames.Serialization.Tests/JupiterTests.cs b/Engine/Source/Programs/Shared/EpicGames.Serialization.Tests/JupiterTests.cs
--- a/Engine/Source/Programs/Shared/EpicGames.Serialization.Tests/JupiterTests.cs
+++ b/Engine/Source/Programs/Shared/EpicGames.Serialization.Tests/JupiterTests.cs
@@ -34,7 +34,7 @@ namespace EpicGames.Serialization.Tests
 			CbWriter Writer1 = new CbWriter();
 			Writer1.BeginObject();
 			Writer1.WriteInteger("a", 1);
-			Writer1.WriteString("b", "hello");
+			Writer1.WriteUtf8String("b", "hello");
 			Writer1.EndObject();
 
 			CbObject Object1 = Writer1.ToObject();
@@ -136,7 +136,7 @@ namespace EpicGames.Serialization.Tests
 
 			CbWriter writer = new CbWriter();
 			writer.BeginObject();
-			writer.WriteString("string", "test");
+			writer.WriteUtf8String("string", "test");
 			writer.WriteBinaryAttachment("hash", hash1);
 			writer.EndObject();
 
@@ -150,7 +150,7 @@ namespace EpicGames.Serialization.Tests
 			Assert.AreEqual(2, fields.Count);
 
 			CbField? stringField = o["string"];
-			Assert.AreEqual("test", stringField!.AsString());
+			Assert.AreEqual("test", stringField!.AsUtf8String());
 
 			CbField? hashField = o["hash"];
 			Assert.AreEqual(hash1, hashField!.AsAttachment());
diff --git a/Engine/Source/Programs/Shared/EpicGames.Serialization/CbField.cs b/Engine/Source/Programs/Shared/EpicGames.Serialization/CbField.cs
--- a/Engine/Source/Programs/Shared/EpicGames.Serialization/CbField.cs
+++ b/Engine/Source/Programs/Shared/EpicGames.Serialization/CbField.cs
@@ -869,7 +869,16 @@ namespace EpicGames.Serialization
 		/// Access the field as binary data.
 		/// </summary>
 		/// <returns></returns>
-		public ReadOnlyMemory<byte> AsBinary(ReadOnlyMemory<byte> Default = default)
+		public ReadOnlyMemory<byte> AsBinary()
+		{
+			return AsBinary(ReadOnlyMemory<byte>.Empty);
+		}
+
+		/// <summary>
+		/// Access the field as binary data.
+		/// </summary>
+		/// <returns></returns>
+		public ReadOnlyMemory<byte> AsBinary(ReadOnlyMemory<byte> Default)
 		{
 			if (CbFieldUtils.IsBinary(TypeWithFlags))
 			{
@@ -885,13 +894,43 @@ namespace EpicGames.Serialization
 			}
 		}
 
+		/// <summary>
+		/// Access the field as binary data.
+		/// </summary>
+		/// <returns></returns>
+		public byte[] AsBinaryArray()
+		{
+			return AsBinaryByteArray(Array.Empty<byte>());
+		}
+
+		/// <summary>
+		/// Access the field as binary data.
+		/// </summary>
+		/// <returns></returns>
+		public byte[] AsBinaryArray(byte[] Default)
+		{
+			return AsBinary(Default).ToArray();
+		}
+
+		/// <summary>
+		/// Access the field as a UTF-8 string.
+		/// </summary>
+		/// <returns></returns>
+		public string AsString() => AsUtf8String().ToString();
+
+		/// <summary>
+		/// Access the field as a UTF-8 string.
+		/// </summary>
+		/// <returns></returns>
+		public string AsString(string Default) => AsUtf8String(Default).ToString();
+
 		/// <summary>
 		/// Access the field as a UTF-8 string.
 		/// </summary>
 		/// <returns></returns>
-		public Utf8String AsString()
+		public Utf8String AsUtf8String()
 		{
-			return AsString(default);
+			return AsUtf8String(default);
 		}
 
 		/// <summary>
@@ -899,7 +938,7 @@ namespace EpicGames.Serialization
 		/// </summary>
 		/// <param name="Default">Default value to return</param>
 		/// <returns></returns>
-		public Utf8String AsString(Utf8String Default)
+		public Utf8String AsUtf8String(Utf8String Default)
 		{
 			if (CbFieldUtils.IsString(TypeWithFlags))
 			{
@@ -1728,11 +1767,11 @@ namespace EpicGames.Serialization
 				case CbFieldType.String:
 					if (Field.HasName())
 					{
-						Writer.WriteString(Field.Name.Span, Field.AsString().Span);
+						Writer.WriteString(Field.Name.Span, Field.AsUtf8String().Span);
 					}
 					else
 					{
-						Writer.WriteStringValue(Field.AsString().Span);
+						Writer.WriteStringValue(Field.AsUtf8String().Span);
 					}
 					break;
 				case CbFieldType.IntegerPositive:
@@ -2634,7 +2673,7 @@ namespace EpicGames.Serialization
 			}
 			else if (Field.IsString())
 			{
-				Writer.WriteString(Field.Name.Span, Field.AsString().Span);
+				Writer.WriteString(Field.Name.Span, Field.AsUtf8String().Span);
 			}
 			else
 			{
diff --git a/Engine/Source/Programs/Shared/EpicGames.Serialization/CbWriter.cs b/Engine/Source/Programs/Shared/EpicGames.Serialization/CbWriter.cs
--- a/Engine/Source/Programs/Shared/EpicGames.Serialization/CbWriter.cs
+++ b/Engine/Source/Programs/Shared/EpicGames.Serialization/CbWriter.cs
@@ -463,6 +463,15 @@ namespace EpicGames.Serialization
 			VarInt.Write(Buffer, Value);
 		}
 
+		/// <summary>
+		/// Writes an unnamed integer field
+		/// </summary>
+		/// <param name="Value">Value to be written</param>
+		public void WriteIntegerValue(int Value)
+		{
+			WriteIntegerValue((long)Value);
+		}
+
 		/// <summary>
 		/// Writes an unnamed integer field
 		/// </summary>
@@ -481,6 +490,16 @@ namespace EpicGames.Serialization
 			}
 		}
 
+		/// <summary>
+		/// Writes an named integer field
+		/// </summary>
+		/// <param name="Name">Name of the field</param>
+		/// <param name="Value">Value to be written</param>
+		public void WriteInteger(Utf8String Name, int Value)
+		{
+			WriteInteger(Name, (long)Value);
+		}
+
 		/// <summary>
 		/// Writes an named integer field
 		/// </summary>
@@ -707,7 +726,20 @@ namespace EpicGames.Serialization
 		/// Writes an unnamed string value
 		/// </summary>
 		/// <param name="Value">Value to be written</param>
-		public void WriteStringValue(Utf8String Value)
+		public void WriteStringValue(string Value) => WriteUtf8StringValue(Value);
+
+		/// <summary>
+		/// Writes a named string value
+		/// </summary>
+		/// <param name="Name">Name of the field</param>
+		/// <param name="Value">Value to be written</param>
+		public void WriteString(Utf8String Name, string Value) => WriteUtf8String(Name, Value);
+
+		/// <summary>
+		/// Writes an unnamed string value
+		/// </summary>
+		/// <param name="Value">Value to be written</param>
+		public void WriteUtf8StringValue(Utf8String Value)
 		{
 			WriteFieldHeader(CbFieldType.String);
 			WriteBinaryPayload(Value.Span);
@@ -718,7 +750,7 @@ namespace EpicGames.Serialization
 		/// </summary>
 		/// <param name="Name">Name of the field</param>
 		/// <param name="Value">Value to be written</param>
-		public void WriteString(Utf8String Name, Utf8String Value)
+		public void WriteUtf8String(Utf8String Name, Utf8String Value)
 		{
 			WriteFieldHeader(CbFieldType.String, Name);
 			WriteBinaryPayload(Value.Span);
@@ -728,7 +760,7 @@ namespace EpicGames.Serialization
 		/// Writes an unnamed binary value
 		/// </summary>
 		/// <param name="Value">Value to be written</param>
-		public void WriteBinaryValue(ReadOnlySpan<byte> Value)
+		public void WriteBinarySpanValue(ReadOnlySpan<byte> Value)
 		{
 			WriteFieldHeader(CbFieldType.Binary);
 			WriteBinaryPayload(Value);
@@ -739,12 +771,50 @@ namespace EpicGames.Serialization
 		/// </summary>
 		/// <param name="Name">Name of the field</param>
 		/// <param name="Value">Value to be written</param>
-		public void WriteBinary(Utf8String Name, ReadOnlySpan<byte> Value)
+		public void WriteBinarySpan(Utf8String Name, ReadOnlySpan<byte> Value)
 		{
 			WriteFieldHeader(CbFieldType.Binary, Name);
 			WriteBinaryPayload(Value);
 		}
 
+		/// <summary>
+		/// Writes an unnamed binary value
+		/// </summary>
+		/// <param name="Value">Value to be written</param>
+		public void WriteBinaryValue(ReadOnlyMemory<byte> Value)
+		{
+			WriteBinarySpanValue(Value.Span);
+		}
+
+		/// <summary>
+		/// Writes a named binary value
+		/// </summary>
+		/// <param name="Name">Name of the field</param>
+		/// <param name="Value">Value to be written</param>
+		public void WriteBinary(Utf8String Name, ReadOnlyMemory<byte> Value)
+		{
+			WriteBinarySpan(Name, Value.Span);
+		}
+
+		/// <summary>
+		/// Writes an unnamed binary value
+		/// </summary>
+		/// <param name="Value">Value to be written</param>
+		public void WriteBinaryArrayValue(byte[] Value)
+		{
+			WriteBinarySpanValue(Value.AsSpan());
+		}
+
+		/// <summary>
+		/// Writes a named binary value
+		/// </summary>
+		/// <param name="Name">Name of the field</param>
+		/// <param name="Value">Value to be written</param>
+		public void WriteBinaryArray(Utf8String Name, byte[] Value)
+		{
+			WriteBinarySpan(Name, Value.AsSpan());
+		}
+
 		/// <summary>
 		/// Check that the given span is the required size
 		/// </summary>
diff --git a/Engine/Source/Programs/Shared/EpicGames.Serialization/Converters/CbDefaultConverter.cs b/Engine/Source/Programs/Shared/EpicGames.Serialization/Converters/CbDefaultConverter.cs
--- a/Engine/Source/Programs/Shared/EpicGames.Serialization/Converters/CbDefaultConverter.cs
+++ b/Engine/Source/Programs/Shared/EpicGames.Serialization/Converters/CbDefaultConverter.cs
@@ -98,39 +98,20 @@ namespace EpicGames.Serialization.Converters
 			}
 		}
 
-		static void WriteInt32Value(CbWriter Writer, int Value) => Writer.WriteIntegerValue(Value);
-		static void WriteInt32(CbWriter Writer, Utf8String Name, int Value) => Writer.WriteInteger(Name, Value);
-		static void WriteInt64Value(CbWriter Writer, long Value) => Writer.WriteIntegerValue(Value);
-		static void WriteInt64(CbWriter Writer, Utf8String Name, long Value) => Writer.WriteInteger(Name, Value);
-		static void WriteDoubleValue(CbWriter Writer, double Value) => Writer.WriteDoubleValue(Value);
-		static void WriteDouble(CbWriter Writer, Utf8String Name, double Value) => Writer.WriteDouble(Name, Value);
-
-		static string ReadString(CbField Field) => Field.AsString().ToString();
-		static void WriteString(CbWriter Writer, Utf8String Name, string Value) => Writer.WriteString(Name, Value);
-		static void WriteStringValue(CbWriter Writer, string Value) => Writer.WriteStringValue(Value);
-
-		static ReadOnlyMemory<byte> ReadMemory(CbField Field) => Field.AsBinary();
-		static void WriteMemory(CbWriter Writer, Utf8String Name, ReadOnlyMemory<byte> Value) => Writer.WriteBinary(Name, Value.Span);
-		static void WriteMemoryValue(CbWriter Writer, ReadOnlyMemory<byte> Value) => Writer.WriteBinaryValue(Value.Span);
-
-		static byte[] ReadByteArray(CbField Field) => Field.AsBinary().ToArray();
-		static void WriteByteArray(CbWriter Writer, Utf8String Name, byte[] Value) => Writer.WriteBinary(Name, Value);
-		static void WriteByteArrayValue(CbWriter Writer, byte[] Value) => Writer.WriteBinaryValue(Value);
-
 		static Dictionary<Type, CbConverterInfo> TypeToConverterInfo = new Dictionary<Type, CbConverterInfo>(new KeyValuePair<Type, CbConverterInfo>[]
 		{
 			GetPodConverterInfo(x => x.AsBool(), (w, v) => w.WriteBoolValue(v), (w, n, v) => w.WriteBool(n, v), SkipIfNullOrZero),
-			GetPodConverterInfo(x => x.AsInt32(), (w, v) => WriteInt32Value(w, v), (w, n, v) => WriteInt32(w, n, v), SkipIfNullOrZero),
-			GetPodConverterInfo(x => x.AsInt64(), (w, v) => WriteInt64Value(w, v), (w, n, v) => WriteInt64(w, n, v), SkipIfNullOrZero),
-			GetPodConverterInfo(x => x.AsDouble(), (w, v) => WriteDoubleValue(w, v), (w, n, v)=> WriteDouble(w, n, v), SkipIfNullOrZero),
-			GetPodConverterInfo(x => x.AsString(), (w, v) => w.WriteStringValue(v), (w, n, v) => w.WriteString(n, v), null),
+			GetPodConverterInfo(x => x.AsInt32(), (w, v) => w.WriteIntegerValue(v), (w, n, v) => w.WriteInteger(n, v), SkipIfNullOrZero),
+			GetPodConverterInfo(x => x.AsInt64(), (w, v) => w.WriteIntegerValue(v), (w, n, v) => w.WriteInteger(n, v), SkipIfNullOrZero),
+			GetPodConverterInfo(x => x.AsDouble(), (w, v) => w.WriteDoubleValue(v), (w, n, v)=> w.WriteDouble(n, v), SkipIfNullOrZero),
+			GetPodConverterInfo(x => x.AsUtf8String(), (w, v) => w.WriteUtf8StringValue(v), (w, n, v) => w.WriteUtf8String(n, v), null),
 			GetPodConverterInfo(x => x.AsHash(), (w, v) => w.WriteHashValue(v), (w, n, v) => w.WriteHash(n, v), null),
 			GetPodConverterInfo(x => x.AsObjectAttachment(), (w, v) => w.WriteObjectAttachmentValue(v.Hash), (w, n, v) => w.WriteObjectAttachment(n, v.Hash), null),
 			GetPodConverterInfo(x => x.AsBinaryAttachment(), (w, v) => w.WriteBinaryAttachmentValue(v.Hash), (w, n, v) => w.WriteBinaryAttachment(n, v.Hash), null),
 			GetPodConverterInfo(x => x.AsDateTime(), (w, v) => w.WriteDateTimeValue(v), (w, n, v) => w.WriteDateTime(n, v), null),
-			GetPodConverterInfo(x => ReadString(x), (w, v) => WriteStringValue(w, v), (w, n, v) => WriteString(w, n, v), SkipIfNullOrZero),
-			GetPodConverterInfo(x => ReadMemory(x), (w, v) => WriteMemoryValue(w, v), (w, n, v) => WriteMemory(w, n, v), null),
-			GetPodConverterInfo(x => ReadByteArray(x), (w, v) => WriteByteArrayValue(w, v), (w, n, v) => WriteByteArray(w, n, v), null),
+			GetPodConverterInfo(x => x.AsString(), (w, v) => w.WriteStringValue(v), (w, n, v) => w.WriteString(n, v), SkipIfNullOrZero),
+			GetPodConverterInfo(x => x.AsBinary(), (w, v) => w.WriteBinaryValue(v), (w, n, v) => w.WriteBinary(n, v), null),
+			GetPodConverterInfo(x => x.AsBinaryArray(), (w, v) => w.WriteBinaryArrayValue(v), (w, n, v) => w.WriteBinaryArray(n, v), null),
 		});
 
 		static KeyValuePair<Type, CbConverterInfo> GetPodConverterInfo<T>(Expression<Func<CbField, T>> Read, Expression<Action<CbWriter, T>> Write, Expression<Action<CbWriter, Utf8String, T>> WriteNamed, SkipDefaultDelegate? SkipDefault)
diff --git a/Engine/Source/Programs/Shared/EpicGames.Serialization/Converters/CbListConverter.cs b/Engine/Source/Programs/Shared/EpicGames.Serialization/Converters/CbListConverter.cs
--- a/Engine/Source/Programs/Shared/EpicGames.Serialization/Converters/CbListConverter.cs
+++ b/Engine/Source/Programs/Shared/EpicGames.Serialization/Converters/CbListConverter.cs
@@ -67,7 +67,7 @@ namespace EpicGames.Serialization.Converters
 			List<Utf8String> List = new List<Utf8String>();
 			foreach (CbField ElementField in Field)
 			{
-				List.Add(ElementField.AsString());
+				List.Add(ElementField.AsUtf8String());
 			}
 			return List;
 		}
@@ -78,7 +78,7 @@ namespace EpicGames.Serialization.Converters
 			Writer.BeginUniformArray(CbFieldType.String);
 			foreach (Utf8String String in List)
 			{
-				Writer.WriteStringValue(String);
+				Writer.WriteUtf8StringValue(String);
 			}
 			Writer.EndUniformArray();
 		}
@@ -91,7 +91,7 @@ namespace EpicGames.Serialization.Converters
 				Writer.BeginUniformArray(Name, CbFieldType.String);
 				foreach (Utf8String String in List)
 				{
-					Writer.WriteStringValue(String);
+					Writer.WriteUtf8StringValue(String);
 				}
 				Writer.EndUniformArray();
 			}
