Horde: Fix build error.

diff --git a/Engine/Source/Programs/Shared/EpicGames.Serialization/CbField.cs b/Engine/Source/Programs/Shared/EpicGames.Serialization/CbField.cs
--- a/Engine/Source/Programs/Shared/EpicGames.Serialization/CbField.cs
+++ b/Engine/Source/Programs/Shared/EpicGames.Serialization/CbField.cs
@@ -900,7 +900,7 @@ namespace EpicGames.Serialization
 		/// <returns></returns>
 		public byte[] AsBinaryArray()
 		{
-			return AsBinaryByteArray(Array.Empty<byte>());
+			return AsBinaryArray(Array.Empty<byte>());
 		}
 
 		/// <summary>
