CS: Support for serialization and deserialization of polymorphic class hierarchies.

The base class in a hierarchy must have the [CbPolymorphic] attribute, and derived classes must have the [CbDiscriminator] attribute, giving an identifier to indicate the derived type.

diff --git a/Engine/Source/Programs/Shared/EpicGames.Serialization/CbSerializer.cs b/Engine/Source/Programs/Shared/EpicGames.Serialization/CbSerializer.cs
--- a/Engine/Source/Programs/Shared/EpicGames.Serialization/CbSerializer.cs
+++ b/Engine/Source/Programs/Shared/EpicGames.Serialization/CbSerializer.cs
@@ -42,6 +42,30 @@ namespace EpicGames.Serialization
 		}
 	}
 
+	/// <summary>
+	/// Attribute used to indicate that this object is the base for a class hierarchy. Each derived class must have a [CbDiscriminator] attribute.
+	/// </summary>
+	public class CbPolymorphicAttribute : Attribute
+	{
+	}
+
+	/// <summary>
+	/// Sets the name used for discriminating between derived classes during serialization
+	/// </summary>
+	public class CbDiscriminatorAttribute : Attribute
+	{
+		/// <summary>
+		/// Name used to identify this class
+		/// </summary>
+		public string Name { get; }
+
+		/// <summary>
+		/// Constructor
+		/// </summary>
+		/// <param name="Name">Name used to identify this class</param>
+		public CbDiscriminatorAttribute(string Name) => this.Name = Name;
+	}
+
 	/// <summary>
 	/// Exception thrown when serializing cb objects
 	/// </summary>
diff --git a/Engine/Source/Programs/Shared/EpicGames.Serialization/Converters/CbDefaultConverter.cs b/Engine/Source/Programs/Shared/EpicGames.Serialization/Converters/CbDefaultConverter.cs
--- a/Engine/Source/Programs/Shared/EpicGames.Serialization/Converters/CbDefaultConverter.cs
+++ b/Engine/Source/Programs/Shared/EpicGames.Serialization/Converters/CbDefaultConverter.cs
@@ -7,6 +7,7 @@ using System.Linq;
 using System.Linq.Expressions;
 using System.Reflection;
 using System.Reflection.Emit;
+using System.Runtime.Serialization;
 using System.Text;
 
 namespace EpicGames.Serialization.Converters
@@ -148,39 +149,47 @@ namespace EpicGames.Serialization.Converters
 		{
 			lock (LockObject)
 			{
-				CbConverterInfo? ConverterInfo;
-				if (!TypeToConverterInfo.TryGetValue(Type, out ConverterInfo))
+				if (Type.IsClass)
 				{
-					if (Type.IsClass)
-					{
-						DynamicMethod ReadMethod = CreateObjectReader(Type);
-						DynamicMethod WriteMethod = CreateObjectWriter(Type);
-						DynamicMethod WriteNamedMethod = CreateNamedObjectWriter(Type);
+					return CreateClassConverter(Type);
+				}
+				else if (Type.IsEnum)
+				{
+					return CreateEnumConverter(Type);
+				}
+				else
+				{
+					throw new NotSupportedException($"Unable to serialize type {Type}");
+				}
+			}
+		}
 
-						Type GenericType = typeof(CbDefaultConverter<>).MakeGenericType(Type);
-						CbConverter Converter = (CbConverter)Activator.CreateInstance(GenericType, ReadMethod, WriteMethod, WriteNamedMethod)!;
+		CbConverter CreateClassConverter(Type Type)
+		{
+			CbDynamicClassMethods Methods = CreateDynamicClassMethods(Type);
 
-						ConverterInfo = new CbConverterInfo(Converter, ReadMethod, WriteMethod, WriteNamedMethod, SkipIfNullOrZero);
-					}
-					else if (Type.IsEnum)
-					{
-						DynamicMethod ReadMethod = CreateEnumReader(Type);
-						DynamicMethod WriteMethod = CreateEnumWriter(Type);
-						DynamicMethod WriteNamedMethod = CreateNamedEnumWriter(Type);
+			Type GenericType = typeof(CbDefaultConverter<>).MakeGenericType(Type);
+			CbConverter Converter = (CbConverter)Activator.CreateInstance(GenericType, Methods.ReadMethod, Methods.WriteMethod, Methods.WriteNamedMethod)!;
 
-						Type GenericType = typeof(CbDefaultConverter<>).MakeGenericType(Type);
-						CbConverter Converter = (CbConverter)Activator.CreateInstance(GenericType, ReadMethod, WriteMethod, WriteNamedMethod)!;
+			CbConverterInfo ConverterInfo = new CbConverterInfo(Converter, Methods.ReadMethod, Methods.WriteMethod, Methods.WriteNamedMethod, SkipIfNullOrZero);
+			TypeToConverterInfo.Add(Type, ConverterInfo);
 
-						ConverterInfo = new CbConverterInfo(Converter, ReadMethod, WriteMethod, WriteNamedMethod, SkipIfNullOrZero);
-					}
-					else
-					{
-						throw new InvalidOperationException($"Unable to serialize type {Type}");
-					}
-					TypeToConverterInfo.Add(Type, ConverterInfo);
-				}
-				return ConverterInfo.Converter;
-			}
+			return Converter;
+		}
+
+		CbConverter CreateEnumConverter(Type Type)
+		{
+			DynamicMethod ReadMethod = CreateEnumReader(Type);
+			DynamicMethod WriteMethod = CreateEnumWriter(Type);
+			DynamicMethod WriteNamedMethod = CreateNamedEnumWriter(Type);
+
+			Type GenericType = typeof(CbDefaultConverter<>).MakeGenericType(Type);
+			CbConverter Converter = (CbConverter)Activator.CreateInstance(GenericType, ReadMethod, WriteMethod, WriteNamedMethod)!;
+
+			CbConverterInfo ConverterInfo = new CbConverterInfo(Converter, ReadMethod, WriteMethod, WriteNamedMethod, SkipIfNullOrZero);
+			TypeToConverterInfo.Add(Type, ConverterInfo);
+
+			return Converter;
 		}
 
 		static void SkipIfNullOrZero(ILGenerator Generator, Label SkipLabel)
@@ -263,280 +272,412 @@ namespace EpicGames.Serialization.Converters
 			return DynamicMethod;
 		}
 
-		static Dictionary<Type, DynamicMethod> WriteMethods = new Dictionary<Type, DynamicMethod>();
-
-		static DynamicMethod CreateObjectWriter(Type Type)
+		class CbDynamicClassMethods
 		{
-			DynamicMethod? DynamicMethod;
-			if (!WriteMethods.TryGetValue(Type, out DynamicMethod))
+			public Type ClassType { get; }
+			public bool IsPolymorphic { get; }
+
+			public DynamicMethod ReadMethod;
+			public DynamicMethod WriteMethod;
+			public DynamicMethod WriteNamedMethod;
+			public DynamicMethod WriteContentsMethod;
+
+			public DynamicMethod ReadConcreteMethod;
+			public DynamicMethod WriteConcreteContentsMethod;
+
+			public CbDynamicClassMethods(Type ClassType)
 			{
-				// Create the new method
-				DynamicMethod = new DynamicMethod("_", null, new Type[] { typeof(CbWriter), Type });
-				WriteMethods.Add(Type, DynamicMethod);
+				this.ClassType = ClassType;
+				this.IsPolymorphic = (ClassType.GetCustomAttribute<CbPolymorphicAttribute>(true) != null);
 
-				// Implement the body
-				ILGenerator Generator = DynamicMethod.GetILGenerator();
-				Generator.Emit(OpCodes.Ldarg_1);
+				ReadConcreteMethod = new DynamicMethod($"ReadConcrete_{ClassType.Name}", ClassType, new Type[] { typeof(CbField) });
+				WriteConcreteContentsMethod = new DynamicMethod($"WriteConcreteContents_{ClassType.Name}", null, new Type[] { typeof(CbWriter), ClassType });
 
-				Label SkipLabel = Generator.DefineLabel();
-				Generator.Emit(OpCodes.Brfalse, SkipLabel);
+				WriteMethod = new DynamicMethod($"Write_{ClassType.Name}", null, new Type[] { typeof(CbWriter), ClassType });
+				WriteNamedMethod = new DynamicMethod($"WriteNamed_{ClassType.Name}", null, new Type[] { typeof(CbWriter), typeof(Utf8String), ClassType });
 
-				Generator.Emit(OpCodes.Ldarg_0);
-				Generator.EmitCall(OpCodes.Call, GetMethodInfo<CbWriter>(x => x.BeginObject()), null);
+				if (IsPolymorphic)
+				{
+					ReadMethod = new DynamicMethod($"Read_{ClassType.Name}", ClassType, new Type[] { typeof(CbField) });
+					WriteContentsMethod = new DynamicMethod($"WriteContents_{ClassType.Name}", null, new Type[] { typeof(CbWriter), ClassType });
+				}
+				else
+				{
+					ReadMethod = ReadConcreteMethod;
+					WriteContentsMethod = WriteConcreteContentsMethod;
+				}
+			}
 
-				Generator.Emit(OpCodes.Ldarg_0);
-				Generator.Emit(OpCodes.Ldarg_1);
-				Generator.EmitCall(OpCodes.Call, CreateObjectContentsWriter(Type), null);
+			public void GenerateBytecode()
+			{
+				// Create the regular methods
+				CreateConcreteObjectReader(ClassType, ReadConcreteMethod.GetILGenerator());
+				CreateConcreteObjectContentsWriter(ClassType, WriteConcreteContentsMethod.GetILGenerator());
 
-				Generator.Emit(OpCodes.Ldarg_0);
-				Generator.EmitCall(OpCodes.Call, GetMethodInfo<CbWriter>(x => x.EndObject()), null);
+				CreateObjectWriter(ClassType, WriteMethod.GetILGenerator(), WriteContentsMethod);
+				CreateNamedObjectWriter(ClassType, WriteNamedMethod.GetILGenerator(), WriteContentsMethod);
 
-				Generator.MarkLabel(SkipLabel);
-				Generator.Emit(OpCodes.Ret);
+				// Create the extra polymorphic methods
+				if (IsPolymorphic)
+				{
+					// Create the dispatch type
+					Type DispatchType = typeof(CbPolymorphicDispatch<>).MakeGenericType(ClassType);
+
+					// Create the read dispatch method
+					{
+						ILGenerator Generator = ReadMethod.GetILGenerator();
+						Generator.Emit(OpCodes.Ldarg_0);
+						Generator.Emit(OpCodes.Call, DispatchType.GetMethod(nameof(CbPolymorphicDispatch<object>.Read))!);
+						Generator.Emit(OpCodes.Ret);
+					}
+
+					// Create the write dispatch method
+					{
+						ILGenerator Generator = WriteContentsMethod.GetILGenerator();
+						Generator.Emit(OpCodes.Ldarg_0);
+						Generator.Emit(OpCodes.Ldarg_1);
+						Generator.Emit(OpCodes.Call, DispatchType.GetMethod(nameof(CbPolymorphicDispatch<object>.WriteContents))!);
+						Generator.Emit(OpCodes.Ret);
+					}
+
+					// Finally, update the dispatch type with all the methods. We should be safe on the recursive path now.
+					PopulateDispatchType(ClassType, DispatchType);
+				}
 			}
-			return DynamicMethod;
 		}
 
-		static Dictionary<Type, DynamicMethod> WriteNamedMethods = new Dictionary<Type, DynamicMethod>();
+		static Dictionary<Type, CbDynamicClassMethods> TypeToClassMethods = new Dictionary<Type, CbDynamicClassMethods>();
 
-		static DynamicMethod CreateNamedObjectWriter(Type Type)
+		static CbDynamicClassMethods CreateDynamicClassMethods(Type ClassType)
 		{
-			DynamicMethod? DynamicMethod;
-			if (!WriteNamedMethods.TryGetValue(Type, out DynamicMethod))
+			CbDynamicClassMethods? Methods;
+			if (!TypeToClassMethods.TryGetValue(ClassType, out Methods))
 			{
-				// Create the new method
-				DynamicMethod = new DynamicMethod("_", null, new Type[] { typeof(CbWriter), typeof(Utf8String), Type });
-				WriteNamedMethods.Add(Type, DynamicMethod);
+				Methods = new CbDynamicClassMethods(ClassType);
+				TypeToClassMethods[ClassType] = Methods;
+				Methods.GenerateBytecode();
+			}
+			return Methods;
+		}
 
-				// Implement the body
-				ILGenerator Generator = DynamicMethod.GetILGenerator();
-				Generator.Emit(OpCodes.Ldarg_2);
+		static void PopulateDispatchType(Type ClassType, Type DispatchType)
+		{
+			Dictionary<Utf8String, Type> DiscriminatorToKnownType = new Dictionary<Utf8String, Type>();
 
-				Label SkipLabel = Generator.DefineLabel();
-				Generator.Emit(OpCodes.Brfalse, SkipLabel);
+			Type[] KnownTypes = ClassType.Assembly.GetTypes();
+			foreach (Type KnownType in KnownTypes)
+			{
+				if (KnownType.IsClass && !KnownType.IsAbstract)
+				{
+					for (Type? BaseType = KnownType; BaseType != null; BaseType = BaseType.BaseType)
+					{
+						if (BaseType == ClassType)
+						{
+							CbDiscriminatorAttribute? Discriminator = KnownType.GetCustomAttribute<CbDiscriminatorAttribute>();
+							if (Discriminator == null)
+							{
+								throw new NotSupportedException();
+							}
+							DiscriminatorToKnownType[Discriminator.Name] = KnownType;
+						}
+					}
+				}
+			}
 
-				Generator.Emit(OpCodes.Ldarg_0);
-				Generator.Emit(OpCodes.Ldarg_1);
-				Generator.EmitCall(OpCodes.Call, GetMethodInfo<CbWriter>(x => x.BeginObject(null!)), null);
+			// Populate the dictionary
+			Dictionary<Utf8String, Func<CbField, object>> NameToReadFunc = (Dictionary<Utf8String, Func<CbField, object>>)DispatchType.GetField(nameof(CbPolymorphicDispatch<object>.NameToReadFunc))!.GetValue(null)!;
+			Dictionary<Type, Action<CbWriter, object>> TypeToWriteContentsFunc = (Dictionary<Type, Action<CbWriter, object>>)DispatchType.GetField(nameof(CbPolymorphicDispatch<object>.TypeToWriteContentsFunc))!.GetValue(null)!;
 
-				Generator.Emit(OpCodes.Ldarg_0);
-				Generator.Emit(OpCodes.Ldarg_2);
-				Generator.EmitCall(OpCodes.Call, CreateObjectContentsWriter(Type), null);
+			foreach ((Utf8String Name, Type KnownType) in DiscriminatorToKnownType)
+			{
+				CbDynamicClassMethods Methods = CreateDynamicClassMethods(KnownType);
 
-				Generator.Emit(OpCodes.Ldarg_0);
-				Generator.EmitCall(OpCodes.Call, GetMethodInfo<CbWriter>(x => x.EndObject()), null);
+				{
+					DynamicMethod DynamicMethod = new DynamicMethod("_", typeof(object), new Type[] { typeof(CbField) });
+					ILGenerator Generator = DynamicMethod.GetILGenerator();
+					Generator.Emit(OpCodes.Ldarg_0);
+					Generator.Emit(OpCodes.Call, Methods.ReadConcreteMethod);
+					Generator.Emit(OpCodes.Castclass, typeof(object));
+					Generator.Emit(OpCodes.Ret);
+					NameToReadFunc[Name] = CreateDelegate<Func<CbField, object>>(DynamicMethod);
+				}
 
-				Generator.MarkLabel(SkipLabel);
-				Generator.Emit(OpCodes.Ret);
+				{
+					DynamicMethod DynamicMethod = new DynamicMethod("_", null, new Type[] { typeof(CbWriter), typeof(object) });
+					ILGenerator Generator = DynamicMethod.GetILGenerator();
+					Generator.Emit(OpCodes.Ldarg_0);
+					Generator.Emit(OpCodes.Ldarg_1);
+					Generator.Emit(OpCodes.Castclass, KnownType);
+					Generator.Emit(OpCodes.Call, Methods.WriteConcreteContentsMethod);
+					Generator.Emit(OpCodes.Ret);
+					TypeToWriteContentsFunc[KnownType] = CreateDelegate<Action<CbWriter, object>>(DynamicMethod);
+				}
 			}
-			return DynamicMethod;
 		}
 
-		static Dictionary<Type, DynamicMethod> WriteObjectProperties = new Dictionary<Type, DynamicMethod>();
+		static void CreateObjectWriter(Type Type, ILGenerator Generator, DynamicMethod ContentsWriter)
+		{
+			Generator.Emit(OpCodes.Ldarg_1);
+
+			Label SkipLabel = Generator.DefineLabel();
+			Generator.Emit(OpCodes.Brfalse, SkipLabel);
+
+			Generator.Emit(OpCodes.Ldarg_0);
+			Generator.EmitCall(OpCodes.Call, GetMethodInfo<CbWriter>(x => x.BeginObject()), null);
+
+			Generator.Emit(OpCodes.Ldarg_0);
+			Generator.Emit(OpCodes.Ldarg_1);
+			Generator.EmitCall(OpCodes.Call, ContentsWriter, null);
+
+			Generator.Emit(OpCodes.Ldarg_0);
+			Generator.EmitCall(OpCodes.Call, GetMethodInfo<CbWriter>(x => x.EndObject()), null);
+
+			Generator.MarkLabel(SkipLabel);
+			Generator.Emit(OpCodes.Ret);
+		}
+
+		static void CreateNamedObjectWriter(Type Type, ILGenerator Generator, DynamicMethod ContentsWriter)
+		{
+			Generator.Emit(OpCodes.Ldarg_2);
+
+			Label SkipLabel = Generator.DefineLabel();
+			Generator.Emit(OpCodes.Brfalse, SkipLabel);
+
+			Generator.Emit(OpCodes.Ldarg_0);
+			Generator.Emit(OpCodes.Ldarg_1);
+			Generator.EmitCall(OpCodes.Call, GetMethodInfo<CbWriter>(x => x.BeginObject(null!)), null);
 
-		static DynamicMethod CreateObjectContentsWriter(Type Type)
+			Generator.Emit(OpCodes.Ldarg_0);
+			Generator.Emit(OpCodes.Ldarg_2);
+			Generator.EmitCall(OpCodes.Call, ContentsWriter, null);
+
+			Generator.Emit(OpCodes.Ldarg_0);
+			Generator.EmitCall(OpCodes.Call, GetMethodInfo<CbWriter>(x => x.EndObject()), null);
+
+			Generator.MarkLabel(SkipLabel);
+			Generator.Emit(OpCodes.Ret);
+		}
+
+		static Utf8String DiscriminatorKey = "_t";
+
+		static void CreateConcreteObjectContentsWriter(Type Type, ILGenerator Generator)
 		{
-			DynamicMethod? DynamicMethod;
-			if (!WriteObjectProperties.TryGetValue(Type, out DynamicMethod))
+			// Find the reflected properties from this type
+			(Utf8String Name, PropertyInfo Property)[] Properties = GetProperties(Type);
+
+			// Create a static type with the required reflection data
+			Type ReflectedType = typeof(CbReflectedTypeInfo<>).MakeGenericType(Type);
+			FieldInfo NamesField = ReflectedType.GetField(nameof(CbReflectedTypeInfo<object>.Names))!;
+			NamesField.SetValue(null, Properties.Select(x => x.Name).ToArray());
+
+			// Write the discriminator
+			CbDiscriminatorAttribute? Discriminator = Type.GetCustomAttribute<CbDiscriminatorAttribute>();
+			if (Discriminator != null)
 			{
-				// Create the method
-				DynamicMethod = new DynamicMethod("_", null, new Type[] { typeof(CbWriter), Type });
-				WriteObjectProperties.Add(Type, DynamicMethod);
+				FieldInfo DiscriminatorKeyField = GetFieldInfo(() => DiscriminatorKey);
+				Generator.Emit(OpCodes.Ldarg_0);
+				Generator.Emit(OpCodes.Ldsfld, DiscriminatorKeyField);
+				Generator.Emit(OpCodes.Ldstr, Discriminator.Name);
+				Generator.Emit(OpCodes.Call, GetMethodInfo<CbWriter>(x => x.WriteString(default, null!)));
+			}
+
+			// Write all the remaining properties
+			for (int Idx = 0; Idx < Properties.Length; Idx++)
+			{
+				PropertyInfo Property = Properties[Idx].Property;
+				Type PropertyType = Property.PropertyType;
 
-				// Find the reflected properties from this type
-				(Utf8String Name, PropertyInfo Property)[] Properties = GetProperties(Type);
+				// Get the field value
+				Generator.Emit(OpCodes.Ldarg_1);
+				Generator.EmitCall(OpCodes.Call, Property.GetMethod!, null);
 
-				// Create a static type with the required reflection data
-				Type ReflectedType = typeof(CbReflectedTypeInfo<>).MakeGenericType(Type);
-				FieldInfo NamesField = ReflectedType.GetField(nameof(CbReflectedTypeInfo<object>.Names))!;
-				NamesField.SetValue(null, Properties.Select(x => x.Name).ToArray());
+				Label SkipLabel = Generator.DefineLabel();
 
-				// Generate code for the method
-				ILGenerator Generator = DynamicMethod.GetILGenerator();
-				for (int Idx = 0; Idx < Properties.Length; Idx++)
+				MethodInfo WriteMethod;
+				if (TypeToClassMethods.TryGetValue(PropertyType, out CbDynamicClassMethods? DynamicMethods))
 				{
-					PropertyInfo Property = Properties[Idx].Property;
-					Type PropertyType = Property.PropertyType;
+					Generator.Emit(OpCodes.Dup);
+					Generator.Emit(OpCodes.Brfalse, SkipLabel);
+					WriteMethod = DynamicMethods.WriteNamedMethod;
+				}
+				else
+				{
+					CbConverterInfo ConverterInfo = GetConverterInfo(PropertyType);
+					ConverterInfo.SkipDefault?.Invoke(Generator, SkipLabel);
+					WriteMethod = ConverterInfo.WriteNamedMethod;
+				}
 
-					// Get the field value
-					Generator.Emit(OpCodes.Ldarg_1);
-					Generator.EmitCall(OpCodes.Call, Property.GetMethod!, null);
+				// Store the variable in a local
+				LocalBuilder Local = Generator.DeclareLocal(PropertyType);
+				Generator.Emit(OpCodes.Dup);
+				Generator.Emit(OpCodes.Stloc, Local);
 
-					Label SkipLabel = Generator.DefineLabel();
+				// Call the writer
+				Generator.Emit(OpCodes.Ldarg_0);
 
-					MethodInfo WriteMethod;
-					if (WriteNamedMethods.TryGetValue(PropertyType, out DynamicMethod? DynamicWriteMethod))
-					{
-						Generator.Emit(OpCodes.Dup);
-						Generator.Emit(OpCodes.Brfalse, SkipLabel);
-						WriteMethod = DynamicWriteMethod;
-					}
-					else
-					{
-						CbConverterInfo ConverterInfo = GetConverterInfo(PropertyType);
-						ConverterInfo.SkipDefault?.Invoke(Generator, SkipLabel);
-						WriteMethod = ConverterInfo.WriteNamedMethod;
-					}
+				Generator.Emit(OpCodes.Ldsfld, NamesField);
+				Generator.Emit(OpCodes.Ldc_I4, Idx);
+				Generator.Emit(OpCodes.Ldelem, typeof(Utf8String));
 
-					// Store the variable in a local
-					LocalBuilder Local = Generator.DeclareLocal(PropertyType);
-					Generator.Emit(OpCodes.Dup);
-					Generator.Emit(OpCodes.Stloc, Local);
+				Generator.Emit(OpCodes.Ldloc, Local);
+				Generator.EmitCall(OpCodes.Call, WriteMethod, null);
 
-					// Call the writer
-					Generator.Emit(OpCodes.Ldarg_0);
+				// Remove the duplicated value from the top of the stack
+				Generator.MarkLabel(SkipLabel);
+				Generator.Emit(OpCodes.Pop);
+			}
+			Generator.Emit(OpCodes.Ret);
+		}
 
-					Generator.Emit(OpCodes.Ldsfld, NamesField);
-					Generator.Emit(OpCodes.Ldc_I4, Idx);
-					Generator.Emit(OpCodes.Ldelem, typeof(Utf8String));
+		class CbPolymorphicDispatch<T>
+		{
+			public static Dictionary<Utf8String, Func<CbField, object>> NameToReadFunc = new Dictionary<Utf8String, Func<CbField, object>>();
+			public static Dictionary<Type, Action<CbWriter, object>> TypeToWriteContentsFunc = new Dictionary<Type, Action<CbWriter, object>>();
 
-					Generator.Emit(OpCodes.Ldloc, Local);
-					Generator.EmitCall(OpCodes.Call, WriteMethod, null);
+			public static object Read(CbField Field)
+			{
+				Utf8String Name = Field.AsObject().Find(DiscriminatorKey).AsUtf8String();
+				return NameToReadFunc[Name](Field);
+			}
 
-					// Remove the duplicated value from the top of the stack
-					Generator.MarkLabel(SkipLabel);
-					Generator.Emit(OpCodes.Pop);
-				}
-				Generator.Emit(OpCodes.Ret);
+			public static void WriteContents(CbWriter Writer, object Value)
+			{
+				Type Type = Value!.GetType();
+				TypeToWriteContentsFunc[Type](Writer, Value);
 			}
-			return DynamicMethod;
 		}
 
-		static Dictionary<Type, DynamicMethod> ReadObjectMethods = new Dictionary<Type, DynamicMethod>();
+		static T CreateDelegate<T>(DynamicMethod Method) where T : Delegate
+		{
+			return (T)Method.CreateDelegate(typeof(T));
+		}
 
-		static DynamicMethod CreateObjectReader(Type Type)
+		static void CreateConcreteObjectReader(Type Type, ILGenerator Generator)
 		{
-			DynamicMethod? DynamicMethod;
-			if (!ReadObjectMethods.TryGetValue(Type, out DynamicMethod))
+			// Construct the object
+			ConstructorInfo? Constructor = Type.GetConstructor(BindingFlags.Public | BindingFlags.NonPublic | BindingFlags.Instance, null, Type.EmptyTypes, null);
+			if (Constructor == null)
 			{
-				// Create the method
-				DynamicMethod = new DynamicMethod("_", Type, new Type[] { typeof(CbField) });
-				ReadObjectMethods.Add(Type, DynamicMethod);
+				throw new CbException($"Unable to find default constructor for {Type}");
+			}
 
-				// Generate the IL for it
-				ILGenerator Generator = DynamicMethod.GetILGenerator();
+			// Find the reflected properties from this type
+			(Utf8String Name, PropertyInfo Property)[] Properties = GetProperties(Type);
 
-				// Construct the object
-				ConstructorInfo? Constructor = Type.GetConstructor(BindingFlags.Public | BindingFlags.NonPublic | BindingFlags.Instance, null, Type.EmptyTypes, null);
-				if (Constructor == null)
-				{
-					throw new CbException($"Unable to find default constructor for {Type}");
-				}
+			// Create a static type with the required reflection data
+			Type ReflectedType = typeof(CbReflectedTypeInfo<>).MakeGenericType(Type);
+			FieldInfo NamesField = ReflectedType.GetField(nameof(CbReflectedTypeInfo<object>.Names))!;
+			NamesField.SetValue(null, Properties.Select(x => x.Name).ToArray());
+			MethodInfo MatchNameMethod = ReflectedType.GetMethod(nameof(CbReflectedTypeInfo<object>.MatchName))!;
 
-				// Find the reflected properties from this type
-				(Utf8String Name, PropertyInfo Property)[] Properties = GetProperties(Type);
+			// NewObjectLocal = new Type()
+			LocalBuilder NewObjectLocal = Generator.DeclareLocal(typeof(object));
+			Generator.Emit(OpCodes.Newobj, Constructor);
+			Generator.Emit(OpCodes.Stloc, NewObjectLocal);
 
-				// Create a static type with the required reflection data
-				Type ReflectedType = typeof(CbReflectedTypeInfo<>).MakeGenericType(Type);
-				FieldInfo NamesField = ReflectedType.GetField(nameof(CbReflectedTypeInfo<object>.Names))!;
-				NamesField.SetValue(null, Properties.Select(x => x.Name).ToArray());
-				MethodInfo MatchNameMethod = ReflectedType.GetMethod(nameof(CbReflectedTypeInfo<object>.MatchName))!;
+			// Stack(0) = CbField.CreateIterator()
+			Generator.Emit(OpCodes.Ldarg_0);
+			Generator.EmitCall(OpCodes.Call, GetMethodInfo<CbField>(x => x.CreateIterator()), null);
 
-				// NewObjectLocal = new Type()
-				LocalBuilder NewObjectLocal = Generator.DeclareLocal(typeof(object));
-				Generator.Emit(OpCodes.Newobj, Constructor);
-				Generator.Emit(OpCodes.Stloc, NewObjectLocal);
+			// CbFieldIterator IteratorLocal = Stack(0)
+			LocalBuilder IteratorLocal = Generator.DeclareLocal(typeof(CbFieldIterator));
+			Generator.Emit(OpCodes.Dup);
+			Generator.Emit(OpCodes.Stloc, IteratorLocal);
+
+			// if(!Stack.Pop().IsValid()) goto ReturnLabel
+			Generator.EmitCall(OpCodes.Call, GetMethodInfo<CbFieldIterator>(x => x.IsValid()), null);
+			Label ReturnLabel = Generator.DefineLabel();
+			Generator.Emit(OpCodes.Brfalse, ReturnLabel);
+
+			// NamesLocal = CbReflectedTypeInfo<Type>.Names
+			LocalBuilder NamesLocal = Generator.DeclareLocal(typeof(Utf8String[]));
+			Generator.Emit(OpCodes.Ldsfld, NamesField);
+			Generator.Emit(OpCodes.Stloc, NamesLocal);
+
+			// IterationLoopLabel:
+			Label IterationLoopLabel = Generator.DefineLabel();
+			Generator.MarkLabel(IterationLoopLabel);
+
+			// bool MatchLocal = false
+			LocalBuilder MatchLocal = Generator.DeclareLocal(typeof(bool));
+			Generator.Emit(OpCodes.Ldc_I4_0);
+			Generator.Emit(OpCodes.Stloc, MatchLocal);
+
+			// Stack(0) = IteratorLocal.GetCurrent()
+			Generator.Emit(OpCodes.Ldloc, IteratorLocal);
+			Generator.EmitCall(OpCodes.Call, GetMethodInfo<CbFieldIterator>(x => x.GetCurrent()), null);
+
+			// Try to parse each of the properties in order. If fields are ordered correctly, we will parse the object in a single pass. Otherwise we can loop and start again.
+			LocalBuilder FieldLocal = Generator.DeclareLocal(typeof(CbField));
+			for (int Idx = 0; Idx < Properties.Length; Idx++)
+			{
+				PropertyInfo Property = Properties[Idx].Property;
 
-				// Stack(0) = CbField.CreateIterator()
-				Generator.Emit(OpCodes.Ldarg_0);
-				Generator.EmitCall(OpCodes.Call, GetMethodInfo<CbField>(x => x.CreateIterator()), null);
+				// Get the read method for this property type
+				MethodInfo ReadMethod;
+				if (TypeToClassMethods.TryGetValue(Property.PropertyType, out CbDynamicClassMethods? DynamicMethods))
+				{
+					ReadMethod = DynamicMethods.ReadMethod;
+				}
+				else
+				{
+					ReadMethod = GetConverterInfo(Property.PropertyType).ReadMethod;
+				}
 
-				// CbFieldIterator IteratorLocal = Stack(0)
-				LocalBuilder IteratorLocal = Generator.DeclareLocal(typeof(CbFieldIterator));
-				Generator.Emit(OpCodes.Dup);
-				Generator.Emit(OpCodes.Stloc, IteratorLocal);
+				// if(!CbReflectedTypeInfo<Type>.MatchName(Stack(0), Idx)) goto SkipPropertyLabel
+				Label SkipPropertyLabel = Generator.DefineLabel();
+				Generator.Emit(OpCodes.Dup); // Current CbField
+				Generator.Emit(OpCodes.Ldc_I4, Idx);
+				Generator.Emit(OpCodes.Call, MatchNameMethod);
+				Generator.Emit(OpCodes.Brfalse, SkipPropertyLabel);
 
-				// if(!Stack.Pop().IsValid()) goto ReturnLabel
-				Generator.EmitCall(OpCodes.Call, GetMethodInfo<CbFieldIterator>(x => x.IsValid()), null);
-				Label ReturnLabel = Generator.DefineLabel();
-				Generator.Emit(OpCodes.Brfalse, ReturnLabel);
+				// FieldLocal = Stack.Pop()
+				Generator.Emit(OpCodes.Stloc, FieldLocal);
 
-				// NamesLocal = CbReflectedTypeInfo<Type>.Names
-				LocalBuilder NamesLocal = Generator.DeclareLocal(typeof(Utf8String[]));
-				Generator.Emit(OpCodes.Ldsfld, NamesField);
-				Generator.Emit(OpCodes.Stloc, NamesLocal);
+				// Property.SetMethod(NewObjectLocal, ReadMethod(FieldLocal))
+				Generator.Emit(OpCodes.Ldloc, NewObjectLocal);
+				Generator.Emit(OpCodes.Ldloc, FieldLocal);
+				Generator.EmitCall(OpCodes.Call, ReadMethod, null);
+				Generator.EmitCall(OpCodes.Call, Property.SetMethod!, null);
 
-				// IterationLoopLabel:
-				Label IterationLoopLabel = Generator.DefineLabel();
-				Generator.MarkLabel(IterationLoopLabel);
+				// if(!IteratorLocal.MoveNext()) goto ReturnLabel
+				Generator.Emit(OpCodes.Ldloc, IteratorLocal);
+				Generator.EmitCall(OpCodes.Call, GetMethodInfo<CbFieldIterator>(x => x.MoveNext()), null);
+				Generator.Emit(OpCodes.Brfalse, ReturnLabel);
 
-				// bool MatchLocal = false
-				LocalBuilder MatchLocal = Generator.DeclareLocal(typeof(bool));
-				Generator.Emit(OpCodes.Ldc_I4_0);
+				// MatchLocal = true
+				Generator.Emit(OpCodes.Ldc_I4_1);
 				Generator.Emit(OpCodes.Stloc, MatchLocal);
 
 				// Stack(0) = IteratorLocal.GetCurrent()
 				Generator.Emit(OpCodes.Ldloc, IteratorLocal);
 				Generator.EmitCall(OpCodes.Call, GetMethodInfo<CbFieldIterator>(x => x.GetCurrent()), null);
 
-				// Try to parse each of the properties in order. If fields are ordered correctly, we will parse the object in a single pass. Otherwise we can loop and start again.
-				LocalBuilder FieldLocal = Generator.DeclareLocal(typeof(CbField));
-				for (int Idx = 0; Idx < Properties.Length; Idx++)
-				{
-					PropertyInfo Property = Properties[Idx].Property;
-
-					// Get the read method for this property type
-					MethodInfo ReadMethod;
-					if (ReadObjectMethods.TryGetValue(Property.PropertyType, out DynamicMethod? DynamicReadMethod))
-					{
-						ReadMethod = DynamicReadMethod;
-					}
-					else
-					{
-						ReadMethod = GetConverterInfo(Property.PropertyType).ReadMethod;
-					}
+				// SkipPropertyLabel:
+				Generator.MarkLabel(SkipPropertyLabel);
+			}
 
-					// if(!CbReflectedTypeInfo<Type>.MatchName(Stack(0), Idx)) goto SkipPropertyLabel
-					Label SkipPropertyLabel = Generator.DefineLabel();
-					Generator.Emit(OpCodes.Dup); // Current CbField
-					Generator.Emit(OpCodes.Ldc_I4, Idx);
-					Generator.Emit(OpCodes.Call, MatchNameMethod);
-					Generator.Emit(OpCodes.Brfalse, SkipPropertyLabel);
-
-					// FieldLocal = Stack.Pop()
-					Generator.Emit(OpCodes.Stloc, FieldLocal);
-
-					// Property.SetMethod(NewObjectLocal, ReadMethod(FieldLocal))
-					Generator.Emit(OpCodes.Ldloc, NewObjectLocal);
-					Generator.Emit(OpCodes.Ldloc, FieldLocal);
-					Generator.EmitCall(OpCodes.Call, ReadMethod, null);
-					Generator.EmitCall(OpCodes.Call, Property.SetMethod!, null);
-
-					// if(!IteratorLocal.MoveNext()) goto ReturnLabel
-					Generator.Emit(OpCodes.Ldloc, IteratorLocal);
-					Generator.EmitCall(OpCodes.Call, GetMethodInfo<CbFieldIterator>(x => x.MoveNext()), null);
-					Generator.Emit(OpCodes.Brfalse, ReturnLabel);
-
-					// MatchLocal = true
-					Generator.Emit(OpCodes.Ldc_I4_1);
-					Generator.Emit(OpCodes.Stloc, MatchLocal);
-
-					// Stack(0) = IteratorLocal.GetCurrent()
-					Generator.Emit(OpCodes.Ldloc, IteratorLocal);
-					Generator.EmitCall(OpCodes.Call, GetMethodInfo<CbFieldIterator>(x => x.GetCurrent()), null);
-
-					// SkipPropertyLabel:
-					Generator.MarkLabel(SkipPropertyLabel);
-				}
+			// Stack.Pop()
+			Generator.Emit(OpCodes.Pop); // Current CbField
 
-				// Stack.Pop()
-				Generator.Emit(OpCodes.Pop); // Current CbField
+			// if(MatchLocal) goto IterationLoopLabel
+			Generator.Emit(OpCodes.Ldloc, MatchLocal);
+			Generator.Emit(OpCodes.Brtrue, IterationLoopLabel);
 
-				// if(MatchLocal) goto IterationLoopLabel
-				Generator.Emit(OpCodes.Ldloc, MatchLocal);
-				Generator.Emit(OpCodes.Brtrue, IterationLoopLabel);
+			// if(IteratorLocal.MoveNext()) goto IterationLoopLabel
+			Generator.Emit(OpCodes.Ldloc, IteratorLocal);
+			Generator.EmitCall(OpCodes.Call, GetMethodInfo<CbFieldIterator>(x => x.MoveNext()), null);
+			Generator.Emit(OpCodes.Brtrue, IterationLoopLabel);
 
-				// if(IteratorLocal.MoveNext()) goto IterationLoopLabel
-				Generator.Emit(OpCodes.Ldloc, IteratorLocal);
-				Generator.EmitCall(OpCodes.Call, GetMethodInfo<CbFieldIterator>(x => x.MoveNext()), null);
-				Generator.Emit(OpCodes.Brtrue, IterationLoopLabel);
+			// return NewObjectLocal
+			Generator.MarkLabel(ReturnLabel);
+			Generator.Emit(OpCodes.Ldloc, NewObjectLocal);
+			Generator.Emit(OpCodes.Ret);
+		}
 
-				// return NewObjectLocal
-				Generator.MarkLabel(ReturnLabel);
-				Generator.Emit(OpCodes.Ldloc, NewObjectLocal);
-				Generator.Emit(OpCodes.Ret);
-			}
-			return DynamicMethod;
+		static FieldInfo GetFieldInfo<T>(Expression<Func<T>> Expr)
+		{
+			return (FieldInfo)((MemberExpression)Expr.Body).Member;
 		}
 
 		static MethodInfo GetMethodInfo(Expression<Action> Expr)
