CS: Support for serialization and deserialization of polymorphic class hierarchies into compact binary.

The base class in a hierarchy must have the [CbPolymorphic] attribute, and derived classes must have the [CbDiscriminator] attribute, giving an identifier to indicate the derived type.

#p4v-preflight-copy 18635755

diff --git a/Engine/Source/Programs/Shared/EpicGames.Horde/Common/Condition.cs b/Engine/Source/Programs/Shared/EpicGames.Horde/Common/Condition.cs
--- a/Engine/Source/Programs/Shared/EpicGames.Horde/Common/Condition.cs
+++ b/Engine/Source/Programs/Shared/EpicGames.Horde/Common/Condition.cs
@@ -600,7 +600,7 @@ namespace EpicGames.Horde.Common
 	/// <summary>
 	/// Converter from conditions to compact binary objects
 	/// </summary>
-	public class ConditionCbConverter : CbConverter<Condition>
+	public class ConditionCbConverter : CbConverterBase<Condition>
 	{
 		/// <inheritdoc/>
 		public override Condition Read(CbField Field) => Condition.TryParse(Field.AsUtf8String().ToString());
diff --git a/Engine/Source/Programs/Shared/EpicGames.Serialization/CbConverter.cs b/Engine/Source/Programs/Shared/EpicGames.Serialization/CbConverter.cs
--- a/Engine/Source/Programs/Shared/EpicGames.Serialization/CbConverter.cs
+++ b/Engine/Source/Programs/Shared/EpicGames.Serialization/CbConverter.cs
@@ -1,8 +1,12 @@
 // Copyright Epic Games, Inc. All Rights Reserved.
 
 using EpicGames.Core;
+using EpicGames.Serialization.Converters;
 using System;
 using System.Collections.Generic;
+using System.Linq.Expressions;
+using System.Reflection;
+using System.Reflection.Emit;
 using System.Text;
 
 namespace EpicGames.Serialization
@@ -31,21 +35,21 @@ namespace EpicGames.Serialization
 	/// <summary>
 	/// Base class for all converters. Deriving from <see cref="CbConverter{T}"/> is more efficient
 	/// </summary>
-	public abstract class CbConverter
+	public interface ICbConverter
 	{
 		/// <summary>
 		/// Reads an object from a field
 		/// </summary>
 		/// <param name="Field"></param>
 		/// <returns></returns>
-		public abstract object? ReadObject(CbField Field);
+		object? ReadObject(CbField Field);
 
 		/// <summary>
 		/// Writes an object to a field
 		/// </summary>
 		/// <param name="Writer"></param>
 		/// <param name="Value"></param>
-		public abstract void WriteObject(CbWriter Writer, object? Value);
+		void WriteObject(CbWriter Writer, object? Value);
 
 		/// <summary>
 		/// Writes an object to a named field, if not equal to the default value
@@ -53,27 +57,27 @@ namespace EpicGames.Serialization
 		/// <param name="Writer"></param>
 		/// <param name="Name"></param>
 		/// <param name="Value"></param>
-		public abstract void WriteNamedObject(CbWriter Writer, Utf8String Name, object? Value);
+		void WriteNamedObject(CbWriter Writer, Utf8String Name, object? Value);
 	}
 
 	/// <summary>
 	/// Converter for a particular type
 	/// </summary>
-	public abstract class CbConverter<T> : CbConverter
+	public interface ICbConverter<T> : ICbConverter
 	{
 		/// <summary>
 		/// Reads an object from a field
 		/// </summary>
 		/// <param name="Field"></param>
 		/// <returns></returns>
-		public abstract T Read(CbField Field);
+		T Read(CbField Field);
 
 		/// <summary>
 		/// Writes an object to a field
 		/// </summary>
 		/// <param name="Writer"></param>
 		/// <param name="Value"></param>
-		public abstract void Write(CbWriter Writer, T Value);
+		void Write(CbWriter Writer, T Value);
 
 		/// <summary>
 		/// Writes an object to a named field, if not equal to the default value
@@ -81,16 +85,107 @@ namespace EpicGames.Serialization
 		/// <param name="Writer"></param>
 		/// <param name="Name"></param>
 		/// <param name="Value"></param>
-		public abstract void WriteNamed(CbWriter Writer, Utf8String Name, T Value);
+		void WriteNamed(CbWriter Writer, Utf8String Name, T Value);
+	}
+
+	/// <summary>
+	/// Interface for obtaining static methods that can be invoked directly
+	/// </summary>
+	public interface ICbConverterMethods
+	{
+		/// <summary>
+		/// Method with the signature CbField -> T
+		/// </summary>
+		public MethodInfo ReadMethod { get; }
+
+		/// <summary>
+		/// Method with the signature CbWriter, T -> void
+		/// </summary>
+		public MethodInfo WriteMethod { get; }
+
+		/// <summary>
+		/// Method with the signature CbWriter, Utf8String, T -> void
+		/// </summary>
+		public MethodInfo WriteNamedMethod { get; }
+	}
+
+	/// <summary>
+	/// Helper class for wrapping regular converters in a ICbConverterMethods interface
+	/// </summary>
+	public static class CbConverterMethods
+	{
+		class CbConverterMethodsWrapper<T> : ICbConverterMethods
+		{
+			static ICbConverter<T> StaticConverter = null!;
+
+			static T Read(CbField Field) => StaticConverter.Read(Field);
+			static void Write(CbWriter Writer, T Value) => StaticConverter.Write(Writer, Value);
+			static void WriteNamed(CbWriter Writer, Utf8String Name, T Value) => StaticConverter.WriteNamed(Writer, Name, Value);
+
+			public MethodInfo ReadMethod { get; } = GetMethodInfo(() => Read(null!));
+			public MethodInfo WriteMethod { get; } = GetMethodInfo(() => Write(null!, default!));
+			public MethodInfo WriteNamedMethod { get; } = GetMethodInfo(() => WriteNamed(null!, null!, default!));
+
+			public CbConverterMethodsWrapper(ICbConverter<T> Converter)
+			{
+				StaticConverter = Converter;
+			}
+
+			static MethodInfo GetMethodInfo(Expression<Action> Expr)
+			{
+				return ((MethodCallExpression)Expr.Body).Method;
+			}
+		}
+
+		static Dictionary<Type, ICbConverterMethods> TypeToMethods = new Dictionary<Type, ICbConverterMethods>();
+
+		static ICbConverterMethods CreateWrapper(Type Type, ICbConverter Converter)
+		{
+			Type ConverterType = typeof(CbConverterMethodsWrapper<>).MakeGenericType(Type);
+			return (ICbConverterMethods)Activator.CreateInstance(ConverterType, new object[] { Converter })!;
+		}
+
+		/// <summary>
+		/// Gets a <see cref="ICbConverterMethods"/> interface for the given type
+		/// </summary>
+		/// <param name="Type"></param>
+		/// <returns></returns>
+		public static ICbConverterMethods Get(Type Type)
+		{
+			ICbConverterMethods? Methods;
+			if (!TypeToMethods.TryGetValue(Type, out Methods))
+			{
+				ICbConverter Converter = CbConverter.GetConverter(Type);
+				Methods = (Converter as ICbConverterMethods) ?? CreateWrapper(Type, Converter);
+				TypeToMethods.Add(Type, Methods);
+			}
+			return Methods;
+		}
+	}
+
+	/// <summary>
+	/// Base class for converter implementations 
+	/// </summary>
+	/// <typeparam name="T"></typeparam>
+	public abstract class CbConverterBase<T> : ICbConverter<T>
+	{
+		/// <inheritdoc/>
+		public object? ReadObject(CbField Field) => Read(Field);
+
+		/// <inheritdoc/>
+		public void WriteObject(CbWriter Writer, object? Value) => Write(Writer, (T)Value!);
 
 		/// <inheritdoc/>
-		public sealed override object? ReadObject(CbField Field) => Read(Field);
+		public void WriteNamedObject(CbWriter Writer, Utf8String Name, object? Value) => WriteNamed(Writer, Name, (T)Value!);
 
 		/// <inheritdoc/>
-		public sealed override void WriteObject(CbWriter Writer, object? Value) => Write(Writer, (T)Value!);
+		public abstract T Read(CbField Field);
+
+		/// <inheritdoc/>
+		public abstract void Write(CbWriter Writer, T Value);
 
 		/// <inheritdoc/>
-		public sealed override void WriteNamedObject(CbWriter Writer, Utf8String Name, object? Value) => WriteNamed(Writer, Name, (T)Value!);
+		public abstract void WriteNamed(CbWriter Writer, Utf8String Name, T Value);
 	}
 
 	/// <summary>
@@ -103,6 +198,142 @@ namespace EpicGames.Serialization
 		/// </summary>
 		/// <param name="Type">The type to create a converter for</param>
 		/// <returns>The converter instance, or null if this factory does not support the given type</returns>
-		public abstract CbConverter? CreateConverter(Type Type);
+		public abstract ICbConverter? CreateConverter(Type Type);
+	}
+
+	/// <summary>
+	/// Utility functions for creating converters
+	/// </summary>
+	public static class CbConverter
+	{
+		/// <summary>
+		/// Object used for locking access to shared objects
+		/// </summary>
+		static object LockObject = new object();
+
+		/// <summary>
+		/// Cache of type to converter
+		/// </summary>
+		public static Dictionary<Type, ICbConverter> TypeToConverter = new Dictionary<Type, ICbConverter>()
+		{
+			[typeof(bool)] = new CbPrimitiveConverter<bool>(x => x.AsBool(), (w, v) => w.WriteBoolValue(v), (w, n, v) => w.WriteBool(n, v)),
+			[typeof(int)] = new CbPrimitiveConverter<int>(x => x.AsInt32(), (w, v) => w.WriteIntegerValue(v), (w, n, v) => w.WriteInteger(n, v)),
+			[typeof(long)] = new CbPrimitiveConverter<long>(x => x.AsInt64(), (w, v) => w.WriteIntegerValue(v), (w, n, v) => w.WriteInteger(n, v)),
+			[typeof(double)] = new CbPrimitiveConverter<double>(x => x.AsDouble(), (w, v) => w.WriteDoubleValue(v), (w, n, v) => w.WriteDouble(n, v)),
+			[typeof(Utf8String)] = new CbPrimitiveConverter<Utf8String>(x => x.AsUtf8String(), (w, v) => w.WriteUtf8StringValue(v), (w, n, v) => w.WriteUtf8String(n, v)),
+			[typeof(IoHash)] = new CbPrimitiveConverter<IoHash>(x => x.AsHash(), (w, v) => w.WriteHashValue(v), (w, n, v) => w.WriteHash(n, v)),
+			[typeof(CbObjectAttachment)] = new CbPrimitiveConverter<CbObjectAttachment>(x => x.AsObjectAttachment(), (w, v) => w.WriteObjectAttachmentValue(v.Hash), (w, n, v) => w.WriteObjectAttachment(n, v.Hash)),
+			[typeof(CbBinaryAttachment)] = new CbPrimitiveConverter<CbBinaryAttachment>(x => x.AsBinaryAttachment(), (w, v) => w.WriteBinaryAttachmentValue(v.Hash), (w, n, v) => w.WriteBinaryAttachment(n, v.Hash)),
+			[typeof(DateTime)] = new CbPrimitiveConverter<DateTime>(x => x.AsDateTime(), (w, v) => w.WriteDateTimeValue(v), (w, n, v) => w.WriteDateTime(n, v)),
+			[typeof(string)] = new CbPrimitiveConverter<string>(x => x.AsString(), (w, v) => w.WriteStringValue(v), (w, n, v) => w.WriteString(n, v)),
+			[typeof(ReadOnlyMemory<byte>)] = new CbPrimitiveConverter<ReadOnlyMemory<byte>>(x => x.AsBinary(), (w, v) => w.WriteBinaryValue(v), (w, n, v) => w.WriteBinary(n, v)),
+			[typeof(byte[])] = new CbPrimitiveConverter<byte[]>(x => x.AsBinaryArray(), (w, v) => w.WriteBinaryArrayValue(v), (w, n, v) => w.WriteBinaryArray(n, v)),
+			[typeof(CbField)] = new CbFieldConverter(),
+			[typeof(CbObject)] = new CbObjectConverter()
+		};
+
+		/// <summary>
+		/// List of converter factories. Must be 
+		/// </summary>
+		public static List<CbConverterFactory> ConverterFactories = new List<CbConverterFactory>
+		{
+			new CbClassConverterFactory(),
+			new CbEnumConverterFactory(),
+			new CbListConverterFactory(),
+			new CbDictionaryConverterFactory(),
+			new CbNullableConverterFactory()
+		};
+
+		/// <summary>
+		/// Class used to cache values returned by CreateConverterInfo without having to do a dictionary lookup.
+		/// </summary>
+		/// <typeparam name="T">The type to be converted</typeparam>
+		class CbConverterCache<T>
+		{
+			/// <summary>
+			/// The converter instance
+			/// </summary>
+			public static ICbConverter<T> Instance { get; } = CreateConverter();
+
+			/// <summary>
+			/// Create a typed converter
+			/// </summary>
+			/// <returns></returns>
+			static ICbConverter<T> CreateConverter()
+			{
+				ICbConverter Converter = GetConverter(typeof(T));
+				return (Converter as ICbConverter<T>) ?? new CbConverterWrapper(Converter);
+			}
+
+			/// <summary>
+			/// Wrapper class to convert an untyped converter into a typed one
+			/// </summary>
+			class CbConverterWrapper : CbConverterBase<T>
+			{
+				ICbConverter Inner;
+
+				public CbConverterWrapper(ICbConverter Inner) => this.Inner = Inner;
+
+				/// <inheritdoc/>
+				public override T Read(CbField Field) => (T)Inner.ReadObject(Field)!;
+
+				/// <inheritdoc/>
+				public override void Write(CbWriter Writer, T Value) => Inner.WriteObject(Writer, Value);
+
+				/// <inheritdoc/>
+				public override void WriteNamed(CbWriter Writer, Utf8String Name, T Value) => Inner.WriteNamedObject(Writer, Name, Value);
+			}
+		}
+
+		/// <summary>
+		/// Gets the converter for a particular type
+		/// </summary>
+		/// <param name="Type"></param>
+		/// <returns></returns>
+		public static ICbConverter GetConverter(Type Type)
+		{
+			ICbConverter? Converter;
+			lock (LockObject)
+			{
+				if (!TypeToConverter.TryGetValue(Type, out Converter))
+				{
+					CbConverterAttribute? ConverterAttribute = Type.GetCustomAttribute<CbConverterAttribute>();
+					if (ConverterAttribute != null)
+					{
+						Type ConverterType = ConverterAttribute.ConverterType;
+						if (Type.IsGenericType && ConverterType.IsGenericTypeDefinition)
+						{
+							ConverterType = ConverterType.MakeGenericType(Type.GetGenericArguments());
+						}
+						Converter = (ICbConverter?)Activator.CreateInstance(ConverterType)!;
+					}
+					else
+					{
+						for (int Idx = ConverterFactories.Count - 1; Idx >= 0 && Converter == null; Idx--)
+						{
+							Converter = ConverterFactories[Idx].CreateConverter(Type);
+						}
+
+						if (Converter == null)
+						{
+							throw new CbException($"Unable to create converter for {Type.Name}");
+						}
+					}
+					TypeToConverter.Add(Type, Converter!);
+				}
+			}
+			return Converter;
+		}
+
+		/// <summary>
+		/// Gets the converter for a given type
+		/// </summary>
+		/// <typeparam name="T">Type to retrieve the converter for</typeparam>
+		/// <returns></returns>
+		public static ICbConverter<T> GetConverter<T>()
+		{
+			return CbConverterCache<T>.Instance;
+		}
 	}
 }
+
diff --git a/Engine/Source/Programs/Shared/EpicGames.Serialization/CbSerializer.cs b/Engine/Source/Programs/Shared/EpicGames.Serialization/CbSerializer.cs
--- a/Engine/Source/Programs/Shared/EpicGames.Serialization/CbSerializer.cs
+++ b/Engine/Source/Programs/Shared/EpicGames.Serialization/CbSerializer.cs
@@ -87,122 +87,6 @@ namespace EpicGames.Serialization
 	/// </summary>
 	public static class CbSerializer
 	{
-		/// <summary>
-		/// Object used for locking access to shared objects
-		/// </summary>
-		static object LockObject = new object();
-
-		/// <summary>
-		/// Cache of type to converter
-		/// </summary>
-		public static Dictionary<Type, CbConverter> TypeToConverter = new Dictionary<Type, CbConverter>()
-		{
-			[typeof(CbField)] = new CbFieldConverter(),
-			[typeof(CbObject)] = new CbObjectConverter()
-		};
-
-		/// <summary>
-		/// List of converter factories. Must be 
-		/// </summary>
-		public static List<CbConverterFactory> ConverterFactories = new List<CbConverterFactory>
-		{
-			new CbDefaultConverterFactory(),
-			new CbListConverterFactory(),
-			new CbDictionaryConverterFactory(),
-			new CbNullableConverterFactory()
-		};
-	
-		/// <summary>
-		/// Class used to cache values returned by CreateConverterInfo without having to do a dictionary lookup.
-		/// </summary>
-		/// <typeparam name="T">The type to be converted</typeparam>
-		class CbConverterCache<T>
-		{
-			/// <summary>
-			/// The converter instance
-			/// </summary>
-			public static CbConverter<T> Instance { get; } = CreateConverter();
-
-			/// <summary>
-			/// Create a typed converter
-			/// </summary>
-			/// <returns></returns>
-			static CbConverter<T> CreateConverter()
-			{
-				CbConverter Converter = GetConverter(typeof(T));
-				return (Converter as CbConverter<T>) ?? new CbConverterWrapper(Converter);
-			}
-
-			/// <summary>
-			/// Wrapper class to convert an untyped converter into a typed one
-			/// </summary>
-			class CbConverterWrapper : CbConverter<T>
-			{
-				CbConverter Inner;
-
-				public CbConverterWrapper(CbConverter Inner) => this.Inner = Inner;
-
-				/// <inheritdoc/>
-				public override T Read(CbField Field) => (T)Inner.ReadObject(Field)!;
-
-				/// <inheritdoc/>
-				public override void Write(CbWriter Writer, T Value) => Inner.WriteObject(Writer, Value);
-
-				/// <inheritdoc/>
-				public override void WriteNamed(CbWriter Writer, Utf8String Name, T Value) => Inner.WriteNamedObject(Writer, Name, Value);
-			}
-		}
-
-		/// <summary>
-		/// Gets the converter for a particular type
-		/// </summary>
-		/// <param name="Type"></param>
-		/// <returns></returns>
-		public static CbConverter GetConverter(Type Type)
-		{
-			CbConverter? Converter;
-			lock (LockObject)
-			{
-				if (!TypeToConverter.TryGetValue(Type, out Converter))
-				{
-					CbConverterAttribute? ConverterAttribute = Type.GetCustomAttribute<CbConverterAttribute>();
-					if (ConverterAttribute != null)
-					{
-						Type ConverterType = ConverterAttribute.ConverterType;
-						if (Type.IsGenericType && ConverterType.IsGenericTypeDefinition)
-						{
-							ConverterType = ConverterType.MakeGenericType(Type.GetGenericArguments());
-						}
-						Converter = (CbConverter?)Activator.CreateInstance(ConverterType)!;
-					}
-					else
-					{
-						for (int Idx = ConverterFactories.Count - 1; Idx >= 0 && Converter == null; Idx--)
-						{
-							Converter = ConverterFactories[Idx].CreateConverter(Type);
-						}
-
-						if (Converter == null)
-						{
-							throw new CbException($"Unable to create converter for {Type.Name}");
-						}
-					}
-					TypeToConverter.Add(Type, Converter!);
-				}
-			}
-			return Converter;
-		}
-
-		/// <summary>
-		/// Gets the converter for a given type
-		/// </summary>
-		/// <typeparam name="T">Type to retrieve the converter for</typeparam>
-		/// <returns></returns>
-		public static CbConverter<T> GetConverter<T>()
-		{
-			return CbConverterCache<T>.Instance;
-		}
-
 		/// <summary>
 		/// Serialize an object
 		/// </summary>
@@ -212,7 +96,7 @@ namespace EpicGames.Serialization
 		public static CbObject Serialize(Type Type, object Value)
 		{
 			CbWriter Writer = new CbWriter();
-			GetConverter(Type).WriteObject(Writer, Value);
+			CbConverter.GetConverter(Type).WriteObject(Writer, Value);
 			return Writer.ToObject();
 		}
 
@@ -225,7 +109,7 @@ namespace EpicGames.Serialization
 		public static CbObject Serialize<T>(T Value)
 		{
 			CbWriter Writer = new CbWriter();
-			GetConverter<T>().Write(Writer, Value);
+			CbConverter.GetConverter<T>().Write(Writer, Value);
 			return Writer.ToObject();
 		}
 
@@ -237,7 +121,7 @@ namespace EpicGames.Serialization
 		/// <param name="Value"></param>
 		public static void Serialize<T>(CbWriter Writer, T Value)
 		{
-			GetConverter<T>().Write(Writer, Value);
+			CbConverter.GetConverter<T>().Write(Writer, Value);
 		}
 
 		/// <summary>
@@ -249,7 +133,7 @@ namespace EpicGames.Serialization
 		/// <param name="Value"></param>
 		public static void Serialize<T>(CbWriter Writer, Utf8String Name, T Value)
 		{
-			GetConverter<T>().WriteNamed(Writer, Name, Value);
+			CbConverter.GetConverter<T>().WriteNamed(Writer, Name, Value);
 		}
 
 		/// <summary>
@@ -260,7 +144,7 @@ namespace EpicGames.Serialization
 		/// <returns></returns>
 		public static object? Deserialize(CbField Field, Type Type)
 		{
-			return GetConverter(Type).ReadObject(Field);
+			return CbConverter.GetConverter(Type).ReadObject(Field);
 		}
 
 		/// <summary>
@@ -271,7 +155,7 @@ namespace EpicGames.Serialization
 		/// <returns></returns>
 		public static T Deserialize<T>(CbField Field)
 		{
-			return GetConverter<T>().Read(Field);
+			return CbConverter.GetConverter<T>().Read(Field);
 		}
 
 		/// <summary>
diff --git a/Engine/Source/Programs/Shared/EpicGames.Serialization/Converters/CbClassConverter.cs b/Engine/Source/Programs/Shared/EpicGames.Serialization/Converters/CbClassConverter.cs
new file mode 100644
--- /dev/null
+++ b/Engine/Source/Programs/Shared/EpicGames.Serialization/Converters/CbClassConverter.cs
@@ -0,0 +1,507 @@
+// Copyright Epic Games, Inc. All Rights Reserved.
+
+using EpicGames.Core;
+using System;
+using System.Collections.Generic;
+using System.Linq;
+using System.Linq.Expressions;
+using System.Reflection;
+using System.Reflection.Emit;
+using System.Text;
+
+namespace EpicGames.Serialization.Converters
+{
+	class CbClassConverterMethods
+	{
+		class CbReflectedTypeInfo<T> where T : class
+		{
+			public static Utf8String[]? Names = null;
+			public static PropertyInfo[]? Properties = null;
+
+			public static bool MatchName(CbField Field, int Idx)
+			{
+				return Field.Name == Names![Idx];
+			}
+		}
+
+		static Utf8String DiscriminatorKey = "_t";
+
+		public Type ClassType { get; }
+		public bool IsPolymorphic { get; }
+
+		public DynamicMethod ReadMethod { get; }
+		public DynamicMethod WriteMethod { get; }
+		public DynamicMethod WriteNamedMethod { get; }
+
+		public DynamicMethod WriteContentsMethod { get; }
+
+		public DynamicMethod ReadConcreteMethod { get; }
+		public DynamicMethod WriteConcreteContentsMethod { get; }
+
+		static Dictionary<Type, CbClassConverterMethods> TypeToMethods = new Dictionary<Type, CbClassConverterMethods>();
+
+		public static CbClassConverterMethods Create(Type ClassType)
+		{
+			CbClassConverterMethods? Methods;
+			if (!TypeToMethods.TryGetValue(ClassType, out Methods))
+			{
+				Methods = new CbClassConverterMethods(ClassType);
+				TypeToMethods.Add(ClassType, Methods);
+				Methods.GenerateBytecode();
+			}
+			return Methods;
+		}
+
+		private CbClassConverterMethods(Type ClassType)
+		{
+			this.ClassType = ClassType;
+			this.IsPolymorphic = (ClassType.GetCustomAttribute<CbPolymorphicAttribute>(true) != null);
+
+			ReadConcreteMethod = new DynamicMethod($"ReadConcrete_{ClassType.Name}", ClassType, new Type[] { typeof(CbField) });
+			WriteConcreteContentsMethod = new DynamicMethod($"WriteConcreteContents_{ClassType.Name}", null, new Type[] { typeof(CbWriter), ClassType });
+
+			WriteMethod = new DynamicMethod($"Write_{ClassType.Name}", null, new Type[] { typeof(CbWriter), ClassType });
+			WriteNamedMethod = new DynamicMethod($"WriteNamed_{ClassType.Name}", null, new Type[] { typeof(CbWriter), typeof(Utf8String), ClassType });
+
+			if (IsPolymorphic)
+			{
+				ReadMethod = new DynamicMethod($"Read_{ClassType.Name}", ClassType, new Type[] { typeof(CbField) });
+				WriteContentsMethod = new DynamicMethod($"WriteContents_{ClassType.Name}", null, new Type[] { typeof(CbWriter), ClassType });
+			}
+			else
+			{
+				ReadMethod = ReadConcreteMethod;
+				WriteContentsMethod = WriteConcreteContentsMethod;
+			}
+		}
+
+		private void GenerateBytecode()
+		{
+			// Create the regular methods
+			CreateConcreteObjectReader(ClassType, ReadConcreteMethod.GetILGenerator());
+			CreateConcreteObjectContentsWriter(ClassType, WriteConcreteContentsMethod.GetILGenerator());
+
+			CreateObjectWriter(ClassType, WriteMethod.GetILGenerator(), WriteContentsMethod);
+			CreateNamedObjectWriter(ClassType, WriteNamedMethod.GetILGenerator(), WriteContentsMethod);
+
+			// Create the extra polymorphic methods
+			if (IsPolymorphic)
+			{
+				// Create the dispatch type
+				Type DispatchType = typeof(CbPolymorphicDispatch<>).MakeGenericType(ClassType);
+
+				// Create the read dispatch method
+				{
+					ILGenerator Generator = ReadMethod.GetILGenerator();
+					Generator.Emit(OpCodes.Ldarg_0);
+					Generator.Emit(OpCodes.Call, DispatchType.GetMethod(nameof(CbPolymorphicDispatch<object>.Read))!);
+					Generator.Emit(OpCodes.Ret);
+				}
+
+				// Create the write dispatch method
+				{
+					ILGenerator Generator = WriteContentsMethod.GetILGenerator();
+					Generator.Emit(OpCodes.Ldarg_0);
+					Generator.Emit(OpCodes.Ldarg_1);
+					Generator.Emit(OpCodes.Call, DispatchType.GetMethod(nameof(CbPolymorphicDispatch<object>.WriteContents))!);
+					Generator.Emit(OpCodes.Ret);
+				}
+
+				// Finally, update the dispatch type with all the methods. We should be safe on the recursive path now.
+				PopulateDispatchType(ClassType, DispatchType);
+			}
+		}
+
+		static void PopulateDispatchType(Type ClassType, Type DispatchType)
+		{
+			Dictionary<Utf8String, Type> DiscriminatorToKnownType = new Dictionary<Utf8String, Type>();
+
+			Type[] KnownTypes = ClassType.Assembly.GetTypes();
+			foreach (Type KnownType in KnownTypes)
+			{
+				if (KnownType.IsClass && !KnownType.IsAbstract)
+				{
+					for (Type? BaseType = KnownType; BaseType != null; BaseType = BaseType.BaseType)
+					{
+						if (BaseType == ClassType)
+						{
+							CbDiscriminatorAttribute? Discriminator = KnownType.GetCustomAttribute<CbDiscriminatorAttribute>();
+							if (Discriminator == null)
+							{
+								throw new NotSupportedException();
+							}
+							DiscriminatorToKnownType[Discriminator.Name] = KnownType;
+						}
+					}
+				}
+			}
+
+			// Populate the dictionary
+			Dictionary<Utf8String, Func<CbField, object>> NameToReadFunc = (Dictionary<Utf8String, Func<CbField, object>>)DispatchType.GetField(nameof(CbPolymorphicDispatch<object>.NameToReadFunc))!.GetValue(null)!;
+			Dictionary<Type, Action<CbWriter, object>> TypeToWriteContentsFunc = (Dictionary<Type, Action<CbWriter, object>>)DispatchType.GetField(nameof(CbPolymorphicDispatch<object>.TypeToWriteContentsFunc))!.GetValue(null)!;
+
+			foreach ((Utf8String Name, Type KnownType) in DiscriminatorToKnownType)
+			{
+				CbClassConverterMethods Methods = Create(KnownType);
+
+				{
+					DynamicMethod DynamicMethod = new DynamicMethod("_", typeof(object), new Type[] { typeof(CbField) });
+					ILGenerator Generator = DynamicMethod.GetILGenerator();
+					Generator.Emit(OpCodes.Ldarg_0);
+					Generator.Emit(OpCodes.Call, Methods.ReadConcreteMethod);
+					Generator.Emit(OpCodes.Castclass, typeof(object));
+					Generator.Emit(OpCodes.Ret);
+					NameToReadFunc[Name] = CreateDelegate<Func<CbField, object>>(DynamicMethod);
+				}
+
+				{
+					DynamicMethod DynamicMethod = new DynamicMethod("_", null, new Type[] { typeof(CbWriter), typeof(object) });
+					ILGenerator Generator = DynamicMethod.GetILGenerator();
+					Generator.Emit(OpCodes.Ldarg_0);
+					Generator.Emit(OpCodes.Ldarg_1);
+					Generator.Emit(OpCodes.Castclass, KnownType);
+					Generator.Emit(OpCodes.Call, Methods.WriteConcreteContentsMethod);
+					Generator.Emit(OpCodes.Ret);
+					TypeToWriteContentsFunc[KnownType] = CreateDelegate<Action<CbWriter, object>>(DynamicMethod);
+				}
+			}
+		}
+
+		static void CreateObjectWriter(Type Type, ILGenerator Generator, DynamicMethod ContentsWriter)
+		{
+			Generator.Emit(OpCodes.Ldarg_1);
+
+			Label SkipLabel = Generator.DefineLabel();
+			Generator.Emit(OpCodes.Brfalse, SkipLabel);
+
+			Generator.Emit(OpCodes.Ldarg_0);
+			Generator.EmitCall(OpCodes.Call, GetMethodInfo<CbWriter>(x => x.BeginObject()), null);
+
+			Generator.Emit(OpCodes.Ldarg_0);
+			Generator.Emit(OpCodes.Ldarg_1);
+			Generator.EmitCall(OpCodes.Call, ContentsWriter, null);
+
+			Generator.Emit(OpCodes.Ldarg_0);
+			Generator.EmitCall(OpCodes.Call, GetMethodInfo<CbWriter>(x => x.EndObject()), null);
+
+			Generator.MarkLabel(SkipLabel);
+			Generator.Emit(OpCodes.Ret);
+		}
+
+		static void CreateNamedObjectWriter(Type Type, ILGenerator Generator, DynamicMethod ContentsWriter)
+		{
+			Generator.Emit(OpCodes.Ldarg_2);
+
+			Label SkipLabel = Generator.DefineLabel();
+			Generator.Emit(OpCodes.Brfalse, SkipLabel);
+
+			Generator.Emit(OpCodes.Ldarg_0);
+			Generator.Emit(OpCodes.Ldarg_1);
+			Generator.EmitCall(OpCodes.Call, GetMethodInfo<CbWriter>(x => x.BeginObject(null!)), null);
+
+			Generator.Emit(OpCodes.Ldarg_0);
+			Generator.Emit(OpCodes.Ldarg_2);
+			Generator.EmitCall(OpCodes.Call, ContentsWriter, null);
+
+			Generator.Emit(OpCodes.Ldarg_0);
+			Generator.EmitCall(OpCodes.Call, GetMethodInfo<CbWriter>(x => x.EndObject()), null);
+
+			Generator.MarkLabel(SkipLabel);
+			Generator.Emit(OpCodes.Ret);
+		}
+
+		static void CreateConcreteObjectContentsWriter(Type Type, ILGenerator Generator)
+		{
+			// Find the reflected properties from this type
+			(Utf8String Name, PropertyInfo Property)[] Properties = GetProperties(Type);
+
+			// Create a static type with the required reflection data
+			Type ReflectedType = typeof(CbReflectedTypeInfo<>).MakeGenericType(Type);
+			FieldInfo NamesField = ReflectedType.GetField(nameof(CbReflectedTypeInfo<object>.Names))!;
+			NamesField.SetValue(null, Properties.Select(x => x.Name).ToArray());
+
+			// Write the discriminator
+			CbDiscriminatorAttribute? Discriminator = Type.GetCustomAttribute<CbDiscriminatorAttribute>();
+			if (Discriminator != null)
+			{
+				FieldInfo DiscriminatorKeyField = GetFieldInfo(() => DiscriminatorKey);
+				Generator.Emit(OpCodes.Ldarg_0);
+				Generator.Emit(OpCodes.Ldsfld, DiscriminatorKeyField);
+				Generator.Emit(OpCodes.Ldstr, Discriminator.Name);
+				Generator.Emit(OpCodes.Call, GetMethodInfo<CbWriter>(x => x.WriteString(default, null!)));
+			}
+
+			// Write all the remaining properties
+			for (int Idx = 0; Idx < Properties.Length; Idx++)
+			{
+				PropertyInfo Property = Properties[Idx].Property;
+				Type PropertyType = Property.PropertyType;
+
+				// Get the field value
+				Generator.Emit(OpCodes.Ldarg_1);
+				Generator.EmitCall(OpCodes.Call, Property.GetMethod!, null);
+
+				Label SkipLabel = Generator.DefineLabel();
+
+				MethodInfo WriteMethod;
+				if (TypeToMethods.TryGetValue(PropertyType, out CbClassConverterMethods? DynamicMethods))
+				{
+					Generator.Emit(OpCodes.Dup);
+					Generator.Emit(OpCodes.Brfalse, SkipLabel);
+					WriteMethod = DynamicMethods.WriteNamedMethod;
+				}
+				else
+				{
+					ICbConverterMethods Methods = CbConverterMethods.Get(PropertyType);
+					WriteMethod = Methods.WriteNamedMethod;
+				}
+
+				// Store the variable in a local
+				LocalBuilder Local = Generator.DeclareLocal(PropertyType);
+				Generator.Emit(OpCodes.Dup);
+				Generator.Emit(OpCodes.Stloc, Local);
+
+				// Call the writer
+				Generator.Emit(OpCodes.Ldarg_0);
+
+				Generator.Emit(OpCodes.Ldsfld, NamesField);
+				Generator.Emit(OpCodes.Ldc_I4, Idx);
+				Generator.Emit(OpCodes.Ldelem, typeof(Utf8String));
+
+				Generator.Emit(OpCodes.Ldloc, Local);
+				Generator.EmitCall(OpCodes.Call, WriteMethod, null);
+
+				// Remove the duplicated value from the top of the stack
+				Generator.MarkLabel(SkipLabel);
+				Generator.Emit(OpCodes.Pop);
+			}
+			Generator.Emit(OpCodes.Ret);
+		}
+
+		class CbPolymorphicDispatch<T>
+		{
+			public static Dictionary<Utf8String, Func<CbField, object>> NameToReadFunc = new Dictionary<Utf8String, Func<CbField, object>>();
+			public static Dictionary<Type, Action<CbWriter, object>> TypeToWriteContentsFunc = new Dictionary<Type, Action<CbWriter, object>>();
+
+			public static object Read(CbField Field)
+			{
+				Utf8String Name = Field.AsObject().Find(DiscriminatorKey).AsUtf8String();
+				return NameToReadFunc[Name](Field);
+			}
+
+			public static void WriteContents(CbWriter Writer, object Value)
+			{
+				Type Type = Value!.GetType();
+				TypeToWriteContentsFunc[Type](Writer, Value);
+			}
+		}
+
+		static T CreateDelegate<T>(DynamicMethod Method) where T : Delegate
+		{
+			return (T)Method.CreateDelegate(typeof(T));
+		}
+
+		static void CreateConcreteObjectReader(Type Type, ILGenerator Generator)
+		{
+			// Construct the object
+			ConstructorInfo? Constructor = Type.GetConstructor(BindingFlags.Public | BindingFlags.NonPublic | BindingFlags.Instance, null, Type.EmptyTypes, null);
+			if (Constructor == null)
+			{
+				throw new CbException($"Unable to find default constructor for {Type}");
+			}
+
+			// Find the reflected properties from this type
+			(Utf8String Name, PropertyInfo Property)[] Properties = GetProperties(Type);
+
+			// Create a static type with the required reflection data
+			Type ReflectedType = typeof(CbReflectedTypeInfo<>).MakeGenericType(Type);
+			FieldInfo NamesField = ReflectedType.GetField(nameof(CbReflectedTypeInfo<object>.Names))!;
+			NamesField.SetValue(null, Properties.Select(x => x.Name).ToArray());
+			MethodInfo MatchNameMethod = ReflectedType.GetMethod(nameof(CbReflectedTypeInfo<object>.MatchName))!;
+
+			// NewObjectLocal = new Type()
+			LocalBuilder NewObjectLocal = Generator.DeclareLocal(typeof(object));
+			Generator.Emit(OpCodes.Newobj, Constructor);
+			Generator.Emit(OpCodes.Stloc, NewObjectLocal);
+
+			// Stack(0) = CbField.CreateIterator()
+			Generator.Emit(OpCodes.Ldarg_0);
+			Generator.EmitCall(OpCodes.Call, GetMethodInfo<CbField>(x => x.CreateIterator()), null);
+
+			// CbFieldIterator IteratorLocal = Stack(0)
+			LocalBuilder IteratorLocal = Generator.DeclareLocal(typeof(CbFieldIterator));
+			Generator.Emit(OpCodes.Dup);
+			Generator.Emit(OpCodes.Stloc, IteratorLocal);
+
+			// if(!Stack.Pop().IsValid()) goto ReturnLabel
+			Generator.EmitCall(OpCodes.Call, GetMethodInfo<CbFieldIterator>(x => x.IsValid()), null);
+			Label ReturnLabel = Generator.DefineLabel();
+			Generator.Emit(OpCodes.Brfalse, ReturnLabel);
+
+			// NamesLocal = CbReflectedTypeInfo<Type>.Names
+			LocalBuilder NamesLocal = Generator.DeclareLocal(typeof(Utf8String[]));
+			Generator.Emit(OpCodes.Ldsfld, NamesField);
+			Generator.Emit(OpCodes.Stloc, NamesLocal);
+
+			// IterationLoopLabel:
+			Label IterationLoopLabel = Generator.DefineLabel();
+			Generator.MarkLabel(IterationLoopLabel);
+
+			// bool MatchLocal = false
+			LocalBuilder MatchLocal = Generator.DeclareLocal(typeof(bool));
+			Generator.Emit(OpCodes.Ldc_I4_0);
+			Generator.Emit(OpCodes.Stloc, MatchLocal);
+
+			// Stack(0) = IteratorLocal.GetCurrent()
+			Generator.Emit(OpCodes.Ldloc, IteratorLocal);
+			Generator.EmitCall(OpCodes.Call, GetMethodInfo<CbFieldIterator>(x => x.GetCurrent()), null);
+
+			// Try to parse each of the properties in order. If fields are ordered correctly, we will parse the object in a single pass. Otherwise we can loop and start again.
+			LocalBuilder FieldLocal = Generator.DeclareLocal(typeof(CbField));
+			for (int Idx = 0; Idx < Properties.Length; Idx++)
+			{
+				PropertyInfo Property = Properties[Idx].Property;
+
+				// Get the read method for this property type
+				MethodInfo ReadMethod;
+				if (TypeToMethods.TryGetValue(Property.PropertyType, out CbClassConverterMethods? DynamicMethods))
+				{
+					ReadMethod = DynamicMethods.ReadMethod;
+				}
+				else
+				{
+					ReadMethod = CbConverterMethods.Get(Property.PropertyType).ReadMethod;
+				}
+
+				// if(!CbReflectedTypeInfo<Type>.MatchName(Stack(0), Idx)) goto SkipPropertyLabel
+				Label SkipPropertyLabel = Generator.DefineLabel();
+				Generator.Emit(OpCodes.Dup); // Current CbField
+				Generator.Emit(OpCodes.Ldc_I4, Idx);
+				Generator.Emit(OpCodes.Call, MatchNameMethod);
+				Generator.Emit(OpCodes.Brfalse, SkipPropertyLabel);
+
+				// FieldLocal = Stack.Pop()
+				Generator.Emit(OpCodes.Stloc, FieldLocal);
+
+				// Property.SetMethod(NewObjectLocal, ReadMethod(FieldLocal))
+				Generator.Emit(OpCodes.Ldloc, NewObjectLocal);
+				Generator.Emit(OpCodes.Ldloc, FieldLocal);
+				Generator.EmitCall(OpCodes.Call, ReadMethod, null);
+				Generator.EmitCall(OpCodes.Call, Property.SetMethod!, null);
+
+				// if(!IteratorLocal.MoveNext()) goto ReturnLabel
+				Generator.Emit(OpCodes.Ldloc, IteratorLocal);
+				Generator.EmitCall(OpCodes.Call, GetMethodInfo<CbFieldIterator>(x => x.MoveNext()), null);
+				Generator.Emit(OpCodes.Brfalse, ReturnLabel);
+
+				// MatchLocal = true
+				Generator.Emit(OpCodes.Ldc_I4_1);
+				Generator.Emit(OpCodes.Stloc, MatchLocal);
+
+				// Stack(0) = IteratorLocal.GetCurrent()
+				Generator.Emit(OpCodes.Ldloc, IteratorLocal);
+				Generator.EmitCall(OpCodes.Call, GetMethodInfo<CbFieldIterator>(x => x.GetCurrent()), null);
+
+				// SkipPropertyLabel:
+				Generator.MarkLabel(SkipPropertyLabel);
+			}
+
+			// Stack.Pop()
+			Generator.Emit(OpCodes.Pop); // Current CbField
+
+			// if(MatchLocal) goto IterationLoopLabel
+			Generator.Emit(OpCodes.Ldloc, MatchLocal);
+			Generator.Emit(OpCodes.Brtrue, IterationLoopLabel);
+
+			// if(IteratorLocal.MoveNext()) goto IterationLoopLabel
+			Generator.Emit(OpCodes.Ldloc, IteratorLocal);
+			Generator.EmitCall(OpCodes.Call, GetMethodInfo<CbFieldIterator>(x => x.MoveNext()), null);
+			Generator.Emit(OpCodes.Brtrue, IterationLoopLabel);
+
+			// return NewObjectLocal
+			Generator.MarkLabel(ReturnLabel);
+			Generator.Emit(OpCodes.Ldloc, NewObjectLocal);
+			Generator.Emit(OpCodes.Ret);
+		}
+
+		static (Utf8String, PropertyInfo)[] GetProperties(Type Type)
+		{
+			List<(Utf8String, PropertyInfo)> PropertyList = new List<(Utf8String, PropertyInfo)>();
+			foreach (PropertyInfo Property in Type.GetProperties(BindingFlags.Public | BindingFlags.NonPublic | BindingFlags.Instance))
+			{
+				CbFieldAttribute? Attribute = Property.GetCustomAttribute<CbFieldAttribute>();
+				if (Attribute != null)
+				{
+					Utf8String Name = Attribute.Name ?? Property.Name;
+					PropertyList.Add((Name, Property));
+				}
+			}
+			if (PropertyList.Count == 0)
+			{
+				throw new CbException($"{Type.Name} does not have any properties marked with [CbField]");
+			}
+			return PropertyList.ToArray();
+		}
+
+		static FieldInfo GetFieldInfo<T>(Expression<Func<T>> Expr)
+		{
+			return (FieldInfo)((MemberExpression)Expr.Body).Member;
+		}
+
+		static MethodInfo GetMethodInfo(Expression<Action> Expr)
+		{
+			return ((MethodCallExpression)Expr.Body).Method;
+		}
+
+		static MethodInfo GetMethodInfo<T>(Expression<Action<T>> Expr)
+		{
+			return ((MethodCallExpression)Expr.Body).Method;
+		}
+	}
+
+	class CbClassConverter<T> : CbConverterBase<T>, ICbConverterMethods where T : class
+	{
+		CbClassConverterMethods Methods;
+
+		Func<CbField, T> ReadFunc;
+		Action<CbWriter, T> WriteFunc;
+		Action<CbWriter, Utf8String, T> WriteNamedFunc;
+
+		public CbClassConverter()
+		{
+			Methods = CbClassConverterMethods.Create(typeof(T));
+
+			ReadFunc = CreateDelegate<Func<CbField, T>>(Methods.ReadMethod);
+			WriteFunc = CreateDelegate<Action<CbWriter, T>>(Methods.WriteMethod);
+			WriteNamedFunc = CreateDelegate<Action<CbWriter, Utf8String, T>>(Methods.WriteNamedMethod);
+		}
+
+		public MethodInfo ReadMethod => Methods.ReadMethod;
+
+		public MethodInfo WriteMethod => Methods.WriteMethod;
+
+		public MethodInfo WriteNamedMethod => Methods.WriteNamedMethod;
+
+		static TDelegate CreateDelegate<TDelegate>(DynamicMethod Method) where TDelegate : Delegate => (TDelegate)Method.CreateDelegate(typeof(TDelegate));
+
+		public override T Read(CbField Field) => ReadFunc(Field);
+
+		public override void Write(CbWriter Writer, T Value) => WriteFunc(Writer, Value);
+
+		public override void WriteNamed(CbWriter Writer, Utf8String Name, T Value) => WriteNamedFunc(Writer, Name, Value);
+	}
+
+	class CbClassConverterFactory : CbConverterFactory
+	{
+		public override ICbConverter? CreateConverter(Type Type)
+		{
+			ICbConverter? Converter = null;
+			if (Type.IsClass)
+			{
+				Type ConverterType = typeof(CbClassConverter<>).MakeGenericType(Type);
+				Converter = (ICbConverter?)Activator.CreateInstance(ConverterType);
+			}
+			return Converter;
+		}
+	}
+}
diff --git a/Engine/Source/Programs/Shared/EpicGames.Serialization/Converters/CbDefaultConverter.cs b/Engine/Source/Programs/Shared/EpicGames.Serialization/Converters/CbDefaultConverter.cs
deleted file mode 100644
diff --git a/Engine/Source/Programs/Shared/EpicGames.Serialization/Converters/CbDictionaryConverter.cs b/Engine/Source/Programs/Shared/EpicGames.Serialization/Converters/CbDictionaryConverter.cs
--- a/Engine/Source/Programs/Shared/EpicGames.Serialization/Converters/CbDictionaryConverter.cs
+++ b/Engine/Source/Programs/Shared/EpicGames.Serialization/Converters/CbDictionaryConverter.cs
@@ -16,7 +16,7 @@ namespace EpicGames.Serialization.Converters
 	/// <summary>
 	/// Converter for dictionary types
 	/// </summary>
-	class CbDictionaryConverter<TKey, TValue> : CbConverter<Dictionary<TKey, TValue>> where TKey : notnull
+	class CbDictionaryConverter<TKey, TValue> : CbConverterBase<Dictionary<TKey, TValue>> where TKey : notnull
 	{
 		/// <inheritdoc/>
 		public override Dictionary<TKey, TValue> Read(CbField Field)
@@ -84,12 +84,12 @@ namespace EpicGames.Serialization.Converters
 	class CbDictionaryConverterFactory : CbConverterFactory
 	{
 		/// <inheritdoc/>
-		public override CbConverter? CreateConverter(Type Type)
+		public override ICbConverter? CreateConverter(Type Type)
 		{
 			if (Type.IsGenericType && Type.GetGenericTypeDefinition() == typeof(Dictionary<,>))
 			{
 				Type ConverterType = typeof(CbDictionaryConverter<,>).MakeGenericType(Type.GenericTypeArguments);
-				return (CbConverter)Activator.CreateInstance(ConverterType)!;
+				return (ICbConverter)Activator.CreateInstance(ConverterType)!;
 			}
 			return null;
 		}
diff --git a/Engine/Source/Programs/Shared/EpicGames.Serialization/Converters/CbEnumConverter.cs b/Engine/Source/Programs/Shared/EpicGames.Serialization/Converters/CbEnumConverter.cs
new file mode 100644
--- /dev/null
+++ b/Engine/Source/Programs/Shared/EpicGames.Serialization/Converters/CbEnumConverter.cs
@@ -0,0 +1,110 @@
+// Copyright Epic Games, Inc. All Rights Reserved.
+
+using EpicGames.Core;
+using System;
+using System.Collections.Generic;
+using System.Linq.Expressions;
+using System.Reflection;
+using System.Reflection.Emit;
+using System.Text;
+
+namespace EpicGames.Serialization.Converters
+{
+	class CbEnumConverter<T> : CbConverterBase<T>, ICbConverterMethods where T : Enum
+	{
+		Func<CbField, T> ReadFunc;
+		Action<CbWriter, T> WriteFunc;
+		Action<CbWriter, Utf8String, T> WriteNamedFunc;
+
+		public CbEnumConverter()
+		{
+			Type Type = typeof(T);
+
+			ReadMethod = new DynamicMethod($"Read_{Type.Name}", Type, new Type[] { typeof(CbField) });
+			CreateEnumReader(Type, ReadMethod.GetILGenerator());
+			ReadFunc = (Func<CbField, T>)ReadMethod.CreateDelegate(typeof(Func<CbField, T>));
+
+			WriteMethod = new DynamicMethod($"Write_{Type.Name}", null, new Type[] { typeof(CbWriter), Type });
+			CreateEnumWriter(Type, WriteMethod.GetILGenerator());
+			WriteFunc = (Action<CbWriter, T>)WriteMethod.CreateDelegate(typeof(Action<CbWriter, T>));
+
+			WriteNamedMethod = new DynamicMethod($"WriteNamed_{Type.Name}", null, new Type[] { typeof(CbWriter), typeof(Utf8String), Type });
+			CreateNamedEnumWriter(Type, WriteNamedMethod.GetILGenerator());
+			WriteNamedFunc = (Action<CbWriter, Utf8String, T>)WriteNamedMethod.CreateDelegate(typeof(Action<CbWriter, Utf8String, T>));
+		}
+
+		public DynamicMethod ReadMethod { get; }
+		MethodInfo ICbConverterMethods.ReadMethod => ReadMethod;
+
+		public DynamicMethod WriteMethod { get; }
+		MethodInfo ICbConverterMethods.WriteMethod => WriteMethod;
+
+		public DynamicMethod WriteNamedMethod { get; }
+		MethodInfo ICbConverterMethods.WriteNamedMethod => WriteNamedMethod;
+
+		public override T Read(CbField Field) => ReadFunc(Field);
+
+		public override void Write(CbWriter Writer, T Value) => WriteFunc(Writer, Value);
+
+		public override void WriteNamed(CbWriter Writer, Utf8String Name, T Value) => WriteNamedFunc(Writer, Name, Value);
+
+		static void CreateEnumReader(Type Type, ILGenerator Generator)
+		{
+			Generator.Emit(OpCodes.Ldarg_0);
+			Generator.EmitCall(OpCodes.Call, GetMethodInfo<CbField>(x => x.AsInt32()), null);
+			Generator.Emit(OpCodes.Ret);
+		}
+
+		static void CreateEnumWriter(Type Type, ILGenerator Generator)
+		{
+			Generator.Emit(OpCodes.Ldarg_1);
+
+			Label SkipLabel = Generator.DefineLabel();
+			Generator.Emit(OpCodes.Brfalse, SkipLabel);
+
+			Generator.Emit(OpCodes.Ldarg_0);
+			Generator.Emit(OpCodes.Ldarg_1);
+			Generator.Emit(OpCodes.Conv_I8);
+			Generator.EmitCall(OpCodes.Call, GetMethodInfo<CbWriter>(x => x.WriteIntegerValue(0L)), null);
+
+			Generator.MarkLabel(SkipLabel);
+			Generator.Emit(OpCodes.Ret);
+		}
+
+		static void CreateNamedEnumWriter(Type Type, ILGenerator Generator)
+		{
+			Generator.Emit(OpCodes.Ldarg_2);
+
+			Label SkipLabel = Generator.DefineLabel();
+			Generator.Emit(OpCodes.Brfalse, SkipLabel);
+
+			Generator.Emit(OpCodes.Ldarg_0);
+			Generator.Emit(OpCodes.Ldarg_1);
+			Generator.Emit(OpCodes.Ldarg_2);
+			Generator.Emit(OpCodes.Conv_I8);
+			Generator.EmitCall(OpCodes.Call, GetMethodInfo<CbWriter>(x => x.WriteInteger(default, 0L)), null);
+
+			Generator.MarkLabel(SkipLabel);
+			Generator.Emit(OpCodes.Ret);
+		}
+
+		static MethodInfo GetMethodInfo<TArg>(Expression<Action<TArg>> Expr)
+		{
+			return ((MethodCallExpression)Expr.Body).Method;
+		}
+	}
+
+	class CbEnumConverterFactory : CbConverterFactory
+	{
+		public override ICbConverter? CreateConverter(Type Type)
+		{
+			ICbConverter? Converter = null;
+			if (Type.IsEnum)
+			{
+				Type ConverterType = typeof(CbEnumConverter<>).MakeGenericType(Type);
+				Converter = (ICbConverter?)Activator.CreateInstance(ConverterType);
+			}
+			return Converter;
+		}
+	}
+}
diff --git a/Engine/Source/Programs/Shared/EpicGames.Serialization/Converters/CbFieldConverter.cs b/Engine/Source/Programs/Shared/EpicGames.Serialization/Converters/CbFieldConverter.cs
--- a/Engine/Source/Programs/Shared/EpicGames.Serialization/Converters/CbFieldConverter.cs
+++ b/Engine/Source/Programs/Shared/EpicGames.Serialization/Converters/CbFieldConverter.cs
@@ -11,7 +11,7 @@ namespace EpicGames.Serialization.Converters
 	/// Converter for raw CbField types
 	/// </summary>
 	/// <typeparam name="T"></typeparam>
-	class CbFieldConverter : CbConverter<CbField>
+	class CbFieldConverter : CbConverterBase<CbField>
 	{
 		/// <inheritdoc/>
 		public override CbField Read(CbField Field)
@@ -36,7 +36,7 @@ namespace EpicGames.Serialization.Converters
 	/// Converter for raw CbObject types
 	/// </summary>
 	/// <typeparam name="T"></typeparam>
-	class CbObjectConverter : CbConverter<CbObject>
+	class CbObjectConverter : CbConverterBase<CbObject>
 	{
 		/// <inheritdoc/>
 		public override CbObject Read(CbField Field)
diff --git a/Engine/Source/Programs/Shared/EpicGames.Serialization/Converters/CbListConverter.cs b/Engine/Source/Programs/Shared/EpicGames.Serialization/Converters/CbListConverter.cs
--- a/Engine/Source/Programs/Shared/EpicGames.Serialization/Converters/CbListConverter.cs
+++ b/Engine/Source/Programs/Shared/EpicGames.Serialization/Converters/CbListConverter.cs
@@ -17,7 +17,7 @@ namespace EpicGames.Serialization.Converters
 	/// Converter for list types
 	/// </summary>
 	/// <typeparam name="T"></typeparam>
-	class CbListConverter<T> : CbConverter<List<T>>
+	class CbListConverter<T> : CbConverterBase<List<T>>
 	{
 		/// <inheritdoc/>
 		public override List<T> Read(CbField Field)
@@ -59,7 +59,7 @@ namespace EpicGames.Serialization.Converters
 	/// <summary>
 	/// Specialization for serializing string lists
 	/// </summary>
-	sealed class CbStringListConverter : CbConverter<List<Utf8String>>
+	sealed class CbStringListConverter : CbConverterBase<List<Utf8String>>
 	{
 		/// <inheritdoc/>
 		public override List<Utf8String> Read(CbField Field)
@@ -104,7 +104,7 @@ namespace EpicGames.Serialization.Converters
 	class CbListConverterFactory : CbConverterFactory
 	{
 		/// <inheritdoc/>
-		public override CbConverter? CreateConverter(Type Type)
+		public override ICbConverter? CreateConverter(Type Type)
 		{
 			if (Type.IsGenericType && Type.GetGenericTypeDefinition() == typeof(List<>))
 			{
@@ -115,7 +115,7 @@ namespace EpicGames.Serialization.Converters
 				else
 				{
 					Type ConverterType = typeof(CbListConverter<>).MakeGenericType(Type.GenericTypeArguments);
-					return (CbConverter)Activator.CreateInstance(ConverterType)!;
+					return (ICbConverter)Activator.CreateInstance(ConverterType)!;
 				}
 			}
 			return null;
diff --git a/Engine/Source/Programs/Shared/EpicGames.Serialization/Converters/CbNullableConverter.cs b/Engine/Source/Programs/Shared/EpicGames.Serialization/Converters/CbNullableConverter.cs
--- a/Engine/Source/Programs/Shared/EpicGames.Serialization/Converters/CbNullableConverter.cs
+++ b/Engine/Source/Programs/Shared/EpicGames.Serialization/Converters/CbNullableConverter.cs
@@ -11,7 +11,7 @@ namespace EpicGames.Serialization.Converters
 	/// Converter for list types
 	/// </summary>
 	/// <typeparam name="T"></typeparam>
-	class CbNullableConverter<T> : CbConverter<Nullable<T>> where T : struct
+	class CbNullableConverter<T> : CbConverterBase<Nullable<T>> where T : struct
 	{
 		/// <inheritdoc/>
 		public override T? Read(CbField Field)
@@ -44,12 +44,12 @@ namespace EpicGames.Serialization.Converters
 	class CbNullableConverterFactory : CbConverterFactory
 	{
 		/// <inheritdoc/>
-		public override CbConverter? CreateConverter(Type Type)
+		public override ICbConverter? CreateConverter(Type Type)
 		{
 			if (Type.IsGenericType && Type.GetGenericTypeDefinition() == typeof(Nullable<>))
 			{
 				Type ConverterType = typeof(CbNullableConverter<>).MakeGenericType(Type.GenericTypeArguments);
-				return (CbConverter)Activator.CreateInstance(ConverterType)!;
+				return (ICbConverter)Activator.CreateInstance(ConverterType)!;
 			}
 			return null;
 		}
diff --git a/Engine/Source/Programs/Shared/EpicGames.Serialization/Converters/CbPrimitiveConverter.cs b/Engine/Source/Programs/Shared/EpicGames.Serialization/Converters/CbPrimitiveConverter.cs
new file mode 100644
--- /dev/null
+++ b/Engine/Source/Programs/Shared/EpicGames.Serialization/Converters/CbPrimitiveConverter.cs
@@ -0,0 +1,41 @@
+// Copyright Epic Games, Inc. All Rights Reserved.
+
+using EpicGames.Core;
+using System;
+using System.Collections.Generic;
+using System.Linq.Expressions;
+using System.Reflection;
+using System.Text;
+
+namespace EpicGames.Serialization.Converters
+{
+	class CbPrimitiveConverter<T> : CbConverterBase<T>, ICbConverterMethods
+	{
+		public MethodInfo ReadMethod { get; }
+		public Func<CbField, T> ReadFunc { get; }
+
+		public MethodInfo WriteMethod { get; }
+		public Action<CbWriter, T> WriteFunc { get; }
+
+		public MethodInfo WriteNamedMethod { get; }
+		public Action<CbWriter, Utf8String, T> WriteNamedFunc { get; }
+
+		public CbPrimitiveConverter(Expression<Func<CbField, T>> Read, Expression<Action<CbWriter, T>> Write, Expression<Action<CbWriter, Utf8String, T>> WriteNamed)
+		{
+			this.ReadMethod = ((MethodCallExpression)Read.Body).Method;
+			this.ReadFunc = Read.Compile();
+
+			this.WriteMethod = ((MethodCallExpression)Write.Body).Method;
+			this.WriteFunc = Write.Compile();
+
+			this.WriteNamedMethod = ((MethodCallExpression)WriteNamed.Body).Method;
+			this.WriteNamedFunc = WriteNamed.Compile();
+		}
+
+		public override T Read(CbField Field) => ReadFunc(Field);
+
+		public override void Write(CbWriter Writer, T Value) => WriteFunc(Writer, Value);
+
+		public override void WriteNamed(CbWriter Writer, Utf8String Name, T Value) => WriteNamed(Writer, Name, Value);
+	}
+}
