Client-facing interface for Horde.Storage.

Idea is that tools could interface with Horde Storage use this API, with concrete implementations that use HTTP, gRPC or (in the case of legacy applications within Horde.Build) communicate internally within the same process to store the data.

Errors should be surfaced using exceptions.

diff --git a/Engine/Source/Programs/Shared/EpicGames.Horde/Storage/BucketId.cs b/Engine/Source/Programs/Shared/EpicGames.Horde/Storage/BucketId.cs
new file mode 100644
--- /dev/null
+++ b/Engine/Source/Programs/Shared/EpicGames.Horde/Storage/BucketId.cs
@@ -0,0 +1,47 @@
+// Copyright Epic Games, Inc. All Rights Reserved.
+
+using System;
+using System.Collections.Generic;
+using System.Diagnostics.CodeAnalysis;
+using System.Text;
+
+namespace EpicGames.Horde.Storage
+{
+	/// <summary>
+	/// Identifier for a storage bucket
+	/// </summary>
+	public struct BucketId : IEquatable<BucketId>
+	{
+		/// <summary>
+		/// The text representing this id
+		/// </summary>
+		readonly StringId Inner;
+
+		/// <summary>
+		/// Constructor
+		/// </summary>
+		/// <param name="Input">Unique id for the string</param>
+		public BucketId(string Input)
+		{
+			Inner = new StringId(Input);
+		}
+
+		/// <inheritdoc/>
+		public override bool Equals(object? Obj) => Obj is BucketId Id && Inner.Equals(Id.Inner);
+
+		/// <inheritdoc/>
+		public override int GetHashCode() => Inner.GetHashCode();
+
+		/// <inheritdoc/>
+		public bool Equals(BucketId Other) => Inner.Equals(Other.Inner);
+
+		/// <inheritdoc/>
+		public override string ToString() => Inner.ToString();
+
+		/// <inheritdoc cref="StringId.op_Equality"/>
+		public static bool operator ==(BucketId Left, BucketId Right) => Left.Inner == Right.Inner;
+
+		/// <inheritdoc cref="StringId.op_Inequality"/>
+		public static bool operator !=(BucketId Left, BucketId Right) => Left.Inner != Right.Inner;
+	}
+}
diff --git a/Engine/Source/Programs/Shared/EpicGames.Horde/Storage/IBlobStore.cs b/Engine/Source/Programs/Shared/EpicGames.Horde/Storage/IBlobStore.cs
new file mode 100644
--- /dev/null
+++ b/Engine/Source/Programs/Shared/EpicGames.Horde/Storage/IBlobStore.cs
@@ -0,0 +1,170 @@
+// Copyright Epic Games, Inc. All Rights Reserved.
+
+using EpicGames.Core;
+using System;
+using System.Collections.Generic;
+using System.IO;
+using System.Text;
+using System.Threading.Tasks;
+
+namespace EpicGames.Horde.Storage
+{
+	/// <summary>
+	/// Base class for exceptions related to the blob store
+	/// </summary>
+	public class BlobException : Exception
+	{
+		/// <summary>
+		/// Namespace containing the blob
+		/// </summary>
+		public NamespaceId NamespaceId { get; }
+
+		/// <summary>
+		/// Hash of the blob
+		/// </summary>
+		public IoHash Hash { get; }
+
+		/// <summary>
+		/// Constructor
+		/// </summary>
+		public BlobException(NamespaceId NamespaceId, IoHash Hash, string Message, Exception? InnerException = null)
+			: base(Message, InnerException)
+		{
+			this.NamespaceId = NamespaceId;
+			this.Hash = Hash;
+		}
+	}
+
+	/// <summary>
+	/// Exception thrown for missing blobs
+	/// </summary>
+	public sealed class BlobNotFoundException : BlobException
+	{
+		/// <summary>
+		/// Constructor
+		/// </summary>
+		public BlobNotFoundException(NamespaceId NamespaceId, IoHash Hash, Exception? InnerException = null)
+			: base(NamespaceId, Hash, $"Unable to find blob {Hash} in {NamespaceId}", InnerException)
+		{
+		}
+	}
+
+	/// <summary>
+	/// A referenced namespace for a blob was not found
+	/// </summary>
+	public class BlobNamespaceNotFoundException : BlobException
+	{
+		/// <summary>
+		/// Constructor
+		/// </summary>
+		public BlobNamespaceNotFoundException(NamespaceId NamespaceId, IoHash Hash, Exception? InnerException = null)
+			: base(NamespaceId, Hash, $"Unable to find namespace {NamespaceId} for blob {Hash}", InnerException)
+		{
+		}
+	}
+
+	/// <summary>
+	/// Interface for a collection of blobs
+	/// </summary>
+	public interface IBlobStore
+	{
+		/// <summary>
+		/// Determines whether a blob exists
+		/// </summary>
+		/// <param name="NamespaceId">Namespace to operate on</param>
+		/// <param name="Hash"></param>
+		/// <returns></returns>
+		Task<bool> ExistsAsync(NamespaceId NamespaceId, IoHash Hash);
+
+		/// <summary>
+		/// Determines which blobs the storage collection needs
+		/// </summary>
+		/// <param name="NamespaceId">Namespace to operate on</param>
+		/// <param name="Hashes"></param>
+		/// <returns></returns>
+		Task<List<IoHash>> NeedsAsync(NamespaceId NamespaceId, IEnumerable<IoHash> Hashes);
+
+		/// <summary>
+		/// Opens a blob read stream
+		/// </summary>
+		/// <param name="NamespaceId">Namespace to operate on</param>
+		/// <param name="Hash">Hash of the blob</param>
+		/// <returns>Stream for the blob, or null if it does not exist</returns>
+		Task<Stream> ReadAsync(NamespaceId NamespaceId, IoHash Hash);
+
+		/// <summary>
+		/// Adds an item to storage
+		/// </summary>
+		/// <param name="NamespaceId">Namespace to operate on</param>
+		/// <param name="Hash">Hash of the blob</param>
+		/// <param name="Stream">The stream to write</param>
+		Task WriteAsync(NamespaceId NamespaceId, IoHash Hash, Stream Stream);
+	}
+
+	/// <summary>
+	/// Extension methods for <see cref="IBlobStore"/>
+	/// </summary>
+	public static class BlobStoreExtensions
+	{
+		const int DefaultMaxInMemoryBlobLength = 50 * 1024 * 1024;
+
+		/// <summary>
+		/// Gets a blob as a byte array
+		/// </summary>
+		/// <param name="Collection"></param>
+		/// <param name="NamespaceId"></param>
+		/// <param name="Hash"></param>
+		/// <returns></returns>
+		public static async Task<byte[]> ReadBytesAsync(this IBlobStore Collection, NamespaceId NamespaceId, IoHash Hash, int MaxInMemoryBlobLength = DefaultMaxInMemoryBlobLength)
+		{
+			using Stream Stream = await Collection.ReadAsync(NamespaceId, Hash);
+
+			long Length = Stream.Length;
+			if (Length > MaxInMemoryBlobLength)
+			{
+				throw new BlobException(NamespaceId, Hash, $"Blob {Hash} is too large ({Length} > {MaxInMemoryBlobLength})");
+			}
+
+			byte[] Buffer = new byte[Length];
+			for(int Offset = 0; Offset < Length; )
+			{
+				int Count = await Stream.ReadAsync(Buffer, Offset, (int)Length - Offset);
+				if (Count == 0)
+				{
+					throw new BlobException(NamespaceId, Hash, $"Unexpected end of stream reading blob {Hash}");
+				}
+				Offset += Count;
+			}
+
+			return Buffer;
+		}
+
+		/// <summary>
+		/// Gets a blob as a byte array
+		/// </summary>
+		/// <param name="Collection"></param>
+		/// <param name="NamespaceId"></param>
+		/// <param name="Data">The data to be written</param>
+		/// <returns></returns>
+		public static async Task<IoHash> WriteBytesAsync(this IBlobStore Collection, NamespaceId NamespaceId, ReadOnlyMemory<byte> Data)
+		{
+			IoHash Hash = IoHash.Compute(Data.Span);
+			await WriteBytesAsync(Collection, NamespaceId, Hash, Data);
+			return Hash;
+		}
+
+		/// <summary>
+		/// Gets a blob as a byte array
+		/// </summary>
+		/// <param name="Collection"></param>
+		/// <param name="NamespaceId"></param>
+		/// <param name="Hash">Hash of the data</param>
+		/// <param name="Data">The data to be written</param>
+		/// <returns></returns>
+		public static async Task WriteBytesAsync(this IBlobStore Collection, NamespaceId NamespaceId, IoHash Hash, ReadOnlyMemory<byte> Data)
+		{
+			using ReadOnlyMemoryStream Stream = new ReadOnlyMemoryStream(Data);
+			await Collection.WriteAsync(NamespaceId, Hash, Stream);
+		}
+	}
+}
diff --git a/Engine/Source/Programs/Shared/EpicGames.Horde/Storage/IRefStore.cs b/Engine/Source/Programs/Shared/EpicGames.Horde/Storage/IRefStore.cs
new file mode 100644
--- /dev/null
+++ b/Engine/Source/Programs/Shared/EpicGames.Horde/Storage/IRefStore.cs
@@ -0,0 +1,225 @@
+// Copyright Epic Games, Inc. All Rights Reserved.
+
+using EpicGames.Core;
+using EpicGames.Serialization;
+using System;
+using System.Collections.Generic;
+using System.Diagnostics.CodeAnalysis;
+using System.IO;
+using System.Text;
+using System.Threading.Tasks;
+
+namespace EpicGames.Horde.Storage
+{
+	/// <summary>
+	/// Base class for exceptions
+	/// </summary>
+	public class RefException : Exception
+	{
+		/// <summary>
+		/// Namespace containing the ref
+		/// </summary>
+		public NamespaceId NamespaceId { get; }
+
+		/// <summary>
+		/// Bucket containing the ref
+		/// </summary>
+		public BucketId BucketId { get; }
+
+		/// <summary>
+		/// Name of the ref
+		/// </summary>
+		public IoHash RefId { get; }
+
+		/// <summary>
+		/// Constructor
+		/// </summary>
+		public RefException(NamespaceId NamespaceId, BucketId BucketId, IoHash RefId, string Message, Exception? InnerException = null)
+			: base(Message, InnerException)
+		{
+			this.NamespaceId = NamespaceId;
+			this.BucketId = BucketId;
+			this.RefId = RefId;
+		}
+	}
+
+	/// <summary>
+	/// Indicates that a named reference wasn't found
+	/// </summary>
+	public sealed class RefNotFoundException : RefException
+	{
+		/// <summary>
+		/// Constructor
+		/// </summary>
+		public RefNotFoundException(NamespaceId NamespaceId, BucketId BucketId, IoHash RefId, Exception? InnerException = null)
+			: base(NamespaceId, BucketId, RefId, $"Ref {NamespaceId}/{BucketId}/{RefId} not found", InnerException)
+		{
+		}
+	}
+
+	/// <summary>
+	/// Indicates that a named reference wasn't found
+	/// </summary>
+	public sealed class RefNamespaceNotFoundException : RefException
+	{
+		/// <summary>
+		/// Constructor
+		/// </summary>
+		public RefNamespaceNotFoundException(NamespaceId NamespaceId, BucketId BucketId, IoHash RefId, Exception? InnerException = null)
+			: base(NamespaceId, BucketId, RefId, $"Namespace not found for ref {NamespaceId}/{BucketId}/{RefId}", InnerException)
+		{
+		}
+	}
+
+	/// <summary>
+	/// Interface for an object reference
+	/// </summary>
+	public interface IRef
+	{
+		/// <summary>
+		/// Namespace identifier
+		/// </summary>
+		NamespaceId NamespaceId { get; }
+
+		/// <summary>
+		/// Bucket identifier
+		/// </summary>
+		BucketId BucketId { get; }
+
+		/// <summary>
+		/// Ref identifier
+		/// </summary>
+		IoHash RefId { get; }
+
+		/// <summary>
+		/// Last access time
+		/// </summary>
+		DateTime LastAccessTime { get; }
+
+		/// <summary>
+		/// Whether the ref has been finalized
+		/// </summary>
+		bool Finalized { get; }
+	}
+
+	/// <summary>
+	/// An object reference with attached data
+	/// </summary>
+	public interface IRefData : IRef, IDisposable
+	{
+		/// <summary>
+		/// Stream for the data
+		/// </summary>
+		Stream Stream { get; }
+	}
+
+	/// <summary>
+	/// Interface for a collection of ref documents
+	/// </summary>
+	public interface IRefStore
+	{
+		/// <summary>
+		/// Gets the given reference
+		/// </summary>
+		/// <param name="NamespaceId">Namespace identifier</param>
+		/// <param name="BucketId">Bucket identifier</param>
+		/// <param name="Name">Name of the reference</param>
+		/// <returns>The reference data if the ref exists</returns>
+		Task<IRef> GetAsync(NamespaceId NamespaceId, BucketId BucketId, IoHash RefId);
+
+		/// <summary>
+		/// Gets the given reference with its data
+		/// </summary>
+		/// <param name="NamespaceId">Namespace identifier</param>
+		/// <param name="BucketId">Bucket identifier</param>
+		/// <param name="RefId">Ref identifier</param>
+		/// <returns>The reference data if the ref exists</returns>
+		Task<IRefData> GetDataAsync(NamespaceId NamespaceId, BucketId BucketId, IoHash RefId);
+
+		/// <summary>
+		/// Gets the given reference
+		/// </summary>
+		/// <param name="NamespaceId">Namespace identifier</param>
+		/// <param name="BucketId">Bucket identifier</param>
+		/// <param name="RefId">Ref identifier</param>
+		/// <returns>The reference data if the ref exists</returns>
+		Task<bool> ExistsAsync(NamespaceId NamespaceId, BucketId BucketId, IoHash RefId);
+
+		/// <summary>
+		/// Determines which refs are missing
+		/// </summary>
+		/// <param name="NamespaceId">Namespace identifier</param>
+		/// <param name="BucketId">Bucket identifier</param>
+		/// <param name="RefIds">Names of the references</param>
+		/// <returns>The reference data if the ref exists</returns>
+		Task<List<string>> NeedsAsync(NamespaceId NamespaceId, BucketId BucketId, List<IoHash> RefIds);
+
+		/// <summary>
+		/// Sets the given reference
+		/// </summary>
+		/// <param name="NamespaceId">Namespace identifier</param>
+		/// <param name="BucketId">Bucket identifier</param>
+		/// <param name="RefId">Ref identifier</param>
+		/// <param name="Stream">New value for the reference, as a compact binary object</param>
+		/// <returns>List of missing references</returns>
+		Task<List<IoHash>> SetAsync(NamespaceId NamespaceId, BucketId BucketId, IoHash RefId, Stream Stream);
+
+		/// <summary>
+		/// Attempts to finalize a reference, turning its references into hard references
+		/// </summary>
+		/// <param name="NamespaceId">Namespace identifier</param>
+		/// <param name="BucketId">Bucket identifier</param>
+		/// <param name="RefId">Ref identifier</param>
+		/// <param name="Hash">Hash of the referenced object</param>
+		/// <returns></returns>
+		Task<List<IoHash>> FinalizeAsync(NamespaceId NamespaceId, BucketId BucketId, IoHash RefId, IoHash Hash);
+
+		/// <summary>
+		/// Removes the given reference
+		/// </summary>
+		/// <param name="NamespaceId">Namespace identifier</param>
+		/// <param name="BucketId">Bucket identifier</param>
+		/// <param name="RefId">Ref identifier</param>
+		/// <returns>True if the ref was deleted, false if it did not exist</returns>
+		Task<bool> DeleteAsync(NamespaceId NamespaceId, BucketId BucketId, IoHash RefId);
+	}
+
+	/// <summary>
+	/// Extension methods for refs
+	/// </summary>
+	public static class RefExtensions
+	{
+		/// <summary>
+		/// Reads a ref data as a compact binary field
+		/// </summary>
+		/// <param name="RefData">The ref data object</param>
+		/// <returns>Deserialized field</returns>
+		public static async ValueTask<CbField> AsCbFieldAsync(this IRefData RefData)
+		{
+			using MemoryStream MemoryStream = new MemoryStream();
+			await RefData.Stream.CopyToAsync(MemoryStream);
+			return new CbField(MemoryStream.ToArray());
+		}
+	}
+
+	/// <summary>
+	/// Extension methods for ref collections
+	/// </summary>
+	public static class RefCollectionExtensions
+	{
+		/// <summary>
+		/// Sets the given reference
+		/// </summary>
+		/// <param name="Collection">Collection to store the ref</param>
+		/// <param name="NamespaceId">Namespace identifier</param>
+		/// <param name="BucketId">Bucket identifier</param>
+		/// <param name="RefId">Ref identifier</param>
+		/// <param name="Object">New value for the reference, as a compact binary object</param>
+		/// <returns>List of missing references</returns>
+		public static async Task<List<IoHash>> SetAsync(this IRefStore Collection, NamespaceId NamespaceId, BucketId BucketId, IoHash RefId, CbObject Object)
+		{
+			using ReadOnlyMemoryStream Stream = new ReadOnlyMemoryStream(Object.GetView());
+			return await Collection.SetAsync(NamespaceId, BucketId, RefId, Stream);
+		}
+	}
+}
diff --git a/Engine/Source/Programs/Shared/EpicGames.Horde/Storage/NamespaceId.cs b/Engine/Source/Programs/Shared/EpicGames.Horde/Storage/NamespaceId.cs
new file mode 100644
--- /dev/null
+++ b/Engine/Source/Programs/Shared/EpicGames.Horde/Storage/NamespaceId.cs
@@ -0,0 +1,47 @@
+// Copyright Epic Games, Inc. All Rights Reserved.
+
+using System;
+using System.Collections.Generic;
+using System.Diagnostics.CodeAnalysis;
+using System.Text;
+
+namespace EpicGames.Horde.Storage
+{
+	/// <summary>
+	/// Identifier for a storage namespace
+	/// </summary>
+	public struct NamespaceId : IEquatable<NamespaceId>
+	{
+		/// <summary>
+		/// The text representing this id
+		/// </summary>
+		readonly StringId Inner;
+
+		/// <summary>
+		/// Constructor
+		/// </summary>
+		/// <param name="Input">Unique id for the string</param>
+		public NamespaceId(string Input)
+		{
+			Inner = new StringId(Input);
+		}
+
+		/// <inheritdoc/>
+		public override bool Equals(object? Obj) => Obj is NamespaceId Id && Inner.Equals(Id.Inner);
+
+		/// <inheritdoc/>
+		public override int GetHashCode() => Inner.GetHashCode();
+
+		/// <inheritdoc/>
+		public bool Equals(NamespaceId Other) => Inner.Equals(Other.Inner);
+
+		/// <inheritdoc/>
+		public override string ToString() => Inner.ToString();
+
+		/// <inheritdoc cref="StringId.op_Equality"/>
+		public static bool operator ==(NamespaceId Left, NamespaceId Right) => Left.Inner == Right.Inner;
+
+		/// <inheritdoc cref="StringId.op_Inequality"/>
+		public static bool operator !=(NamespaceId Left, NamespaceId Right) => Left.Inner != Right.Inner;
+	}
+}
diff --git a/Engine/Source/Programs/Shared/EpicGames.Horde/StringId.cs b/Engine/Source/Programs/Shared/EpicGames.Horde/StringId.cs
new file mode 100644
--- /dev/null
+++ b/Engine/Source/Programs/Shared/EpicGames.Horde/StringId.cs
@@ -0,0 +1,115 @@
+// Copyright Epic Games, Inc. All Rights Reserved.
+
+using System;
+using System.Collections.Generic;
+using System.Diagnostics.CodeAnalysis;
+using System.Text;
+
+namespace EpicGames.Horde
+{
+	/// <summary>
+	/// Normalized string identifier for a resource
+	/// </summary>
+	struct StringId : IEquatable<StringId>
+	{
+		/// <summary>
+		/// The text representing this id
+		/// </summary>
+		public string Text { get; }
+
+		/// <summary>
+		/// Constructor
+		/// </summary>
+		/// <param name="Input">Unique id for the string</param>
+		[SuppressMessage("Globalization", "CA1308:Normalize strings to uppercase", Justification = "Known limited character set")]
+		public StringId(string Input)
+		{
+			this.Text = Input;
+
+			if (Text.Length == 0)
+			{
+				throw new ArgumentException("String id may not be empty");
+			}
+
+			const int MaxLength = 64;
+			if (Text.Length > MaxLength)
+			{
+				throw new ArgumentException($"String id may not be longer than {MaxLength} characters");
+			}
+
+			for (int Idx = 0; Idx < Text.Length; Idx++)
+			{
+				char Character = Text[Idx];
+				if (!IsValidCharacter(Character))
+				{
+					if (Character >= 'A' && Character <= 'Z')
+					{
+						Text = Text.ToLowerInvariant();
+					}
+					else
+					{
+						throw new ArgumentException($"{Text} is not a valid string id");
+					}
+				}
+			}
+		}
+
+		/// <summary>
+		/// Checks whether this StringId is set
+		/// </summary>
+		public bool IsEmpty
+		{
+			get { return String.IsNullOrEmpty(Text); }
+		}
+
+		/// <summary>
+		/// Checks whether the given character is valid within a string id
+		/// </summary>
+		/// <param name="Character">The character to check</param>
+		/// <returns>True if the character is valid</returns>
+		static bool IsValidCharacter(char Character)
+		{
+			if (Character >= 'a' && Character <= 'z')
+			{
+				return true;
+			}
+			if (Character >= '0' && Character <= '9')
+			{
+				return true;
+			}
+			if (Character == '-' || Character == '_' || Character == '.')
+			{
+				return true;
+			}
+			return false;
+		}
+
+		/// <inheritdoc/>
+		public override bool Equals(object? Obj) => Obj is StringId Id && Equals(Id);
+
+		/// <inheritdoc/>
+		public override int GetHashCode() => Text.GetHashCode(StringComparison.Ordinal);
+
+		/// <inheritdoc/>
+		public bool Equals(StringId Other) => Text.Equals(Other.Text, StringComparison.Ordinal);
+
+		/// <inheritdoc/>
+		public override string ToString() => Text;
+
+		/// <summary>
+		/// Compares two string ids for equality
+		/// </summary>
+		/// <param name="Left">The first string id</param>
+		/// <param name="Right">Second string id</param>
+		/// <returns>True if the two string ids are equal</returns>
+		public static bool operator ==(StringId Left, StringId Right) => Left.Equals(Right);
+
+		/// <summary>
+		/// Compares two string ids for inequality
+		/// </summary>
+		/// <param name="Left">The first string id</param>
+		/// <param name="Right">Second string id</param>
+		/// <returns>True if the two string ids are not equal</returns>
+		public static bool operator !=(StringId Left, StringId Right) => !Left.Equals(Right);
+	}
+}
