Fix build error due to duplicated type in EpicGames.Horde.Client.

diff --git a/Engine/Source/Programs/Shared/EpicGames.Horde.Client/Storage/BlobStoreExtensions.cs b/Engine/Source/Programs/Shared/EpicGames.Horde.Client/Storage/BlobStoreExtensions.cs
--- a/Engine/Source/Programs/Shared/EpicGames.Horde.Client/Storage/BlobStoreExtensions.cs
+++ b/Engine/Source/Programs/Shared/EpicGames.Horde.Client/Storage/BlobStoreExtensions.cs
@@ -14,7 +14,7 @@ namespace EpicGames.Horde.Storage
 	/// <summary>
 	/// Exception thrown when a blob is not found in the object store
 	/// </summary>
-	public sealed class BlobNotFoundException : Exception
+	public sealed class LegacyBlobNotFoundException : Exception
 	{
 		/// <summary>
 		/// Hash of the missing blob
@@ -25,7 +25,7 @@ namespace EpicGames.Horde.Storage
 		/// Constructor
 		/// </summary>
 		/// <param name="Hash"></param>
-		public BlobNotFoundException(IoHash Hash)
+		public LegacyBlobNotFoundException(IoHash Hash)
 			: base("Blob {Hash} was not found in the object store")
 		{
 			this.Hash = Hash;
@@ -59,7 +59,7 @@ namespace EpicGames.Horde.Storage
 				}
 			}
 
-			throw new BlobNotFoundException(Hash);
+			throw new LegacyBlobNotFoundException(Hash);
 		}
 
 		/// <summary>
