Horde: Updating client API for storage.

* Now uses a single interface for the blob and ref stores (IStorageClient).
* Added an OAuth message handler which refreshes an access token as necessary.
* Added HTTP implementation of IStorageClient.

diff --git a/Engine/Source/Programs/Shared/EpicGames.Horde/Auth/OAuthMessageHandler.cs b/Engine/Source/Programs/Shared/EpicGames.Horde/Auth/OAuthMessageHandler.cs
new file mode 100644
--- /dev/null
+++ b/Engine/Source/Programs/Shared/EpicGames.Horde/Auth/OAuthMessageHandler.cs
@@ -0,0 +1,115 @@
+// Copyright Epic Games, Inc. All Rights Reserved.
+
+using EpicGames.Core;
+using EpicGames.Serialization;
+using System;
+using System.Collections.Generic;
+using System.Net;
+using System.Net.Http;
+using System.Text;
+using System.Text.Json;
+using System.Threading;
+using System.Threading.Tasks;
+
+namespace EpicGames.Horde.Compute
+{
+	/// <summary>
+	/// Exception thrown due to failed authorization
+	/// </summary>
+	public class AuthenticationException : Exception
+	{
+		public AuthenticationException(string Message, Exception? InnerException)
+			: base(Message, InnerException)
+		{
+		}
+	}
+
+	/// <summary>
+	/// Http message handler which adds an OAuth authorization header using a cached/periodically refreshed bearer token
+	/// </summary>
+	public class OAuthMessageHandler : DelegatingHandler
+	{
+		class ClientCredentialsResponse
+		{
+			public string? access_token { get; set; }
+			public string? token_type { get; set; }
+			public int? expires_in { get; set; }
+			public string? scope { get; set; }
+		}
+
+		readonly string AuthUrl;
+		readonly string GrantType;
+		readonly string ClientId;
+		readonly string ClientSecret;
+		readonly string Scope;
+
+		HttpClient Client;
+		string CachedAccessToken = String.Empty;
+		DateTime ExpiresAt = DateTime.MinValue;
+
+		public OAuthMessageHandler(HttpClient Client, string AuthUrl, string GrantType, string ClientId, string ClientSecret, string Scope)
+		{
+			this.Client = Client;
+			this.AuthUrl = AuthUrl;
+			this.GrantType = GrantType;
+			this.ClientId = ClientId;
+			this.ClientSecret = ClientSecret;
+			this.Scope = Scope;
+		}
+
+		/// <inheritdoc/>
+		protected override async Task<HttpResponseMessage> SendAsync(HttpRequestMessage request, CancellationToken cancellationToken)
+		{
+			if (DateTime.UtcNow > ExpiresAt)
+			{
+				await UpdateAccessTokenAsync();
+			}
+
+			request.Headers.Add("Authorization", $"Bearer {CachedAccessToken}");
+			return await base.SendAsync(request, cancellationToken);
+		}
+
+		/// <summary>
+		/// Updates the current access token
+		/// </summary>
+		/// <returns></returns>
+		async Task UpdateAccessTokenAsync()
+		{
+			KeyValuePair<string, string>[] Content = new KeyValuePair<string, string>[]
+			{
+				new KeyValuePair<string, string>("grant_type", GrantType),
+				new KeyValuePair<string, string>("client_id", ClientId),
+				new KeyValuePair<string, string>("client_secret", ClientSecret),
+				new KeyValuePair<string, string>("scope", Scope)
+			};
+
+			try
+			{
+				using HttpRequestMessage Message = new HttpRequestMessage(HttpMethod.Post, AuthUrl);
+				Message.Content = new FormUrlEncodedContent(Content);
+
+				HttpResponseMessage Response = await Client.SendAsync(Message);
+				if (!Response.IsSuccessStatusCode)
+				{
+					throw new AuthenticationException($"Authentication failed. Response: {Response.Content}", null);
+				}
+
+				byte[] ResponseData = await Response.Content.ReadAsByteArrayAsync();
+				ClientCredentialsResponse Result = JsonSerializer.Deserialize<ClientCredentialsResponse>(ResponseData)!;
+
+				string? accessToken = Result?.access_token;
+				if (string.IsNullOrEmpty(accessToken))
+				{
+					throw new AuthenticationException("The authentication token received by the server is null or empty. Body received was: " + Encoding.UTF8.GetString(ResponseData), null);
+				}
+				CachedAccessToken = accessToken;
+				// renew after half the renewal time
+				ExpiresAt = DateTime.UtcNow + TimeSpan.FromSeconds((Result?.expires_in ?? 3200) / 2.0);
+			}
+			catch (WebException Ex)
+			{
+				throw new AuthenticationException("Unable to authenticate.", Ex);
+			}
+		}
+	}
+}
diff --git a/Engine/Source/Programs/Shared/EpicGames.Horde/Common/Condition.cs b/Engine/Source/Programs/Shared/EpicGames.Horde/Common/Condition.cs
--- a/Engine/Source/Programs/Shared/EpicGames.Horde/Common/Condition.cs
+++ b/Engine/Source/Programs/Shared/EpicGames.Horde/Common/Condition.cs
@@ -640,7 +640,12 @@ namespace EpicGames.Horde.Common
 
 		public override Condition Read(ref Utf8JsonReader reader, Type typeToConvert, JsonSerializerOptions options)
 		{
-			return Condition.Parse(reader.GetString());
+			string? String = reader.GetString();
+			if(String == null)
+			{
+				throw new InvalidOperationException();
+			}
+			return Condition.Parse(String);
 		}
 
 		public override void Write(Utf8JsonWriter writer, Condition value, JsonSerializerOptions options)
diff --git a/Engine/Source/Programs/Shared/EpicGames.Horde/EpicGames.Horde.csproj b/Engine/Source/Programs/Shared/EpicGames.Horde/EpicGames.Horde.csproj
--- a/Engine/Source/Programs/Shared/EpicGames.Horde/EpicGames.Horde.csproj
+++ b/Engine/Source/Programs/Shared/EpicGames.Horde/EpicGames.Horde.csproj
@@ -36,6 +36,7 @@
     </PackageReference>
     <PackageReference Include="Microsoft.CodeAnalysis.CSharp.Workspaces" Version="3.10.0" />
     <PackageReference Include="Microsoft.Extensions.Logging" Version="5.0.0" />
+    <PackageReference Include="System.Net.Http.Json" Version="6.0.0" />
   </ItemGroup>
   <ItemGroup>
     <ProjectReference Include="..\EpicGames.Core\EpicGames.Core.csproj" />
diff --git a/Engine/Source/Programs/Shared/EpicGames.Horde/Storage/IBlobStore.cs b/Engine/Source/Programs/Shared/EpicGames.Horde/Storage/IBlobStore.cs
deleted file mode 100644
diff --git a/Engine/Source/Programs/Shared/EpicGames.Horde/Storage/IRefStore.cs b/Engine/Source/Programs/Shared/EpicGames.Horde/Storage/IRefStore.cs
deleted file mode 100644
diff --git a/Engine/Source/Programs/Shared/EpicGames.Horde/Storage/IStorageClient.cs b/Engine/Source/Programs/Shared/EpicGames.Horde/Storage/IStorageClient.cs
new file mode 100644
--- /dev/null
+++ b/Engine/Source/Programs/Shared/EpicGames.Horde/Storage/IStorageClient.cs
@@ -0,0 +1,289 @@
+using EpicGames.Core;
+using EpicGames.Serialization;
+using System;
+using System.Collections.Generic;
+using System.IO;
+using System.Text;
+using System.Threading.Tasks;
+
+namespace EpicGames.Horde.Storage
+{
+	/// <summary>
+	/// Base class for exceptions related to store
+	/// </summary>
+	public class StorageException : Exception
+	{
+		/// <summary>
+		/// Constructor
+		/// </summary>
+		public StorageException(string Message, Exception? InnerException)
+			: base(Message, InnerException)
+		{
+		}
+	}
+
+	/// <summary>
+	/// Base class for blob exceptions
+	/// </summary>
+	public class BlobException : StorageException
+	{
+		/// <summary>
+		/// Namespace containing the blob
+		/// </summary>
+		public NamespaceId NamespaceId { get; }
+
+		/// <summary>
+		/// Hash of the blob
+		/// </summary>
+		public IoHash Hash { get; }
+
+		/// <summary>
+		/// Constructor
+		/// </summary>
+		public BlobException(NamespaceId NamespaceId, IoHash Hash, string Message, Exception? InnerException = null)
+			: base(Message, InnerException)
+		{
+			this.NamespaceId = NamespaceId;
+			this.Hash = Hash;
+		}
+	}
+
+	/// <summary>
+	/// Exception thrown for missing blobs
+	/// </summary>
+	public sealed class BlobNotFoundException : BlobException
+	{
+		/// <summary>
+		/// Constructor
+		/// </summary>
+		public BlobNotFoundException(NamespaceId NamespaceId, IoHash Hash, Exception? InnerException = null)
+			: base(NamespaceId, Hash, $"Unable to find blob {Hash} in {NamespaceId}", InnerException)
+		{
+		}
+	}
+
+	/// <summary>
+	/// Base class for ref exceptions
+	/// </summary>
+	public class RefException : StorageException
+	{
+		/// <summary>
+		/// Namespace containing the ref
+		/// </summary>
+		public NamespaceId NamespaceId { get; }
+
+		/// <summary>
+		/// Bucket containing the ref
+		/// </summary>
+		public BucketId BucketId { get; }
+
+		/// <summary>
+		/// Identifier for the ref
+		/// </summary>
+		public IoHash RefId { get; }
+
+		/// <summary>
+		/// Constructor
+		/// </summary>
+		public RefException(NamespaceId NamespaceId, BucketId BucketId, IoHash RefId, string Message, Exception? InnerException = null)
+			: base(Message, InnerException)
+		{
+			this.NamespaceId = NamespaceId;
+			this.BucketId = BucketId;
+			this.RefId = RefId;
+		}
+	}
+
+	/// <summary>
+	/// Indicates that a named reference wasn't found
+	/// </summary>
+	public sealed class RefNotFoundException : RefException
+	{
+		/// <summary>
+		/// Constructor
+		/// </summary>
+		public RefNotFoundException(NamespaceId NamespaceId, BucketId BucketId, IoHash RefId, Exception? InnerException = null)
+			: base(NamespaceId, BucketId, RefId, $"Ref {NamespaceId}/{BucketId}/{RefId} not found", InnerException)
+		{
+		}
+	}
+
+	/// <summary>
+	/// Interface for an object reference
+	/// </summary>
+	public interface IRef
+	{
+		/// <summary>
+		/// Namespace identifier
+		/// </summary>
+		NamespaceId NamespaceId { get; }
+
+		/// <summary>
+		/// Bucket identifier
+		/// </summary>
+		BucketId BucketId { get; }
+
+		/// <summary>
+		/// Ref identifier
+		/// </summary>
+		IoHash RefId { get; }
+
+		/// <summary>
+		/// The value stored for this ref
+		/// </summary>
+		CbObject Value { get; }
+	}
+
+	/// <summary>
+	/// Interface for a collection of ref documents
+	/// </summary>
+	public interface IStorageClient
+	{
+		#region Blobs
+
+		/// <summary>
+		/// Opens a blob read stream
+		/// </summary>
+		/// <param name="NamespaceId">Namespace to operate on</param>
+		/// <param name="Hash">Hash of the blob</param>
+		/// <returns>Stream for the blob, or null if it does not exist</returns>
+		Task<Stream> ReadBlobAsync(NamespaceId NamespaceId, IoHash Hash);
+
+		/// <summary>
+		/// Writes a blob to storage
+		/// </summary>
+		/// <param name="NamespaceId">Namespace to operate on</param>
+		/// <param name="Hash">Hash of the blob</param>
+		/// <param name="Stream">The stream to write</param>
+		Task WriteBlobAsync(NamespaceId NamespaceId, IoHash Hash, Stream Stream);
+
+		#endregion
+
+		#region Refs
+
+		/// <summary>
+		/// Gets the given reference
+		/// </summary>
+		/// <param name="NamespaceId">Namespace identifier</param>
+		/// <param name="BucketId">Bucket identifier</param>
+		/// <param name="RefId">Name of the reference</param>
+		/// <returns>The reference data if the ref exists</returns>
+		Task<IRef> GetRefAsync(NamespaceId NamespaceId, BucketId BucketId, IoHash RefId);
+
+		/// <summary>
+		/// Gets the given reference
+		/// </summary>
+		/// <param name="NamespaceId">Namespace identifier</param>
+		/// <param name="BucketId">Bucket identifier</param>
+		/// <param name="RefId">Ref identifier</param>
+		/// <returns>The reference data if the ref exists</returns>
+		Task<bool> HasRefAsync(NamespaceId NamespaceId, BucketId BucketId, IoHash RefId);
+
+		/// <summary>
+		/// Determines which refs are missing
+		/// </summary>
+		/// <param name="NamespaceId">Namespace identifier</param>
+		/// <param name="BucketId">Bucket identifier</param>
+		/// <param name="RefIds">Names of the references</param>
+		/// <returns>List of missing references</returns>
+		Task<List<IoHash>> FindMissingRefsAsync(NamespaceId NamespaceId, BucketId BucketId, List<IoHash> RefIds);
+
+		/// <summary>
+		/// Sets the given reference
+		/// </summary>
+		/// <param name="NamespaceId">Namespace identifier</param>
+		/// <param name="BucketId">Bucket identifier</param>
+		/// <param name="RefId">Ref identifier</param>
+		/// <param name="Value">New value for the reference</param>
+		/// <returns>List of missing references</returns>
+		Task<List<IoHash>> SetRefAsync(NamespaceId NamespaceId, BucketId BucketId, IoHash RefId, CbObject Value);
+
+		/// <summary>
+		/// Attempts to finalize a reference, turning its references into hard references
+		/// </summary>
+		/// <param name="NamespaceId">Namespace identifier</param>
+		/// <param name="BucketId">Bucket identifier</param>
+		/// <param name="RefId">Ref identifier</param>
+		/// <param name="Hash">Hash of the referenced object</param>
+		/// <returns></returns>
+		Task<List<IoHash>> FinalizeRefAsync(NamespaceId NamespaceId, BucketId BucketId, IoHash RefId, IoHash Hash);
+
+		/// <summary>
+		/// Removes the given reference
+		/// </summary>
+		/// <param name="NamespaceId">Namespace identifier</param>
+		/// <param name="BucketId">Bucket identifier</param>
+		/// <param name="RefId">Ref identifier</param>
+		/// <returns>True if the ref was deleted, false if it did not exist</returns>
+		Task<bool> DeleteRefAsync(NamespaceId NamespaceId, BucketId BucketId, IoHash RefId);
+
+		#endregion
+	}
+
+	/// <summary>
+	/// Extension methods for <see cref="IStorageClient"/>
+	/// </summary>
+	public static class StorageClientExtensions
+	{
+		const int DefaultMaxInMemoryBlobLength = 50 * 1024 * 1024;
+
+		/// <summary>
+		/// Gets a blob as a byte array
+		/// </summary>
+		/// <param name="StorageClient"></param>
+		/// <param name="NamespaceId"></param>
+		/// <param name="Hash"></param>
+		/// <returns></returns>
+		public static async Task<byte[]> ReadBlobToMemoryAsync(this IStorageClient StorageClient, NamespaceId NamespaceId, IoHash Hash, int MaxInMemoryBlobLength = DefaultMaxInMemoryBlobLength)
+		{
+			using Stream Stream = await StorageClient.ReadBlobAsync(NamespaceId, Hash);
+
+			long Length = Stream.Length;
+			if (Length > MaxInMemoryBlobLength)
+			{
+				throw new BlobException(NamespaceId, Hash, $"Blob {Hash} is too large ({Length} > {MaxInMemoryBlobLength})");
+			}
+
+			byte[] Buffer = new byte[Length];
+			for (int Offset = 0; Offset < Length;)
+			{
+				int Count = await Stream.ReadAsync(Buffer, Offset, (int)Length - Offset);
+				if (Count == 0)
+				{
+					throw new BlobException(NamespaceId, Hash, $"Unexpected end of stream reading blob {Hash}");
+				}
+				Offset += Count;
+			}
+
+			return Buffer;
+		}
+
+		/// <summary>
+		/// Gets a blob as a byte array
+		/// </summary>
+		/// <param name="StorageClient"></param>
+		/// <param name="NamespaceId"></param>
+		/// <param name="Data">The data to be written</param>
+		/// <returns></returns>
+		public static async Task<IoHash> WriteBlobFromMemoryAsync(this IStorageClient StorageClient, NamespaceId NamespaceId, ReadOnlyMemory<byte> Data)
+		{
+			IoHash Hash = IoHash.Compute(Data.Span);
+			await WriteBlobFromMemoryAsync(StorageClient, NamespaceId, Hash, Data);
+			return Hash;
+		}
+
+		/// <summary>
+		/// Gets a blob as a byte array
+		/// </summary>
+		/// <param name="StorageClient"></param>
+		/// <param name="NamespaceId"></param>
+		/// <param name="Hash">Hash of the data</param>
+		/// <param name="Data">The data to be written</param>
+		/// <returns></returns>
+		public static async Task WriteBlobFromMemoryAsync(this IStorageClient StorageClient, NamespaceId NamespaceId, IoHash Hash, ReadOnlyMemory<byte> Data)
+		{
+			using ReadOnlyMemoryStream Stream = new ReadOnlyMemoryStream(Data);
+			await StorageClient.WriteBlobAsync(NamespaceId, Hash, Stream);
+		}
+	}
+}
diff --git a/Engine/Source/Programs/Shared/EpicGames.Horde/Storage/Impl/HttpStorageClient.cs b/Engine/Source/Programs/Shared/EpicGames.Horde/Storage/Impl/HttpStorageClient.cs
new file mode 100644
--- /dev/null
+++ b/Engine/Source/Programs/Shared/EpicGames.Horde/Storage/Impl/HttpStorageClient.cs
@@ -0,0 +1,250 @@
+// Copyright Epic Games, Inc. All Rights Reserved.
+
+using EpicGames.Core;
+using EpicGames.Serialization;
+using Google.Protobuf;
+using System;
+using System.Collections.Generic;
+using System.IO;
+using System.Linq;
+using System.Net.Http;
+using System.Net.Http.Headers;
+using System.Net.Http.Json;
+using System.Text;
+using System.Text.Json;
+using System.Threading.Tasks;
+
+namespace EpicGames.Horde.Storage.Impl
+{
+	/// <summary>
+	/// Response from adding a new object
+	/// </summary>
+	public class RefPutResponse
+	{
+		/// <summary>
+		/// List of missing hashes
+		/// </summary>
+		public List<IoHash> Needs { get; set; } = new List<IoHash>();
+	}
+
+	/// <summary>
+	/// Response from posting to the /exists endpoint
+	/// </summary>
+	public class RefExistsResponse
+	{
+		/// <summary>
+		/// List of hashes that the blob store needs.
+		/// </summary>
+		public List<IoHash> Needs { get; set; } = new List<IoHash>();
+	}
+
+	/// <summary>
+	/// Implementation of <see cref="IStorageClient"/> using REST requests.
+	/// </summary>
+	public sealed class HttpStorageClient : IStorageClient
+	{
+		class RefImpl : IRef
+		{
+			public NamespaceId NamespaceId { get; set; }
+
+			public BucketId BucketId { get; set; }
+
+			public IoHash RefId { get; set; }
+
+			public CbObject Value { get; }
+
+			public RefImpl(NamespaceId NamespaceId, BucketId BucketId, IoHash RefId, CbObject Value)
+			{
+				this.NamespaceId = NamespaceId;
+				this.BucketId = BucketId;
+				this.RefId = RefId;
+				this.Value = Value;
+			}
+		}
+
+		const string HashHeaderName = "X-Jupiter-IoHash";
+		const string LastAccessHeaderName = "X-Jupiter-LastAccess";
+
+		const string CompactBinaryMimeType = "application/x-ue-cb";
+
+		HttpClient HttpClient;
+
+		/// <summary>
+		/// Constructor
+		/// </summary>
+		/// <param name="HttpClient">Http client for the blob store. Should be pre-configured with a base address and appropriate authentication headers.</param>
+		public HttpStorageClient(HttpClient HttpClient)
+		{
+			this.HttpClient = HttpClient;
+		}
+
+		#region Blobs
+
+		/// <inheritdoc/>
+		public async Task<Stream> ReadBlobAsync(NamespaceId NamespaceId, IoHash Hash)
+		{
+			using HttpResponseMessage Response = await HttpClient.GetAsync($"/api/v1/blobs/{NamespaceId}/{Hash}");
+			return await Response.Content.ReadAsStreamAsync();
+		}
+
+		/// <inheritdoc/>
+		public async Task WriteBlobAsync(NamespaceId NamespaceId, IoHash Hash, Stream Stream)
+		{
+			StreamContent Content = new StreamContent(Stream);
+			Content.Headers.ContentType = new MediaTypeHeaderValue("application/octet-stream");
+
+			HttpResponseMessage Response = await HttpClient.PutAsync($"api/v1/blobs/{NamespaceId}/{Hash}", Content);
+			Response.EnsureSuccessStatusCode();
+		}
+
+		#endregion
+
+		#region Refs
+
+		/// <inheritdoc/>
+		public async Task<IRef> GetRefAsync(NamespaceId NamespaceId, BucketId BucketId, IoHash RefId)
+		{
+			using HttpRequestMessage Request = new HttpRequestMessage(HttpMethod.Get, $"api/v1/refs/{NamespaceId}/{BucketId}/{RefId}");
+			Request.Headers.Accept.Clear();
+			Request.Headers.Accept.Add(new MediaTypeWithQualityHeaderValue(CompactBinaryMimeType));
+
+			using HttpResponseMessage Response = await HttpClient.SendAsync(Request);
+			if (!Response.IsSuccessStatusCode)
+			{
+				if (Response.StatusCode == System.Net.HttpStatusCode.NotFound)
+				{
+					throw new RefNotFoundException(NamespaceId, BucketId, RefId);
+				}
+				else
+				{
+					throw new RefException(NamespaceId, BucketId, RefId, await GetMessageFromResponse(Response));
+				}
+			}
+
+			byte[] Data = await Response.Content.ReadAsByteArrayAsync();
+			return new RefImpl(NamespaceId, BucketId, RefId, new CbObject(Data));
+		}
+
+		/// <inheritdoc/>
+		public async Task<List<IoHash>> SetRefAsync(NamespaceId NamespaceId, BucketId BucketId, IoHash RefId, CbObject Value)
+		{
+			using ReadOnlyMemoryStream Stream = new ReadOnlyMemoryStream(Value.GetView());
+
+			StreamContent Content = new StreamContent(Stream);
+			Content.Headers.ContentType = new MediaTypeHeaderValue(CompactBinaryMimeType);
+			Content.Headers.Add(HashHeaderName, IoHash.Compute(Value.GetView().Span).ToString());
+
+			using HttpRequestMessage Request = new HttpRequestMessage(HttpMethod.Put, $"api/v1/refs/{NamespaceId}/{BucketId}/{RefId}");
+			Request.Headers.Accept.Clear();
+			Request.Headers.Accept.Add(new MediaTypeWithQualityHeaderValue("application/json"));
+			Request.Content = Content;
+
+			using HttpResponseMessage Response = await HttpClient.SendAsync(Request);
+			if (!Response.IsSuccessStatusCode)
+			{
+				throw new RefException(NamespaceId, BucketId, RefId, await GetMessageFromResponse(Response));
+			}
+
+			RefPutResponse? ResponseBody = await Response.Content.ReadFromJsonAsync<RefPutResponse>();
+			if (ResponseBody == null)
+			{
+				throw new RefException(NamespaceId, BucketId, RefId, "Unable to parse response body");
+			}
+
+			return ResponseBody.Needs;
+		}
+
+		/// <inheritdoc/>
+		public async Task<List<IoHash>> FinalizeRefAsync(NamespaceId NamespaceId, BucketId BucketId, IoHash RefId, IoHash Hash)
+		{
+			using HttpRequestMessage Request = new HttpRequestMessage(HttpMethod.Post, $"api/v1/refs/{NamespaceId}/{BucketId}/{RefId}/finalize/{Hash}");
+			Request.Headers.Accept.Clear();
+			Request.Headers.Accept.Add(new MediaTypeWithQualityHeaderValue("application/json"));
+
+			using HttpResponseMessage Response = await HttpClient.SendAsync(Request);
+			if (!Response.IsSuccessStatusCode)
+			{
+				throw new RefException(NamespaceId, BucketId, RefId, await GetMessageFromResponse(Response));
+			}
+
+			RefPutResponse? ResponseBody = await Response.Content.ReadFromJsonAsync<RefPutResponse>();
+			if (ResponseBody == null)
+			{
+				throw new RefException(NamespaceId, BucketId, RefId, "Unable to parse response body");
+			}
+
+			return ResponseBody.Needs;
+		}
+
+		/// <inheritdoc/>
+		public async Task<bool> DeleteRefAsync(NamespaceId NamespaceId, BucketId BucketId, IoHash RefId)
+		{
+			using HttpResponseMessage Response = await HttpClient.DeleteAsync($"api/v1/refs/{NamespaceId}/{BucketId}/{RefId}");
+			return Response.IsSuccessStatusCode;
+		}
+
+		/// <inheritdoc/>
+		public async Task<bool> HasRefAsync(NamespaceId NamespaceId, BucketId BucketId, IoHash RefId)
+		{
+			using (HttpRequestMessage Request = new HttpRequestMessage(HttpMethod.Head, $"api/v1/refs/{NamespaceId}/{BucketId}/{RefId}"))
+			{
+				using HttpResponseMessage Response = await HttpClient.SendAsync(Request);
+				if (Response.IsSuccessStatusCode)
+				{
+					return true;
+				}
+				else if (Response.StatusCode == System.Net.HttpStatusCode.NotFound)
+				{
+					return false;
+				}
+				else
+				{
+					throw new RefException(NamespaceId, BucketId, RefId, $"Unexpected response for {nameof(HasRefAsync)} call: {Response.StatusCode}");
+				}
+			}
+		}
+
+		/// <inheritdoc/>
+		public async Task<List<IoHash>> FindMissingRefsAsync(NamespaceId NamespaceId, BucketId BucketId, List<IoHash> RefIds)
+		{
+			string RefList = String.Join("&", RefIds.Select(x => $"id={x}"));
+
+			using HttpResponseMessage Response = await HttpClient.PostAsync($"/api/v1/refs/{NamespaceId}/{BucketId}/exists?{RefList}", null);
+			if (!Response.IsSuccessStatusCode)
+			{
+				throw new RefException(NamespaceId, BucketId, IoHash.Zero, await GetMessageFromResponse(Response));
+			}
+
+			RefExistsResponse? ResponseContent = await Response.Content.ReadFromJsonAsync<RefExistsResponse>();
+			if (ResponseContent == null)
+			{
+				throw new RefException(NamespaceId, BucketId, IoHash.Zero, "Unable to parse response body");
+			}
+
+			return ResponseContent.Needs;
+		}
+
+		#endregion
+
+		class ProblemDetails
+		{
+			public string? Title { get; set; }
+		}
+
+		static async Task<string> GetMessageFromResponse(HttpResponseMessage Response)
+		{
+			StringBuilder Message = new StringBuilder($"HTTP {Response.StatusCode}");
+			try
+			{
+				byte[] Data = await Response.Content.ReadAsByteArrayAsync();
+				ProblemDetails? Details = JsonSerializer.Deserialize<ProblemDetails>(Data);
+				Message.Append($": {Details?.Title ?? "No description available"}");
+			}
+			catch
+			{
+				Message.Append(" (Unable to parse response data)");
+			}
+			return Message.ToString();
+		}
+	}
+}
