Horde: Remove dependency on System.Net.Http.Json; causing analyzer errors during build.

diff --git a/Engine/Source/Programs/Horde/Horde.Agent/Execution/BuildGraphExecutor.cs b/Engine/Source/Programs/Horde/Horde.Agent/Execution/BuildGraphExecutor.cs
--- a/Engine/Source/Programs/Horde/Horde.Agent/Execution/BuildGraphExecutor.cs
+++ b/Engine/Source/Programs/Horde/Horde.Agent/Execution/BuildGraphExecutor.cs
@@ -911,7 +911,7 @@ namespace HordeAgent.Execution
 					FileReference TraceFile = FileReference.Combine(TelemetryDir, "Trace.json");
 					using (FileStream Stream = FileReference.Open(TraceFile, FileMode.Create))
 					{
-						JsonSerializerOptions Options = new JsonSerializerOptions { DefaultIgnoreCondition = JsonIgnoreCondition.WhenWritingNull };
+						JsonSerializerOptions Options = new JsonSerializerOptions { IgnoreNullValues = true };
 						await JsonSerializer.SerializeAsync(Stream, RootSpan, Options);
 					}
 					await ArtifactUploader.UploadAsync(RpcConnection, JobId, BatchId, Step.StepId, "Trace.json", TraceFile, Logger, CancellationToken.None);
diff --git a/Engine/Source/Programs/Shared/EpicGames.Horde/EpicGames.Horde.csproj b/Engine/Source/Programs/Shared/EpicGames.Horde/EpicGames.Horde.csproj
--- a/Engine/Source/Programs/Shared/EpicGames.Horde/EpicGames.Horde.csproj
+++ b/Engine/Source/Programs/Shared/EpicGames.Horde/EpicGames.Horde.csproj
@@ -36,7 +36,6 @@
     </PackageReference>
     <PackageReference Include="Microsoft.CodeAnalysis.CSharp.Workspaces" Version="3.10.0" />
     <PackageReference Include="Microsoft.Extensions.Logging" Version="5.0.0" />
-    <PackageReference Include="System.Net.Http.Json" Version="6.0.0" />
   </ItemGroup>
   <ItemGroup>
     <ProjectReference Include="..\EpicGames.Core\EpicGames.Core.csproj" />
diff --git a/Engine/Source/Programs/Shared/EpicGames.Horde/Storage/Impl/HttpStorageClient.cs b/Engine/Source/Programs/Shared/EpicGames.Horde/Storage/Impl/HttpStorageClient.cs
--- a/Engine/Source/Programs/Shared/EpicGames.Horde/Storage/Impl/HttpStorageClient.cs
+++ b/Engine/Source/Programs/Shared/EpicGames.Horde/Storage/Impl/HttpStorageClient.cs
@@ -9,7 +9,6 @@ using System.IO;
 using System.Linq;
 using System.Net.Http;
 using System.Net.Http.Headers;
-using System.Net.Http.Json;
 using System.Text;
 using System.Text.Json;
 using System.Threading.Tasks;
@@ -145,7 +144,7 @@ namespace EpicGames.Horde.Storage.Impl
 				throw new RefException(NamespaceId, BucketId, RefId, await GetMessageFromResponse(Response));
 			}
 
-			RefPutResponse? ResponseBody = await Response.Content.ReadFromJsonAsync<RefPutResponse>();
+			RefPutResponse? ResponseBody = await ReadJsonResponse<RefPutResponse>(Response.Content);
 			if (ResponseBody == null)
 			{
 				throw new RefException(NamespaceId, BucketId, RefId, "Unable to parse response body");
@@ -167,7 +166,7 @@ namespace EpicGames.Horde.Storage.Impl
 				throw new RefException(NamespaceId, BucketId, RefId, await GetMessageFromResponse(Response));
 			}
 
-			RefPutResponse? ResponseBody = await Response.Content.ReadFromJsonAsync<RefPutResponse>();
+			RefPutResponse? ResponseBody = await ReadJsonResponse<RefPutResponse>(Response.Content);
 			if (ResponseBody == null)
 			{
 				throw new RefException(NamespaceId, BucketId, RefId, "Unable to parse response body");
@@ -215,7 +214,7 @@ namespace EpicGames.Horde.Storage.Impl
 				throw new RefException(NamespaceId, BucketId, IoHash.Zero, await GetMessageFromResponse(Response));
 			}
 
-			RefExistsResponse? ResponseContent = await Response.Content.ReadFromJsonAsync<RefExistsResponse>();
+			RefExistsResponse? ResponseContent = await ReadJsonResponse<RefExistsResponse>(Response.Content);
 			if (ResponseContent == null)
 			{
 				throw new RefException(NamespaceId, BucketId, IoHash.Zero, "Unable to parse response body");
@@ -246,5 +245,11 @@ namespace EpicGames.Horde.Storage.Impl
 			}
 			return Message.ToString();
 		}
+
+		static async Task<T> ReadJsonResponse<T>(HttpContent Content)
+		{
+			byte[] Data = await Content.ReadAsByteArrayAsync();
+			return JsonSerializer.Deserialize<T>(Data, new JsonSerializerOptions { PropertyNameCaseInsensitive = true });
+		}
 	}
 }
