Horde: Mising file.

diff --git a/Engine/Source/Programs/Shared/EpicGames.Horde/Auth/NullAuthProvider.cs b/Engine/Source/Programs/Shared/EpicGames.Horde/Auth/NullAuthProvider.cs
new file mode 100644
--- /dev/null
+++ b/Engine/Source/Programs/Shared/EpicGames.Horde/Auth/NullAuthProvider.cs
@@ -0,0 +1,20 @@
+// Copyright Epic Games, Inc. All Rights Reserved.
+
+using System;
+using System.Collections.Generic;
+using System.Net.Http;
+using System.Text;
+using System.Threading;
+using System.Threading.Tasks;
+
+namespace EpicGames.Horde.Auth
+{
+	/// <summary>
+	/// Empty implementation of <see cref="IAuthProvider"/>
+	/// </summary>
+	public class NullAuthProvider<T> : IAuthProvider<T>
+	{
+		/// <inheritdoc/>
+		public Task AddAuthorizationAsync(HttpRequestMessage Request, CancellationToken CancellationToken) => Task.CompletedTask;
+	}
+}
