Horde: Fix test failures due to DI service lifetimes.

diff --git a/Engine/Source/Programs/Shared/EpicGames.Horde/Storage/StorageExtensions.cs b/Engine/Source/Programs/Shared/EpicGames.Horde/Storage/StorageExtensions.cs
--- a/Engine/Source/Programs/Shared/EpicGames.Horde/Storage/StorageExtensions.cs
+++ b/Engine/Source/Programs/Shared/EpicGames.Horde/Storage/StorageExtensions.cs
@@ -57,12 +57,12 @@ namespace EpicGames.Horde.Storage
 		{
 			Services.AddOptions<StorageSettings>();
 
-			Services.AddTransient<NullAuthProvider<HttpStorageClient>>();
+			Services.AddSingleton<NullAuthProvider<HttpStorageClient>>();
 
-			Services.AddScoped<OAuth2AuthProviderFactory<HttpStorageClient>>();
+			Services.AddSingleton<OAuth2AuthProviderFactory<HttpStorageClient>>();
 			Services.AddHttpClient<OAuth2AuthProviderFactory<HttpStorageClient>>();
 
-			Services.AddScoped<IAuthProvider<HttpStorageClient>>(ServiceProvider => CreateAuthProvider(ServiceProvider));
+			Services.AddSingleton<IAuthProvider<HttpStorageClient>>(ServiceProvider => CreateAuthProvider(ServiceProvider));
 
 			Services.AddTransient<IStorageClient, HttpStorageClient>();
 			Services.AddHttpClient<IStorageClient, HttpStorageClient>();
