Horde: Add missing JsonConverter for BucketId.

diff --git a/Engine/Source/Programs/Horde/Horde.Build/Compute/Impl/ComputeService.cs b/Engine/Source/Programs/Horde/Horde.Build/Compute/Impl/ComputeService.cs
--- a/Engine/Source/Programs/Horde/Horde.Build/Compute/Impl/ComputeService.cs
+++ b/Engine/Source/Programs/Horde/Horde.Build/Compute/Impl/ComputeService.cs
@@ -265,8 +265,6 @@ namespace HordeServer.Compute.Impl
 			{
 				throw new KeyNotFoundException();
 			}
-
-			Logger.LogInformation("Found compute cluster: {ClusterId}, Reqs {ReqId} {ReqId2} {Json}", ClusterId, Config.RequestBucketId, new ClusterInfo(Config).RequestBucketId, JsonSerializer.Serialize(new ClusterInfo(Config)));
 			return new ClusterInfo(Config);
 		}
 
diff --git a/Engine/Source/Programs/Shared/EpicGames.Horde/Storage/BucketId.cs b/Engine/Source/Programs/Shared/EpicGames.Horde/Storage/BucketId.cs
--- a/Engine/Source/Programs/Shared/EpicGames.Horde/Storage/BucketId.cs
+++ b/Engine/Source/Programs/Shared/EpicGames.Horde/Storage/BucketId.cs
@@ -2,14 +2,20 @@
 
 using System;
 using System.Collections.Generic;
+using System.ComponentModel;
 using System.Diagnostics.CodeAnalysis;
+using System.Globalization;
 using System.Text;
+using System.Text.Json;
+using System.Text.Json.Serialization;
 
 namespace EpicGames.Horde.Storage
 {
 	/// <summary>
 	/// Identifier for a storage bucket
 	/// </summary>
+	[JsonConverter(typeof(BucketIdJsonConverter))]
+	[TypeConverter(typeof(BucketIdTypeConverter))]
 	public struct BucketId : IEquatable<BucketId>
 	{
 		/// <summary>
@@ -44,4 +50,28 @@ namespace EpicGames.Horde.Storage
 		/// <inheritdoc cref="StringId.op_Inequality"/>
 		public static bool operator !=(BucketId Left, BucketId Right) => Left.Inner != Right.Inner;
 	}
+
+	/// <summary>
+	/// Type converter for BucketId to and from JSON
+	/// </summary>
+	sealed class BucketIdJsonConverter : JsonConverter<BucketId>
+	{
+		/// <inheritdoc/>
+		public override BucketId Read(ref Utf8JsonReader Reader, Type TypeToConvert, JsonSerializerOptions Options) => new BucketId(Reader.GetString() ?? String.Empty);
+
+		/// <inheritdoc/>
+		public override void Write(Utf8JsonWriter Writer, BucketId Value, JsonSerializerOptions Options) => Writer.WriteStringValue(Value.ToString());
+	}
+
+	/// <summary>
+	/// Type converter from strings to BucketId objects
+	/// </summary>
+	sealed class BucketIdTypeConverter : TypeConverter
+	{
+		/// <inheritdoc/>
+		public override bool CanConvertFrom(ITypeDescriptorContext Context, Type SourceType) => SourceType == typeof(string);
+
+		/// <inheritdoc/>
+		public override object ConvertFrom(ITypeDescriptorContext Context, CultureInfo Culture, object Value) => new BucketId((string)Value);
+	}
 }
