Horde: Support for running standalone compute requests using HordeAgent.

- JSON job descriptor contains storage/compute server settings.
- Supports auth using OAuth2 or bearer token.
- Added static extension methods to configure Horde Storage and Horde Compute via dependency injection.

diff --git a/Engine/Source/Programs/Horde/Horde.Agent/AgentSettings.cs b/Engine/Source/Programs/Horde/Horde.Agent/AgentSettings.cs
--- a/Engine/Source/Programs/Horde/Horde.Agent/AgentSettings.cs
+++ b/Engine/Source/Programs/Horde/Horde.Agent/AgentSettings.cs
@@ -1,6 +1,7 @@
 // Copyright Epic Games, Inc. All Rights Reserved.
 
 using EpicGames.Core;
+using EpicGames.Horde.Compute;
 using EpicGames.Horde.Storage;
 using HordeCommon.Rpc;
 using Microsoft.Extensions.Configuration;
@@ -95,7 +96,7 @@ namespace HordeAgent
 		/// <summary>
 		/// Storage settings for using this server
 		/// </summary>
-		public StorageSettings Storage { get; set; } = new StorageSettings();
+		public StorageOptions Storage { get; set; } = new StorageOptions();
 
 		/// <summary>
 		/// Checks whether the given certificate thumbprint should be trusted
diff --git a/Engine/Source/Programs/Horde/Horde.Agent/Commands/Execution/ExecuteCommandV2.cs b/Engine/Source/Programs/Horde/Horde.Agent/Commands/Execution/ExecuteCommandV2.cs
--- a/Engine/Source/Programs/Horde/Horde.Agent/Commands/Execution/ExecuteCommandV2.cs
+++ b/Engine/Source/Programs/Horde/Horde.Agent/Commands/Execution/ExecuteCommandV2.cs
@@ -45,6 +45,8 @@ namespace HordeAgent.Commands
 
 		class JsonComputeTask
 		{
+			public ComputeOptions ComputeServer { get; set; } = new ComputeOptions();
+			public StorageOptions StorageServer { get; set; } = new StorageOptions();
 			public string Executable { get; set; } = String.Empty;
 			public List<string> Arguments { get; set; } = new List<string>();
 			public Dictionary<string, string> EnvVars { get; set; } = new Dictionary<string, string>();
@@ -153,10 +155,9 @@ namespace HordeAgent.Commands
 
 			IConfiguration Configuration = new ConfigurationBuilder()
 				.AddEnvironmentVariables()
-				.AddJsonFile("appsettings.json")
+				.AddJsonFile(Input.FullName)
 				.Build();
 
-			IConfigurationSection ConfigSection = Configuration.GetSection(AgentSettings.SectionName);
 			IHostBuilder HostBuilder = Host.CreateDefaultBuilder()
 				.ConfigureLogging(Builder =>
 				{
@@ -164,14 +165,13 @@ namespace HordeAgent.Commands
 				})
 				.ConfigureServices(Services =>
 				{
-					Services.AddOptions<AgentSettings>().Configure(Options => ConfigSection.Bind(Options)).ValidateDataAnnotations();
 					Services.AddLogging();
 
-					IConfigurationSection StorageSettings = ConfigSection.GetCurrentServerProfile().GetSection(nameof(ServerProfile.Storage));
-					Services.AddHordeStorage(Settings => StorageSettings.Bind(Settings));
+					IConfigurationSection ComputeSettings = Configuration.GetSection(nameof(JsonComputeTask.ComputeServer));
+					Services.AddHordeCompute(Settings => ComputeSettings.Bind(Settings));
 
-					Services.AddSingleton<GrpcService>();
-					Services.AddSingleton<IComputeClient, HttpComputeClient>(SP => new HttpComputeClient(new HttpClient() { BaseAddress = new Uri(Server) }));
+					IConfigurationSection StorageSettings = Configuration.GetSection(nameof(JsonComputeTask.StorageServer));
+					Services.AddHordeStorage(Settings => StorageSettings.Bind(Settings));
 				});
 
 			using (IHost Host = HostBuilder.Build())
@@ -179,9 +179,6 @@ namespace HordeAgent.Commands
 				IStorageClient StorageClient = Host.Services.GetRequiredService<IStorageClient>();
 				IComputeClient ComputeClient = Host.Services.GetRequiredService<IComputeClient>();
 
-				byte[] Json = await FileReference.ReadAllBytesAsync(Input);
-				JsonComputeTask JsonComputeTask = JsonSerializer.Deserialize<JsonComputeTask>(Json, new JsonSerializerOptions { PropertyNameCaseInsensitive = true })!;
-
 				ByteString SaltBytes = ByteString.Empty;
 				if (Salt != null)
 				{
@@ -200,6 +197,9 @@ namespace HordeAgent.Commands
 				Dictionary<IoHash, byte[]> Blobs = new Dictionary<IoHash, byte[]>();
 				(_, IoHash SandboxHash) = CreateSandbox(InputDir.ToDirectoryInfo(), Blobs);
 
+				JsonComputeTask JsonComputeTask = new JsonComputeTask();
+				Configuration.Bind(JsonComputeTask);
+
 				JsonRequirements JsonRequirements = JsonComputeTask.Requirements;
 				Requirements Requirements = new Requirements(JsonRequirements.Condition ?? String.Empty);
 				Requirements.Resources = JsonRequirements.Resources.ToDictionary(x => x.Key, x => x.Value);
diff --git a/Engine/Source/Programs/Horde/Horde.Build/ServerSettings.cs b/Engine/Source/Programs/Horde/Horde.Build/ServerSettings.cs
--- a/Engine/Source/Programs/Horde/Horde.Build/ServerSettings.cs
+++ b/Engine/Source/Programs/Horde/Horde.Build/ServerSettings.cs
@@ -565,7 +565,7 @@ namespace HordeServer
 		/// <summary>
 		/// Settings for the storage service
 		/// </summary>
-		public StorageSettings? Storage { get; set; }
+		public StorageOptions? Storage { get; set; }
 
 		/// <summary>
 		/// Lazily computed timezone value
diff --git a/Engine/Source/Programs/Shared/EpicGames.Horde/Auth/AuthExtensions.cs b/Engine/Source/Programs/Shared/EpicGames.Horde/Auth/AuthExtensions.cs
new file mode 100644
--- /dev/null
+++ b/Engine/Source/Programs/Shared/EpicGames.Horde/Auth/AuthExtensions.cs
@@ -0,0 +1,88 @@
+// Copyright Epic Games, Inc. All Rights Reserved.
+
+using Microsoft.Extensions.DependencyInjection;
+using System;
+using System.Collections.Generic;
+using System.ComponentModel.DataAnnotations;
+using System.Net.Http;
+using System.Text;
+
+namespace EpicGames.Horde.Auth
+{
+	/// <summary>
+	/// Base class for configuring HTTP service clients
+	/// </summary>
+	public class HttpServiceClientOptions : IOAuthOptions, ITokenAuthOptions
+	{
+		/// <summary>
+		/// Base address for http requests
+		/// </summary>
+		[Required]
+		public string Url { get; set; } = String.Empty;
+
+		#region OAuth2
+
+		/// <inheritdoc/>
+		public string AuthUrl { get; set; } = String.Empty;
+
+		/// <inheritdoc/>
+		public string GrantType { get; set; } = String.Empty;
+
+		/// <inheritdoc/>
+		public string ClientId { get; set; } = String.Empty;
+
+		/// <inheritdoc/>
+		public string ClientSecret { get; set; } = String.Empty;
+
+		/// <inheritdoc/>
+		public string Scope { get; set; } = String.Empty;
+
+		#endregion
+
+		#region Bearer token
+
+		/// <inheritdoc/>
+		public string Token { get; set; } = String.Empty;
+
+		#endregion
+	}
+
+	internal static class AuthExtensions
+	{
+		public static void AddHttpClientWithAuth<TClient, TImplementation>(this IServiceCollection Services, Func<IServiceProvider, HttpServiceClientOptions> GetOptions) 
+			where TClient : class 
+			where TImplementation : class, TClient
+		{
+			Services.AddScoped<OAuthHandlerFactory>();
+			Services.AddHttpClient<OAuthHandlerFactory>();
+			Services.AddScoped<OAuthHandler<TImplementation>>(ServiceProvider => ServiceProvider.GetRequiredService<OAuthHandlerFactory>().Create<TImplementation>(GetOptions(ServiceProvider)));
+
+			Services.AddHttpClient<TClient, TImplementation>((ServiceProvider, Client) =>
+				{
+					HttpServiceClientOptions Options = GetOptions(ServiceProvider);
+					Client.BaseAddress = new Uri(Options.Url);
+				})
+				.ConfigurePrimaryHttpMessageHandler(ServiceProvider =>
+				{
+					HttpServiceClientOptions Options = GetOptions(ServiceProvider);
+					return CreateMessageHandler<TImplementation>(ServiceProvider, Options);
+				});
+		}
+
+		static HttpMessageHandler CreateMessageHandler<TImplementation>(IServiceProvider ServiceProvider, HttpServiceClientOptions Options)
+		{
+			if (!String.IsNullOrEmpty(Options.AuthUrl))
+			{
+				return ServiceProvider.GetRequiredService<OAuthHandler<TImplementation>>();
+			}
+			else if (!String.IsNullOrEmpty(Options.Token))
+			{
+				return new TokenHandler<TImplementation>(Options);
+			}
+			else
+			{
+				return new HttpClientHandler();
+			}
+		}
+	}
+}
diff --git a/Engine/Source/Programs/Shared/EpicGames.Horde/Auth/IAuthProvider.cs b/Engine/Source/Programs/Shared/EpicGames.Horde/Auth/IAuthProvider.cs
deleted file mode 100644
diff --git a/Engine/Source/Programs/Shared/EpicGames.Horde/Auth/NullAuthProvider.cs b/Engine/Source/Programs/Shared/EpicGames.Horde/Auth/NullAuthProvider.cs
deleted file mode 100644
diff --git a/Engine/Source/Programs/Shared/EpicGames.Horde/Auth/OAuth2AuthProvider.cs b/Engine/Source/Programs/Shared/EpicGames.Horde/Auth/OAuth2AuthProvider.cs
deleted file mode 100644
diff --git a/Engine/Source/Programs/Shared/EpicGames.Horde/Auth/OAuthHandler.cs b/Engine/Source/Programs/Shared/EpicGames.Horde/Auth/OAuthHandler.cs
new file mode 100644
--- /dev/null
+++ b/Engine/Source/Programs/Shared/EpicGames.Horde/Auth/OAuthHandler.cs
@@ -0,0 +1,163 @@
+// Copyright Epic Games, Inc. All Rights Reserved.
+
+using EpicGames.Core;
+using EpicGames.Serialization;
+using Microsoft.Extensions.Options;
+using System;
+using System.Collections.Generic;
+using System.ComponentModel.DataAnnotations;
+using System.Net;
+using System.Net.Http;
+using System.Text;
+using System.Text.Json;
+using System.Threading;
+using System.Threading.Tasks;
+
+namespace EpicGames.Horde.Auth
+{
+	/// <summary>
+	/// Exception thrown due to failed authorization
+	/// </summary>
+	public class AuthenticationException : Exception
+	{
+		public AuthenticationException(string Message, Exception? InnerException)
+			: base(Message, InnerException)
+		{
+		}
+	}
+
+	/// <summary>
+	/// Options for authenticating particular requests
+	/// </summary>
+	public interface IOAuthOptions
+	{
+		/// <summary>
+		/// Url of the auth server
+		/// </summary>
+		string AuthUrl { get; }
+
+		/// <summary>
+		/// Type of grant
+		/// </summary>
+		string GrantType { get; }
+
+		/// <summary>
+		/// Client id
+		/// </summary>
+		string ClientId { get; }
+
+		/// <summary>
+		/// Client secret
+		/// </summary>
+		string ClientSecret { get; }
+
+		/// <summary>
+		/// Scope of the token
+		/// </summary>
+		string Scope { get; }
+	}
+
+	/// <summary>
+	/// Http message handler which adds an OAuth authorization header using a cached/periodically refreshed bearer token
+	/// </summary>
+	public class OAuthHandler<T> : HttpClientHandler
+	{
+		class ClientCredentialsResponse
+		{
+			public string? access_token { get; set; }
+			public string? token_type { get; set; }
+			public int? expires_in { get; set; }
+			public string? scope { get; set; }
+		}
+
+		HttpClient Client;
+		IOAuthOptions Options;
+		string CachedAccessToken = String.Empty;
+		DateTime ExpiresAt = DateTime.MinValue;
+
+		public OAuthHandler(HttpClient Client, IOAuthOptions Options)
+		{
+			this.Client = Client;
+			this.Options = Options;
+		}
+
+		/// <inheritdoc/>
+		protected override async Task<HttpResponseMessage> SendAsync(HttpRequestMessage Request, CancellationToken CancellationToken)
+		{
+			if (DateTime.UtcNow > ExpiresAt)
+			{
+				await UpdateAccessTokenAsync(CancellationToken);
+			}
+
+			Request.Headers.Add("Authorization", $"Bearer {CachedAccessToken}");
+			return await base.SendAsync(Request, CancellationToken);
+		}
+
+		/// <summary>
+		/// Updates the current access token
+		/// </summary>
+		/// <returns></returns>
+		async Task UpdateAccessTokenAsync(CancellationToken CancellationToken)
+		{
+			KeyValuePair<string, string>[] Content = new KeyValuePair<string, string>[]
+			{
+				new KeyValuePair<string, string>("grant_type", Options.GrantType),
+				new KeyValuePair<string, string>("client_id", Options.ClientId),
+				new KeyValuePair<string, string>("client_secret", Options.ClientSecret),
+				new KeyValuePair<string, string>("scope", Options.Scope)
+			};
+
+			try
+			{
+				using HttpRequestMessage Message = new HttpRequestMessage(HttpMethod.Post, Options.AuthUrl);
+				Message.Content = new FormUrlEncodedContent(Content);
+
+				HttpResponseMessage Response = await Client.SendAsync(Message, CancellationToken);
+				if (!Response.IsSuccessStatusCode)
+				{
+					throw new AuthenticationException($"Authentication failed. Response: {Response.Content}", null);
+				}
+
+				byte[] ResponseData = await Response.Content.ReadAsByteArrayAsync();
+				ClientCredentialsResponse Result = JsonSerializer.Deserialize<ClientCredentialsResponse>(ResponseData)!;
+
+				string? accessToken = Result?.access_token;
+				if (string.IsNullOrEmpty(accessToken))
+				{
+					throw new AuthenticationException("The authentication token received by the server is null or empty. Body received was: " + Encoding.UTF8.GetString(ResponseData), null);
+				}
+				CachedAccessToken = accessToken;
+				// renew after half the renewal time
+				ExpiresAt = DateTime.UtcNow + TimeSpan.FromSeconds((Result?.expires_in ?? 3200) / 2.0);
+			}
+			catch (WebException Ex)
+			{
+				throw new AuthenticationException("Unable to authenticate.", Ex);
+			}
+		}
+	}
+
+	/// <summary>
+	/// Factory for creating OAuth2AuthProvider instances from a set of options
+	/// </summary>
+	public class OAuthHandlerFactory
+	{
+		HttpClient HttpClient;
+
+		/// <summary>
+		/// Constructor
+		/// </summary>
+		/// <param name="HttpClient"></param>
+		public OAuthHandlerFactory(HttpClient HttpClient)
+		{
+			this.HttpClient = HttpClient;
+		}
+
+		/// <summary>
+		/// Create an instance of the auth provider
+		/// </summary>
+		/// <param name="Options"></param>
+		/// <returns></returns>
+		public OAuthHandler<T> Create<T>(IOAuthOptions Options) => new OAuthHandler<T>(HttpClient, Options);
+	}
+}
diff --git a/Engine/Source/Programs/Shared/EpicGames.Horde/Auth/TokenHandler.cs b/Engine/Source/Programs/Shared/EpicGames.Horde/Auth/TokenHandler.cs
new file mode 100644
--- /dev/null
+++ b/Engine/Source/Programs/Shared/EpicGames.Horde/Auth/TokenHandler.cs
@@ -0,0 +1,48 @@
+// Copyright Epic Games, Inc. All Rights Reserved.
+
+using EpicGames.Core;
+using EpicGames.Serialization;
+using Microsoft.Extensions.Options;
+using System;
+using System.Collections.Generic;
+using System.ComponentModel.DataAnnotations;
+using System.Net;
+using System.Net.Http;
+using System.Text;
+using System.Text.Json;
+using System.Threading;
+using System.Threading.Tasks;
+
+namespace EpicGames.Horde.Auth
+{
+	/// <summary>
+	/// Options for authenticating particular requests
+	/// </summary>
+	public interface ITokenAuthOptions
+	{
+		/// <summary>
+		/// Bearer token for auth
+		/// </summary>
+		string Token { get; }
+	}
+
+	/// <summary>
+	/// Http message handler which adds an OAuth authorization header using a cached/periodically refreshed bearer token
+	/// </summary>
+	public class TokenHandler<T> : HttpClientHandler
+	{
+		ITokenAuthOptions Options;
+
+		public TokenHandler(ITokenAuthOptions Options)
+		{
+			this.Options = Options;
+		}
+
+		/// <inheritdoc/>
+		protected override async Task<HttpResponseMessage> SendAsync(HttpRequestMessage Request, CancellationToken CancellationToken)
+		{
+			Request.Headers.Add("Authorization", $"Bearer {Options.Token}");
+			return await base.SendAsync(Request, CancellationToken);
+		}
+	}
+}
diff --git a/Engine/Source/Programs/Shared/EpicGames.Horde/Compute/ComputeExtensions.cs b/Engine/Source/Programs/Shared/EpicGames.Horde/Compute/ComputeExtensions.cs
new file mode 100644
--- /dev/null
+++ b/Engine/Source/Programs/Shared/EpicGames.Horde/Compute/ComputeExtensions.cs
@@ -0,0 +1,48 @@
+// Copyright Epic Games, Inc. All Rights Reserved.
+
+using EpicGames.Horde.Auth;
+using EpicGames.Horde.Compute.Impl;
+using Microsoft.Extensions.DependencyInjection;
+using Microsoft.Extensions.Options;
+using System;
+using System.Collections.Generic;
+using System.Text;
+
+namespace EpicGames.Horde.Compute
+{
+	/// <summary>
+	/// Settings to configure a connected Horde.Build instance
+	/// </summary>
+	public class ComputeOptions : HttpServiceClientOptions
+	{
+	}
+
+	/// <summary>
+	/// Extension methods for configuring Horde Compute
+	/// </summary>
+	public static class ComputeExtensions
+	{
+		/// <summary>
+		/// Registers services for Horde Compute
+		/// </summary>
+		/// <param name="Services">The current service collection</param>
+		public static void AddHordeCompute(this IServiceCollection Services)
+		{
+			Services.AddOptions<ComputeOptions>();
+
+			Services.AddScoped<IComputeClient, HttpComputeClient>();
+			Services.AddHttpClientWithAuth<IComputeClient, HttpComputeClient>(ServiceProvider => ServiceProvider.GetRequiredService<IOptions<ComputeOptions>>().Value);
+		}
+
+		/// <summary>
+		/// Registers services for Horde Compute
+		/// </summary>
+		/// <param name="Services">The current service collection</param>
+		/// <param name="Configure">Callback for configuring the storage service</param>
+		public static void AddHordeCompute(this IServiceCollection Services, Action<ComputeOptions> Configure)
+		{
+			AddHordeCompute(Services);
+			Services.Configure<ComputeOptions>(Configure);
+		}
+	}
+}
diff --git a/Engine/Source/Programs/Shared/EpicGames.Horde/Storage/Impl/HttpStorageClient.cs b/Engine/Source/Programs/Shared/EpicGames.Horde/Storage/Impl/HttpStorageClient.cs
--- a/Engine/Source/Programs/Shared/EpicGames.Horde/Storage/Impl/HttpStorageClient.cs
+++ b/Engine/Source/Programs/Shared/EpicGames.Horde/Storage/Impl/HttpStorageClient.cs
@@ -69,20 +69,14 @@ namespace EpicGames.Horde.Storage.Impl
 		const string CompactBinaryMimeType = "application/x-ue-cb";
 
 		HttpClient HttpClient;
-		StorageSettings Settings;
-		Uri BaseUrl;
-		IAuthProvider AuthProvider;
 
 		/// <summary>
 		/// Constructor
 		/// </summary>
 		/// <param name="HttpClient">Http client for the blob store. Should be pre-configured with a base address and appropriate authentication headers.</param>
-		public HttpStorageClient(HttpClient HttpClient, IOptions<StorageSettings> Settings, IAuthProvider<HttpStorageClient> AuthProvider)
+		public HttpStorageClient(HttpClient HttpClient)
 		{
 			this.HttpClient = HttpClient;
-			this.Settings = Settings.Value;
-			this.BaseUrl = new Uri(Settings.Value.Url);
-			this.AuthProvider = AuthProvider;
 		}
 
 		#region Blobs
@@ -90,8 +84,7 @@ namespace EpicGames.Horde.Storage.Impl
 		/// <inheritdoc/>
 		public async Task<Stream> ReadBlobAsync(NamespaceId NamespaceId, IoHash Hash, CancellationToken CancellationToken)
 		{
-			using HttpRequestMessage Request = new HttpRequestMessage(HttpMethod.Get, new Uri(BaseUrl, $"api/v1/blobs/{NamespaceId}/{Hash}"));
-			await AuthProvider.AddAuthorizationAsync(Request, CancellationToken);
+			using HttpRequestMessage Request = new HttpRequestMessage(HttpMethod.Get, $"api/v1/blobs/{NamespaceId}/{Hash}");
 
 			HttpResponseMessage Response = await HttpClient.SendAsync(Request, CancellationToken);
 			try
@@ -114,8 +107,7 @@ namespace EpicGames.Horde.Storage.Impl
 		/// <inheritdoc/>
 		public async Task WriteBlobAsync(NamespaceId NamespaceId, IoHash Hash, Stream Stream, CancellationToken CancellationToken)
 		{
-			using HttpRequestMessage Request = new HttpRequestMessage(HttpMethod.Put, new Uri(BaseUrl, $"api/v1/blobs/{NamespaceId}/{Hash}"));
-			await AuthProvider.AddAuthorizationAsync(Request, CancellationToken);
+			using HttpRequestMessage Request = new HttpRequestMessage(HttpMethod.Put, $"api/v1/blobs/{NamespaceId}/{Hash}");
 
 			StreamContent Content = new StreamContent(Stream);
 			Content.Headers.ContentType = new MediaTypeHeaderValue("application/octet-stream");
@@ -132,8 +124,7 @@ namespace EpicGames.Horde.Storage.Impl
 		/// <inheritdoc/>
 		public async Task<IRef> GetRefAsync(NamespaceId NamespaceId, BucketId BucketId, RefId RefId, CancellationToken CancellationToken)
 		{
-			using HttpRequestMessage Request = new HttpRequestMessage(HttpMethod.Get, new Uri(BaseUrl, $"api/v1/refs/{NamespaceId}/{BucketId}/{RefId}"));
-			await AuthProvider.AddAuthorizationAsync(Request, CancellationToken);
+			using HttpRequestMessage Request = new HttpRequestMessage(HttpMethod.Get, $"api/v1/refs/{NamespaceId}/{BucketId}/{RefId}");
 
 			Request.Headers.Accept.Clear();
 			Request.Headers.Accept.Add(new MediaTypeWithQualityHeaderValue(CompactBinaryMimeType));
@@ -164,8 +155,7 @@ namespace EpicGames.Horde.Storage.Impl
 			Content.Headers.ContentType = new MediaTypeHeaderValue(CompactBinaryMimeType);
 			Content.Headers.Add(HashHeaderName, IoHash.Compute(Value.GetView().Span).ToString());
 
-			using HttpRequestMessage Request = new HttpRequestMessage(HttpMethod.Put, new Uri(BaseUrl, $"api/v1/refs/{NamespaceId}/{BucketId}/{RefId}"));
-			await AuthProvider.AddAuthorizationAsync(Request, CancellationToken);
+			using HttpRequestMessage Request = new HttpRequestMessage(HttpMethod.Put, $"api/v1/refs/{NamespaceId}/{BucketId}/{RefId}");
 
 			Request.Headers.Accept.Clear();
 			Request.Headers.Accept.Add(new MediaTypeWithQualityHeaderValue("application/json"));
@@ -189,8 +179,7 @@ namespace EpicGames.Horde.Storage.Impl
 		/// <inheritdoc/>
 		public async Task<List<IoHash>> TryFinalizeRefAsync(NamespaceId NamespaceId, BucketId BucketId, RefId RefId, IoHash Hash, CancellationToken CancellationToken)
 		{
-			using HttpRequestMessage Request = new HttpRequestMessage(HttpMethod.Post, new Uri(BaseUrl, $"api/v1/refs/{NamespaceId}/{BucketId}/{RefId}/finalize/{Hash}"));
-			await AuthProvider.AddAuthorizationAsync(Request, CancellationToken);
+			using HttpRequestMessage Request = new HttpRequestMessage(HttpMethod.Post, $"api/v1/refs/{NamespaceId}/{BucketId}/{RefId}/finalize/{Hash}");
 
 			Request.Headers.Accept.Clear();
 			Request.Headers.Accept.Add(new MediaTypeWithQualityHeaderValue("application/json"));
@@ -213,8 +202,7 @@ namespace EpicGames.Horde.Storage.Impl
 		/// <inheritdoc/>
 		public async Task<bool> DeleteRefAsync(NamespaceId NamespaceId, BucketId BucketId, RefId RefId, CancellationToken CancellationToken)
 		{
-			using HttpRequestMessage Request = new HttpRequestMessage(HttpMethod.Delete, new Uri(BaseUrl, $"api/v1/refs/{NamespaceId}/{BucketId}/{RefId}"));
-			await AuthProvider.AddAuthorizationAsync(Request, CancellationToken);
+			using HttpRequestMessage Request = new HttpRequestMessage(HttpMethod.Delete, $"api/v1/refs/{NamespaceId}/{BucketId}/{RefId}");
 
 			using HttpResponseMessage Response = await HttpClient.SendAsync(Request, CancellationToken);
 			return Response.IsSuccessStatusCode;
@@ -223,8 +211,7 @@ namespace EpicGames.Horde.Storage.Impl
 		/// <inheritdoc/>
 		public async Task<bool> HasRefAsync(NamespaceId NamespaceId, BucketId BucketId, RefId RefId, CancellationToken CancellationToken)
 		{
-			using HttpRequestMessage Request = new HttpRequestMessage(HttpMethod.Head, new Uri(BaseUrl, $"api/v1/refs/{NamespaceId}/{BucketId}/{RefId}"));
-			await AuthProvider.AddAuthorizationAsync(Request, CancellationToken);
+			using HttpRequestMessage Request = new HttpRequestMessage(HttpMethod.Head, $"api/v1/refs/{NamespaceId}/{BucketId}/{RefId}");
 
 			using HttpResponseMessage Response = await HttpClient.SendAsync(Request, CancellationToken);
 			if (Response.IsSuccessStatusCode)
@@ -246,10 +233,9 @@ namespace EpicGames.Horde.Storage.Impl
 		{
 			string RefList = String.Join("&", RefIds.Select(x => $"id={x}"));
 
-			using HttpRequestMessage Request = new HttpRequestMessage(HttpMethod.Post, new Uri(BaseUrl, $"/api/v1/refs/{NamespaceId}/{BucketId}/exists?{RefList}"));
-			await AuthProvider.AddAuthorizationAsync(Request, CancellationToken);
+			using HttpRequestMessage Request = new HttpRequestMessage(HttpMethod.Post, $"/api/v1/refs/{NamespaceId}/{BucketId}/exists?{RefList}");
 
-			using HttpResponseMessage Response = await HttpClient.PostAsync(new Uri(BaseUrl, $"/api/v1/refs/{NamespaceId}/{BucketId}/exists?{RefList}"), null);
+			using HttpResponseMessage Response = await HttpClient.SendAsync(Request, CancellationToken);
 			if (!Response.IsSuccessStatusCode)
 			{
 				throw new RefException(NamespaceId, BucketId, new RefId(IoHash.Zero), await GetMessageFromResponse(Response));
diff --git a/Engine/Source/Programs/Shared/EpicGames.Horde/Storage/StorageExtensions.cs b/Engine/Source/Programs/Shared/EpicGames.Horde/Storage/StorageExtensions.cs
--- a/Engine/Source/Programs/Shared/EpicGames.Horde/Storage/StorageExtensions.cs
+++ b/Engine/Source/Programs/Shared/EpicGames.Horde/Storage/StorageExtensions.cs
@@ -11,37 +11,12 @@ using System.Text;
 
 namespace EpicGames.Horde.Storage
 {
-	/// <summary>
-	/// Auth type for the storage system
-	/// </summary>
-	public enum StorageAuthType
-	{
-		/// <summary>
-		/// No auth necessary
-		/// </summary>
-		None,
-
-		/// <summary>
-		/// Use OAuth2 for auth
-		/// </summary>
-		OAuth2,
-	}
-
 	/// <summary>
 	/// Settings to configure a connected Horde.Storage instance
 	/// </summary>
 	[System.Diagnostics.CodeAnalysis.SuppressMessage("Design", "CA1056:URI-like properties should not be strings")]
-	public class StorageSettings
+	public class StorageOptions : HttpServiceClientOptions
 	{
-		/// <summary>
-		/// URL for the storage service
-		/// </summary>
-		public string Url { get; set; } = String.Empty;
-
-		/// <summary>
-		/// Options for OAuth2
-		/// </summary>
-		public OAuth2Options? OAuth2 { get; set; }
 	}
 
 	/// <summary>
@@ -55,17 +30,10 @@ namespace EpicGames.Horde.Storage
 		/// <param name="Services">The current service collection</param>
 		public static void AddHordeStorage(this IServiceCollection Services)
 		{
-			Services.AddOptions<StorageSettings>();
-
-			Services.AddSingleton<NullAuthProvider<HttpStorageClient>>();
-
-			Services.AddSingleton<OAuth2AuthProviderFactory<HttpStorageClient>>();
-			Services.AddHttpClient<OAuth2AuthProviderFactory<HttpStorageClient>>();
-
-			Services.AddSingleton<IAuthProvider<HttpStorageClient>>(ServiceProvider => CreateAuthProvider(ServiceProvider));
+			Services.AddOptions<StorageOptions>();
 
-			Services.AddTransient<IStorageClient, HttpStorageClient>();
-			Services.AddHttpClient<IStorageClient, HttpStorageClient>();
+			Services.AddScoped<IStorageClient, HttpStorageClient>();
+			Services.AddHttpClientWithAuth<IStorageClient, HttpStorageClient>(ServiceProvider => ServiceProvider.GetRequiredService<IOptions<StorageOptions>>().Value);
 		}
 
 		/// <summary>
@@ -73,28 +41,10 @@ namespace EpicGames.Horde.Storage
 		/// </summary>
 		/// <param name="Services">The current service collection</param>
 		/// <param name="Configure">Callback for configuring the storage service</param>
-		public static void AddHordeStorage(this IServiceCollection Services, Action<StorageSettings> Configure)
+		public static void AddHordeStorage(this IServiceCollection Services, Action<StorageOptions> Configure)
 		{
 			AddHordeStorage(Services);
-			Services.Configure<StorageSettings>(Configure);
-		}
-
-		/// <summary>
-		/// Creates an auth provider based on the configured settings
-		/// </summary>
-		/// <param name="Services"></param>
-		/// <returns></returns>
-		static IAuthProvider<HttpStorageClient> CreateAuthProvider(IServiceProvider Services)
-		{
-			StorageSettings Settings = Services.GetRequiredService<IOptions<StorageSettings>>().Value;
-			if (Settings.OAuth2 != null)
-			{
-				return Services.GetRequiredService<OAuth2AuthProviderFactory<HttpStorageClient>>().Create(Settings.OAuth2);
-			}
-			else
-			{
-				return Services.GetRequiredService<NullAuthProvider<HttpStorageClient>>();
-			}
+			Services.Configure<StorageOptions>(Configure);
 		}
 	}
 }
