Horde: Fix incorrect method called for Utf8String.LastIndexOf.

diff --git a/Engine/Source/Programs/Shared/EpicGames.Core/Utf8String.cs b/Engine/Source/Programs/Shared/EpicGames.Core/Utf8String.cs
--- a/Engine/Source/Programs/Shared/EpicGames.Core/Utf8String.cs
+++ b/Engine/Source/Programs/Shared/EpicGames.Core/Utf8String.cs
@@ -190,7 +190,7 @@ namespace EpicGames.Core
 		/// <inheritdoc cref="String.LastIndexOf(char)"/>
 		public int LastIndexOf(byte Char)
 		{
-			return Span.IndexOf(Char);
+			return Span.LastIndexOf(Char);
 		}
 
 		/// <inheritdoc cref="String.LastIndexOf(char)"/>
@@ -198,11 +198,11 @@ namespace EpicGames.Core
 		{
 			if (Char < 0x80)
 			{
-				return Span.IndexOf((byte)Char);
+				return Span.LastIndexOf((byte)Char);
 			}
 			else
 			{
-				return Span.IndexOf(Encoding.UTF8.GetBytes(new[] { Char }));
+				return Span.LastIndexOf(Encoding.UTF8.GetBytes(new[] { Char }));
 			}
 		}
 
