Make the type field optional for OpenRecord parsing

if -s is passed through the type field will not be present.

#rnx

diff --git a/Engine/Source/Programs/Shared/EpicGames.Perforce/Records/OpenedRecord.cs b/Engine/Source/Programs/Shared/EpicGames.Perforce/Records/OpenedRecord.cs
--- a/Engine/Source/Programs/Shared/EpicGames.Perforce/Records/OpenedRecord.cs
+++ b/Engine/Source/Programs/Shared/EpicGames.Perforce/Records/OpenedRecord.cs
@@ -60,7 +60,7 @@ namespace EpicGames.Perforce
 		/// <summary>
 		/// New filetype for the file
 		/// </summary>
-		[PerforceTag("type")]
+		[PerforceTag("type", Optional = true)]
 		public string Type { get; set; } = String.Empty;
 
 		/// <summary>
