Change IoHash to be a pure struct rather than wrapping a ReadOnlyMemory<byte> buffer.

diff --git a/Engine/Source/Programs/Horde/Horde.Build/Utilities/ProtoExtensions.cs b/Engine/Source/Programs/Horde/Horde.Build/Utilities/ProtoExtensions.cs
--- a/Engine/Source/Programs/Horde/Horde.Build/Utilities/ProtoExtensions.cs
+++ b/Engine/Source/Programs/Horde/Horde.Build/Utilities/ProtoExtensions.cs
@@ -17,7 +17,7 @@ namespace HordeCommon
 		/// <param name="Hash"></param>
 		public IoHashWrapper(IoHash Hash)
 		{
-			this.Hash = ByteString.CopyFrom(Hash.Memory.ToArray());
+			this.Hash = ByteString.CopyFrom(Hash.ToByteArray());
 		}
 
 		/// <summary>
@@ -50,7 +50,7 @@ namespace HordeCommon
 		/// <param name="RefId"></param>
 		public RefIdWrapper(RefId RefId)
 		{
-			this.Hash = ByteString.CopyFrom(RefId.Hash.Span);
+			this.Hash = ByteString.CopyFrom(RefId.Hash.ToByteArray());
 		}
 
 		/// <summary>
diff --git a/Engine/Source/Programs/Shared/EpicGames.Core/MemoryWriter.cs b/Engine/Source/Programs/Shared/EpicGames.Core/MemoryWriter.cs
--- a/Engine/Source/Programs/Shared/EpicGames.Core/MemoryWriter.cs
+++ b/Engine/Source/Programs/Shared/EpicGames.Core/MemoryWriter.cs
@@ -176,5 +176,17 @@ namespace EpicGames.Core
 				WriteItem(Array[Idx]);
 			}
 		}
+
+		/// <summary>
+		/// Allocate a writable span from the buffer
+		/// </summary>
+		/// <param name="Length">Length of the span to allocate</param>
+		/// <returns>Span that can be written to</returns>
+		public Span<byte> AllocateSpan(int Length)
+		{
+			Span<byte> Span = Memory.Span.Slice(Offset, Length);
+			Offset += Length;
+			return Span;
+		}
 	}
 }
diff --git a/Engine/Source/Programs/Shared/EpicGames.IoHash/IoHash.cs b/Engine/Source/Programs/Shared/EpicGames.IoHash/IoHash.cs
--- a/Engine/Source/Programs/Shared/EpicGames.IoHash/IoHash.cs
+++ b/Engine/Source/Programs/Shared/EpicGames.IoHash/IoHash.cs
@@ -31,44 +31,44 @@ namespace EpicGames.Core
 		/// </summary>
 		public const int NumBits = NumBytes * 8;
 
-		/// <summary>
-		/// Memory storing the digest data
-		/// </summary>
-		public ReadOnlyMemory<byte> Memory;
+		ulong A;
+		ulong B;
+		uint C;
 
 		/// <summary>
-		/// Span for the underlying memory
+		/// Hash consisting of zeroes
 		/// </summary>
-		public ReadOnlySpan<byte> Span => Memory.Span;
+		public static IoHash Zero { get; } = new IoHash(0, 0, 0);
 
 		/// <summary>
-		/// Hash consisting of zeroes
+		/// Constructor
 		/// </summary>
-		public static IoHash Zero { get; } = new IoHash(new byte[NumBytes]);
+		/// <param name="Memory">Memory to construct from</param>
+		public IoHash(ReadOnlySpan<byte> Span)
+			: this(BinaryPrimitives.ReadUInt64BigEndian(Span), BinaryPrimitives.ReadUInt64BigEndian(Span.Slice(8)), BinaryPrimitives.ReadUInt32BigEndian(Span.Slice(16)))
+		{
+		}
 
 		/// <summary>
 		/// Constructor
 		/// </summary>
 		/// <param name="Memory">Memory to construct from</param>
-		public IoHash(ReadOnlyMemory<byte> Memory)
+		public IoHash(ulong A, ulong B, uint C)
 		{
-			if (Memory.Length != NumBytes)
-			{
-				throw new ArgumentException($"IoHash must be {NumBytes} bytes long");
-			}
-
-			this.Memory = Memory;
+			this.A = A;
+			this.B = B;
+			this.C = C;
 		}
 
 		/// <summary>
 		/// Construct 
 		/// </summary>
 		/// <param name="Hasher">The hasher to construct from</param>
-		public IoHash(Blake3.Hasher Hasher)
+		public static IoHash FromBlake3(Blake3.Hasher Hasher)
 		{
 			byte[] Output = new byte[32];
 			Hasher.Finalize(Output);
-			Memory = Output.AsMemory(0, NumBytes);
+			return new IoHash(Output);
 		}
 
 		/// <summary>
@@ -80,7 +80,7 @@ namespace EpicGames.Core
 		{
 			byte[] Output = new byte[32];
 			Blake3.Hasher.Hash(Data, Output);
-			return new IoHash(Output.AsMemory(0, 20));
+			return new IoHash(Output);
 		}
 
 		/// <summary>
@@ -106,39 +106,57 @@ namespace EpicGames.Core
 		/// <inheritdoc cref="IComparable{T}.CompareTo(T)"/>
 		public int CompareTo(IoHash Other)
 		{
-			ReadOnlySpan<byte> A = Span;
-			ReadOnlySpan<byte> B = Other.Span;
-
-			for (int Idx = 0; Idx < A.Length && Idx < B.Length; Idx++)
+			if (A != Other.A)
 			{
-				int Compare = A[Idx] - B[Idx];
-				if (Compare != 0)
-				{
-					return Compare;
-				}
+				return (A < Other.A) ? -1 : +1;
+			}
+			else if (B != Other.B)
+			{
+				return (B < Other.B) ? -1 : +1;
+			}
+			else
+			{
+				return (C < Other.C) ? -1 : +1;
 			}
-			return A.Length - B.Length;
 		}
 
 		/// <inheritdoc/>
-		public bool Equals(IoHash Other) => Span.SequenceEqual(Other.Span);
+		public bool Equals(IoHash Other) => A == Other.A && B == Other.B && C == Other.C;
 
 		/// <inheritdoc/>
-		public override bool Equals(object? Obj) => (Obj is IoHash Hash) && Hash.Span.SequenceEqual(Span);
+		public override bool Equals(object? Obj) => (Obj is IoHash Hash) && Equals(Hash);
 
 		/// <inheritdoc/>
-		public override int GetHashCode() => BinaryPrimitives.ReadInt32LittleEndian(Span);
+		public override int GetHashCode() => (int)A;
 
 		/// <inheritdoc/>
-		public Utf8String ToUtf8String() => StringUtils.FormatUtf8HexString(Memory.Span);
+		public Utf8String ToUtf8String() => StringUtils.FormatUtf8HexString(ToByteArray());
 
 		/// <inheritdoc/>
-		public override string ToString() => StringUtils.FormatHexString(Memory.Span);
+		public override string ToString() => StringUtils.FormatHexString(ToByteArray());
+
+		public byte[] ToByteArray()
+		{
+			byte[] Data = new byte[NumBytes];
+			CopyTo(Data);
+			return Data;
+		}
+
+		/// <summary>
+		/// Copies this hash into a span
+		/// </summary>
+		/// <param name="Span"></param>
+		public void CopyTo(Span<byte> Span)
+		{
+			BinaryPrimitives.WriteUInt64BigEndian(Span, A);
+			BinaryPrimitives.WriteUInt64BigEndian(Span[8..], B);
+			BinaryPrimitives.WriteUInt32BigEndian(Span[16..], C);
+		}
 
 		/// <summary>
 		/// Test two hash values for equality
 		/// </summary>
-		public static bool operator ==(IoHash A, IoHash B) => A.Span.SequenceEqual(B.Span);
+		public static bool operator ==(IoHash A, IoHash B) => A.Equals(B);
 
 		/// <summary>
 		/// Test two hash values for equality
@@ -171,7 +189,7 @@ namespace EpicGames.Core
 		/// <param name="Hash"></param>
 		public static implicit operator IoHash(Blake3Hash Hash)
 		{
-			return new IoHash(Hash.Memory.Slice(0, NumBytes));
+			return new IoHash(Hash.Span.Slice(0, NumBytes));
 		}
 	}
 
@@ -187,7 +205,7 @@ namespace EpicGames.Core
 		/// <returns></returns>
 		public static IoHash ReadIoHash(this MemoryReader Reader)
 		{
-			return new IoHash(Reader.ReadFixedLengthBytes(IoHash.NumBytes));
+			return new IoHash(Reader.ReadFixedLengthBytes(IoHash.NumBytes).Span);
 		}
 
 		/// <summary>
@@ -197,7 +215,7 @@ namespace EpicGames.Core
 		/// <param name="Hash"></param>
 		public static void WriteIoHash(this MemoryWriter Writer, IoHash Hash)
 		{
-			Writer.WriteFixedLengthBytes(Hash.Span);
+			Hash.CopyTo(Writer.AllocateSpan(IoHash.NumBytes));
 		}
 	}
 
diff --git a/Engine/Source/Programs/Shared/EpicGames.LruCache/LruCache.cs b/Engine/Source/Programs/Shared/EpicGames.LruCache/LruCache.cs
--- a/Engine/Source/Programs/Shared/EpicGames.LruCache/LruCache.cs
+++ b/Engine/Source/Programs/Shared/EpicGames.LruCache/LruCache.cs
@@ -198,7 +198,7 @@ namespace HordeCommon
 					Slot = NextSlot;
 				}
 
-				Hash.Memory.CopyTo(GetHashMemory(Slot));
+				Hash.CopyTo(GetHashMemory(Slot).Span);
 				Items[Slot] = ItemData;
 			}
 
@@ -209,7 +209,7 @@ namespace HordeCommon
 				return ProbeLen > CompareProbeLen;
 			}
 
-			public IoHash GetHash(int Slot) => new IoHash(GetHashMemory(Slot));
+			public IoHash GetHash(int Slot) => new IoHash(GetHashMemory(Slot).Span);
 			public Item GetItem(int Slot) => new Item(Items[Slot]);
 
 			public int Find(IoHash Hash)
@@ -574,7 +574,7 @@ namespace HordeCommon
 					Item Item = Lookup.GetItem(Slot);
 					if (Item.IsValid)
 					{
-						Writer.WriteFixedLengthBytes(Lookup.GetHash(Slot).Span);
+						Writer.WriteIoHash(Lookup.GetHash(Slot));
 						Writer.WriteInt64(Item.Data);
 
 						if (Item.IsLargeItem)
diff --git a/Engine/Source/Programs/Shared/EpicGames.Serialization/CbField.cs b/Engine/Source/Programs/Shared/EpicGames.Serialization/CbField.cs
--- a/Engine/Source/Programs/Shared/EpicGames.Serialization/CbField.cs
+++ b/Engine/Source/Programs/Shared/EpicGames.Serialization/CbField.cs
@@ -1197,7 +1197,7 @@ namespace EpicGames.Serialization
 			if (CbFieldUtils.IsObjectAttachment(TypeWithFlags))
 			{
 				Error = CbFieldError.None;
-				return new IoHash(Payload);
+				return new IoHash(Payload.Span);
 			}
 			else
 			{
@@ -1223,7 +1223,7 @@ namespace EpicGames.Serialization
 			if (CbFieldUtils.IsBinaryAttachment(TypeWithFlags))
 			{
 				Error = CbFieldError.None;
-				return new IoHash(Payload);
+				return new IoHash(Payload.Span);
 			}
 			else
 			{
@@ -1249,7 +1249,7 @@ namespace EpicGames.Serialization
 			if (CbFieldUtils.IsAttachment(TypeWithFlags))
 			{
 				Error = CbFieldError.None;
-				return new IoHash(Payload);
+				return new IoHash(Payload.Span);
 			}
 			else
 			{
@@ -1274,7 +1274,7 @@ namespace EpicGames.Serialization
 			if (CbFieldUtils.IsHash(TypeWithFlags))
 			{
 				Error = CbFieldError.None;
-				return new IoHash(Payload);
+				return new IoHash(Payload.Span);
 			}
 			else
 			{
@@ -2669,7 +2669,7 @@ namespace EpicGames.Serialization
 			}
 			else if (Field.IsHash())
 			{
-				Writer.WriteString(Field.Name.Span, StringUtils.FormatUtf8HexString(Field.AsHash().Span).Span);
+				Writer.WriteString(Field.Name.Span, Field.AsHash().ToUtf8String().Span);
 			}
 			else if (Field.IsString())
 			{
diff --git a/Engine/Source/Programs/Shared/EpicGames.Serialization/CbWriter.cs b/Engine/Source/Programs/Shared/EpicGames.Serialization/CbWriter.cs
--- a/Engine/Source/Programs/Shared/EpicGames.Serialization/CbWriter.cs
+++ b/Engine/Source/Programs/Shared/EpicGames.Serialization/CbWriter.cs
@@ -599,7 +599,7 @@ namespace EpicGames.Serialization
 		void WriteHashPayload(IoHash Hash)
 		{
 			Span<byte> Buffer = Allocate(IoHash.NumBytes).Span;
-			Hash.Span.CopyTo(Buffer);
+			Hash.CopyTo(Buffer);
 		}
 
 		/// <summary>
