P4: Add support for intercepting I/O through native Perforce library. This allows streaming sync data directly into memory, avoiding any writes to disk.

diff --git a/Engine/Source/Programs/Shared/EpicGames.Perforce.Native/Client.cpp b/Engine/Source/Programs/Shared/EpicGames.Perforce.Native/Client.cpp
--- a/Engine/Source/Programs/Shared/EpicGames.Perforce.Native/Client.cpp
+++ b/Engine/Source/Programs/Shared/EpicGames.Perforce.Native/Client.cpp
@@ -1,7 +1,11 @@
 // Copyright Epic Games, Inc. All Rights Reserved.
 
+#pragma warning(push)
+#pragma warning(disable: 4244)
 #include "clientapi.h"
 #include "strtable.h"
+#include "datetime.h"
+#pragma warning(pop)
 #include <assert.h>
 
 #ifdef _MSC_VER
@@ -55,6 +59,7 @@ private:
 public:
 	const char* Func = nullptr;
 	const char* PromptResponse = nullptr;
+	bool InterceptIo = false;
 
 	FClientUser(FWriteBuffer* WriteBuffer, FOnBufferReadyFn* OnBufferReady)
 	{
@@ -62,6 +67,8 @@ public:
 		this->OnBufferReady = OnBufferReady;
 	}
 
+	virtual FileSys* File(FileSysType type) override;
+
 	virtual void InputData(StrBuf* strbuf, Error* e) override
 	{
 		if (InputBuffer != nullptr)
@@ -200,6 +207,43 @@ public:
 		assert(false);
 	}
 
+	void OutputIo(int FileId, const char* Command, const void* Payload, int PayloadLen)
+	{
+		while (!TryOutputIo(FileId, Command, Payload, PayloadLen))
+		{
+			Flush();
+		}
+	}
+
+	bool TryOutputIo(int FileId, const char* Command, const void* Payload, int PayloadLen)
+	{
+		static const char CodeKey[] = "code";
+		static const char FileKey[] = "file";
+		static const char CommandKey[] = "command";
+		static const char PayloadKey[] = "payload";
+
+		static const char IoCode[] = "io";
+
+		unsigned int CommandLen = (unsigned int)strlen(Command);
+		unsigned int RecordLen = (unsigned int)(1 + MeasureStringField(CodeKey, IoCode) + MeasureIntField(FileKey) + MeasureStringField(CommandKey, CommandLen) + MeasureStringField(PayloadKey, PayloadLen) + 1);
+		if (Length + RecordLen > MaxLength)
+		{
+			return false;
+		}
+
+		unsigned char* Pos = Data + Length;
+		*(Pos++) = '{';
+		Pos = WriteStringField(Pos, CodeKey, IoCode);
+		Pos = WriteIntField(Pos, FileKey, FileId);
+		Pos = WriteStringField(Pos, CommandKey, Command, CommandLen);
+		Pos = WriteStringField(Pos, PayloadKey, (const char*)Payload, PayloadLen);
+		*(Pos++) = '0';
+		assert(Pos == Data + Length + RecordLen);
+
+		Length += RecordLen;
+		return true;
+	}
+
 	virtual void OutputInfo(char Level, const char* Data) override
 	{
 		while (!TryOutputInfo(Level, Data))
@@ -395,6 +439,136 @@ public:
 	}
 };
 
+class FFileSys : public FileSys
+{
+public:
+	static int NextFileId;
+	int FileId = -1;
+
+	const FileSysType Type;
+	class FClientUser& User;
+
+	FFileSys(FileSysType InType, FClientUser& InUser)
+		: Type(InType)
+		, User(InUser)
+	{
+	}
+
+	virtual void	ChmodTimeHP(const DateTimeHighPrecision& /* modTime */, Error* /* e */) override {};
+	virtual void	SetAttribute(FileSysAttr, Error*) override { };
+
+	virtual bool	HasOnlyPerm(FilePerm perms) override { return false; }
+	virtual int	GetFd() override { return -1; }
+	virtual int     GetOwner() override { return 0; }
+	virtual offL_t	GetSize() override { return 0; }
+	virtual void	Seek(offL_t offset, Error*) override { }
+	virtual offL_t	Tell() override { return 0; }
+
+	virtual void	MakeLocalTemp(char* file) override { assert(false); }
+	virtual void	SetDeleteOnClose() override { }
+	virtual void	ClearDeleteOnClose() override {}
+	virtual StrArray* ScanDir(Error* e) override { return nullptr; }
+
+	virtual void	MkDir(const StrPtr& p, Error* e) override { }
+	virtual void	PurgeDir(const char* p, Error* e) override { }
+	virtual void	RmDir(const StrPtr& p, Error* e) override { }
+	virtual int	ReadLine(StrBuf* buf, Error* e) override { return 0; }
+
+	virtual void Open(FileOpenMode mode, Error* e) override
+	{
+		int PathLen = path.Length() + 1;
+
+		int BufferLen = PathLen + sizeof(int) + sizeof(int);
+
+		char* Buffer = new char[BufferLen];
+		memcpy(Buffer, path.Text(), (size_t)PathLen);
+
+		int TypeInt = (int)Type;
+		memcpy(Buffer + PathLen, &TypeInt, sizeof(int));
+
+		int ModeInt = (int)mode;
+		memcpy(Buffer + PathLen + sizeof(int), &ModeInt, sizeof(int));
+
+		FileId = ++NextFileId;
+		User.OutputIo(FileId, "open", Buffer, BufferLen);
+
+		delete[] Buffer;
+	}
+
+	virtual void Write(const char* buf, int len, Error* e) override
+	{
+		User.OutputIo(FileId, "write", buf, len);
+	}
+
+	virtual int Read(char* buf, int len, Error* e) override
+	{
+		return 0;
+	}
+
+	virtual void Close(Error* e) override
+	{
+		User.OutputIo(FileId, "close", nullptr, 0);
+	}
+
+	virtual int Stat() override
+	{
+		return 0;
+	}
+
+	virtual int StatModTime() override
+	{
+		return 0;
+	}
+
+	virtual void StatModTimeHP(DateTimeHighPrecision* modTime) override
+	{
+	}
+
+	virtual void Truncate(Error* e) override
+	{
+		assert(false);
+	}
+
+	virtual void Truncate(offL_t offset, Error* e) override
+	{
+		assert(false);
+	}
+
+	virtual void Unlink(Error* e = 0) override
+	{
+		User.OutputIo(FileId, "unlink", path.Text(), path.Length());
+	}
+
+	virtual void Rename(FileSys* target, Error* e) override
+	{
+		assert(false);
+	}
+
+	virtual void Chmod(FilePerm perms, Error* e) override
+	{
+		assert(false);
+	}
+
+	virtual void ChmodTime(Error* e) override
+	{
+		assert(false);
+	}
+};
+
+int FFileSys::NextFileId = 100;
+
+FileSys* FClientUser::File(FileSysType type)
+{
+	if (InterceptIo)
+	{
+		return new FFileSys(type, *this);
+	}
+	else
+	{
+		return ClientUser::File(type);
+	}
+}
+
 class FClient
 {
 public:
@@ -458,14 +632,16 @@ extern "C" NATIVE_API void Client_Login(FClient * Client, const char* Password)
 	Client->User.PromptResponse = nullptr;
 }
 
-extern "C" NATIVE_API void Client_Command(FClient* Client, const char* Func, int ArgCount, const char** Args, const char* InputData, int InputLength)
+extern "C" NATIVE_API void Client_Command(FClient* Client, const char* Func, int ArgCount, const char** Args, const char* InputData, int InputLength, bool InterceptIo)
 {
+	Client->User.InterceptIo = InterceptIo;
 	Client->User.Func = Func;
 	Client->User.SetInputBuffer(InputData, InputLength);
 	Client->ClientApi.SetArgv(ArgCount, (char* const*)Args);
 	Client->ClientApi.Run(Func, &Client->User);
 	Client->User.Flush();
 	Client->User.Func = nullptr;
+	Client->User.InterceptIo = false;
 }
 
 extern "C" NATIVE_API void Client_Destroy(FClient* Client)
diff --git a/Engine/Source/Programs/Shared/EpicGames.Perforce/IPerforceConnection.cs b/Engine/Source/Programs/Shared/EpicGames.Perforce/IPerforceConnection.cs
--- a/Engine/Source/Programs/Shared/EpicGames.Perforce/IPerforceConnection.cs
+++ b/Engine/Source/Programs/Shared/EpicGames.Perforce/IPerforceConnection.cs
@@ -31,8 +31,9 @@ namespace EpicGames.Perforce
 		/// <param name="Arguments">Arguments for the command</param>
 		/// <param name="FileArguments">File arguments (may be put into a response file)</param>
 		/// <param name="InputData">Input data to be passed to the command</param>
+		/// <param name="InterceptIo">Whether to intercept file I/O and return it in the reponse stream. Only supported by the native client.</param>
 		/// <returns>Response object</returns>
-		Task<IPerforceOutput> CommandAsync(string Command, IReadOnlyList<string> Arguments, IReadOnlyList<string>? FileArguments, byte[]? InputData);
+		Task<IPerforceOutput> CommandAsync(string Command, IReadOnlyList<string> Arguments, IReadOnlyList<string>? FileArguments, byte[]? InputData, bool InterceptIo);
 
 		/// <summary>
 		/// Execute the 'login' command
diff --git a/Engine/Source/Programs/Shared/EpicGames.Perforce/IPerforceOutput.cs b/Engine/Source/Programs/Shared/EpicGames.Perforce/IPerforceOutput.cs
--- a/Engine/Source/Programs/Shared/EpicGames.Perforce/IPerforceOutput.cs
+++ b/Engine/Source/Programs/Shared/EpicGames.Perforce/IPerforceOutput.cs
@@ -52,6 +52,7 @@ namespace EpicGames.Perforce
 			public static Utf8String Stat = "stat";
 			public static Utf8String Info = "info";
 			public static Utf8String Error = "error";
+			public static Utf8String Io = "io";
 		}
 
 		/// <summary>
@@ -467,6 +468,14 @@ namespace EpicGames.Perforce
 					return false;
 				}
 			}
+			else if (Code == ReadOnlyUtf8StringConstants.Io)
+			{
+				if (!TryReadTypedRecord(Buffer, ref BufferPos, Utf8String.Empty, PerforceReflection.IoRecordInfo, out Record))
+				{
+					Response = null;
+					return false;
+				}
+			}
 			else
 			{
 				throw new PerforceException("Unknown return code for record: {0}", Code);
diff --git a/Engine/Source/Programs/Shared/EpicGames.Perforce/NativePerforceConnection.cs b/Engine/Source/Programs/Shared/EpicGames.Perforce/NativePerforceConnection.cs
--- a/Engine/Source/Programs/Shared/EpicGames.Perforce/NativePerforceConnection.cs
+++ b/Engine/Source/Programs/Shared/EpicGames.Perforce/NativePerforceConnection.cs
@@ -75,7 +75,7 @@ namespace EpicGames.Perforce
 		static extern void Client_Login(IntPtr Client, [MarshalAs(UnmanagedType.LPStr)] string Password);
 
 		[DllImport(NativeDll, CallingConvention = CallingConvention.Cdecl, CharSet = CharSet.Ansi, BestFitMapping = false, ThrowOnUnmappableChar = true)]
-		static extern void Client_Command(IntPtr Client, [MarshalAs(UnmanagedType.LPStr)] string Command, int NumArgs, string[] Args, byte[]? InputData, int InputLength);
+		static extern void Client_Command(IntPtr Client, [MarshalAs(UnmanagedType.LPStr)] string Command, int NumArgs, string[] Args, byte[]? InputData, int InputLength, bool InterceptIo);
 
 		[DllImport(NativeDll, CallingConvention = CallingConvention.Cdecl)]
 		static extern void Client_Destroy(IntPtr Client);
@@ -485,7 +485,7 @@ namespace EpicGames.Perforce
 		}
 
 		/// <inheritdoc/>
-		public Task<IPerforceOutput> CommandAsync(string Command, IReadOnlyList<string> Arguments, IReadOnlyList<string>? FileArguments, byte[]? InputData)
+		public Task<IPerforceOutput> CommandAsync(string Command, IReadOnlyList<string> Arguments, IReadOnlyList<string>? FileArguments, byte[]? InputData, bool InterceptIo)
 		{
 			byte[]? SpecData = null;
 			if (InputData != null)
@@ -500,7 +500,7 @@ namespace EpicGames.Perforce
 			}
 
 			Response Response = new Response(this);
-			Requests.Add((() => Client_Command(Client, Command, AllArguments.Count, AllArguments.ToArray(), SpecData, SpecData?.Length ?? 0), Response));
+			Requests.Add((() => Client_Command(Client, Command, AllArguments.Count, AllArguments.ToArray(), SpecData, SpecData?.Length ?? 0, InterceptIo), Response));
 			return Task.FromResult<IPerforceOutput>(Response);
 		}
 
diff --git a/Engine/Source/Programs/Shared/EpicGames.Perforce/PerforceConnection.cs b/Engine/Source/Programs/Shared/EpicGames.Perforce/PerforceConnection.cs
--- a/Engine/Source/Programs/Shared/EpicGames.Perforce/PerforceConnection.cs
+++ b/Engine/Source/Programs/Shared/EpicGames.Perforce/PerforceConnection.cs
@@ -309,8 +309,12 @@ namespace EpicGames.Perforce
 		}
 
 		/// <inheritdoc/>
-		public Task<IPerforceOutput> CommandAsync(string Command, IReadOnlyList<string> Arguments, IReadOnlyList<string>? FileArguments, byte[]? InputData)
+		public Task<IPerforceOutput> CommandAsync(string Command, IReadOnlyList<string> Arguments, IReadOnlyList<string>? FileArguments, byte[]? InputData, bool InterceptIo)
 		{
+			if (InterceptIo)
+			{
+				throw new NotSupportedException($"{nameof(InterceptIo)} option is not supported through legacy Perforce client");
+			}
 			return Task.FromResult<IPerforceOutput>(new PerforceChildProcess(Command, Arguments, FileArguments, InputData, GetGlobalArguments(), Logger));
 		}
 
@@ -421,7 +425,7 @@ namespace EpicGames.Perforce
 		/// <returns>List of objects returned by the server</returns>
 		public static async Task<List<PerforceResponse>> CommandAsync(this IPerforceConnection Perforce, string Command, IReadOnlyList<string> Arguments, IReadOnlyList<string>? FileArguments, byte[]? InputData, Type? StatRecordType, CancellationToken CancellationToken = default)
 		{
-			await using (IPerforceOutput Response = await Perforce.CommandAsync(Command, Arguments, FileArguments, InputData))
+			await using (IPerforceOutput Response = await Perforce.CommandAsync(Command, Arguments, FileArguments, InputData, false))
 			{
 				return await Response.ReadResponsesAsync(StatRecordType, CancellationToken);
 			}
@@ -436,11 +440,12 @@ namespace EpicGames.Perforce
 		/// <param name="FileArguments">File arguments for the command</param>
 		/// <param name="InputData">Input data to pass to Perforce</param>
 		/// <param name="StatRecordType">The type of records to return for "stat" responses</param>
+		/// <param name="InterceptIo">Whether to intercept Io operations and return them in the response output</param>
 		/// <param name="CancellationToken">Token used to cancel the operation</param>
 		/// <returns>List of objects returned by the server</returns>
-		public static async IAsyncEnumerable<PerforceResponse> StreamCommandAsync(this IPerforceConnection Perforce, string Command, IReadOnlyList<string> Arguments, IReadOnlyList<string>? FileArguments, byte[]? InputData, Type? StatRecordType, [EnumeratorCancellation] CancellationToken CancellationToken)
+		public static async IAsyncEnumerable<PerforceResponse> StreamCommandAsync(this IPerforceConnection Perforce, string Command, IReadOnlyList<string> Arguments, IReadOnlyList<string>? FileArguments, byte[]? InputData, Type? StatRecordType, bool InterceptIo, [EnumeratorCancellation] CancellationToken CancellationToken)
 		{
-			await using (IPerforceOutput Output = await Perforce.CommandAsync(Command, Arguments, FileArguments, InputData))
+			await using (IPerforceOutput Output = await Perforce.CommandAsync(Command, Arguments, FileArguments, InputData, InterceptIo))
 			{
 				await foreach (PerforceResponse Response in Output.ReadStreamingResponsesAsync(StatRecordType, CancellationToken))
 				{
@@ -461,7 +466,7 @@ namespace EpicGames.Perforce
 		/// <returns>List of objects returned by the server</returns>
 		public static async IAsyncEnumerable<PerforceResponse<T>> StreamCommandAsync<T>(this IPerforceConnection Perforce, string Command, IReadOnlyList<string> Arguments, IReadOnlyList<string>? FileArguments = null, byte[]? InputData = null, [EnumeratorCancellation] CancellationToken CancellationToken = default) where T : class
 		{
-			await using (IPerforceOutput Output = await Perforce.CommandAsync(Command, Arguments, FileArguments, InputData))
+			await using (IPerforceOutput Output = await Perforce.CommandAsync(Command, Arguments, FileArguments, InputData, false))
 			{
 				Type StatRecordType = typeof(T);
 				await foreach (PerforceResponse Response in Output.ReadStreamingResponsesAsync(StatRecordType, CancellationToken))
@@ -483,7 +488,7 @@ namespace EpicGames.Perforce
 		/// <returns>List of objects returned by the server</returns>
 		public static async Task RecordCommandAsync(this IPerforceConnection Perforce, string Command, IReadOnlyList<string> Arguments, byte[]? InputData, Action<PerforceRecord> HandleRecord, CancellationToken CancellationToken = default)
 		{
-			await using (IPerforceOutput Response = await Perforce.CommandAsync(Command, Arguments, null, InputData))
+			await using (IPerforceOutput Response = await Perforce.CommandAsync(Command, Arguments, null, InputData, false))
 			{
 				await Response.ReadRecordsAsync(HandleRecord, CancellationToken);
 			}
diff --git a/Engine/Source/Programs/Shared/EpicGames.Perforce/PerforceIo.cs b/Engine/Source/Programs/Shared/EpicGames.Perforce/PerforceIo.cs
new file mode 100644
--- /dev/null
+++ b/Engine/Source/Programs/Shared/EpicGames.Perforce/PerforceIo.cs
@@ -0,0 +1,62 @@
+// Copyright Epic Games, Inc. All Rights Reserved.
+
+using System;
+using System.Collections.Generic;
+using System.Text;
+
+namespace EpicGames.Perforce
+{
+	/// <summary>
+	/// 
+	/// </summary>
+	public enum PerforceIoCommand
+	{
+		/// <summary>
+		/// 
+		/// </summary>
+		[PerforceEnum("open")]
+		Open,
+
+		/// <summary>
+		/// 
+		/// </summary>
+		[PerforceEnum("write")]
+		Write,
+
+		/// <summary>
+		/// 
+		/// </summary>
+		[PerforceEnum("close")]
+		Close,
+
+		/// <summary>
+		/// 
+		/// </summary>
+		[PerforceEnum("unlink")]
+		Unlink,
+	}
+
+	/// <summary>
+	/// 
+	/// </summary>
+	public class PerforceIo
+	{
+		/// <summary>
+		/// The severity of this error
+		/// </summary>
+		[PerforceTagAttribute("file")]
+		public int File { get; set; }
+
+		/// <summary>
+		/// The generic error code associated with this message
+		/// </summary>
+		[PerforceTagAttribute("command")]
+		public PerforceIoCommand Command { get; set; }
+
+		/// <summary>
+		/// The message text
+		/// </summary>
+		[PerforceTagAttribute("payload")]
+		public ReadOnlyMemory<byte> Payload { get; set; }
+	}
+}
diff --git a/Engine/Source/Programs/Shared/EpicGames.Perforce/PerforceReflection.cs b/Engine/Source/Programs/Shared/EpicGames.Perforce/PerforceReflection.cs
--- a/Engine/Source/Programs/Shared/EpicGames.Perforce/PerforceReflection.cs
+++ b/Engine/Source/Programs/Shared/EpicGames.Perforce/PerforceReflection.cs
@@ -355,6 +355,11 @@ namespace EpicGames.Perforce
 		/// </summary>
 		public static CachedRecordInfo ErrorRecordInfo = GetCachedRecordInfo(typeof(PerforceError));
 
+		/// <summary>
+		/// Default type for errors
+		/// </summary>
+		public static CachedRecordInfo IoRecordInfo = GetCachedRecordInfo(typeof(PerforceIo));
+
 		/// <summary>
 		/// Gets a mapping of flags to enum values for the given type
 		/// </summary>
@@ -467,6 +472,10 @@ namespace EpicGames.Perforce
 						{
 							TagInfo.SetFromString = (Obj, String) => ((List<string>)PropertyCopy.GetValue(Obj)!).Add(String.ToString());
 						}
+						else if (FieldType == typeof(ReadOnlyMemory<byte>))
+						{
+							TagInfo.SetFromString = (Obj, String) => PropertyCopy.SetValue(Obj, String.Memory);
+						}
 						else
 						{
 							throw new PerforceException("Unsupported type of {0}.{1} for tag '{2}'", RecordType.Name, FieldType.Name, TagName);
diff --git a/Engine/Source/Programs/Shared/EpicGames.Perforce/PerforceResponse.cs b/Engine/Source/Programs/Shared/EpicGames.Perforce/PerforceResponse.cs
--- a/Engine/Source/Programs/Shared/EpicGames.Perforce/PerforceResponse.cs
+++ b/Engine/Source/Programs/Shared/EpicGames.Perforce/PerforceResponse.cs
@@ -71,6 +71,14 @@ namespace EpicGames.Perforce
 			get { return InternalData as PerforceError; }
 		}
 
+		/// <summary>
+		/// Returns the io data, or null if this is a regular response.
+		/// </summary>
+		public PerforceIo? Io
+		{
+			get { return InternalData as PerforceIo; }
+		}
+
 		/// <summary>
 		/// Throws an exception if the response is an error
 		/// </summary>
