Horde: Add helper classes for facilitate packing merkle trees into larger blobs, exploiting temporal coherency when updating them. Supports content-aware slicing of larger files.

The "logical" tree consists of nodes prefixed by a byte indicating its type (TreePackNodeType). Supported node types are Binary (leaf data with no external references), Object (a generic compact-binary encoded object), Directory (a directory structure which may reference other files or directories by name), and Concat (a concatenated sequence of Binary objects).

Nodes are packed together into "physical" blobs for storage, which default to 10mb in size and may reference other blobs. Each blob has an import and export table listing nodes they contain and other blobs they require nodes from. The hierarchy of blob dependencies also form an acyclic graph.

Physical blobs are reused between subsequent runs, with updated/replacement data being added into new blobs at the root of the graph. Physical blobs can be discarded and re-packed along with the new data once a costing heuristic indicates it is more efficient to do so. The costing heuristic takes into account the size of useful vs wasted space in a pack, its age, and an estimate for how costly it would be for machines to download it (see TreePack.GetCostHeuristic).

HordeAgent includes two new utility commands that can be use to create pack files from a tree of files or Perforce depot path, with the following syntax:

    HordeAgent.exe archive create -source=[Local or Perforce Depot Path To Read From] -target=[Path to an output TXT file containing the hash of the root object; blobs are placed alongside it in the same directory] -base=[Path to a previous archive TXT file to generate a delta from]
    HordeAgent.exe archive extract -input=[Archive root hash TXT file] -output=[Output directory]

diff --git a/Engine/Source/Programs/Horde/Horde.Agent/Commands/Archive/CreateArchiveCommand.cs b/Engine/Source/Programs/Horde/Horde.Agent/Commands/Archive/CreateArchiveCommand.cs
new file mode 100644
--- /dev/null
+++ b/Engine/Source/Programs/Horde/Horde.Agent/Commands/Archive/CreateArchiveCommand.cs
@@ -0,0 +1,239 @@
+// Copyright Epic Games, Inc. All Rights Reserved.
+
+using Google.Protobuf;
+using Microsoft.Extensions.Logging;
+using System;
+using System.Collections.Concurrent;
+using System.Collections.Generic;
+using System.IO;
+using System.Linq;
+using System.Text;
+using System.Threading.Tasks;
+using EpicGames.Core;
+using System.Text.Json;
+using HordeAgent.Services;
+using Grpc.Net.Client;
+using System.Net.Http.Headers;
+using Grpc.Core;
+using Microsoft.Extensions.Hosting;
+using Microsoft.Extensions.DependencyInjection;
+using Microsoft.Extensions.Configuration;
+using Serilog.Events;
+using EpicGames.Serialization;
+using EpicGames.Horde.Compute;
+using EpicGames.Horde.Storage;
+using System.Net.Http;
+using EpicGames.Horde.Compute.Impl;
+using EpicGames.Horde.Storage.Impl;
+using System.Threading;
+using EpicGames.Perforce;
+
+namespace HordeAgent.Commands.Archive
+{
+	class FileSystemStorageClient : IBlobStorageClient
+	{
+		public DirectoryReference BaseDir;
+		ILogger Logger;
+
+		public FileSystemStorageClient(DirectoryReference BaseDir, ILogger Logger)
+		{
+			this.BaseDir = BaseDir;
+			this.Logger = Logger;
+
+			DirectoryReference.CreateDirectory(BaseDir);
+		}
+
+		public Task<Stream> ReadBlobAsync(NamespaceId NamespaceId, IoHash Hash, CancellationToken CancellationToken = default)
+		{
+			FileReference File = FileReference.Combine(BaseDir, $"{NamespaceId}-{Hash}.dat");
+			Logger.LogInformation("Reading {File} ({Size} bytes)", File, new FileInfo(File.FullName).Length);
+			return Task.FromResult<Stream>(FileReference.Open(File, FileMode.Open, FileAccess.Read, FileShare.Read));
+		}
+
+		public async Task WriteBlobAsync(NamespaceId NamespaceId, IoHash Hash, Stream Stream, CancellationToken CancellationToken = default)
+		{
+			FileReference File = FileReference.Combine(BaseDir, $"{NamespaceId}-{Hash}.dat");
+			Logger.LogInformation("Writing {File} ({Size} bytes)", File, Stream.Length);
+			using (Stream OutputStream = FileReference.Open(File, FileMode.Create, FileAccess.Write, FileShare.Read))
+			{
+				await Stream.CopyToAsync(OutputStream);
+			}
+		}
+	}
+
+	/// <summary>
+	/// Installs the agent as a service
+	/// </summary>
+	[Command("Archive", "Create", "Creates an archive containing a set of files or folders")]
+	class CreateArchiveCommand : Command
+	{
+		/// <summary>
+		/// Input path. May be a local filesystem wildcard (eg. "E:\P4\...") or Perforce depot path.
+		/// </summary>
+		[CommandLine("-Source=", Required = true)]
+		public string SourcePath = null!;
+
+		/// <summary>
+		/// Path to a base file to generate an archive from
+		/// </summary>
+		[CommandLine("-Base=")]
+		public FileReference? BasePath = null;
+
+		/// <summary>
+		/// Output path to an index file.
+		/// </summary>
+		[CommandLine("-Target=", Required = true)]
+		public FileReference TargetPath = null!;
+
+		/// <inheritdoc/>
+		public override async Task<int> ExecuteAsync(ILogger Logger)
+		{
+			TreePackOptions Options = new TreePackOptions();
+//			Options.MaxBlobSize = 64 * 1024;
+//			Options.MaxInlineBlobSize = 20 * 1024;
+
+			DirectoryReference TargetDir = TargetPath.Directory;
+			NamespaceId NamespaceId = new NamespaceId("test");
+			FileSystemStorageClient Storage = new FileSystemStorageClient(TargetDir, Logger);
+
+			TreePack TargetPack = new TreePack(Storage, NamespaceId, Options);
+
+			TreePackDirWriter TargetWriter = new TreePackDirWriter(TargetPack);
+			if (BasePath != null)
+			{
+				IoHash RootHash = IoHash.Parse((await FileReference.ReadAllTextAsync(BasePath)).Trim());
+				byte[] RootObject = await FileReference.ReadAllBytesAsync(BasePath.ChangeExtension(".dat"));
+				TargetPack.AddRootBlob(RootObject);
+
+				ReadOnlyMemory<byte> RootDirData = await TargetPack.GetDataAsync(RootHash);
+				TreePackDirNode RootDirNode = TreePackDirNode.Parse(RootDirData);
+
+				TargetWriter = new TreePackDirWriter(TargetPack, RootDirNode);
+			}
+
+			if (SourcePath.StartsWith("//"))
+			{
+				await ReadFilesFromPerforceAsync(SourcePath, TargetPack, TargetWriter);
+			}
+			else
+			{
+				await ReadFilesFromDiskAsync(SourcePath, TargetPack, TargetWriter);
+			}
+
+			IoHash NewHash = await TargetWriter.FinalizeAsync();
+			TreePackObject Pack = await TargetPack.FlushAsync(NewHash, DateTime.UtcNow);
+			await WriteSummaryFiles(NamespaceId, Pack, new HashSet<IoHash>(), Storage, TargetDir);
+
+			Logger.LogInformation("Writing root to {TargetPath}", TargetPath);
+			await FileReference.WriteAllTextAsync(TargetPath, NewHash.ToString());
+
+			Logger.LogInformation("Writing pack to {TargetPath}", TargetPath.ChangeExtension(".dat"));
+			await FileReference.WriteAllBytesAsync(TargetPath.ChangeExtension(".dat"), CbSerializer.Serialize(Pack).GetView().ToArray());
+
+			Logger.LogInformation("Writing yaml to {TargetPath}", TargetPath.ChangeExtension(".yml"));
+			TreePack.WriteSummary(TargetPath.ChangeExtension(".yml"), Pack);
+
+			return 0;
+		}
+
+		async Task WriteSummaryFiles(NamespaceId NamespaceId, TreePackObject RootObject, HashSet<IoHash> VisitedObjects, IBlobStorageClient StorageClient, DirectoryReference TargetDir)
+		{
+			foreach (TreePackObjectImport Import in RootObject.ObjectImports)
+			{
+				if (VisitedObjects.Add(Import.Object.Hash))
+				{
+					TreePackObject Object = await StorageClient.ReadObjectAsync<TreePackObject>(NamespaceId, Import.Object.Hash);
+					TreePack.WriteSummary(FileReference.Combine(TargetDir, $"{NamespaceId}-{Import.Object.Hash}.yml"), Object);
+					await WriteSummaryFiles(NamespaceId, Object, VisitedObjects, StorageClient, TargetDir);
+				}
+			}
+		}
+
+		public async Task ReadFilesFromPerforceAsync(string SourcePath, TreePack TargetPack, TreePackDirWriter TargetWriter)
+		{
+			PerforceSettings Settings = new PerforceSettings(PerforceSettings.Default);
+			Settings.PreferNativeClient = true;
+
+			using IPerforceConnection Perforce = await PerforceConnection.CreateAsync(Settings, Log.Logger);
+			ClientRecord Client = await Perforce.GetClientAsync(null);
+
+			Dictionary<int, (string Path, TreePackFileWriter Writer)> Files = new Dictionary<int, (string, TreePackFileWriter)>();
+			await foreach (PerforceResponse Response in Perforce.StreamCommandAsync("sync", new string[] { "-f" }, new string[] { $"{SourcePath}#have" }, null, typeof(SyncRecord), true, default))
+			{
+				PerforceIo? Io = Response.Io;
+				if (Io != null)
+				{
+					if (Io.Command == PerforceIoCommand.Open)
+					{
+						string Path = GetClientRelativePath(Io.Payload, Client.Root);
+						Files[Io.File] = (Path, new TreePackFileWriter(TargetPack));
+					}
+					else if (Io.Command == PerforceIoCommand.Write)
+					{
+						TreePackFileWriter FileWriter = Files[Io.File].Writer;
+						await FileWriter.WriteAsync(Io.Payload, false);
+					}
+					else if (Io.Command == PerforceIoCommand.Close)
+					{
+						(string Path, TreePackFileWriter FileWriter) = Files[Io.File];
+						IoHash FileHash = await FileWriter.FinalizeAsync();
+						await TargetWriter.FindOrAddFileByPathAsync(Path, TreePackDirEntryFlags.File, FileHash, Sha1Hash.Zero);
+					}
+				}
+			}
+		}
+
+		static string GetClientRelativePath(ReadOnlyMemory<byte> Data, string ClientRoot)
+		{
+			ReadOnlySpan<byte> Span = Data.Span;
+
+			string Path = Encoding.UTF8.GetString(Span.Slice(0, Span.IndexOf((byte)0)));
+			if (!Path.StartsWith(ClientRoot))
+			{
+				throw new ArgumentException($"Unable to make path {Path} relative to client root {ClientRoot}");
+			}
+
+			return Path.Substring(ClientRoot.Length).Replace('\\', '/').TrimStart('/');
+		}
+
+		public async Task ReadFilesFromDiskAsync(string SourcePath, TreePack TargetPack, TreePackDirWriter TargetWriter)
+		{
+			int MaxIdx = FileFilter.FindWildcardIndex(SourcePath);
+			if (MaxIdx < 0)
+			{
+				MaxIdx = SourcePath.Length - 1;
+			}
+			while (SourcePath[MaxIdx] != Path.DirectorySeparatorChar)
+			{
+				MaxIdx--;
+			}
+
+			DirectoryReference BaseDir = new DirectoryReference(SourcePath.Substring(0, MaxIdx));
+			string SearchPath = SourcePath.Substring(MaxIdx + 1);
+
+			foreach (FileReference File in FileFilter.ResolveWildcard(BaseDir, SearchPath))
+			{
+				TreePackFileWriter FileWriter = new TreePackFileWriter(TargetPack);
+
+				byte[] Buffer = new byte[65536];
+				using (FileStream Stream = FileReference.Open(File, FileMode.Open, FileAccess.Read))
+				{
+					for (; ; )
+					{
+						int ReadBytes = await Stream.ReadAsync(Buffer);
+						if (ReadBytes == 0)
+						{
+							break;
+						}
+						await FileWriter.WriteAsync(Buffer.AsMemory(0, ReadBytes), false);
+					}
+				}
+
+				IoHash Hash = await FileWriter.FinalizeAsync();
+
+				string Path = File.MakeRelativeTo(BaseDir).Replace('\\', '/');
+				await TargetWriter.FindOrAddFileByPathAsync(Path, TreePackDirEntryFlags.File, Hash);
+			}
+		}
+	}
+}
diff --git a/Engine/Source/Programs/Horde/Horde.Agent/Commands/Archive/ExtractArchiveCommand.cs b/Engine/Source/Programs/Horde/Horde.Agent/Commands/Archive/ExtractArchiveCommand.cs
new file mode 100644
--- /dev/null
+++ b/Engine/Source/Programs/Horde/Horde.Agent/Commands/Archive/ExtractArchiveCommand.cs
@@ -0,0 +1,256 @@
+// Copyright Epic Games, Inc. All Rights Reserved.
+
+using Google.Protobuf;
+using Microsoft.Extensions.Logging;
+using System;
+using System.Collections.Concurrent;
+using System.Collections.Generic;
+using System.IO;
+using System.Linq;
+using System.Text;
+using System.Threading.Tasks;
+using EpicGames.Core;
+using System.Text.Json;
+using HordeAgent.Services;
+using Grpc.Net.Client;
+using System.Net.Http.Headers;
+using Grpc.Core;
+using Microsoft.Extensions.Hosting;
+using Microsoft.Extensions.DependencyInjection;
+using Microsoft.Extensions.Configuration;
+using Serilog.Events;
+using EpicGames.Serialization;
+using EpicGames.Horde.Compute;
+using EpicGames.Horde.Storage;
+using System.Net.Http;
+using EpicGames.Horde.Compute.Impl;
+using EpicGames.Horde.Storage.Impl;
+using System.Threading;
+using EpicGames.Perforce;
+using System.Diagnostics;
+
+namespace HordeAgent.Commands.Archive
+{
+	/// <summary>
+	/// Installs the agent as a service
+	/// </summary>
+	[Command("Archive", "Extract", "Creates an archive containing a set of files or folders")]
+	class ExtractArchiveCommand : Command
+	{
+		/// <summary>
+		/// Input path.
+		/// </summary>
+		[CommandLine("-Input=", Required = true)]
+		public FileReference InputPath = null!;
+
+		/// <summary>
+		/// Output path to an index file.
+		/// </summary>
+		[CommandLine("-Output=", Required = true)]
+		public DirectoryReference TargetDir = null!;
+
+		/// <inheritdoc/>
+		public override async Task<int> ExecuteAsync(ILogger Logger)
+		{
+			int TargetSize = 64 * 1024;
+			int MaxValue = (int)((1L << 32) / TargetSize);
+
+			List<int> Sizes = new List<int>();
+			byte[] Data = File.ReadAllBytes(@"E:\P4 UE5\Engine\Source\ThirdParty\CEF3\cef_binary_3.2623.1395.g3034273_linux64\Release\libcef.so");
+			BuzHash Hash = new BuzHash();
+
+			int LastStart = 0;
+			for (int Idx = 0; Idx < Data.Length; Idx++)
+			{
+				Hash.add(Data[Idx]);
+				if (Hash.get() < MaxValue)
+				{
+					Sizes.Add((Idx + 1) - LastStart);
+					Hash.reset();
+					LastStart = Idx + 1;
+				}
+			}
+
+			double Average = Sizes.Average();
+
+
+			TreePackOptions Options = new TreePackOptions();
+			Options.MaxBlobSize = 64 * 1024;
+			Options.MaxInlineBlobSize = 20 * 1024;
+
+			DirectoryReference Inputdir = InputPath.Directory;
+			NamespaceId NamespaceId = new NamespaceId("test");
+			FileSystemStorageClient Storage = new FileSystemStorageClient(Inputdir, Logger);
+
+			TreePack TreePack = new TreePack(Storage, NamespaceId, Options);
+
+			IoHash RootHash = IoHash.Parse((await FileReference.ReadAllTextAsync(InputPath)).Trim());
+			byte[] RootObject = await FileReference.ReadAllBytesAsync(InputPath.ChangeExtension(".dat"));
+			TreePack.AddRootBlob(RootObject);
+
+			await ExtractFilesAsync(TreePack, RootHash, TargetDir, Logger);
+			return 0;
+		}
+
+		async Task ExtractFilesAsync(TreePack TreePack, IoHash Hash, DirectoryReference TargetDir, ILogger Logger)
+		{
+			Logger.LogInformation("Writing directory {TargetDir}...", TargetDir);
+			DirectoryReference.CreateDirectory(TargetDir);
+
+			ReadOnlyMemory<byte> Data = await TreePack.GetDataAsync(Hash);
+			Debug.Assert(Data.Span[0] == (byte)TreePackNodeType.Directory);
+
+			TreePackDirNode Dir = TreePackDirNode.Parse(Data);
+			foreach (TreePackDirEntry Entry in Dir.Entries)
+			{
+				if ((Entry.Flags & TreePackDirEntryFlags.Directory) != 0)
+				{
+					DirectoryReference TargetSubDir = DirectoryReference.Combine(TargetDir, Entry.Name.ToString());
+					await ExtractFilesAsync(TreePack, Entry.Hash, TargetSubDir, Logger);
+				}
+				else if ((Entry.Flags & TreePackDirEntryFlags.File) != 0)
+				{
+					FileReference File = FileReference.Combine(TargetDir, Entry.Name.ToString());
+					Logger.LogInformation("Writing {File}...", File);
+
+					using (FileStream Stream = FileReference.Open(File, FileMode.Create, FileAccess.Write, FileShare.Read))
+					{
+						await ExtractFileAsync(TreePack, Entry.Hash, Stream, Logger);
+					}
+				}
+				else
+				{
+					throw new InvalidDataException($"Unknown entry type for {Entry.Name}");
+				}
+			}
+		}
+
+		async Task ExtractFileAsync(TreePack TreePack, IoHash Hash, Stream OutputStream, ILogger Logger)
+		{
+			ReadOnlyMemory<byte> Data = await TreePack.GetDataAsync(Hash);
+			TreePackNodeType Type = (TreePackNodeType)Data.Span[0];
+
+			if (Type == TreePackNodeType.Binary)
+			{
+				Logger.LogInformation("Copying data from binary node {Hash}", Hash);
+				await OutputStream.WriteAsync(Data.Slice(1));
+			}
+			else if (Type == TreePackNodeType.Concat)
+			{
+				Logger.LogInformation("Start data from concat node {Hash}", Hash);
+
+				TreePackConcatNode Concat = TreePackConcatNode.Parse(Data.Span);
+				foreach (TreePackConcatEntry Entry in Concat.Entries)
+				{
+					await ExtractFileAsync(TreePack, Entry.Hash, OutputStream, Logger);
+				}
+
+				Logger.LogInformation("Finish data from concat node {Hash}", Hash);
+			}
+			else
+			{
+				throw new InvalidDataException();
+			}
+		}
+
+		async Task WriteSummaryFiles(NamespaceId NamespaceId, TreePackObject RootObject, HashSet<IoHash> VisitedObjects, IBlobStorageClient StorageClient, DirectoryReference TargetDir)
+		{
+			foreach (TreePackObjectImport Import in RootObject.ObjectImports)
+			{
+				if (VisitedObjects.Add(Import.Object.Hash))
+				{
+					TreePackObject Object = await StorageClient.ReadObjectAsync<TreePackObject>(NamespaceId, Import.Object.Hash);
+					TreePack.WriteSummary(FileReference.Combine(TargetDir, $"{NamespaceId}-{Import.Object.Hash}.yml"), Object);
+					await WriteSummaryFiles(NamespaceId, Object, VisitedObjects, StorageClient, TargetDir);
+				}
+			}
+		}
+
+		public async Task ReadFilesFromPerforceAsync(string SourcePath, TreePack TargetPack, TreePackDirWriter TargetWriter)
+		{
+			PerforceSettings Settings = new PerforceSettings(PerforceSettings.Default);
+			Settings.PreferNativeClient = true;
+
+			using IPerforceConnection Perforce = await PerforceConnection.CreateAsync(Settings, Log.Logger);
+			ClientRecord Client = await Perforce.GetClientAsync(null);
+
+			Dictionary<int, (string Path, TreePackFileWriter Writer)> Files = new Dictionary<int, (string, TreePackFileWriter)>();
+			await foreach (PerforceResponse Response in Perforce.StreamCommandAsync("sync", new string[] { "-f" }, new string[] { $"{SourcePath}#have" }, null, typeof(SyncRecord), true, default))
+			{
+				PerforceIo? Io = Response.Io;
+				if (Io != null)
+				{
+					if (Io.Command == PerforceIoCommand.Open)
+					{
+						string Path = GetClientRelativePath(Io.Payload, Client.Root);
+						Files[Io.File] = (Path, new TreePackFileWriter(TargetPack));
+					}
+					else if (Io.Command == PerforceIoCommand.Write)
+					{
+						TreePackFileWriter FileWriter = Files[Io.File].Writer;
+						await FileWriter.WriteAsync(Io.Payload, false);
+					}
+					else if (Io.Command == PerforceIoCommand.Close)
+					{
+						(string Path, TreePackFileWriter FileWriter) = Files[Io.File];
+						IoHash FileHash = await FileWriter.FinalizeAsync();
+						await TargetWriter.FindOrAddFileByPathAsync(Path, TreePackDirEntryFlags.File, FileHash, Sha1Hash.Zero);
+					}
+				}
+			}
+		}
+
+		static string GetClientRelativePath(ReadOnlyMemory<byte> Data, string ClientRoot)
+		{
+			ReadOnlySpan<byte> Span = Data.Span;
+
+			string Path = Encoding.UTF8.GetString(Span.Slice(0, Span.IndexOf((byte)0)));
+			if (!Path.StartsWith(ClientRoot))
+			{
+				throw new ArgumentException($"Unable to make path {Path} relative to client root {ClientRoot}");
+			}
+
+			return Path.Substring(ClientRoot.Length).Replace('\\', '/').TrimStart('/');
+		}
+
+		public async Task ReadFilesFromDiskAsync(string SourcePath, TreePack TargetPack, TreePackDirWriter TargetWriter)
+		{
+			int MaxIdx = FileFilter.FindWildcardIndex(SourcePath);
+			if (MaxIdx < 0)
+			{
+				MaxIdx = SourcePath.Length - 1;
+			}
+			while (SourcePath[MaxIdx] != Path.DirectorySeparatorChar)
+			{
+				MaxIdx--;
+			}
+
+			DirectoryReference BaseDir = new DirectoryReference(SourcePath.Substring(0, MaxIdx));
+			string SearchPath = SourcePath.Substring(MaxIdx + 1);
+
+			foreach (FileReference File in FileFilter.ResolveWildcard(BaseDir, SearchPath))
+			{
+				TreePackFileWriter FileWriter = new TreePackFileWriter(TargetPack);
+
+				byte[] Buffer = new byte[65536];
+				using (FileStream Stream = FileReference.Open(File, FileMode.Open, FileAccess.Read))
+				{
+					for (; ; )
+					{
+						int ReadBytes = await Stream.ReadAsync(Buffer);
+						if (ReadBytes == 0)
+						{
+							break;
+						}
+						await FileWriter.WriteAsync(Buffer.AsMemory(0, ReadBytes), false);
+					}
+				}
+
+				IoHash Hash = await FileWriter.FinalizeAsync();
+
+				string Path = File.MakeRelativeTo(BaseDir).Replace('\\', '/');
+				await TargetWriter.FindOrAddFileByPathAsync(Path, TreePackDirEntryFlags.File, Hash);
+			}
+		}
+	}
+}
diff --git a/Engine/Source/Programs/Shared/EpicGames.Core/BuzHash.cs b/Engine/Source/Programs/Shared/EpicGames.Core/BuzHash.cs
new file mode 100644
--- /dev/null
+++ b/Engine/Source/Programs/Shared/EpicGames.Core/BuzHash.cs
@@ -0,0 +1,101 @@
+// Copyright Epic Games, Inc. All Rights Reserved.
+
+using System;
+using System.Collections.Generic;
+using System.Text;
+
+namespace EpicGames.Core
+{
+	/// <summary>
+	/// Rolling hash by cyclic polynomial
+	/// </summary>
+	public class BuzHash
+	{
+		static uint[] Table = new uint[256]
+		{
+			0x458be752, 0xc10748cc, 0xfbbcdbb8, 0x6ded5b68, 0xb10a82b5, 0x20d75648, 0xdfc5665f, 0xa8428801, 0x7ebf5191, 0x841135c7, 0x65cc53b3,
+			0x280a597c, 0x16f60255, 0xc78cbc3e, 0x294415f5, 0xb938d494, 0xec85c4e6, 0xb7d33edc, 0xe549b544, 0xfdeda5aa, 0x882bf287, 0x3116737c,
+			0x05569956, 0xe8cc1f68, 0x0806ac5e, 0x22a14443, 0x15297e10, 0x50d090e7, 0x4ba60f6f, 0xefd9f1a7, 0x5c5c885c, 0x82482f93, 0x9bfd7c64,
+			0x0b3e7276, 0xf2688e77, 0x8fad8abc, 0xb0509568, 0xf1ada29f, 0xa53efdfe, 0xcb2b1d00, 0xf2a9e986, 0x6463432b, 0x95094051, 0x5a223ad2,
+			0x9be8401b, 0x61e579cb, 0x1a556a14, 0x5840fdc2, 0x9261ddf6, 0xcde002bb, 0x52432bb0, 0xbf17373e, 0x7b7c222f, 0x2955ed16, 0x9f10ca59,
+			0xe840c4c9, 0xccabd806, 0x14543f34, 0x1462417a, 0x0d4a1f9c, 0x087ed925, 0xd7f8f24c, 0x7338c425, 0xcf86c8f5, 0xb19165cd, 0x9891c393,
+			0x325384ac, 0x0308459d, 0x86141d7e, 0xc922116a, 0xe2ffa6b6, 0x53f52aed, 0x2cd86197, 0xf5b9f498, 0xbf319c8f, 0xe0411fae, 0x977eb18c,
+			0xd8770976, 0x9833466a, 0xc674df7f, 0x8c297d45, 0x8ca48d26, 0xc49ed8e2, 0x7344f874, 0x556f79c7, 0x6b25eaed, 0xa03e2b42, 0xf68f66a4,
+			0x8e8b09a2, 0xf2e0e62a, 0x0d3a9806, 0x9729e493, 0x8c72b0fc, 0x160b94f6, 0x450e4d3d, 0x7a320e85, 0xbef8f0e1, 0x21d73653, 0x4e3d977a,
+			0x1e7b3929, 0x1cc6c719, 0xbe478d53, 0x8d752809, 0xe6d8c2c6, 0x275f0892, 0xc8acc273, 0x4cc21580, 0xecc4a617, 0xf5f7be70, 0xe795248a,
+			0x375a2fe9, 0x425570b6, 0x8898dcf8, 0xdc2d97c4, 0x0106114b, 0x364dc22f, 0x1e0cad1f, 0xbe63803c, 0x5f69fac2, 0x4d5afa6f, 0x1bc0dfb5,
+			0xfb273589, 0x0ea47f7b, 0x3c1c2b50, 0x21b2a932, 0x6b1223fd, 0x2fe706a8, 0xf9bd6ce2, 0xa268e64e, 0xe987f486, 0x3eacf563, 0x1ca2018c,
+			0x65e18228, 0x2207360a, 0x57cf1715, 0x34c37d2b, 0x1f8f3cde, 0x93b657cf, 0x31a019fd, 0xe69eb729, 0x8bca7b9b, 0x4c9d5bed, 0x277ebeaf,
+			0xe0d8f8ae, 0xd150821c, 0x31381871, 0xafc3f1b0, 0x927db328, 0xe95effac, 0x305a47bd, 0x426ba35b, 0x1233af3f, 0x686a5b83, 0x50e072e5,
+			0xd9d3bb2a, 0x8befc475, 0x487f0de6, 0xc88dff89, 0xbd664d5e, 0x971b5d18, 0x63b14847, 0xd7d3c1ce, 0x7f583cf3, 0x72cbcb09, 0xc0d0a81c,
+			0x7fa3429b, 0xe9158a1b, 0x225ea19a, 0xd8ca9ea3, 0xc763b282, 0xbb0c6341, 0x020b8293, 0xd4cd299d, 0x58cfa7f8, 0x91b4ee53, 0x37e4d140,
+			0x95ec764c, 0x30f76b06, 0x5ee68d24, 0x679c8661, 0xa41979c2, 0xf2b61284, 0x4fac1475, 0x0adb49f9, 0x19727a23, 0x15a7e374, 0xc43a18d5,
+			0x3fb1aa73, 0x342fc615, 0x924c0793, 0xbee2d7f0, 0x8a279de9, 0x4aa2d70c, 0xe24dd37f, 0xbe862c0b, 0x177c22c2, 0x5388e5ee, 0xcd8a7510,
+			0xf901b4fd, 0xdbc13dbc, 0x6c0bae5b, 0x64efe8c7, 0x48b02079, 0x80331a49, 0xca3d8ae6, 0xf3546190, 0xfed7108b, 0xc49b941b, 0x32baf4a9,
+			0xeb833a4a, 0x88a3f1a5, 0x3a91ce0a, 0x3cc27da1, 0x7112e684, 0x4a3096b1, 0x3794574c, 0xa3c8b6f3, 0x1d213941, 0x6e0a2e00, 0x233479f1,
+			0x0f4cd82f, 0x6093edd2, 0x5d7d209e, 0x464fe319, 0xd4dcac9e, 0x0db845cb, 0xfb5e4bc3, 0xe0256ce1, 0x09fb4ed1, 0x0914be1e, 0xa5bdb2c3,
+			0xc6eb57bb, 0x30320350, 0x3f397e91, 0xa67791bc, 0x86bc0e2c, 0xefa0a7e2, 0xe9ff7543, 0xe733612c, 0xd185897b, 0x329e5388, 0x91dd236b,
+			0x2ecb0d93, 0xf4d82a3d, 0x35b5c03f, 0xe4e606f0, 0x05b21843, 0x37b45964, 0x5eff22f4, 0x6027f4cc, 0x77178b3c, 0xae507131, 0x7bf7cabc,
+			0xf9c18d66, 0x593ade65, 0xd95ddf11,
+		};
+
+		ulong Count = 0;
+		uint State = 0;
+
+		/// <summary>
+		/// Resets the state of the 
+		/// </summary>
+		public void Reset()
+		{ 
+			Count = 0; 
+			State = 0; 
+		}
+
+		/// <summary>
+		/// Add a byte to the window and update the hash
+		/// </summary>
+		/// <param name="Value">New value to append to the window</param>
+		public void Add(byte Value)
+		{
+			State = Rol32(State, 1) ^ Table[Value];
+			Count++;
+		}
+
+		/// <summary>
+		/// Removes a byte from the start of the hash window
+		/// </summary>
+		/// <param name="Value">Value at the start of the window</param>
+		public void Sub(byte Value)
+		{
+			State = State ^ Rol32(Table[Value], (int)(Count - 1));
+			Count--;
+		}
+
+		/// <summary>
+		/// Update the hash with the given data
+		/// </summary>
+		/// <param name="Data"></param>
+		public void Update(ReadOnlySpan<byte> Data)
+		{
+			for (int Idx = 0; Idx < Data.Length; ++Idx)
+			{
+				Add(Data[Idx]);
+			}
+		}
+
+		/// <summary>
+		/// Gets the current hash value
+		/// </summary>
+		/// <returns>The current hash value</returns>
+		public uint Get()
+		{
+			return State;
+		}
+
+		static uint Rol32(uint V, int N)
+		{
+			N &= 31;
+			return ((V) << (N)) | ((V) >> (32 - N));
+		}
+	}
+}
diff --git a/Engine/Source/Programs/Shared/EpicGames.Horde/EpicGames.Horde.csproj b/Engine/Source/Programs/Shared/EpicGames.Horde/EpicGames.Horde.csproj
--- a/Engine/Source/Programs/Shared/EpicGames.Horde/EpicGames.Horde.csproj
+++ b/Engine/Source/Programs/Shared/EpicGames.Horde/EpicGames.Horde.csproj
@@ -28,6 +28,7 @@
     </Compile>
   </ItemGroup>
   <ItemGroup>
+    <PackageReference Include="K4os.Compression.LZ4" Version="1.2.6" />
     <PackageReference Include="Microsoft.CodeAnalysis.CSharp.Workspaces" Version="3.10.0" />
     <PackageReference Include="Microsoft.Extensions.DependencyInjection" Version="6.0.0" />
     <PackageReference Include="Microsoft.Extensions.Http" Version="3.1.22" />
diff --git a/Engine/Source/Programs/Shared/EpicGames.Horde/Storage/IStorageClient.cs b/Engine/Source/Programs/Shared/EpicGames.Horde/Storage/IStorageClient.cs
--- a/Engine/Source/Programs/Shared/EpicGames.Horde/Storage/IStorageClient.cs
+++ b/Engine/Source/Programs/Shared/EpicGames.Horde/Storage/IStorageClient.cs
@@ -152,12 +152,10 @@ namespace EpicGames.Horde.Storage
 	}
 
 	/// <summary>
-	/// Interface for a collection of ref documents
+	/// Base interface for a storage client that only records blobs.
 	/// </summary>
-	public interface IStorageClient
+	public interface IBlobStorageClient
 	{
-		#region Blobs
-
 		/// <summary>
 		/// Opens a blob read stream
 		/// </summary>
@@ -175,11 +173,13 @@ namespace EpicGames.Horde.Storage
 		/// <param name="Stream">The stream to write</param>
 		/// <param name="CancellationToken">Cancellation token for the operation</param>
 		Task WriteBlobAsync(NamespaceId NamespaceId, IoHash Hash, Stream Stream, CancellationToken CancellationToken = default);
+	}
 
-		#endregion
-
-		#region Refs
-
+	/// <summary>
+	/// Interface for a collection of ref documents
+	/// </summary>
+	public interface IStorageClient : IBlobStorageClient
+	{
 		/// <summary>
 		/// Gets the given reference
 		/// </summary>
@@ -241,8 +241,6 @@ namespace EpicGames.Horde.Storage
 		/// <param name="CancellationToken">Cancellation token for the operation</param>
 		/// <returns>True if the ref was deleted, false if it did not exist</returns>
 		Task<bool> DeleteRefAsync(NamespaceId NamespaceId, BucketId BucketId, RefId RefId, CancellationToken CancellationToken = default);
-
-		#endregion
 	}
 
 	/// <summary>
@@ -260,7 +258,7 @@ namespace EpicGames.Horde.Storage
 		/// <param name="Hash">Hash of the blob to read</param>
 		/// <param name="MaxInMemoryBlobLength">Maximum allowed memory allocation to store the blob</param>
 		/// <returns>Data for the blob that was read. Throws an exception if the blob was not found.</returns>
-		public static async Task<byte[]> ReadBlobToMemoryAsync(this IStorageClient StorageClient, NamespaceId NamespaceId, IoHash Hash, int MaxInMemoryBlobLength = DefaultMaxInMemoryBlobLength)
+		public static async Task<byte[]> ReadBlobToMemoryAsync(this IBlobStorageClient StorageClient, NamespaceId NamespaceId, IoHash Hash, int MaxInMemoryBlobLength = DefaultMaxInMemoryBlobLength)
 		{
 			using Stream Stream = await StorageClient.ReadBlobAsync(NamespaceId, Hash);
 
@@ -292,7 +290,7 @@ namespace EpicGames.Horde.Storage
 		/// <param name="Hash">Hash of the blob to read</param>
 		/// <param name="MaxInMemoryBlobLength">Maximum allowed memory allocation to store the blob</param>
 		/// <returns>The decoded object</returns>
-		public static async Task<CbObject> ReadObjectAsync(this IStorageClient StorageClient, NamespaceId NamespaceId, IoHash Hash, int MaxInMemoryBlobLength = DefaultMaxInMemoryBlobLength)
+		public static async Task<CbObject> ReadObjectAsync(this IBlobStorageClient StorageClient, NamespaceId NamespaceId, IoHash Hash, int MaxInMemoryBlobLength = DefaultMaxInMemoryBlobLength)
 		{
 			return new CbObject(await ReadBlobToMemoryAsync(StorageClient, NamespaceId, Hash, MaxInMemoryBlobLength));
 		}
@@ -306,7 +304,7 @@ namespace EpicGames.Horde.Storage
 		/// <param name="Hash">Hash of the blob to read</param>
 		/// <param name="MaxInMemoryBlobLength">Maximum allowed memory allocation to store the blob</param>
 		/// <returns>The decoded object</returns>
-		public static async Task<T> ReadObjectAsync<T>(this IStorageClient StorageClient, NamespaceId NamespaceId, IoHash Hash, int MaxInMemoryBlobLength = DefaultMaxInMemoryBlobLength)
+		public static async Task<T> ReadObjectAsync<T>(this IBlobStorageClient StorageClient, NamespaceId NamespaceId, IoHash Hash, int MaxInMemoryBlobLength = DefaultMaxInMemoryBlobLength)
 		{
 			return CbSerializer.Deserialize<T>(await ReadBlobToMemoryAsync(StorageClient, NamespaceId, Hash, MaxInMemoryBlobLength));
 		}
@@ -318,7 +316,7 @@ namespace EpicGames.Horde.Storage
 		/// <param name="NamespaceId">Namespace containing the blob</param>
 		/// <param name="Data">Data to write</param>
 		/// <returns></returns>
-		public static async Task<IoHash> WriteBlobFromMemoryAsync(this IStorageClient StorageClient, NamespaceId NamespaceId, ReadOnlyMemory<byte> Data)
+		public static async Task<IoHash> WriteBlobFromMemoryAsync(this IBlobStorageClient StorageClient, NamespaceId NamespaceId, ReadOnlyMemory<byte> Data)
 		{
 			IoHash Hash = IoHash.Compute(Data.Span);
 			await WriteBlobFromMemoryAsync(StorageClient, NamespaceId, Hash, Data);
@@ -332,7 +330,7 @@ namespace EpicGames.Horde.Storage
 		/// <param name="NamespaceId">Namespace containing the blob</param>
 		/// <param name="Object">The object to be written</param>
 		/// <returns></returns>
-		public static Task<IoHash> WriteObjectAsync(this IStorageClient StorageClient, NamespaceId NamespaceId, CbObject Object)
+		public static Task<IoHash> WriteObjectAsync(this IBlobStorageClient StorageClient, NamespaceId NamespaceId, CbObject Object)
 		{
 			return WriteBlobFromMemoryAsync(StorageClient, NamespaceId, Object.GetView());
 		}
@@ -344,7 +342,7 @@ namespace EpicGames.Horde.Storage
 		/// <param name="NamespaceId">Namespace containing the blob</param>
 		/// <param name="Object">The object to be written</param>
 		/// <returns></returns>
-		public static Task<IoHash> WriteObjectAsync<T>(this IStorageClient StorageClient, NamespaceId NamespaceId, T Object)
+		public static Task<IoHash> WriteObjectAsync<T>(this IBlobStorageClient StorageClient, NamespaceId NamespaceId, T Object)
 		{
 			return WriteObjectAsync(StorageClient, NamespaceId, CbSerializer.Serialize<T>(Object));
 		}
@@ -357,7 +355,7 @@ namespace EpicGames.Horde.Storage
 		/// <param name="Hash">Hash of the data</param>
 		/// <param name="Data">The data to be written</param>
 		/// <returns></returns>
-		public static async Task WriteBlobFromMemoryAsync(this IStorageClient StorageClient, NamespaceId NamespaceId, IoHash Hash, ReadOnlyMemory<byte> Data)
+		public static async Task WriteBlobFromMemoryAsync(this IBlobStorageClient StorageClient, NamespaceId NamespaceId, IoHash Hash, ReadOnlyMemory<byte> Data)
 		{
 			using ReadOnlyMemoryStream Stream = new ReadOnlyMemoryStream(Data);
 			await StorageClient.WriteBlobAsync(NamespaceId, Hash, Stream);
diff --git a/Engine/Source/Programs/Shared/EpicGames.Horde/Storage/TreePack.cs b/Engine/Source/Programs/Shared/EpicGames.Horde/Storage/TreePack.cs
new file mode 100644
--- /dev/null
+++ b/Engine/Source/Programs/Shared/EpicGames.Horde/Storage/TreePack.cs
@@ -0,0 +1,1237 @@
+// Copyright Epic Games, Inc. All Rights Reserved.
+
+using EpicGames.Core;
+using EpicGames.Serialization;
+using K4os.Compression.LZ4;
+using Microsoft.Extensions.Logging;
+using System;
+using System.Buffers.Binary;
+using System.Collections;
+using System.Collections.Generic;
+using System.Diagnostics;
+using System.Diagnostics.CodeAnalysis;
+using System.IO;
+using System.IO.Compression;
+using System.Linq;
+using System.Security.Cryptography;
+using System.Text;
+using System.Text.Json;
+using System.Threading.Tasks;
+
+namespace EpicGames.Horde.Storage
+{
+	#region Storage primitives
+
+	/// <summary>
+	/// Blob within the tree pack
+	/// </summary>
+	public class TreePackObject
+	{
+		/// <summary>
+		/// Time that this object was minted. Used to determine how likely it will be that a client will already have it.
+		/// </summary>
+		[CbField("time")]
+		public DateTime Time { get; set; }
+
+		/// <summary>
+		/// Imports from object attachments.
+		/// </summary>
+		[CbField("objimp")]
+		public List<TreePackObjectImport> ObjectImports { get; set; } = new List<TreePackObjectImport>();
+
+		/// <summary>
+		/// Imports from binary attachments.
+		/// </summary>
+		[CbField("binimp")]
+		public List<TreePackBinaryImport> BinaryImports { get; set; } = new List<TreePackBinaryImport>();
+
+		/// <summary>
+		/// Exported items in the data 
+		/// </summary>
+		[CbField("exports")]
+		public ReadOnlyMemory<byte> Exports { get; set; }
+
+		/// <summary>
+		/// Raw data for this exported entries in this blob.
+		/// </summary>
+		[CbField("data")]
+		public ReadOnlyMemory<byte> Data { get; set; }
+	}
+
+	/// <summary>
+	/// Reference to another tree pack object
+	/// </summary>
+	public class TreePackObjectImport
+	{
+		/// <summary>
+		/// Imported object. The <see cref="Nodes"/> field contains a flat array of hashes imported from the object.
+		/// </summary>
+		[CbField("object")]
+		public CbObjectAttachment Object { get; set; }
+
+		/// <summary>
+		/// List of keys imported from this object.
+		/// </summary>
+		[CbField("keys")]
+		public ReadOnlyMemory<byte> Keys { get; set; }
+	}
+
+	/// <summary>
+	/// Reference to another tree pack object
+	/// </summary>
+	public class TreePackBinaryImport
+	{
+		/// <summary>
+		/// Imported binary. Binary objects do not have a compact binary header, and effectively consist solely of the <see cref="TreePack.Data"/> field. 
+		/// The <see cref="Nodes"/> field contains encoded a list of encoded entry objects.
+		/// </summary>
+		[CbField("binary")]
+		public CbBinaryAttachment Binary { get; set; }
+
+		/// <summary>
+		/// Time that the binary was minted
+		/// </summary>
+		[CbField("time")]
+		public DateTime Time { get; set; }
+
+		/// <summary>
+		/// Length of the imported binary.
+		/// </summary>
+		[CbField("length")]
+		public int Length { get; set; }
+
+		/// <summary>
+		/// Packed list of exports taken from this binary
+		/// </summary>
+		[CbField("exports")]
+		public ReadOnlyMemory<byte> Exports { get; set; }
+	}
+
+	#endregion
+
+	#region Node types
+
+	/// <summary>
+	/// Type of a node in a tree pack. The first byte of any node should be this value.
+	/// </summary>
+	public enum TreePackNodeType : byte
+	{
+		/// <summary>
+		/// Raw binary data
+		/// </summary>
+		Binary,
+
+		/// <summary>
+		/// A compact binary encoded object
+		/// </summary>
+		Object,
+
+		/// <summary>
+		/// A directory listing
+		/// </summary>
+		Directory,
+
+		/// <summary>
+		/// A concat node
+		/// </summary>
+		Concat,
+	}
+
+	/// <summary>
+	/// Flags for a directory entry
+	/// </summary>
+	public enum TreePackDirEntryFlags : byte
+	{
+		/// <summary>
+		/// This item is a directory
+		/// </summary>
+		Directory = 1,
+
+		/// <summary>
+		/// This item is a file
+		/// </summary>
+		File = 2,
+
+		/// <summary>
+		/// Indicates that the referenced file is executable
+		/// </summary>
+		Executable = 4,
+
+		/// <summary>
+		/// File should be stored as read-only
+		/// </summary>
+		ReadOnly = 8,
+
+		/// <summary>
+		/// File contents are utf-8 encoded text. Client may want to replace line-endings with OS-specific format.
+		/// </summary>
+		Text = 16,
+
+		/// <summary>
+		/// The attached entry includes a Git SHA1 of the corresponding blob/tree contents.
+		/// </summary>
+		HasGitSha1 = 32,
+	}
+
+	/// <summary>
+	/// Entry within a directory
+	/// </summary>
+	public class TreePackDirEntry
+	{
+		/// <summary>
+		/// Flags for this entry
+		/// </summary>
+		public TreePackDirEntryFlags Flags { get; }
+
+		/// <summary>
+		/// Name of the entry
+		/// </summary>
+		public Utf8String Name { get; }
+
+		/// <summary>
+		/// Hash of this node's contents.
+		/// </summary>
+		public IoHash Hash { get; }
+
+		/// <summary>
+		/// Hash of the corresponding Git object. Only included if <see cref="Flags"/> specifies the <see cref="TreePackDirEntryFlags.HasGitSha1"/> flag.
+		/// </summary>
+		public Sha1Hash GitSha1 { get; }
+
+		/// <summary>
+		/// Constructor
+		/// </summary>
+		public TreePackDirEntry(TreePackDirEntryFlags Flags, Utf8String Name, IoHash Hash, Sha1Hash GitSha1)
+		{
+			this.Flags = Flags;
+			this.Name = Name;
+			this.Hash = Hash;
+			this.GitSha1 = GitSha1;
+		}
+	}
+
+	/// <summary>
+	/// A directory node
+	/// </summary>
+	public class TreePackDirNode
+	{
+		/// <summary>
+		/// Entries for this directory
+		/// </summary>
+		public List<TreePackDirEntry> Entries { get; } = new List<TreePackDirEntry>();
+
+		/// <summary>
+		/// Parses a node from a block of memory
+		/// </summary>
+		public static TreePackDirNode Parse(ReadOnlyMemory<byte> Data)
+		{
+			ReadOnlySpan<byte> Span = Data.Span;
+
+			Debug.Assert(Span[0] == (byte)TreePackNodeType.Directory);
+			Span = Span.Slice(1);
+
+			TreePackDirNode Node = new TreePackDirNode();
+			while (Span.Length > 0)
+			{
+				TreePackDirEntryFlags Flags = (TreePackDirEntryFlags)Span[0];
+				Span = Span.Slice(1);
+
+				int NameLength = Span.IndexOf((byte)0);
+				Utf8String Name = new Utf8String(Span.Slice(0, NameLength).ToArray());
+				Span = Span.Slice(NameLength + 1);
+
+				IoHash Hash = new IoHash(Span);
+				Span = Span.Slice(IoHash.NumBytes);
+
+				Sha1Hash GitSha1 = Sha1Hash.Zero;
+				if ((Flags & TreePackDirEntryFlags.HasGitSha1) != 0)
+				{
+					GitSha1 = new Sha1Hash(Span);
+					Span = Span.Slice(Sha1Hash.NumBytes);
+				}
+
+				Node.Entries.Add(new TreePackDirEntry(Flags, Name, Hash, GitSha1));
+			}
+
+			Debug.Assert(Span.Length == 0);
+			return Node;
+		}
+
+		/// <summary>
+		/// Serializes this entry to a byte array
+		/// </summary>
+		/// <returns></returns>
+		public byte[] ToByteArray()
+		{
+			int Size = 1;
+			foreach (TreePackDirEntry Entry in Entries)
+			{
+				Size += 1 + (Entry.Name.Length + 1) + IoHash.NumBytes;
+				if ((Entry.Flags & TreePackDirEntryFlags.HasGitSha1) != 0)
+				{
+					Size += Sha1Hash.NumBytes;
+				}
+			}
+
+			byte[] Data = new byte[Size];
+			Data[0] = (byte)TreePackNodeType.Directory;
+
+			Span<byte> Span = Data.AsSpan(1);
+			foreach (TreePackDirEntry Entry in Entries)
+			{
+				Span[0] = (byte)Entry.Flags;
+				Span = Span.Slice(1);
+
+				Entry.Name.Span.CopyTo(Span);
+				Span = Span.Slice(Entry.Name.Length + 1);
+
+				Entry.Hash.CopyTo(Span);
+				Span = Span.Slice(IoHash.NumBytes);
+
+				if ((Entry.Flags & TreePackDirEntryFlags.HasGitSha1) != 0)
+				{
+					Entry.GitSha1.CopyTo(Span);
+					Span = Span.Slice(Sha1Hash.NumBytes);
+				}
+			}
+
+			Debug.Assert(Span.Length == 0);
+			return Data;
+		}
+	}
+
+	/// <summary>
+	/// Flags for <see cref="TreePackConcatEntry"/>
+	/// </summary>
+	public enum TreePackConcatNodeFlags : int
+	{
+		/// <summary>
+		/// Indicates that the referenced object is a leaf node
+		/// </summary>
+		Leaf = 1,
+	}
+
+	/// <summary>
+	/// Entry in a concat node
+	/// </summary>
+	public class TreePackConcatEntry
+	{
+		/// <summary>
+		/// Flags for this entry
+		/// </summary>
+		public TreePackConcatNodeFlags Flags { get; }
+
+		/// <summary>
+		/// Hash of this node/tree
+		/// </summary>
+		public IoHash Hash { get; }
+
+		/// <summary>
+		/// Length of this node/tree
+		/// </summary>
+		public long Length { get; }
+
+		/// <summary>
+		/// Constructor
+		/// </summary>
+		public TreePackConcatEntry(TreePackConcatNodeFlags Flags, IoHash Hash, long Length)
+		{
+			this.Flags = Flags;
+			this.Hash = Hash;
+			this.Length = Length;
+		}
+
+		/// <inheritdoc/>
+		public override string ToString() => Hash.ToString();
+	}
+
+	/// <summary>
+	/// A node which concatenates other nodes together
+	/// </summary>
+	public class TreePackConcatNode
+	{
+		/// <summary>
+		/// Entries in this concat node
+		/// </summary>
+		public List<TreePackConcatEntry> Entries = new List<TreePackConcatEntry>();
+
+		/// <summary>
+		/// Parses a concat node from a block of memory
+		/// </summary>
+		public static TreePackConcatNode Parse(ReadOnlySpan<byte> Span)
+		{
+			Debug.Assert(Span[0] == (byte)TreePackNodeType.Concat);
+			Span = Span.Slice(1);
+
+			TreePackConcatNode Node = new TreePackConcatNode();
+			while (Span.Length > 0)
+			{
+				TreePackConcatNodeFlags Flags = (TreePackConcatNodeFlags)BinaryPrimitives.ReadInt32LittleEndian(Span);
+				Span = Span.Slice(sizeof(int));
+
+				IoHash Hash = new IoHash(Span);
+				Span = Span.Slice(IoHash.NumBytes);
+
+				long Length = BinaryPrimitives.ReadInt64LittleEndian(Span);
+				Span = Span.Slice(sizeof(long));
+
+				Node.Entries.Add(new TreePackConcatEntry(Flags, Hash, Length));
+			}
+			return Node;
+		}
+
+		/// <summary>
+		/// Serializes into a byte array 
+		/// </summary>
+		/// <returns></returns>
+		public byte[] ToByteArray()
+		{
+			byte[] Data = new byte[1 + (Entries.Count * (sizeof(int) + IoHash.NumBytes + sizeof(long)))];
+			Data[0] = (byte)TreePackNodeType.Concat;
+
+			Span<byte> Span = Data.AsSpan(1);
+			foreach (TreePackConcatEntry Entry in Entries)
+			{
+				BinaryPrimitives.WriteInt32LittleEndian(Span, (int)Entry.Flags);
+				Span = Span.Slice(sizeof(int));
+
+				Entry.Hash.CopyTo(Span);
+				Span = Span.Slice(IoHash.NumBytes);
+
+				BinaryPrimitives.WriteInt64LittleEndian(Span, Entry.Length);
+				Span = Span.Slice(sizeof(long));
+			}
+
+			Debug.Assert(Span.Length == 0);
+			return Data;
+		}
+	}
+
+	#endregion
+
+	/// <summary>
+	/// Options for configuring a tree pack object
+	/// </summary>
+	public class TreePackOptions
+	{
+		/// <summary>
+		/// Maximum payload size fo a blob
+		/// </summary>
+		public int MaxBlobSize { get; set; } = 10 * 1024 * 1024;
+
+		/// <summary>
+		/// Maximum payload size for the root blob
+		/// </summary>
+		public int MaxInlineBlobSize { get; set; } = 1024 * 64;
+
+		/// <summary>
+		/// Target chunk size for content-slicing
+		/// </summary>
+		public int TargetChunkSize { get; set; } = 32 * 1024;
+
+		/// <summary>
+		/// Minimum chunk size
+		/// </summary>
+		public int MinChunkSize { get; set; } = 32 * 1024;
+
+		/// <summary>
+		/// Maximum chunk size. Chunks will be split on this boundary if another match is not found.
+		/// </summary>
+		public int MaxChunkSize { get; set; } = 256 * 1024;
+	}
+
+	/// <summary>
+	/// Manages storage of a pack of tree objects
+	/// </summary>
+	public class TreePack
+	{
+		/// <summary>
+		/// Information about a known node in memory. Nodes may come from three places:
+		/// <list type="bullet">
+		///   <item>New nodes (have data, no blob, no offset/length)</item>
+		///   <item>Existing object nodes (have data, have blob, have offset/length)</item>
+		///   <item>Existing binary nodes (no data, have blob, have offset/length)</item>
+		/// </list>
+		/// </summary>
+		class NodeInfo
+		{
+			public IoHash Hash { get; }
+			public BlobInfo? Blob { get; set; }
+			public int Offset { get; set; }
+			public int CompressedLength { get; set; }
+			public int UncompressedLength { get; set; }
+			public ReadOnlyMemory<byte> CompressedData { get; set; }
+
+			// Temp for computing live set
+			public int Rank;
+
+			public NodeInfo(IoHash Hash, BlobInfo? Blob, int Offset, int CompressedLength, int UncompressedLength, ReadOnlyMemory<byte> CompressedData)
+			{
+				this.Hash = Hash;
+				this.Blob = Blob;
+				this.Offset = Offset;
+				this.CompressedLength = CompressedLength;
+				this.UncompressedLength = UncompressedLength;
+				this.CompressedData = CompressedData;
+			}
+
+			public override string ToString() => $"{Hash}";
+		}
+
+		/// <summary>
+		/// Information about a blob in memory. 
+		/// </summary>
+		class BlobInfo
+		{
+			public IoHash Hash { get; }
+			public DateTime Time;
+			public int Length;
+			public bool IsLeaf;
+			public ReadOnlyMemory<byte> Data;
+
+			// Temp for computing live set
+			public HashSet<BlobInfo> ReferencedBy = new HashSet<BlobInfo>();
+
+			public BlobInfo(IoHash Hash, DateTime Time, int Length, bool IsLeaf)
+			{
+				this.Hash = Hash;
+				this.Time = Time;
+				this.Length = Length;
+				this.IsLeaf = IsLeaf;
+			}
+
+			public override string ToString() => IsLeaf ? $"{Hash} (Leaf)" : $"{Hash}";
+		}
+
+		/// <summary>
+		/// Entry for a block of data exported from a blob
+		/// </summary>
+		class ExportInfo
+		{
+			public IoHash Hash { get; }
+			public int Offset { get; }
+			public int CompressedLength { get; }
+			public int UncompressedLength { get; }
+
+			public ExportInfo(IoHash Hash, int Offset, int CompressedLength, int UncompressedLength)
+			{
+				this.Hash = Hash;
+				this.Offset = Offset;
+				this.CompressedLength = CompressedLength;
+				this.UncompressedLength = UncompressedLength;
+			}
+		}
+
+		IBlobStorageClient StorageClient { get; }
+		NamespaceId NamespaceId { get; }
+
+		public TreePackOptions Options { get; }
+		Dictionary<IoHash, NodeInfo> Nodes = new Dictionary<IoHash, NodeInfo>();
+		Dictionary<IoHash, BlobInfo> Blobs = new Dictionary<IoHash, BlobInfo>();
+
+		List<NodeInfo> LeafNodes = new List<NodeInfo>();
+		long LeafNodesSize = 0;
+
+		/// <summary>
+		/// Constructor
+		/// </summary>
+		public TreePack(IBlobStorageClient StorageClient, NamespaceId NamespaceId)
+			: this(StorageClient, NamespaceId, new TreePackOptions())
+		{
+		}
+
+		/// <summary>
+		/// Constructor
+		/// </summary>
+		public TreePack(IBlobStorageClient StorageClient, NamespaceId NamespaceId, TreePackOptions Options)
+		{
+			this.StorageClient = StorageClient;
+			this.NamespaceId = NamespaceId;
+			this.Options = Options;
+		}
+
+		/// <summary>
+		/// Adds a root blob to the pack
+		/// </summary>
+		/// <param name="Data"></param>
+		public void AddRootBlob(ReadOnlyMemory<byte> Data)
+		{
+			TreePackObject Object = CbSerializer.Deserialize<TreePackObject>(Data);
+			RegisterObject(Object, null);
+		}
+
+		/// <summary>
+		/// Gets references from a node to other nodes
+		/// </summary>
+		/// <param name="Data"></param>
+		/// <returns></returns>
+		private static IEnumerable<IoHash> GetReferences(ReadOnlyMemory<byte> Data)
+		{
+			TreePackNodeType Type = (TreePackNodeType)Data.Span[0];
+			switch (Type)
+			{
+				case TreePackNodeType.Binary:
+					return Enumerable.Empty<IoHash>();
+				case TreePackNodeType.Object:
+					HashSet<IoHash> Hashes = new HashSet<IoHash>();
+					new CbField(Data.Slice(1)).IterateAttachments(x => Hashes.Add(x.AsAttachment()));
+					return Hashes;
+				case TreePackNodeType.Concat:
+					return TreePackConcatNode.Parse(Data.Span).Entries.Select(x => x.Hash);
+				case TreePackNodeType.Directory:
+					return TreePackDirNode.Parse(Data).Entries.Select(x => x.Hash);
+				default:
+					throw new NotImplementedException();
+			}
+		}
+
+		BlobInfo FindOrAddBlob(IoHash Hash, DateTime UtcNow, int Length, bool IsLeaf)
+		{
+			BlobInfo? BlobInfo;
+			if (!Blobs.TryGetValue(Hash, out BlobInfo))
+			{
+				BlobInfo = new BlobInfo(Hash, UtcNow, Length, IsLeaf);
+				Blobs.Add(Hash, BlobInfo);
+			}
+			return BlobInfo;
+		}
+
+		NodeInfo FindOrAddNode(IoHash Hash, int Offset, int CompressedLength, int UncompressedLength, BlobInfo? BlobInfo)
+		{
+			NodeInfo? NodeInfo;
+			if (!Nodes.TryGetValue(Hash, out NodeInfo))
+			{
+				NodeInfo = new NodeInfo(Hash, BlobInfo, Offset, CompressedLength, UncompressedLength, ReadOnlyMemory<byte>.Empty);
+				Nodes.Add(Hash, NodeInfo);
+			}
+			return NodeInfo;
+		}
+
+		/// <summary>
+		/// Adds information about an object to the store
+		/// </summary>
+		/// <param name="Object"></param>
+		public void AddNodes(TreePackObject RootObject)
+		{
+			RegisterObject(RootObject, null);
+		}
+
+		/// <summary>
+		/// Adds a new node to pack
+		/// </summary>
+		/// <param name="Data">Node data</param>
+		public async Task<IoHash> AddNodeAsync(ReadOnlyMemory<byte> Data)
+		{
+			IoHash Hash = IoHash.Compute(Data.Span);
+			await AddNodeAsync(Hash, Data);
+			return Hash;
+		}
+
+		private static ReadOnlyMemory<byte> CompressData(ReadOnlyMemory<byte> Data)
+		{
+			byte[] Output = new byte[LZ4Codec.MaximumOutputSize(Data.Length)];
+			int Length = LZ4Codec.Encode(Data.Span, Output);
+			return Output.AsMemory(0, Length);
+		}
+
+		private static byte[] UncompressData(ReadOnlyMemory<byte> Data, int Length)
+		{
+			byte[] Output = new byte[Length];
+			LZ4Codec.Decode(Data.Span, Output);
+			return Output;
+		}
+
+		/// <summary>
+		/// Adds a new node to pack
+		/// </summary>
+		/// <param name="Hash">Hash for the node</param>
+		/// <param name="Data">Node data</param>
+		public async Task AddNodeAsync(IoHash Hash, ReadOnlyMemory<byte> Data)
+		{
+			if (!Nodes.ContainsKey(Hash))
+			{
+				ReadOnlyMemory<byte> CompressedData = CompressData(Data);
+
+				NodeInfo Node = new NodeInfo(Hash, null, -1, CompressedData.Length, Data.Length, CompressedData);
+				if (!GetReferences(Data).Any())
+				{
+					if (LeafNodesSize + CompressedData.Length > Options.MaxBlobSize)
+					{
+						await WriteObjectAsync(LeafNodes, DateTime.UtcNow);
+						LeafNodes.Clear();
+						LeafNodesSize = 0;
+					}
+
+					LeafNodes.Add(Node);
+					LeafNodesSize += CompressedData.Length;
+				}
+				Nodes.Add(Hash, Node);
+			}
+		}
+
+		/// <summary>
+		/// Gets data for the given node
+		/// </summary>
+		/// <param name="Key"></param>
+		/// <returns></returns>
+		public ValueTask<ReadOnlyMemory<byte>> GetDataAsync(IoHash Hash)
+		{
+			NodeInfo Node = Nodes[Hash];
+			return GetDataAsync(Node);
+		}
+
+		/// <summary>
+		/// Gets data for the given node
+		/// </summary>
+		/// <param name="Key"></param>
+		/// <returns></returns>
+		private async ValueTask<ReadOnlyMemory<byte>> GetDataAsync(NodeInfo Node)
+		{
+			if (Node.CompressedData.Length == 0 && Node.Blob != null)
+			{
+				BlobInfo? Blob = Node.Blob;
+				if (Blob.Data.Length == 0)
+				{
+					Blob.Data = await StorageClient.ReadBlobToMemoryAsync(NamespaceId, Blob.Hash);
+				}
+
+				if (Node.Blob.IsLeaf)
+				{
+					ReadOnlyMemory<byte> BlobData = Node.Blob.Data;
+					Node.CompressedData = BlobData.Slice(Node.Offset, Node.CompressedLength);
+				}
+				else
+				{
+					TreePackObject Object = CbSerializer.Deserialize<TreePackObject>(Node.Blob.Data);
+					RegisterObject(Object, Node.Blob);
+				}
+			}
+			return UncompressData(Node.CompressedData, Node.UncompressedLength);
+		}
+
+		/// <summary>
+		/// Flush any pending blobs to disk
+		/// </summary>
+		/// <param name="RootHash">New root key</param>
+		/// <param name="UtcNow">Timestamp for any new blobs</param>
+		public async Task<TreePackObject> FlushAsync(IoHash RootHash, DateTime UtcNow)
+		{
+			// Reset all the current values
+			foreach (NodeInfo Node in Nodes.Values)
+			{
+				Node.Rank = -1;
+			}
+
+			// Ensure that the root object is written out as part of the output
+			NodeInfo RootNode = Nodes[RootHash];
+
+			// Find the live set of objects
+			HashSet<NodeInfo> LiveNodes = new HashSet<NodeInfo>();
+			await FindLiveSetAsync(RootNode, LiveNodes);
+
+			// Find the live set of blobs
+			HashSet<BlobInfo> LiveBlobs = new HashSet<BlobInfo>();
+			foreach (NodeInfo LiveNode in LiveNodes)
+			{
+				BlobInfo? LiveBlob = LiveNode.Blob;
+				if (LiveBlob != null)
+				{
+					ExpandReferences(LiveBlob, LiveBlobs);
+				}
+			}
+
+			// Find the total cost of all the current blobs, then loop through the blobs trying to find a more optimal arrangement
+			List<NodeInfo> SortedLiveNodes = LiveNodes.OrderByDescending(x => x.Rank).ThenByDescending(x => x.Hash).ToList();
+			for (; ; )
+			{
+				BlobInfo? MergeBlob = null;
+				double MergeCost = GetNewCostHeuristic(SortedLiveNodes, new HashSet<BlobInfo>(), UtcNow);
+
+				foreach (BlobInfo LiveBlob in LiveBlobs)
+				{
+					double PotentialCost = GetNewCostHeuristic(SortedLiveNodes, LiveBlob.ReferencedBy, UtcNow);
+					if (PotentialCost < MergeCost)
+					{
+						MergeBlob = LiveBlob;
+						MergeCost = PotentialCost;
+					}
+				}
+
+				// Bail out if we didn't find anything to merge
+				if (MergeBlob == null)
+				{
+					break;
+				}
+
+				// Find all the live nodes in this blob
+				List<NodeInfo> MergeNodes = LiveNodes.Where(x => x.Blob == MergeBlob).ToList();
+
+				// If it's a leaf node, we'll need to merge it in
+				if (MergeBlob.IsLeaf)
+				{
+					byte[] Data = await StorageClient.ReadBlobToMemoryAsync(NamespaceId, MergeBlob.Hash);
+					foreach (NodeInfo MergeNode in MergeNodes)
+					{
+						MergeNode.CompressedData = Data.AsMemory(MergeNode.Offset, MergeNode.CompressedLength).ToArray();
+					}
+				}
+
+				// Read them all in and add them to the live set
+				foreach (NodeInfo NodeInfo in MergeNodes)
+				{
+					NodeInfo.Blob = null;
+				}
+
+				// Remove the merged blob from the live list
+				LiveBlobs.Remove(MergeBlob);
+			}
+
+			// Clear the list of leaf nodes, since they're all being flushed anyway
+			LeafNodes.Clear();
+			LeafNodesSize = 0;
+
+			// Write all the new blobs. Sort by rank and go from end to start, updating the list as we go so that we can figure out imports correctly.
+			NodeInfo[] WriteNodes = LiveNodes.Where(x => x.Blob == null).OrderBy(x => x.Rank).ToArray();
+
+			int Length = WriteNodes.Length;
+
+			int NextBlobSize = WriteNodes[WriteNodes.Length - 1].CompressedLength;
+			for (int Idx = WriteNodes.Length - 2; Idx >= 1; Idx--)
+			{
+				NextBlobSize += WriteNodes[Idx].CompressedLength;
+				if (NextBlobSize > Options.MaxBlobSize)
+				{
+					await WriteObjectAsync(WriteNodes.Slice(Idx + 1, Length - (Idx + 1)), UtcNow);
+					NextBlobSize = WriteNodes[Idx].CompressedLength;
+					Length = Idx + 1;
+				}
+			}
+
+			return BuildObject(WriteNodes.Slice(0, Length), UtcNow);
+		}
+
+		void ExpandReferences(BlobInfo Blob, HashSet<BlobInfo> VisitedBlobs)
+		{
+			if (VisitedBlobs.Add(Blob))
+			{
+				foreach (BlobInfo ReferencedBlob in Blob.ReferencedBy)
+				{
+					ExpandReferences(ReferencedBlob, VisitedBlobs);
+					Blob.ReferencedBy.UnionWith(ReferencedBlob.ReferencedBy);
+				}
+				Blob.ReferencedBy.Add(Blob);
+			}
+		}
+
+		/// <summary>
+		/// Finds the live set for a particular tree, and updates tree entries with the size of used items within them
+		/// </summary>
+		/// <param name="Node"></param>
+		/// <param name="LiveNodes"></param>
+		async Task FindLiveSetAsync(NodeInfo Node, HashSet<NodeInfo> LiveNodes)
+		{
+			if (LiveNodes.Add(Node))
+			{
+				int MaxRank = 0;
+				if (Node.Blob == null || !Node.Blob.IsLeaf)
+				{
+					ReadOnlyMemory<byte> Data = await GetDataAsync(Node.Hash);
+					foreach (IoHash RefHash in GetReferences(Data))
+					{
+						NodeInfo RefNode = Nodes[RefHash];
+						await FindLiveSetAsync(RefNode, LiveNodes);
+
+						if (Node.Blob != null && RefNode.Blob != null)
+						{
+							RefNode.Blob.ReferencedBy.Add(Node.Blob);
+						}
+
+						MaxRank = Math.Max(MaxRank, Node.Rank + 1);
+					}
+				}
+				Node.Rank = MaxRank;
+			}
+		}
+
+		void RegisterObject(TreePackObject Object, BlobInfo? BlobInfo)
+		{
+			// Register all the exports from the object
+			foreach (ExportInfo Export in ReadExports(Object.Exports))
+			{
+				NodeInfo NodeInfo = FindOrAddNode(Export.Hash, Export.Offset, Export.CompressedLength, Export.UncompressedLength, BlobInfo);
+				NodeInfo.UncompressedLength = Export.UncompressedLength;
+				NodeInfo.CompressedLength = Export.CompressedLength;
+				NodeInfo.CompressedData = Object.Data.Slice(Export.Offset, Export.CompressedLength);
+			}
+
+			// Register all the object imports
+			foreach (TreePackObjectImport Import in Object.ObjectImports)
+			{
+				BlobInfo ImportedBlob = FindOrAddBlob(Import.Object.Hash, DateTime.MinValue, 0, false);
+
+				IoHash[] Keys = ReadKeys(Import.Keys);
+				foreach (IoHash Key in Keys)
+				{
+					FindOrAddNode(Key, -1, -1, -1, ImportedBlob);
+				}
+			}
+
+			// Register all the binary imports
+			foreach (TreePackBinaryImport Import in Object.BinaryImports)
+			{
+				BlobInfo ImportedBlob = FindOrAddBlob(Import.Binary.Hash, Import.Time, Import.Length, true);
+
+				ExportInfo[] Exports = ReadExports(Import.Exports);
+				foreach (ExportInfo Export in Exports)
+				{
+					FindOrAddNode(Export.Hash, Export.Offset, Export.CompressedLength, Export.UncompressedLength, ImportedBlob);
+				}
+			}
+		}
+
+		TreePackObject BuildObject(IReadOnlyList<NodeInfo> Nodes, DateTime UtcNow)
+		{
+			byte[] Data = new byte[Nodes.Sum(x => x.CompressedLength)];
+			int NextOffset = 0;
+
+			HashSet<NodeInfo> References = new HashSet<NodeInfo>();
+			foreach (NodeInfo Node in Nodes)
+			{
+				Memory<byte> NodeData = Data.AsMemory(NextOffset, Node.CompressedLength);
+				Node.CompressedData.CopyTo(NodeData);
+
+				Node.Offset = NextOffset;
+				Node.CompressedData = NodeData;
+
+				NextOffset += Node.CompressedLength;
+
+				byte[] OutputData = UncompressData(Node.CompressedData, Node.UncompressedLength);
+				foreach (IoHash Hash in GetReferences(OutputData))
+				{
+					NodeInfo? RefNode;
+					if (this.Nodes.TryGetValue(Hash, out RefNode) && RefNode.Blob != null)
+					{
+						References.Add(RefNode);
+					}
+				}
+			}
+
+			TreePackObject Object = new TreePackObject();
+			Object.Time = UtcNow;
+			Object.Data = Data;
+			Object.Exports = WriteExports(Nodes.Select(x => new ExportInfo(x.Hash, x.Offset, x.CompressedLength, x.UncompressedLength)).ToArray());
+
+			foreach (IGrouping<BlobInfo, NodeInfo> Group in References.GroupBy(x => x.Blob!))
+			{
+				BlobInfo ImportBlob = Group.Key;
+				if (ImportBlob.IsLeaf)
+				{
+					TreePackBinaryImport BinaryImport = new TreePackBinaryImport();
+					BinaryImport.Binary = new CbBinaryAttachment(ImportBlob.Hash);
+					BinaryImport.Time = ImportBlob.Time;
+					BinaryImport.Length = ImportBlob.Length;
+					BinaryImport.Exports = WriteExports(Group.Select(x => new ExportInfo(x.Hash, x.Offset, x.CompressedLength, x.UncompressedLength)).ToArray());
+					Object.BinaryImports.Add(BinaryImport);
+				}
+				else
+				{
+					TreePackObjectImport ObjectImport = new TreePackObjectImport();
+					ObjectImport.Object = new CbObjectAttachment(ImportBlob.Hash);
+					ObjectImport.Keys = WriteKeys(Group.Select(x => x.Hash).ToArray());
+					Object.ObjectImports.Add(ObjectImport);
+				}
+			}
+
+			return Object;
+		}
+
+		async Task<BlobInfo> WriteObjectAsync(IReadOnlyList<NodeInfo> Nodes, DateTime UtcNow)
+		{
+			TreePackObject Object = BuildObject(Nodes, UtcNow);
+
+			BlobInfo Blob;
+			if (Object.ObjectImports.Count == 0 && Object.BinaryImports.Count == 0)
+			{
+				IoHash Hash = await StorageClient.WriteBlobFromMemoryAsync(NamespaceId, Object.Data);
+				Blob = new BlobInfo(Hash, Object.Time, Object.Data.Length, true);
+			}
+			else
+			{
+				IoHash Hash = await StorageClient.WriteObjectAsync(NamespaceId, Object);
+				Blob = new BlobInfo(Hash, Object.Time, Object.Data.Length, false);
+			}
+
+			foreach (NodeInfo Node in Nodes)
+			{
+				Node.Blob = Blob;
+				Node.CompressedData = ReadOnlyMemory<byte>.Empty;
+			}
+
+			return Blob;
+		}
+
+		static IoHash[] ReadKeys(ReadOnlyMemory<byte> KeyData)
+		{
+			int NumEntries = KeyData.Length / IoHash.NumBytes;
+
+			IoHash[] Keys = new IoHash[NumEntries];
+
+			ReadOnlySpan<byte> Span = KeyData.Span;
+			for (int Idx = 0; Idx < NumEntries; Idx++)
+			{
+				Keys[Idx] = new IoHash(Span);
+				Span = Span.Slice(IoHash.NumBytes);
+			}
+
+			return Keys;
+		}
+
+		static byte[] WriteKeys(IoHash[] Keys)
+		{
+			byte[] Data = new byte[Keys.Length * IoHash.NumBytes];
+
+			Span<byte> Output = Data;
+			foreach (IoHash Key in Keys)
+			{
+				Key.CopyTo(Output);
+				Output = Output.Slice(IoHash.NumBytes);
+			}
+
+			return Data;
+		}
+
+		static ExportInfo[] ReadExports(ReadOnlyMemory<byte> ExportData)
+		{
+			int NumEntries = ExportData.Length / (IoHash.NumBytes + sizeof(int) + sizeof(int) + sizeof(int));
+
+			ExportInfo[] Exports = new ExportInfo[NumEntries];
+
+			ReadOnlySpan<byte> Span = ExportData.Span;
+			for (int Idx = 0; Idx < NumEntries; Idx++)
+			{
+				IoHash Key = new IoHash(Span);
+				Span = Span.Slice(IoHash.NumBytes);
+
+				int Offset = BinaryPrimitives.ReadInt32LittleEndian(Span);
+				Span = Span.Slice(sizeof(int));
+
+				int CompressedLength = BinaryPrimitives.ReadInt32LittleEndian(Span);
+				Span = Span.Slice(sizeof(int));
+
+				int UncompressedLength = BinaryPrimitives.ReadInt32LittleEndian(Span);
+				Span = Span.Slice(sizeof(int));
+
+				Exports[Idx] = new ExportInfo(Key, Offset, CompressedLength, UncompressedLength);
+			}
+
+			return Exports;
+		}
+
+		static byte[] WriteExports(ExportInfo[] Exports)
+		{
+			byte[] Data = new byte[Exports.Length * (IoHash.NumBytes + sizeof(int) + sizeof(int) + sizeof(int))];
+
+			Span<byte> Output = Data;
+			foreach (ExportInfo Export in Exports.OrderBy(x => x.Offset))
+			{
+				Export.Hash.CopyTo(Output);
+				Output = Output.Slice(IoHash.NumBytes);
+
+				BinaryPrimitives.WriteInt32LittleEndian(Output, Export.Offset);
+				Output = Output.Slice(sizeof(int));
+
+				BinaryPrimitives.WriteInt32LittleEndian(Output, Export.CompressedLength);
+				Output = Output.Slice(sizeof(int));
+
+				BinaryPrimitives.WriteInt32LittleEndian(Output, Export.UncompressedLength);
+				Output = Output.Slice(sizeof(int));
+			}
+
+			return Data;
+		}
+
+		/// <inheritdoc cref="GetCostHeuristic(int, TimeSpan)"/>
+		/// <param name="Index">Index to calculate the heuristic for</param>
+		double GetCostHeuristic(BlobInfo Info, DateTime UtcNow) => GetCostHeuristic(Info.Length, UtcNow - Info.Time);
+
+		/// <summary>
+		/// Gets the cost of packaging the given set of nodes which are either not in a blob or are in a blob listed to be collapsed
+		/// </summary>
+		/// <param name="Nodes">The full list of nodes</param>
+		/// <param name="CollapseBlobs">Set of blobs that should be collapsed</param>
+		/// <returns></returns>
+		double GetNewCostHeuristic(List<NodeInfo> Nodes, HashSet<BlobInfo> CollapseBlobs, DateTime UtcNow)
+		{
+			double Cost = 0.0;
+			foreach (BlobInfo CollapseBlob in CollapseBlobs)
+			{
+				Cost -= GetCostHeuristic(CollapseBlob, UtcNow);
+			}
+
+			int BlobSize = 0;
+			for (int Idx = Nodes.Count - 1; Idx >= 0; Idx--)
+			{
+				NodeInfo Node = Nodes[Idx];
+				if (Node.Blob == null || CollapseBlobs.Contains(Node.Blob))
+				{
+					if (BlobSize > 0 && BlobSize + Node.CompressedLength > Options.MaxBlobSize)
+					{
+						Cost += GetCostHeuristic(BlobSize, TimeSpan.Zero);
+						BlobSize = 0;
+					}
+					BlobSize += Node.CompressedLength;
+				}
+			}
+			Cost += GetCostHeuristic(BlobSize, TimeSpan.Zero);
+
+			return Cost;
+		}
+
+		/// <summary>
+		/// Heuristic which estimates the cost of a particular blob. This is used to compare scenarios of merging blobs to reduce download
+		/// size against keeping older blobs which a lot of agents already have.
+		/// </summary>
+		/// <param name="Size">Size of the blob</param>
+		/// <param name="Age">Age of the blob</param>
+		/// <returns>Heuristic for the cost of a blob</returns>
+		static double GetCostHeuristic(int Size, TimeSpan Age)
+		{
+			// Time overhead to starting a download
+			const double DownloadInit = 0.1;
+
+			// Download speed for agents, in bytes/sec
+			const double DownloadRate = 1024 * 1024;
+
+			// Probability of an agent having to download everything. Prevents bias against keeping a large number of files.
+			const double CleanSyncProbability = 0.2;
+
+			// Average length of time between agents having to update
+			TimeSpan AverageCoherence = TimeSpan.FromHours(4.0);
+
+			// Scale the age into a -1.0 -> 1.0 range around AverageCoherence
+			double ScaledAge = (AverageCoherence - Age).TotalSeconds / AverageCoherence.TotalSeconds;
+
+			// Get the probability of agents having to sync this blob based on its age. This is modeled as a logistic function (1 / (1 + e^-x))
+			// with value of 0.5 at AverageCoherence, and MaxInterval at zero.
+
+			// Find the scale factor for the 95% interval
+			//    1 / (1 + e^-x) = MaxInterval
+			//    e^-x = (1 / MaxInterval) - 1
+			//    x = -ln((1 / MaxInterval) - 1)
+			const double MaxInterval = 0.95;
+			double SigmoidScale = -Math.Log((1.0 / MaxInterval) - 1.0);
+
+			// Find the probability of having to sync this 
+			double Param = ScaledAge * SigmoidScale;
+			double Probability = 1.0 / (1.0 + Math.Exp(-Param));
+
+			// Scale the probability against having to do a full sync
+			Probability = CleanSyncProbability + (Probability * (1.0 - CleanSyncProbability));
+
+			// Compute the final cost estimate; the amount of time we expect agents to spend downloading the file
+			return Probability * (DownloadInit + (Size / DownloadRate));
+		}
+
+		public static void WriteSummary(FileReference File, TreePackObject Object)
+		{
+			Dictionary<IoHash, string> HashToLocator = new Dictionary<IoHash, string>();
+			foreach (TreePackObjectImport ObjectImport in Object.ObjectImports)
+			{
+				foreach (IoHash Key in ReadKeys(ObjectImport.Keys))
+				{
+					HashToLocator[Key] = $"{ObjectImport.Object}";
+				}
+			}
+			foreach (TreePackBinaryImport BinaryImport in Object.BinaryImports)
+			{
+				foreach (ExportInfo Export in ReadExports(BinaryImport.Exports))
+				{
+					HashToLocator[Export.Hash] = $"{BinaryImport.Binary} [{Export.Offset},{Export.CompressedLength}]";
+				}
+			}
+
+			using StreamWriter Writer = new StreamWriter(FileReference.Open(File, FileMode.Create));
+
+			Writer.WriteLine("time: {0}", Object.Time.ToString("u"));
+
+			if (Object.ObjectImports.Count > 0)
+			{
+				Writer.WriteLine("objimp:");
+				foreach (TreePackObjectImport ObjectImport in Object.ObjectImports)
+				{
+					Writer.WriteLine($"- blob: {ObjectImport.Object.Hash}");
+					Writer.WriteLine("  keys:");
+					foreach (IoHash Key in ReadKeys(ObjectImport.Keys))
+					{
+						Writer.WriteLine($"  - \"{Key}\"");
+					}
+					Writer.WriteLine();
+				}
+			}
+
+			if (Object.BinaryImports.Count > 0)
+			{
+				Writer.WriteLine("binimp:");
+				foreach (TreePackBinaryImport BinaryImport in Object.BinaryImports)
+				{
+					Writer.WriteLine($"- blob: {BinaryImport.Binary.Hash}");
+					Writer.WriteLine("  keys:");
+					foreach (ExportInfo Export in ReadExports(BinaryImport.Exports))
+					{
+						Writer.WriteLine($"  - \"{Export.Hash}\"");
+					}
+					Writer.WriteLine();
+				}
+			}
+
+			Writer.WriteLine("exports:");
+
+			ExportInfo[] Exports = ReadExports(Object.Exports);
+			foreach (ExportInfo Export in Exports)
+			{
+				ReadOnlyMemory<byte> Data = UncompressData(Object.Data.Slice(Export.Offset, Export.CompressedLength), Export.UncompressedLength);
+				TreePackNodeType Type = (TreePackNodeType)Data.Span[0];
+
+				Writer.WriteLine($"- hash: {Export.Hash}");
+				Writer.WriteLine($"  type: {Type.ToString()}");
+
+				if (Type == TreePackNodeType.Concat)
+				{
+					Writer.WriteLine("  entries:");
+
+					TreePackConcatNode Node = TreePackConcatNode.Parse(Data.Span);
+					foreach (TreePackConcatEntry Entry in Node.Entries)
+					{
+						Writer.WriteLine($"    - hash: \"{Entry.Hash}\"");
+						HashToLocator.TryGetValue(Entry.Hash, out string? Locator);
+						Writer.WriteLine($"      from: \"{Locator ?? "current"}\"");
+					}
+
+					Writer.WriteLine();
+				}
+				else if (Type == TreePackNodeType.Directory)
+				{
+					Writer.WriteLine("  entries:");
+
+					TreePackDirNode Node = TreePackDirNode.Parse(Data);
+					foreach (TreePackDirEntry Entry in Node.Entries)
+					{
+						Writer.WriteLine($"    - name: \"{Entry.Name}\"");
+						Writer.WriteLine($"      type: {Entry.Flags}");
+						Writer.WriteLine($"      hash: \"{Entry.Hash}\"");
+
+						HashToLocator.TryGetValue(Entry.Hash, out string? Locator);
+						Writer.WriteLine($"      from: \"{Locator ?? "current"}\"");
+
+						Writer.WriteLine();
+					}
+				}
+				else
+				{
+					Writer.WriteLine("  length: {0}", Export.UncompressedLength);
+					Writer.WriteLine();
+				}
+			}
+		}
+	}
+}
diff --git a/Engine/Source/Programs/Shared/EpicGames.Horde/Storage/TreePackDirWriter.cs b/Engine/Source/Programs/Shared/EpicGames.Horde/Storage/TreePackDirWriter.cs
new file mode 100644
--- /dev/null
+++ b/Engine/Source/Programs/Shared/EpicGames.Horde/Storage/TreePackDirWriter.cs
@@ -0,0 +1,151 @@
+// Copyright Epic Games, Inc. All Rights Reserved.
+
+using EpicGames.Core;
+using System;
+using System.Collections.Generic;
+using System.Text;
+using System.Threading.Tasks;
+
+namespace EpicGames.Horde.Storage
+{
+	/// <summary>
+	/// Helper class for writing/updating directory trees within a tree pack
+	/// </summary>
+	public class TreePackDirWriter
+	{
+		TreePack TreePack;
+		Dictionary<Utf8String, TreePackDirEntry> NameToEntry = new Dictionary<Utf8String, TreePackDirEntry>();
+		Dictionary<Utf8String, TreePackDirWriter> NameToSubDir = new Dictionary<Utf8String, TreePackDirWriter>();
+
+		public IEnumerable<TreePackDirEntry> Entries => NameToEntry.Values;
+
+		/// <summary>
+		/// Constructor. Creates an empty, new directory structure.
+		/// </summary>
+		/// <param name="TreePack"></param>
+		public TreePackDirWriter(TreePack TreePack)
+		{
+			this.TreePack = TreePack;
+		}
+
+		/// <summary>
+		/// Constructor. Initializes the contents of the tree to an existing directory node.
+		/// </summary>
+		/// <param name="TreePack"></param>
+		/// <param name="Node"></param>
+		public TreePackDirWriter(TreePack TreePack, TreePackDirNode Node)
+		{
+			this.TreePack = TreePack;
+
+			foreach (TreePackDirEntry Entry in Node.Entries)
+			{
+				NameToEntry[Entry.Name] = Entry;
+			}
+		}
+
+		/// <summary>
+		/// Clears the current directory. 
+		/// </summary>
+		public void Clear()
+		{
+			NameToEntry.Clear();
+			NameToSubDir.Clear();
+		}
+
+		/// <summary>
+		/// Adds a file to the tree pack with the given path. 
+		/// </summary>
+		/// <param name="Path">Path to the file. Should be specified using forward slashes as directory separators, and without a leading slash.</param>
+		/// <param name="Flags">Flags for the new file entry.</param>
+		/// <param name="Hash">Hash of the file.</param>
+		/// <param name="GitSha1">SHA1 of the file (optional)</param>
+		/// <returns>New directory entry</returns>
+		public async ValueTask<TreePackDirEntry> FindOrAddFileByPathAsync(Utf8String Path, TreePackDirEntryFlags Flags, IoHash Hash, Sha1Hash GitSha1 = default)
+		{
+			TreePackDirWriter Writer = this;
+			for (; ; )
+			{
+				int Index = Path.IndexOf('/');
+				if (Index != -1)
+				{
+					Writer = await Writer.FindOrAddDirAsync(Path.Substring(0, Index));
+					Path = Path.Substring(Index + 1);
+					continue;
+				}
+				return Writer.FindOrAddFile(Path, Flags, Hash, GitSha1);
+			}
+		}
+
+		/// <summary>
+		/// Adds a file to the current tree level with the given name. 
+		/// </summary>
+		/// <param name="Name">Name of the file. Should not contain any path information.</param>
+		/// <param name="Flags">Flags for the new file entry.</param>
+		/// <param name="Hash">Hash of the file.</param>
+		/// <param name="GitSha1">SHA1 of the file (optional)</param>
+		/// <returns>New directory entry</returns>
+		public TreePackDirEntry FindOrAddFile(Utf8String Name, TreePackDirEntryFlags Flags, IoHash Hash, Sha1Hash GitSha1 = default)
+		{
+			TreePackDirEntry Entry = new TreePackDirEntry(Flags, Name, Hash, GitSha1);
+			NameToEntry[Name] = Entry;
+			NameToSubDir.Remove(Name);
+			return Entry;
+		}
+
+		/// <summary>
+		/// Adds a directory to the current tree level with the given name
+		/// </summary>
+		/// <param name="Name"></param>
+		/// <returns></returns>
+		public async ValueTask<TreePackDirWriter> FindOrAddDirAsync(Utf8String Name)
+		{
+			TreePackDirWriter? Writer;
+			if (!NameToSubDir.TryGetValue(Name, out Writer))
+			{
+				TreePackDirEntry? Entry;
+				if (NameToEntry.TryGetValue(Name, out Entry) && (Entry.Flags & TreePackDirEntryFlags.Directory) != 0)
+				{
+					Writer = await FromDirNodeAsync(TreePack, Entry.Hash);
+				}
+				else
+				{
+					Writer = new TreePackDirWriter(TreePack);
+				}
+				NameToEntry.Remove(Name);
+				NameToSubDir.Add(Name, Writer);
+			}
+			return Writer;
+		}
+
+		/// <summary>
+		/// Flushes the current tree state and writes it to the tree pack
+		/// </summary>
+		/// <returns>Hash of the tree</returns>
+		public async Task<IoHash> FinalizeAsync()
+		{
+			foreach ((Utf8String Name, TreePackDirWriter Writer) in NameToSubDir)
+			{
+				IoHash Hash = await Writer.FinalizeAsync();
+				NameToEntry[Name] = new TreePackDirEntry(TreePackDirEntryFlags.Directory, Name, Hash, Sha1Hash.Zero);
+			}
+			NameToSubDir.Clear();
+
+			TreePackDirNode Node = new TreePackDirNode();
+			Node.Entries.AddRange(NameToEntry.Values);
+			return await TreePack.AddNodeAsync(Node.ToByteArray());
+		}
+
+		/// <summary>
+		/// Creates a writer from a particular directory hash
+		/// </summary>
+		/// <param name="TreePack"></param>
+		/// <param name="Hash"></param>
+		/// <returns></returns>
+		public static async Task<TreePackDirWriter> FromDirNodeAsync(TreePack TreePack, IoHash Hash)
+		{
+			ReadOnlyMemory<byte> Data = await TreePack.GetDataAsync(Hash);
+			TreePackDirNode SubDirNode = TreePackDirNode.Parse(Data);
+			return new TreePackDirWriter(TreePack, SubDirNode);
+		}
+	}
+}
diff --git a/Engine/Source/Programs/Shared/EpicGames.Horde/Storage/TreePackFileWriter.cs b/Engine/Source/Programs/Shared/EpicGames.Horde/Storage/TreePackFileWriter.cs
new file mode 100644
--- /dev/null
+++ b/Engine/Source/Programs/Shared/EpicGames.Horde/Storage/TreePackFileWriter.cs
@@ -0,0 +1,124 @@
+// Copyright Epic Games, Inc. All Rights Reserved.
+
+using EpicGames.Core;
+using System;
+using System.Collections.Generic;
+using System.Text;
+using System.Threading.Tasks;
+
+namespace EpicGames.Horde.Storage
+{
+	/// <summary>
+	/// Helper class for writing files split along content-aware boundaries
+	/// </summary>
+	public class TreePackFileWriter
+	{
+		public TreePack TreePack { get; }
+		int MinBlockSize;
+		int MaxBlockSize;
+		uint WindowHashThreshold;
+		TreePackConcatNode Node = new TreePackConcatNode();
+		byte[] Buffer = Array.Empty<byte>();
+		int BufferLen;
+		BuzHash BuzHash = new BuzHash();
+
+		/// <summary>
+		/// Constructor
+		/// </summary>
+		/// <param name="TreePack"></param>
+		public TreePackFileWriter(TreePack TreePack)
+		{
+			this.TreePack = TreePack;
+			this.MinBlockSize = TreePack.Options.MinChunkSize;
+			this.MaxBlockSize = TreePack.Options.MaxChunkSize;
+			this.WindowHashThreshold = (uint)((1L << 32) / TreePack.Options.TargetChunkSize);
+		}
+
+		/// <summary>
+		/// Resets the currently accumulated data
+		/// </summary>
+		public void Reset()
+		{
+			Node.Entries.Clear();
+		}
+
+		/// <summary>
+		/// Writes data to the pack
+		/// </summary>
+		/// <param name="Data">Data to be written</param>
+		/// <param name="Flush">Whether to flush the contents of the buffer after writing the data</param>
+		public async Task WriteAsync(ReadOnlyMemory<byte> Data, bool Flush)
+		{
+			int Length = 0;
+			for (; Length < Data.Length; Length++)
+			{
+				BuzHash.Add(Data.Span[Length]);
+
+				uint WindowHash = BuzHash.Get();
+
+				int BlockSize = BufferLen + Length;
+				if ((BlockSize >= MinBlockSize && WindowHash < WindowHashThreshold) || (BlockSize >= MaxBlockSize))
+				{
+					await WriteCurrentChunkAsync(Data.Slice(0, Length));
+					Data = Data.Slice(Length);
+					Length = 0;
+				}
+			}
+
+			if (Data.Length > 0)
+			{
+				if (BufferLen + Data.Length > Buffer.Length)
+				{
+					Array.Resize(ref Buffer, BufferLen + Data.Length);
+				}
+
+				Data.CopyTo(Buffer.AsMemory(BufferLen));
+				BufferLen += Data.Length;
+			}
+
+			if (BufferLen > 0 && Flush)
+			{
+				await WriteCurrentChunkAsync(Data.Slice(0, Length));
+			}
+		}
+
+		async Task WriteCurrentChunkAsync(ReadOnlyMemory<byte> Data)
+		{
+			byte[] Chunk = new byte[1 + BufferLen + Data.Length];
+			Chunk[0] = (byte)TreePackNodeType.Binary;
+			Buffer.AsMemory(0, BufferLen).CopyTo(Chunk.AsMemory(1));
+			Data.CopyTo(Chunk.AsMemory(1 + BufferLen));
+
+			IoHash Hash = await TreePack.AddNodeAsync(Chunk);
+			Node.Entries.Add(new TreePackConcatEntry(TreePackConcatNodeFlags.Leaf, Hash, Chunk.Length));
+
+			BufferLen = 0;
+			BuzHash.Reset();
+		}
+
+		/// <summary>
+		/// Flush the contents of the writer
+		/// </summary>
+		public Task FlushAsync()
+		{
+			return WriteAsync(ReadOnlyMemory<byte>.Empty, true);
+		}
+
+		/// <summary>
+		/// Finalizes the current data.
+		/// </summary>
+		/// <returns>Hash of the finalized data</returns>
+		public async Task<IoHash> FinalizeAsync()
+		{
+			await FlushAsync();
+			if (Node.Entries.Count == 1)
+			{
+				return Node.Entries[0].Hash;
+			}
+			else
+			{
+				return await TreePack.AddNodeAsync(Node.ToByteArray());
+			}
+		}
+	}
+}
