Horde: Missing file.

diff --git a/Engine/Source/Programs/Shared/EpicGames.Core/Sha1Hash.cs b/Engine/Source/Programs/Shared/EpicGames.Core/Sha1Hash.cs
new file mode 100644
--- /dev/null
+++ b/Engine/Source/Programs/Shared/EpicGames.Core/Sha1Hash.cs
@@ -0,0 +1,238 @@
+// Copyright Epic Games, Inc. All Rights Reserved.
+
+using System;
+using System.Collections.Generic;
+using System.Buffers;
+using System.Security.Cryptography;
+using System.Text;
+using System.Buffers.Binary;
+using System.Diagnostics.CodeAnalysis;
+using System.ComponentModel;
+using System.Globalization;
+using System.Text.Json.Serialization;
+using System.Text.Json;
+
+namespace EpicGames.Core
+{
+	/// <summary>
+	/// Struct representing a strongly typed Sha1Hash value.
+	/// </summary>
+	[JsonConverter(typeof(Sha1HashJsonConverter))]
+	[TypeConverter(typeof(Sha1HashTypeConverter))]
+	public struct Sha1Hash : IEquatable<Sha1Hash>, IComparable<Sha1Hash>
+	{
+		/// <summary>
+		/// Length of an Sha1Hash
+		/// </summary>
+		public const int NumBytes = 20;
+
+		/// <summary>
+		/// Length of the hash in bits
+		/// </summary>
+		public const int NumBits = NumBytes * 8;
+
+		ulong A;
+		ulong B;
+		uint C;
+
+		/// <summary>
+		/// Hash consisting of zeroes
+		/// </summary>
+		public static Sha1Hash Zero { get; } = new Sha1Hash(0, 0, 0);
+
+		/// <summary>
+		/// Constructor
+		/// </summary>
+		/// <param name="Memory">Memory to construct from</param>
+		public Sha1Hash(ReadOnlySpan<byte> Span)
+			: this(BinaryPrimitives.ReadUInt64BigEndian(Span), BinaryPrimitives.ReadUInt64BigEndian(Span.Slice(8)), BinaryPrimitives.ReadUInt32BigEndian(Span.Slice(16)))
+		{
+		}
+
+		/// <summary>
+		/// Constructor
+		/// </summary>
+		/// <param name="Memory">Memory to construct from</param>
+		public Sha1Hash(ulong A, ulong B, uint C)
+		{
+			this.A = A;
+			this.B = B;
+			this.C = C;
+		}
+
+		/// <summary>
+		/// Creates a content hash for a block of data, using a given algorithm.
+		/// </summary>
+		/// <param name="Data">Data to compute the hash for</param>
+		/// <returns>New content hash instance containing the hash of the data</returns>
+		public static Sha1Hash Compute(ReadOnlySpan<byte> Data)
+		{
+			byte[] Output = new byte[20];
+			using (SHA1 Sha1 = SHA1.Create())
+			{
+				int BytesWritten;
+				if (!Sha1.TryComputeHash(Data, Output, out BytesWritten) || BytesWritten != NumBytes)
+				{
+					throw new Exception($"Unable to hash data");
+				}
+			}
+			return new Sha1Hash(Output);
+		}
+
+		/// <summary>
+		/// Parses a digest from the given hex string
+		/// </summary>
+		/// <param name="Text"></param>
+		/// <returns></returns>
+		public static Sha1Hash Parse(string Text)
+		{
+			return new Sha1Hash(StringUtils.ParseHexString(Text));
+		}
+
+		/// <summary>
+		/// Parses a digest from the given hex string
+		/// </summary>
+		/// <param name="Text"></param>
+		/// <returns></returns>
+		public static Sha1Hash Parse(ReadOnlySpan<byte> Text)
+		{
+			return new Sha1Hash(StringUtils.ParseHexString(Text));
+		}
+
+		/// <inheritdoc cref="IComparable{T}.CompareTo(T)"/>
+		public int CompareTo(Sha1Hash Other)
+		{
+			if (A != Other.A)
+			{
+				return (A < Other.A) ? -1 : +1;
+			}
+			else if (B != Other.B)
+			{
+				return (B < Other.B) ? -1 : +1;
+			}
+			else
+			{
+				return (C < Other.C) ? -1 : +1;
+			}
+		}
+
+		/// <inheritdoc/>
+		public bool Equals(Sha1Hash Other) => A == Other.A && B == Other.B && C == Other.C;
+
+		/// <inheritdoc/>
+		public override bool Equals(object? Obj) => (Obj is Sha1Hash Hash) && Equals(Hash);
+
+		/// <inheritdoc/>
+		public override int GetHashCode() => (int)A;
+
+		/// <inheritdoc/>
+		public Utf8String ToUtf8String() => StringUtils.FormatUtf8HexString(ToByteArray());
+
+		/// <inheritdoc/>
+		public override string ToString() => StringUtils.FormatHexString(ToByteArray());
+
+		public byte[] ToByteArray()
+		{
+			byte[] Data = new byte[NumBytes];
+			CopyTo(Data);
+			return Data;
+		}
+
+		/// <summary>
+		/// Copies this hash into a span
+		/// </summary>
+		/// <param name="Span"></param>
+		public void CopyTo(Span<byte> Span)
+		{
+			BinaryPrimitives.WriteUInt64BigEndian(Span, A);
+			BinaryPrimitives.WriteUInt64BigEndian(Span[8..], B);
+			BinaryPrimitives.WriteUInt32BigEndian(Span[16..], C);
+		}
+
+		/// <summary>
+		/// Test two hash values for equality
+		/// </summary>
+		public static bool operator ==(Sha1Hash A, Sha1Hash B) => A.Equals(B);
+
+		/// <summary>
+		/// Test two hash values for equality
+		/// </summary>
+		public static bool operator !=(Sha1Hash A, Sha1Hash B) => !(A == B);
+
+		/// <summary>
+		/// Tests whether A > B
+		/// </summary>
+		public static bool operator >(Sha1Hash A, Sha1Hash B) => A.CompareTo(B) > 0;
+
+		/// <summary>
+		/// Tests whether A is less than B
+		/// </summary>
+		public static bool operator <(Sha1Hash A, Sha1Hash B) => A.CompareTo(B) < 0;
+
+		/// <summary>
+		/// Tests whether A is greater than or equal to B
+		/// </summary>
+		public static bool operator >=(Sha1Hash A, Sha1Hash B) => A.CompareTo(B) >= 0;
+
+		/// <summary>
+		/// Tests whether A is less than or equal to B
+		/// </summary>
+		public static bool operator <=(Sha1Hash A, Sha1Hash B) => A.CompareTo(B) <= 0;
+	}
+
+	/// <summary>
+	/// Extension methods for dealing with Sha1Hash values
+	/// </summary>
+	public static class Sha1HashExtensions
+	{
+		/// <summary>
+		/// Read an <see cref="Sha1Hash"/> from a memory reader
+		/// </summary>
+		/// <param name="Reader"></param>
+		/// <returns></returns>
+		public static Sha1Hash ReadSha1Hash(this MemoryReader Reader)
+		{
+			return new Sha1Hash(Reader.ReadFixedLengthBytes(Sha1Hash.NumBytes).Span);
+		}
+
+		/// <summary>
+		/// Write an <see cref="Sha1Hash"/> to a memory writer
+		/// </summary>
+		/// <param name="Writer"></param>
+		/// <param name="Hash"></param>
+		public static void WriteSha1Hash(this MemoryWriter Writer, Sha1Hash Hash)
+		{
+			Hash.CopyTo(Writer.AllocateSpan(Sha1Hash.NumBytes));
+		}
+	}
+
+	/// <summary>
+	/// Type converter for Sha1Hash to and from JSON
+	/// </summary>
+	sealed class Sha1HashJsonConverter : JsonConverter<Sha1Hash>
+	{
+		/// <inheritdoc/>
+		public override Sha1Hash Read(ref Utf8JsonReader Reader, Type TypeToConvert, JsonSerializerOptions Options) => Sha1Hash.Parse(Reader.ValueSpan);
+
+		/// <inheritdoc/>
+		public override void Write(Utf8JsonWriter Writer, Sha1Hash Value, JsonSerializerOptions Options) => Writer.WriteStringValue(Value.ToUtf8String().Span);
+	}
+
+	/// <summary>
+	/// Type converter from strings to Sha1Hash objects
+	/// </summary>
+	sealed class Sha1HashTypeConverter : TypeConverter
+	{
+		/// <inheritdoc/>
+		public override bool CanConvertFrom(ITypeDescriptorContext Context, Type SourceType)
+		{
+			return SourceType == typeof(string);
+		}
+
+		/// <inheritdoc/>
+		public override object ConvertFrom(ITypeDescriptorContext Context, CultureInfo Culture, object Value)
+		{
+			return Sha1Hash.Parse((string)Value);
+		}
+	}
+}
