UGS: Fix exception when opening a file from Perforce that doesn't exist on the local disk.

diff --git a/Engine/Source/Programs/Shared/EpicGames.Core/DirectoryUtils.cs b/Engine/Source/Programs/Shared/EpicGames.Core/DirectoryUtils.cs
--- a/Engine/Source/Programs/Shared/EpicGames.Core/DirectoryUtils.cs
+++ b/Engine/Source/Programs/Shared/EpicGames.Core/DirectoryUtils.cs
@@ -37,20 +37,23 @@ namespace EpicGames.Core
 			else
 			{
 				ParentInfo = FindCorrectCase(ParentInfo);
-				try
+				if (Info.Exists)
 				{
-					foreach (DirectoryInfo ChildInfo in ParentInfo.EnumerateDirectories())
+					try
 					{
-						if (String.Equals(ChildInfo.Name, Info.Name, DirectoryReference.Comparison))
+						foreach (DirectoryInfo ChildInfo in ParentInfo.EnumerateDirectories())
 						{
-							return ChildInfo;
+							if (String.Equals(ChildInfo.Name, Info.Name, DirectoryReference.Comparison))
+							{
+								return ChildInfo;
+							}
 						}
 					}
-				}
-				catch (Exception)
-				{
-					// System.Security.SecurityException is expected here if we try to enumerate a directory
-					// we don't have permission to read, such as another user's mounted volume.
+					catch (Exception)
+					{
+						// System.Security.SecurityException is expected here if we try to enumerate a directory
+						// we don't have permission to read, such as another user's mounted volume.
+					}
 				}
 
 				return new DirectoryInfo(Path.Combine(ParentInfo.FullName, Info.Name));
diff --git a/Engine/Source/Programs/Shared/EpicGames.Core/FileUtils.cs b/Engine/Source/Programs/Shared/EpicGames.Core/FileUtils.cs
--- a/Engine/Source/Programs/Shared/EpicGames.Core/FileUtils.cs
+++ b/Engine/Source/Programs/Shared/EpicGames.Core/FileUtils.cs
@@ -86,11 +86,14 @@ namespace EpicGames.Core
 		public static FileInfo FindCorrectCase(FileInfo Info)
 		{
 			DirectoryInfo ParentInfo = DirectoryUtils.FindCorrectCase(Info.Directory);
-			foreach (FileInfo ChildInfo in ParentInfo.EnumerateFiles())
+			if (Info.Exists)
 			{
-				if (String.Equals(ChildInfo.Name, Info.Name, FileReference.Comparison))
+				foreach (FileInfo ChildInfo in ParentInfo.EnumerateFiles())
 				{
-					return ChildInfo;
+					if (String.Equals(ChildInfo.Name, Info.Name, FileReference.Comparison))
+					{
+						return ChildInfo;
+					}
 				}
 			}
 			return new FileInfo(Path.Combine(ParentInfo.FullName, Info.Name));
