Implement IComparable<> for FileReference and DirectoryReference, allowing them to be sorted in lists.

diff --git a/Engine/Source/Programs/Shared/EpicGames.Core/DirectoryReference.cs b/Engine/Source/Programs/Shared/EpicGames.Core/DirectoryReference.cs
--- a/Engine/Source/Programs/Shared/EpicGames.Core/DirectoryReference.cs
+++ b/Engine/Source/Programs/Shared/EpicGames.Core/DirectoryReference.cs
@@ -14,7 +14,7 @@ namespace EpicGames.Core
 	/// Representation of an absolute directory path. Allows fast hashing and comparisons.
 	/// </summary>
 	[Serializable]
-	public class DirectoryReference : FileSystemReference, IEquatable<DirectoryReference>
+	public class DirectoryReference : FileSystemReference, IEquatable<DirectoryReference>, IComparable<DirectoryReference>
 	{
 		/// <summary>
 		/// Special value used to invoke the non-sanitizing constructor overload
@@ -206,6 +206,9 @@ namespace EpicGames.Core
 			return Comparer.GetHashCode(FullName);
 		}
 
+		/// <inheritdoc/>
+		public int CompareTo(DirectoryReference? Other) => Comparer.Compare(FullName, Other?.FullName);
+
 		/// <summary>
 		/// Helper function to create a remote directory reference. Unlike normal DirectoryReference objects, these aren't converted to a full path in the local filesystem.
 		/// </summary>
diff --git a/Engine/Source/Programs/Shared/EpicGames.Core/FileReference.cs b/Engine/Source/Programs/Shared/EpicGames.Core/FileReference.cs
--- a/Engine/Source/Programs/Shared/EpicGames.Core/FileReference.cs
+++ b/Engine/Source/Programs/Shared/EpicGames.Core/FileReference.cs
@@ -15,7 +15,7 @@ namespace EpicGames.Core
 	/// Representation of an absolute file path. Allows fast hashing and comparisons.
 	/// </summary>
 	[Serializable]
-	public class FileReference : FileSystemReference, IEquatable<FileReference>
+	public class FileReference : FileSystemReference, IEquatable<FileReference>, IComparable<FileReference>
 	{
 		/// <summary>
 		/// Dummy enum to allow invoking the constructor which takes a sanitized full path
@@ -239,6 +239,9 @@ namespace EpicGames.Core
 			return Comparer.GetHashCode(FullName);
 		}
 
+		/// <inheritdoc/>
+		public int CompareTo(FileReference? Other) => Comparer.Compare(FullName, Other?.FullName);
+
 		/// <summary>
 		/// Helper function to create a remote file reference. Unlike normal FileReference objects, these aren't converted to a full path in the local filesystem, but are
 		/// left as they are passed in.
