Horde: Move generic functionality for hosting different commands within an application into EpicGames.Core. Simplifies command line parsing, error handling, help output, etc...

diff --git a/Engine/Source/Programs/AutomationTool/AutomationUtils/Automation.cs b/Engine/Source/Programs/AutomationTool/AutomationUtils/Automation.cs
--- a/Engine/Source/Programs/AutomationTool/AutomationUtils/Automation.cs
+++ b/Engine/Source/Programs/AutomationTool/AutomationUtils/Automation.cs
@@ -305,7 +305,7 @@ namespace AutomationTool
 @"Executes scripted commands
 
 AutomationTool.exe [-verbose] [-compileonly] [-p4] Command0 [-Arg0 -Arg1 -Arg2 ...] Command1 [-Arg0 -Arg1 ...] Command2 [-Arg0 ...] Commandn ... [EnvVar0=MyValue0 ... EnvVarn=MyValuen]",
-				ParamDict.ToList(), HelpUtils.WindowWidth, Log.Logger);
+				ParamDict.ToList());
 			CommandUtils.LogHelp(typeof(Automation));
 		}
 
diff --git a/Engine/Source/Programs/AutomationTool/AutomationUtils/HelpUtils.cs b/Engine/Source/Programs/AutomationTool/AutomationUtils/HelpUtils.cs
--- a/Engine/Source/Programs/AutomationTool/AutomationUtils/HelpUtils.cs
+++ b/Engine/Source/Programs/AutomationTool/AutomationUtils/HelpUtils.cs
@@ -146,7 +146,7 @@ namespace AutomationTool
 			}
 
 			Log.TraceInformation("");
-			HelpUtils.PrintHelp(String.Format("{0} Help:", Command.Name), Description, ParamDict.ToList(), HelpUtils.WindowWidth, Log.Logger);
+			HelpUtils.PrintHelp(String.Format("{0} Help:", Command.Name), Description, ParamDict.ToList());
 		}
 
 		/// <summary>
diff --git a/Engine/Source/Programs/Horde/Horde.Agent/Commands/Command.cs b/Engine/Source/Programs/Horde/Horde.Agent/Commands/Command.cs
deleted file mode 100644
diff --git a/Engine/Source/Programs/Horde/Horde.Agent/Commands/CommandAttribute.cs b/Engine/Source/Programs/Horde/Horde.Agent/Commands/CommandAttribute.cs
deleted file mode 100644
diff --git a/Engine/Source/Programs/Horde/Horde.Agent/Commands/Service/UninstallCommand.cs b/Engine/Source/Programs/Horde/Horde.Agent/Commands/Service/UninstallCommand.cs
--- a/Engine/Source/Programs/Horde/Horde.Agent/Commands/Service/UninstallCommand.cs
+++ b/Engine/Source/Programs/Horde/Horde.Agent/Commands/Service/UninstallCommand.cs
@@ -1,5 +1,6 @@
 // Copyright Epic Games, Inc. All Rights Reserved.
 
+using EpicGames.Core;
 using HordeAgent.Utility;
 using Microsoft.Extensions.Logging;
 using System;
diff --git a/Engine/Source/Programs/Horde/Horde.Agent/Commands/Workspace/ClearCommand.cs b/Engine/Source/Programs/Horde/Horde.Agent/Commands/Workspace/ClearCommand.cs
--- a/Engine/Source/Programs/Horde/Horde.Agent/Commands/Workspace/ClearCommand.cs
+++ b/Engine/Source/Programs/Horde/Horde.Agent/Commands/Workspace/ClearCommand.cs
@@ -1,5 +1,6 @@
 // Copyright Epic Games, Inc. All Rights Reserved.
 
+using EpicGames.Core;
 using EpicGames.Perforce;
 using EpicGames.Perforce.Managed;
 using Microsoft.Extensions.Logging;
diff --git a/Engine/Source/Programs/Horde/Horde.Agent/Commands/Workspace/DumpCommand.cs b/Engine/Source/Programs/Horde/Horde.Agent/Commands/Workspace/DumpCommand.cs
--- a/Engine/Source/Programs/Horde/Horde.Agent/Commands/Workspace/DumpCommand.cs
+++ b/Engine/Source/Programs/Horde/Horde.Agent/Commands/Workspace/DumpCommand.cs
@@ -1,5 +1,6 @@
 // Copyright Epic Games, Inc. All Rights Reserved.
 
+using EpicGames.Core;
 using EpicGames.Perforce;
 using EpicGames.Perforce.Managed;
 using Microsoft.Extensions.Logging;
diff --git a/Engine/Source/Programs/Horde/Horde.Agent/Commands/Workspace/PopulateCacheCommand.cs b/Engine/Source/Programs/Horde/Horde.Agent/Commands/Workspace/PopulateCacheCommand.cs
--- a/Engine/Source/Programs/Horde/Horde.Agent/Commands/Workspace/PopulateCacheCommand.cs
+++ b/Engine/Source/Programs/Horde/Horde.Agent/Commands/Workspace/PopulateCacheCommand.cs
@@ -14,7 +14,7 @@ using EpicGames.Perforce;
 
 namespace HordeAgent.Commands.Workspace
 {
-	[Command("Workspace", "PopulateWorkspaceCache", "Populates the cache with the head revision of the given streams")]
+	[Command("Workspace", "PopulateCache", "Populates the cache with the head revision of the given streams")]
 	class PopulateCacheCommand : WorkspaceCommand
 	{
 		[CommandLine("-ClientAndStream=")]
diff --git a/Engine/Source/Programs/Horde/Horde.Agent/Commands/Workspace/RepairCommand.cs b/Engine/Source/Programs/Horde/Horde.Agent/Commands/Workspace/RepairCommand.cs
--- a/Engine/Source/Programs/Horde/Horde.Agent/Commands/Workspace/RepairCommand.cs
+++ b/Engine/Source/Programs/Horde/Horde.Agent/Commands/Workspace/RepairCommand.cs
@@ -1,5 +1,6 @@
 // Copyright Epic Games, Inc. All Rights Reserved.
 
+using EpicGames.Core;
 using EpicGames.Perforce;
 using EpicGames.Perforce.Managed;
 using Microsoft.Extensions.Logging;
diff --git a/Engine/Source/Programs/Horde/Horde.Agent/Commands/Workspace/StatusCommand.cs b/Engine/Source/Programs/Horde/Horde.Agent/Commands/Workspace/StatusCommand.cs
--- a/Engine/Source/Programs/Horde/Horde.Agent/Commands/Workspace/StatusCommand.cs
+++ b/Engine/Source/Programs/Horde/Horde.Agent/Commands/Workspace/StatusCommand.cs
@@ -1,5 +1,6 @@
 // Copyright Epic Games, Inc. All Rights Reserved.
 
+using EpicGames.Core;
 using EpicGames.Perforce;
 using EpicGames.Perforce.Managed;
 using Microsoft.Extensions.Logging;
diff --git a/Engine/Source/Programs/Horde/Horde.Agent/Program.cs b/Engine/Source/Programs/Horde/Horde.Agent/Program.cs
--- a/Engine/Source/Programs/Horde/Horde.Agent/Program.cs
+++ b/Engine/Source/Programs/Horde/Horde.Agent/Program.cs
@@ -24,6 +24,11 @@ using System.Text.Json;
 using System.Threading;
 using System.Threading.Tasks;
 using EpicGames.Core;
+using EpicGames.Horde.Storage;
+using EpicGames.Serialization;
+using System.Text;
+using EpicGames.Perforce;
+using System.Security.Cryptography;
 
 namespace HordeAgent
 {
@@ -57,14 +62,6 @@ namespace HordeAgent
 		/// </summary>
 		public static string Version { get; } = GetVersion();
 
-		/// <summary>
-		/// Width to use for printing out help
-		/// </summary>
-		static int HelpWidth
-		{
-			get { return HelpUtils.WindowWidth - 20; }
-		}
-
 		/// <summary>
 		/// Entry point
 		/// </summary>
@@ -92,24 +89,6 @@ namespace HordeAgent
 			}
 		}
 
-		static bool MatchCommand(string[] Args, CommandAttribute Attribute)
-		{
-			if(Args.Length < Attribute.Names.Length)
-			{
-				return false;
-			}
-
-			for (int Idx = 0; Idx < Attribute.Names.Length; Idx++)
-			{
-				if (!Attribute.Names[Idx].Equals(Args[Idx], StringComparison.OrdinalIgnoreCase))
-				{
-					return false;
-				}
-			}
-
-			return true;
-		}
-
 		/// <summary>
 		/// Actual Main function, without exception guards
 		/// </summary>
@@ -120,116 +99,6 @@ namespace HordeAgent
 		{
 			// Enable unencrypted HTTP/2 for gRPC channel without TLS
 			AppContext.SetSwitch("System.Net.Http.SocketsHttpHandler.Http2UnencryptedSupport", true);
-			
-			// Find all the command types
-			List<(CommandAttribute, Type)> Commands = new List<(CommandAttribute, Type)>();
-			foreach (Type Type in Assembly.GetExecutingAssembly().GetTypes())
-			{
-				CommandAttribute? Attribute = Type.GetCustomAttribute<CommandAttribute>();
-				if (Attribute != null)
-				{
-					Commands.Add((Attribute, Type));
-				}
-			}
-
-			// Check if there's a matching command
-			Type? CommandType = null;
-			CommandAttribute? CommandAttribute = null;
-
-			string[] DefaultArgs = new string[] { "Service", "Run" };
-;			Type? DefaultCommandType = null;
-			CommandAttribute? DefaultCommandAttribute = null;
-
-			foreach ((CommandAttribute Attribute, Type Type) in Commands.OrderBy(x => x.Item1.Names.Length))
-			{
-				if (MatchCommand(Args, Attribute))
-				{
-					CommandType = Type;
-					CommandAttribute = Attribute;
-				}
-
-				if (MatchCommand(DefaultArgs, Attribute))
-				{
-					DefaultCommandType = Type;
-					DefaultCommandAttribute = Attribute;
-				}
-
-			}
-
-			// Check if there are any commands specified on the command line.
-			if (CommandType == null)
-			{
-				if (Args.Length > 0 && !Args[0].StartsWith("-"))
-				{
-					Logger.LogError("Invalid command");
-					Logger.LogInformation("");
-					Logger.LogInformation("Available commands");
-
-					PrintCommands(Commands.Select(x => x.Item1), Logger);
-					return 1;
-				}
-				else if (Args.Length > 0 && Args[0].Equals("-Help", StringComparison.OrdinalIgnoreCase))
-				{
-					Logger.LogInformation("HordeAgent");
-					Logger.LogInformation("");
-					Logger.LogInformation("Utility for managing automated processes on build machines.");
-					Logger.LogInformation("");
-					Logger.LogInformation("Usage:");
-					Logger.LogInformation("    HordeAgent.exe [Command] [-Option1] [-Option2]...");
-					Logger.LogInformation("");
-					Logger.LogInformation("Commands:");
-
-					PrintCommands(Commands.Select(x => x.Item1), Logger);
-
-					Logger.LogInformation("");
-					Logger.LogInformation("Specify \"Command -Help\" for command-specific help");
-					return 0;
-				}
-				else
-				{
-					CommandType = DefaultCommandType;
-					CommandAttribute = DefaultCommandAttribute;
-
-					if (CommandType == null || CommandAttribute == null)
-					{
-						Logger.LogError("Invalid default command");
-						return 1;
-					}
-				}
-			}
-
-			// Extract the 
-			string CommandName = String.Join(" ", CommandAttribute!.Names);
-
-			// Build an argument list for the command, including all the global arguments as well as arguments until the next command
-			CommandLineArguments CommandArguments = new CommandLineArguments(Args.Skip(CommandAttribute.Names.Length).ToArray());
-
-			// Create the command instance
-			Command Command = (Command)Activator.CreateInstance(CommandType)!;
-
-			// If the help flag is specified, print the help info and exit immediately
-			if (CommandArguments.HasOption("-Help"))
-			{
-				HelpUtils.PrintHelp(CommandName, CommandAttribute.Description, Command.GetParameters(CommandArguments), HelpWidth, Logger);
-				return 1;
-			}
-
-			// Configure the command
-			try
-			{
-				Command.Configure(CommandArguments, Logger);
-				CommandArguments.CheckAllArgumentsUsed(Logger);
-			}
-			catch (CommandLineArgumentException Ex)
-			{
-				Logger.LogError("{0}: {1}", CommandName, Ex.Message);
-
-				Logger.LogInformation("");
-				Logger.LogInformation("Arguments for {0}:", CommandName);
-
-				HelpUtils.PrintTable(Command.GetParameters(CommandArguments), 4, 24, HelpWidth, Logger);
-				return 1;
-			}
 
 			if (RuntimeInformation.IsOSPlatform(OSPlatform.Windows))
 			{
@@ -243,22 +112,7 @@ namespace HordeAgent
 			}
 
 			// Execute all the commands
-			return await Command.ExecuteAsync(Logger);
-		}
-
-		/// <summary>
-		/// Print a formatted list of all the available commands
-		/// </summary>
-		/// <param name="Attributes">List of command attributes</param>
-		/// <param name="Logger">The logging output device</param>
-		static void PrintCommands(IEnumerable<CommandAttribute> Attributes, ILogger Logger)
-		{
-			List<KeyValuePair<string, string>> Commands = new List<KeyValuePair<string, string>>();
-			foreach(CommandAttribute Attribute in Attributes)
-			{
-				Commands.Add(new KeyValuePair<string, string>(String.Join(" ", Attribute.Names), Attribute.Description));
-			}
-			HelpUtils.PrintTable(Commands.OrderBy(x => x.Key).ToList(), 4, 20, HelpWidth, Logger);
+			return await CommandHost.RunAsync(new CommandLineArguments(Args), typeof(HordeAgent.Modes.Service.RunCommand), Logger);
 		}
 
 		/// <summary>
diff --git a/Engine/Source/Programs/Shared/EpicGames.BuildGraph/BgGraph.cs b/Engine/Source/Programs/Shared/EpicGames.BuildGraph/BgGraph.cs
--- a/Engine/Source/Programs/Shared/EpicGames.BuildGraph/BgGraph.cs
+++ b/Engine/Source/Programs/Shared/EpicGames.BuildGraph/BgGraph.cs
@@ -658,7 +658,14 @@ namespace EpicGames.BuildGraph
 					Logger.LogInformation("");
 					Logger.LogInformation("Options:");
 					Logger.LogInformation("");
-					HelpUtils.PrintTable(Parameters, 4, 24, Logger);
+
+					List<string> Lines = new List<string>();
+					HelpUtils.FormatTable(Parameters, 4, 24, HelpUtils.WindowWidth - 20, Lines);
+
+					foreach (string Line in Lines)
+					{
+						Logger.Log(LogLevel.Information, Line);
+					}
 				}
 			}
 
diff --git a/Engine/Source/Programs/Shared/EpicGames.Core/CommandHost.cs b/Engine/Source/Programs/Shared/EpicGames.Core/CommandHost.cs
new file mode 100644
--- /dev/null
+++ b/Engine/Source/Programs/Shared/EpicGames.Core/CommandHost.cs
@@ -0,0 +1,244 @@
+// Copyright Epic Games, Inc. All Rights Reserved.
+
+using Microsoft.Extensions.Logging;
+using System;
+using System.Collections.Generic;
+using System.Linq;
+using System.Reflection;
+using System.Text;
+using System.Threading.Tasks;
+
+namespace EpicGames.Core
+{
+	/// <summary>
+	/// Attribute used to specify names of program modes, and help text
+	/// </summary>
+	public class CommandAttribute : Attribute
+	{
+		/// <summary>
+		/// Names for this command
+		/// </summary>
+		public string[] Names;
+
+		/// <summary>
+		/// Short description for the mode. Will be displayed in the help text.
+		/// </summary>
+		public string Description;
+
+		/// <summary>
+		/// Constructor
+		/// </summary>
+		/// <param name="Name">Name of the mode</param>
+		/// <param name="Description">Short description for display in the help text</param>
+		public CommandAttribute(string Name, string Description)
+		{
+			this.Names = new string[] { Name };
+			this.Description = Description;
+		}
+
+		/// <summary>
+		/// Constructor
+		/// </summary>
+		/// <param name="Category">Category for this command</param>
+		/// <param name="Name">Name of the mode</param>
+		/// <param name="Description">Short description for display in the help text</param>
+		public CommandAttribute(string Category, string Name, string Description)
+		{
+			this.Names = new string[] { Category, Name };
+			this.Description = Description;
+		}
+	}
+
+	/// <summary>
+	/// Base class for all commands that can be executed by HordeAgent
+	/// </summary>
+	public abstract class Command
+	{
+		/// <summary>
+		/// Configure this object with the given command line arguments
+		/// </summary>
+		/// <param name="Arguments">Command line arguments</param>
+		/// <param name="Logger">Logging output device</param>
+		public virtual void Configure(CommandLineArguments Arguments, ILogger Logger)
+		{
+			Arguments.ApplyTo(this, Logger);
+		}
+
+		/// <summary>
+		/// Gets all command line parameters to show in help for this command
+		/// </summary>
+		/// <param name="Arguments">The command line arguments</param>
+		/// <returns>List of name/description pairs</returns>
+		public virtual List<KeyValuePair<string, string>> GetParameters(CommandLineArguments Arguments)
+		{
+			return CommandLineArguments.GetParameters(GetType());
+		}
+
+		/// <summary>
+		/// Execute this command
+		/// </summary>
+		/// <param name="Logger">The logger to use for this command</param>
+		/// <returns>Exit code</returns>
+		public abstract Task<int> ExecuteAsync(ILogger Logger);
+	}
+
+	/// <summary>
+	/// Entry point for dispatching commands
+	/// </summary>
+	public static class CommandHost
+	{
+		/// <summary>
+		/// Entry point for executing registered command types in a particular assembly
+		/// </summary>
+		/// <param name="Args">Command line arguments</param>
+		/// <param name="DefaultCommandType">The default command type</param>
+		/// <param name="Logger">Logging device to pass to the command object</param>
+		/// <returns>Return code from the command</returns>
+		public static Task<int> RunAsync(CommandLineArguments Args, Type? DefaultCommandType, ILogger Logger)
+		{
+			return RunAsync(Args, Assembly.GetEntryAssembly()!, DefaultCommandType, Logger);
+		}
+
+		/// <summary>
+		/// Entry point for executing registered command types in a particular assembly
+		/// </summary>
+		/// <param name="Args">Command line arguments</param>
+		/// <param name="CommandAssembly">Assembly to scan for command types</param>
+		/// <param name="DefaultCommandType">The default command type</param>
+		/// <param name="Logger">Logging device to pass to the command object</param>
+		/// <returns>Return code from the command</returns>
+		public static async Task<int> RunAsync(CommandLineArguments Args, Assembly CommandAssembly, Type? DefaultCommandType, ILogger Logger)
+		{
+			// Find all the command types
+			List<(CommandAttribute, Type)> Commands = new List<(CommandAttribute, Type)>();
+			foreach (Type Type in CommandAssembly.GetTypes())
+			{
+				CommandAttribute? Attribute = Type.GetCustomAttribute<CommandAttribute>();
+				if (Attribute != null)
+				{
+					Commands.Add((Attribute, Type));
+				}
+			}
+
+			// Check if there's a matching command
+			Type? CommandType = null;
+			CommandAttribute? CommandAttribute = null;
+
+			// Parse the positional arguments for the command name
+			string[] PositionalArgs = Args.GetPositionalArguments();
+			if (PositionalArgs.Length == 0)
+			{
+				if (DefaultCommandType == null || Args.HasOption("-Help"))
+				{
+					Console.WriteLine(CommandAssembly.GetName().Name);
+					Console.WriteLine("");
+
+					AssemblyDescriptionAttribute? Description = CommandAssembly.GetCustomAttribute<AssemblyDescriptionAttribute>();
+					if (Description != null)
+					{
+						Console.WriteLine(Description.Description);
+					}
+
+					Console.WriteLine("Usage:");
+					Console.WriteLine("    [Command] [-Option1] [-Option2]...");
+					Console.WriteLine("");
+					Console.WriteLine("Commands:");
+
+					PrintCommands(Commands.Select(x => x.Item1));
+
+					Console.WriteLine("");
+					Console.WriteLine("Specify \"<CommandName> -Help\" for command-specific help");
+					return 0;
+				}
+				else 
+				{
+					CommandType = DefaultCommandType;
+				}
+			}
+			else
+			{
+				foreach ((CommandAttribute Attribute, Type Type) in Commands.OrderBy(x => x.Item1.Names.Length))
+				{
+					if (Attribute.Names.SequenceEqual(PositionalArgs, StringComparer.OrdinalIgnoreCase))
+					{
+						CommandType = Type;
+						CommandAttribute = Attribute;
+						break;
+					}
+				}
+				if (CommandType == null)
+				{
+					ConsoleUtils.WriteError($"Invalid command '{String.Join(" ", PositionalArgs)}'");
+					Console.WriteLine("");
+					Console.WriteLine("Available commands:");
+
+					PrintCommands(Commands.Select(x => x.Item1));
+					return 1;
+				}
+			}
+
+			// Create the command instance
+			Command Command = (Command)Activator.CreateInstance(CommandType)!;
+
+			// If the help flag is specified, print the help info and exit immediately
+			if (Args.HasOption("-Help"))
+			{
+				if (CommandAttribute == null)
+				{
+					HelpUtils.PrintHelp(null, null, Command.GetParameters(Args));
+				}
+				else
+				{
+					HelpUtils.PrintHelp(String.Join(" ", CommandAttribute.Names), CommandAttribute.Description, Command.GetParameters(Args));
+				}
+				return 1;
+			}
+
+			// Configure the command
+			try
+			{
+				Command.Configure(Args, Logger);
+				Args.CheckAllArgumentsUsed(Logger);
+			}
+			catch (CommandLineArgumentException Ex)
+			{
+				ConsoleUtils.WriteError(Ex.Message);
+				Console.WriteLine("");
+				Console.WriteLine("Valid parameters:");
+
+				HelpUtils.PrintTable(Command.GetParameters(Args), 4, 24);
+				return 1;
+			}
+
+			// Execute all the commands
+			try
+			{
+				return await Command.ExecuteAsync(Logger);
+			}
+			catch (FatalErrorException Ex)
+			{
+				Logger.LogCritical(Ex, "Fatal error.");
+				return Ex.ExitCode;
+			}
+			catch (Exception Ex)
+			{
+				Logger.LogCritical(Ex, "Fatal error.");
+				return 1;
+			}
+		}
+
+		/// <summary>
+		/// Print a formatted list of all the available commands
+		/// </summary>
+		/// <param name="Attributes">List of command attributes</param>
+		static void PrintCommands(IEnumerable<CommandAttribute> Attributes)
+		{
+			List<KeyValuePair<string, string>> Commands = new List<KeyValuePair<string, string>>();
+			foreach (CommandAttribute Attribute in Attributes)
+			{
+				Commands.Add(new KeyValuePair<string, string>(String.Join(" ", Attribute.Names), Attribute.Description));
+			}
+			HelpUtils.PrintTable(Commands.OrderBy(x => x.Key).ToList(), 4, 20);
+		}
+	}
+}
diff --git a/Engine/Source/Programs/Shared/EpicGames.Core/ConsoleUtils.cs b/Engine/Source/Programs/Shared/EpicGames.Core/ConsoleUtils.cs
--- a/Engine/Source/Programs/Shared/EpicGames.Core/ConsoleUtils.cs
+++ b/Engine/Source/Programs/Shared/EpicGames.Core/ConsoleUtils.cs
@@ -60,5 +60,27 @@ namespace EpicGames.Core
 				Console.WriteLine(Line);
 			}
 		}
+
+		/// <summary>
+		/// Writes an colored warning message to the console
+		/// </summary>
+		/// <param name="Text">The message to output</param>
+		public static void WriteWarning(string Text)
+		{
+			Console.ForegroundColor = ConsoleColor.Yellow;
+			Console.WriteLine(Text);
+			Console.ResetColor();
+		}
+
+		/// <summary>
+		/// Writes an colored error message to the console
+		/// </summary>
+		/// <param name="Text">The message to output</param>
+		public static void WriteError(string Text)
+		{
+			Console.ForegroundColor = ConsoleColor.Red;
+			Console.WriteLine(Text);
+			Console.ResetColor();
+		}
 	}
 }
diff --git a/Engine/Source/Programs/Shared/EpicGames.Core/HelpUtils.cs b/Engine/Source/Programs/Shared/EpicGames.Core/HelpUtils.cs
--- a/Engine/Source/Programs/Shared/EpicGames.Core/HelpUtils.cs
+++ b/Engine/Source/Programs/Shared/EpicGames.Core/HelpUtils.cs
@@ -24,9 +24,9 @@ namespace EpicGames.Core
 		/// Prints help for the given object type
 		/// </summary>
 		/// <param name="Type">Type to print help for</param>
-		public static void PrintHelp(string Title, Type Type, ILogger Logger)
+		public static void PrintHelp(string Title, Type Type)
 		{
-			PrintHelp(Title, GetDescription(Type), CommandLineArguments.GetParameters(Type), WindowWidth, Logger);
+			PrintHelp(Title, GetDescription(Type), CommandLineArguments.GetParameters(Type));
 		}
 
 		/// <summary>
@@ -35,14 +35,12 @@ namespace EpicGames.Core
 		/// <param name="Title">Title for the help text</param>
 		/// <param name="Description">Description for the command</param>
 		/// <param name="Parameters">List of parameters</param>
-		/// <param name="MaxWidth">Maximum width for each line</param>
-		/// <param name="Logger">The output logger</param>
-		public static void PrintHelp(string Title, string Description, List<KeyValuePair<string, string>> Parameters, int MaxWidth, ILogger Logger)
+		public static void PrintHelp(string? Title, string? Description, List<KeyValuePair<string, string>> Parameters)
 		{
 			bool bFirstLine = true;
 			if (!String.IsNullOrEmpty(Title))
 			{
-				PrintParagraph(Title, MaxWidth, Logger);
+				PrintParagraph(Title);
 				bFirstLine = false;
 			}
 
@@ -50,9 +48,9 @@ namespace EpicGames.Core
 			{
 				if (!bFirstLine)
 				{
-					Logger.LogInformation("");
+					Console.WriteLine("");
 				}
-				PrintParagraph(Description, MaxWidth, Logger);
+				PrintParagraph(Description);
 				bFirstLine = false;
 			}
 
@@ -60,11 +58,11 @@ namespace EpicGames.Core
 			{
 				if (!bFirstLine)
 				{
-					Logger.LogInformation("");
+					Console.WriteLine("");
 				}
 
-				Logger.LogInformation("Parameters:");
-				PrintTable(Parameters, 4, 24, MaxWidth, Logger);
+				Console.WriteLine("Parameters:");
+				PrintTable(Parameters, 4, 24);
 			}
 		}
 
@@ -92,9 +90,9 @@ namespace EpicGames.Core
 		/// </summary>
 		/// <param name="Text">Text to print</param>
 		/// <param name="Logger">Logger implementation to write to</param>
-		public static void PrintParagraph(string Text, ILogger Logger)
+		public static void PrintParagraph(string Text)
 		{
-			PrintParagraph(Text, WindowWidth - 1, Logger);
+			PrintParagraph(Text, WindowWidth - 1);
 		}
 
 		/// <summary>
@@ -102,33 +100,50 @@ namespace EpicGames.Core
 		/// </summary>
 		/// <param name="Text">Text to print</param>
 		/// <param name="MaxWidth">Maximum width for each line</param>
-		/// <param name="Logger">Logger implementation to write to</param>
-		public static void PrintParagraph(string Text, int MaxWidth, ILogger Logger)
+		public static void PrintParagraph(string Text, int MaxWidth)
 		{
 			IEnumerable<string> Lines = StringUtils.WordWrap(Text, MaxWidth);
 			foreach (string Line in Lines)
 			{
-				Logger.LogInformation(Line);
+				Console.WriteLine(Line);
 			}
 		}
 
 		/// <summary>
-		/// Formats the given parameters as so:
-		///     -Param1     Param1 Description
-		///
-		///     -Param2      Param2 Description, this description is
-		///                  longer and splits onto a separate line. 
-		///
-		///     -Param3      Param3 Description continues as before. 
+		/// Prints an argument list to the console
 		/// </summary>
 		/// <param name="Items">List of parameters arranged as "-ParamName Param Description"</param>
 		/// <param name="Indent">Indent from the left hand side</param>
 		/// <param name="MinFirstColumnWidth">The minimum padding from the start of the param name to the start of the description (resizes with larger param names)</param>
-		/// <param name="Logger">Logger implementation to write to</param>
 		/// <returns></returns>
-		public static void PrintTable(List<KeyValuePair<string, string>> Items, int Indent, int MinFirstColumnWidth, ILogger Logger)
+		public static void PrintTable(List<KeyValuePair<string, string>> Items, int Indent, int MinFirstColumnWidth)
 		{
-			PrintTable(Items, Indent, MinFirstColumnWidth, WindowWidth - 1, Logger);
+			List<string> Lines = new List<string>();
+			FormatTable(Items, Indent, MinFirstColumnWidth, WindowWidth - 1, Lines);
+
+			foreach (string Line in Lines)
+			{
+				Console.WriteLine(Line);
+			}
+		}
+
+		/// <summary>
+		/// Prints a table of items to a logging device
+		/// </summary>
+		/// <param name="Items"></param>
+		/// <param name="Indent"></param>
+		/// <param name="MinFirstColumnWidth"></param>
+		/// <param name="MaxWidth"></param>
+		/// <param name="Logger"></param>
+		public static void PrintTable(List<KeyValuePair<string, string>> Items, int Indent, int MinFirstColumnWidth, int MaxWidth, ILogger Logger)
+		{
+			List<string> Lines = new List<string>();
+			FormatTable(Items, Indent, MinFirstColumnWidth, MaxWidth, Lines);
+
+			foreach (string Line in Lines)
+			{
+				Logger.LogInformation("{Line)", Line);
+			}
 		}
 
 		/// <summary>
@@ -143,9 +158,8 @@ namespace EpicGames.Core
 		/// <param name="Items">List of parameters arranged as "-ParamName Param Description"</param>
 		/// <param name="Indent">Indent from the left hand side</param>
 		/// <param name="MinFirstColumnWidth">The minimum padding from the start of the param name to the start of the description (resizes with larger param names)</param>
-		/// <param name="Logger">Logger implementation to write to</param>
-		/// <returns></returns>
-		public static void PrintTable(List<KeyValuePair<string, string>> Items, int Indent, int MinFirstColumnWidth, int MaxWidth, ILogger Logger)
+		/// <returns>Sequence of formatted lines in the table</returns>
+		public static void FormatTable(IReadOnlyList<KeyValuePair<string, string>> Items, int Indent, int MinFirstColumnWidth, int MaxWidth, List<string> Lines)
 		{
 			if(Items.Count > 0)
 			{
@@ -153,7 +167,7 @@ namespace EpicGames.Core
 				string IndentString = new string(' ', Indent);
 
 				// default the padding value
-				int RightPadding = Math.Max(MinFirstColumnWidth, Items.Max(x => x.Key.Length + 1));
+				int RightPadding = Math.Max(MinFirstColumnWidth, Items.Max(x => x.Key.Length + 2));
 
 				// Build the formatted params
 				foreach(KeyValuePair<string, string> Item in Items)
@@ -168,7 +182,7 @@ namespace EpicGames.Core
 					{
 						// Formatting as following:
 						// <Indent>-param<Right Padding>Description<New line>
-						Logger.LogInformation(ParamString + DescriptionLine);
+						Lines.Add(ParamString + DescriptionLine);
 
 						// we replace the param string on subsequent lines with white space of the same length
 						ParamString = string.Empty.PadRight(IndentString.Length + RightPadding);
