Horde: Use dependency injection for command discovery in CommandHost, and support other commands within the server application too. Makes it much easier to perform utility operations using standard services without needing to register a REST endpoint to do it.

diff --git a/Engine/Source/Programs/Horde/Horde.Agent/Commands/Service/RunCommand.cs b/Engine/Source/Programs/Horde/Horde.Agent/Commands/Service/RunCommand.cs
--- a/Engine/Source/Programs/Horde/Horde.Agent/Commands/Service/RunCommand.cs
+++ b/Engine/Source/Programs/Horde/Horde.Agent/Commands/Service/RunCommand.cs
@@ -35,11 +35,6 @@ namespace HordeAgent.Modes.Service
 	[Command("Service", "Run", "Runs the service in listen mode")]
 	class RunCommand : Command
 	{
-		/// <summary>
-		/// The logger instance
-		/// </summary>
-		ILogger Logger = null!;
-
 		/// <summary>
 		/// Override for the server to use
 		/// </summary>
@@ -55,12 +50,9 @@ namespace HordeAgent.Modes.Service
 		/// <summary>
 		/// Runs the service indefinitely
 		/// </summary>
-		/// <param name="Logger">Logger to use</param>
 		/// <returns>Exit code</returns>
 		public override async Task<int> ExecuteAsync(ILogger Logger)
 		{
-			this.Logger = Logger;
-
 			IHostBuilder HostBuilder = Host.CreateDefaultBuilder();
 
 			// Attempt to setup this process as a Windows service. A race condition inside Microsoft.Extensions.Hosting.WindowsServices.WindowsServiceHelpers.IsWindowsService
diff --git a/Engine/Source/Programs/Horde/Horde.Agent/Program.cs b/Engine/Source/Programs/Horde/Horde.Agent/Program.cs
--- a/Engine/Source/Programs/Horde/Horde.Agent/Program.cs
+++ b/Engine/Source/Programs/Horde/Horde.Agent/Program.cs
@@ -71,32 +71,10 @@ namespace HordeAgent
 		{
 			Program.Args = Args;
 
-			ILogger Logger = new Logging.HordeLoggerProvider().CreateLogger("HordeAgent");
-			try
-			{
-				int Result = await GuardedMain(Args, Logger);
-				return Result;
-			}
-			catch (FatalErrorException Ex)
-			{
-				Logger.LogCritical(Ex, "Fatal error.");
-				return Ex.ExitCode;
-			}
-			catch (Exception Ex)
-			{
-				Logger.LogCritical(Ex, "Fatal error.");
-				return 1;
-			}
-		}
+			IServiceCollection Services = new ServiceCollection();
+			Services.AddCommandsFromAssembly(Assembly.GetExecutingAssembly());
+			Services.AddLogging(Builder => Builder.AddProvider(new Logging.HordeLoggerProvider()));
 
-		/// <summary>
-		/// Actual Main function, without exception guards
-		/// </summary>
-		/// <param name="Args">Command-line arguments</param>
-		/// <param name="Logger">The logger interface</param>
-		/// <returns>Exit code</returns>
-		static async Task<int> GuardedMain(string[] Args, ILogger Logger)
-		{
 			// Enable unencrypted HTTP/2 for gRPC channel without TLS
 			AppContext.SetSwitch("System.Net.Http.SocketsHttpHandler.Http2UnencryptedSupport", true);
 
@@ -112,7 +90,8 @@ namespace HordeAgent
 			}
 
 			// Execute all the commands
-			return await CommandHost.RunAsync(new CommandLineArguments(Args), typeof(HordeAgent.Modes.Service.RunCommand), Logger);
+			IServiceProvider ServiceProvider = Services.BuildServiceProvider();
+			return await CommandHost.RunAsync(new CommandLineArguments(Args), ServiceProvider, typeof(HordeAgent.Modes.Service.RunCommand));
 		}
 
 		/// <summary>
diff --git a/Engine/Source/Programs/Horde/Horde.Build.Tests/ProgramTests.cs b/Engine/Source/Programs/Horde/Horde.Build.Tests/ProgramTests.cs
--- a/Engine/Source/Programs/Horde/Horde.Build.Tests/ProgramTests.cs
+++ b/Engine/Source/Programs/Horde/Horde.Build.Tests/ProgramTests.cs
@@ -5,6 +5,7 @@ using System.IO;
 using System.Runtime.InteropServices;
 using System.Security.Cryptography.X509Certificates;
 using EpicGames.Core;
+using Horde.Build.Commands;
 using HordeServer;
 using Microsoft.VisualStudio.TestTools.UnitTesting;
 
@@ -28,11 +29,11 @@ namespace Horde.Build.Tests
 			File.WriteAllBytes(TempCertPath, TempCertData);
 			
 			// No cert given
-			Assert.IsNull(Program.ReadGrpcCertificate(new () { ServerPrivateCert = null }));
+			Assert.IsNull(ServerCommand.ReadGrpcCertificate(new () { ServerPrivateCert = null }));
 
 			// Cert as file path
 			{
-				X509Certificate2? Cert = Program.ReadGrpcCertificate(new() { ServerPrivateCert = TempCertPath });
+				X509Certificate2? Cert = ServerCommand.ReadGrpcCertificate(new() { ServerPrivateCert = TempCertPath });
 				Assert.IsNotNull(Cert);
 				Assert.AreEqual(FriendlyName, Cert!.FriendlyName);
 			}
@@ -40,7 +41,7 @@ namespace Horde.Build.Tests
 			// Cert as base64 data
 			{
 				string TempCertBase64 = Convert.ToBase64String(TempCertData);
-				X509Certificate2? Cert = Program.ReadGrpcCertificate(new() { ServerPrivateCert = "base64:" + TempCertBase64 });
+				X509Certificate2? Cert = ServerCommand.ReadGrpcCertificate(new() { ServerPrivateCert = "base64:" + TempCertBase64 });
 				Assert.IsNotNull(Cert);
 				Assert.AreEqual(FriendlyName, Cert!.FriendlyName);	
 			}
diff --git a/Engine/Source/Programs/Horde/Horde.Build/Agents/Impl/AgentSoftwareService.cs b/Engine/Source/Programs/Horde/Horde.Build/Agents/Impl/AgentSoftwareService.cs
--- a/Engine/Source/Programs/Horde/Horde.Build/Agents/Impl/AgentSoftwareService.cs
+++ b/Engine/Source/Programs/Horde/Horde.Build/Agents/Impl/AgentSoftwareService.cs
@@ -21,6 +21,7 @@ using System.Runtime.InteropServices;
 using Microsoft.Extensions.Options;
 using System.IO.Compression;
 using System.Security.Cryptography.X509Certificates;
+using Horde.Build.Commands;
 
 namespace HordeServer.Services
 {
@@ -311,7 +312,7 @@ namespace HordeServer.Services
 
             byte[] Bytes = await File.ReadAllBytesAsync(AgentZip.ToString());
 
-			using X509Certificate2? GrpcCertificate = Program.ReadGrpcCertificate(Settings.CurrentValue);
+			using X509Certificate2? GrpcCertificate = ServerCommand.ReadGrpcCertificate(Settings.CurrentValue);
 
 			if (GrpcCertificate == null)
 			{
diff --git a/Engine/Source/Programs/Horde/Horde.Build/Commands/Config/SchemasCommand.cs b/Engine/Source/Programs/Horde/Horde.Build/Commands/Config/SchemasCommand.cs
new file mode 100644
--- /dev/null
+++ b/Engine/Source/Programs/Horde/Horde.Build/Commands/Config/SchemasCommand.cs
@@ -0,0 +1,32 @@
+// Copyright Epic Games, Inc. All Rights Reserved.
+
+using EpicGames.Core;
+using HordeServer.Controllers;
+using HordeServer.Utilities;
+using Microsoft.Extensions.Logging;
+using System;
+using System.Threading.Tasks;
+
+namespace HordeServer.Commands.Config
+{
+	[Command("config", "schemas", "Writes JSON schemas for server settings")]
+	class SchemasCommand : Command
+	{
+		[CommandLine]
+		DirectoryReference? OutputDir = null!;
+
+		public override Task<int> ExecuteAsync(ILogger Logger)
+		{
+			OutputDir ??= DirectoryReference.Combine(Program.AppDir, "Schemas");
+
+			DirectoryReference.CreateDirectory(OutputDir);
+			foreach (Type SchemaType in SchemaController.ConfigSchemas)
+			{
+				FileReference OutputFile = FileReference.Combine(OutputDir, $"{SchemaType.Name}.json");
+				Schemas.CreateSchema(SchemaType).Write(OutputFile);
+			}
+
+			return Task.FromResult(0);
+		}
+	}
+}
diff --git a/Engine/Source/Programs/Horde/Horde.Build/Commands/ServerCommand.cs b/Engine/Source/Programs/Horde/Horde.Build/Commands/ServerCommand.cs
new file mode 100644
--- /dev/null
+++ b/Engine/Source/Programs/Horde/Horde.Build/Commands/ServerCommand.cs
@@ -0,0 +1,137 @@
+// Copyright Epic Games, Inc. All Rights Reserved.
+
+using EpicGames.Core;
+using HordeServer;
+using Microsoft.AspNetCore.Hosting;
+using Microsoft.AspNetCore.Server.Kestrel.Core;
+using Microsoft.Extensions.Configuration;
+using Microsoft.Extensions.Hosting;
+using Microsoft.Extensions.Logging;
+using Serilog;
+using System;
+using System.Collections.Generic;
+using System.IO;
+using System.Linq;
+using System.Security.Cryptography.X509Certificates;
+using System.Threading.Tasks;
+
+namespace Horde.Build.Commands
+{
+	using ILogger = Microsoft.Extensions.Logging.ILogger;
+
+	[Command("server", "Runs the Horde Build server (default)")]
+	class ServerCommand : Command
+	{
+		ServerSettings HordeSettings;
+		IConfiguration Config;
+		string[] Args = Array.Empty<string>();
+
+		public ServerCommand(ServerSettings Settings, IConfiguration Config)
+		{
+			this.HordeSettings = Settings;
+			this.Config = Config;
+		}
+
+		public override void Configure(CommandLineArguments Arguments, ILogger Logger)
+		{
+			base.Configure(Arguments, Logger);
+			this.Args = Arguments.GetRawArray();
+		}
+
+		public override async Task<int> ExecuteAsync(ILogger Logger)
+		{
+			using (X509Certificate2? GrpcCertificate = ReadGrpcCertificate(HordeSettings))
+			{
+				List<IHost> Hosts = new List<IHost>();
+				Hosts.Add(CreateHostBuilderWithCert(Args, Config, HordeSettings, GrpcCertificate).Build());
+#if WITH_HORDE_STORAGE
+				IHostBuilder StorageHostBuilder = Horde.Storage.Program.CreateHostBuilder(Args);
+				StorageHostBuilder.ConfigureWebHostDefaults(Builder =>
+				{
+					Builder.ConfigureKestrel(Options =>
+					{
+						Options.ListenAnyIP(57000);
+						Options.ListenAnyIP(57001, Configure => Configure.UseHttps());
+					});
+				});
+				Hosts.Add(StorageHostBuilder.Build());
+#endif
+				await Task.WhenAll(Hosts.Select(x => x.RunAsync()));
+				return 0;
+			}
+		}
+
+		// Used by WebApplicationFactory in controller tests. Uses reflection to call this exact function signature.
+		public static IHostBuilder CreateHostBuilder(string[] Args)
+		{
+			ServerSettings HordeSettings = new ServerSettings();
+			return CreateHostBuilderWithCert(Args, new ConfigurationBuilder().Build(), HordeSettings, null);
+		}
+
+		static IHostBuilder CreateHostBuilderWithCert(string[] Args, IConfiguration Config, ServerSettings ServerSettings, X509Certificate2? SslCert)
+		{
+			return Host.CreateDefaultBuilder(Args)
+				.UseSerilog()
+				.ConfigureAppConfiguration(Builder => Builder.AddConfiguration(Config))
+				.ConfigureWebHostDefaults(WebBuilder =>
+				{
+					WebBuilder.ConfigureKestrel(Options =>
+					{
+						Options.Limits.MaxRequestBodySize = 100 * 1024 * 1024;
+
+						if (ServerSettings.HttpPort != 0)
+						{
+							Options.ListenAnyIP(ServerSettings.HttpPort, Configure => { Configure.Protocols = HttpProtocols.Http1AndHttp2; });
+						}
+
+						if (ServerSettings.HttpsPort != 0)
+						{
+							Options.ListenAnyIP(ServerSettings.HttpsPort, Configure => { if (SslCert != null) { Configure.UseHttps(SslCert); } });
+						}
+
+						// To serve HTTP/2 with gRPC *without* TLS enabled, a separate port for HTTP/2 must be used.
+						// This is useful when having a load balancer in front that terminates TLS.
+						if (ServerSettings.Http2Port != 0)
+						{
+							Options.ListenAnyIP(ServerSettings.Http2Port, Configure => { Configure.Protocols = HttpProtocols.Http2; });
+						}
+					});
+					WebBuilder.UseStartup<Startup>();
+				});
+		}
+
+		/// <summary>
+		/// Gets the certificate to use for Grpc endpoints
+		/// </summary>
+		/// <returns>Custom certificate to use for Grpc endpoints, or null for the default.</returns>
+		public static X509Certificate2? ReadGrpcCertificate(ServerSettings HordeSettings)
+		{
+			string Base64Prefix = "base64:";
+
+			if (HordeSettings.ServerPrivateCert == null)
+			{
+				return null;
+			}
+			else if (HordeSettings.ServerPrivateCert.StartsWith(Base64Prefix, StringComparison.Ordinal))
+			{
+				byte[] CertData = Convert.FromBase64String(HordeSettings.ServerPrivateCert.Replace(Base64Prefix, "", StringComparison.Ordinal));
+				return new X509Certificate2(CertData);
+			}
+			else
+			{
+				FileReference? ServerPrivateCert = null;
+
+				if (!Path.IsPathRooted(HordeSettings.ServerPrivateCert))
+				{
+					ServerPrivateCert = FileReference.Combine(Program.AppDir, HordeSettings.ServerPrivateCert);
+				}
+				else
+				{
+					ServerPrivateCert = new FileReference(HordeSettings.ServerPrivateCert);
+				}
+
+				return new X509Certificate2(FileReference.ReadAllBytes(ServerPrivateCert));
+			}
+		}
+	}
+}
diff --git a/Engine/Source/Programs/Horde/Horde.Build/Controllers/SchemaController.cs b/Engine/Source/Programs/Horde/Horde.Build/Controllers/SchemaController.cs
--- a/Engine/Source/Programs/Horde/Horde.Build/Controllers/SchemaController.cs
+++ b/Engine/Source/Programs/Horde/Horde.Build/Controllers/SchemaController.cs
@@ -40,6 +40,24 @@ namespace HordeServer.Controllers
 			public List<CatalogItem> Schemas { get; set; } = new List<CatalogItem>();
 		}
 
+		/// <summary>
+		/// Array of types included in the schema
+		/// </summary>
+		public static Type[] ConfigSchemas { get; } = FindSchemaTypes();
+
+		static Type[] FindSchemaTypes()
+		{
+			List<Type> SchemaTypes = new List<Type>();
+			foreach (Type Type in Assembly.GetExecutingAssembly().GetTypes())
+			{
+				if (Type.GetCustomAttribute<JsonSchemaAttribute>() != null)
+				{
+					SchemaTypes.Add(Type);
+				}
+			}
+			return SchemaTypes.ToArray();
+		}
+
 		/// <summary>
 		/// Get the catalog for config schema
 		/// </summary>
@@ -62,7 +80,7 @@ namespace HordeServer.Controllers
 			}
 
 			CatalogRoot Root = new CatalogRoot();
-			foreach (Type SchemaType in Program.ConfigSchemas)
+			foreach (Type SchemaType in ConfigSchemas)
 			{
 				JsonSchemaAttribute? SchemaAttribute = SchemaType.GetCustomAttribute<JsonSchemaAttribute>();
 				if (SchemaAttribute != null)
@@ -87,7 +105,7 @@ namespace HordeServer.Controllers
 		[Route("/api/v1/schema/types/{TypeName}.json")]
 		public ActionResult GetSchema(string TypeName)
 		{
-			foreach (Type SchemaType in Program.ConfigSchemas)
+			foreach (Type SchemaType in ConfigSchemas)
 			{
 				if (SchemaType.Name.Equals(TypeName, StringComparison.OrdinalIgnoreCase))
 				{
diff --git a/Engine/Source/Programs/Horde/Horde.Build/Startup.cs b/Engine/Source/Programs/Horde/Horde.Build/Startup.cs
--- a/Engine/Source/Programs/Horde/Horde.Build/Startup.cs
+++ b/Engine/Source/Programs/Horde/Horde.Build/Startup.cs
@@ -944,5 +944,15 @@ namespace HordeServer
 			}
 			return LogEventLevel.Information;
 		}
+
+		public static IServiceProvider CreateServiceProvider(IConfiguration Configuration)
+		{
+			IServiceCollection ServiceCollection = new ServiceCollection();
+
+			Startup Startup = new Startup(Configuration);
+			Startup.ConfigureServices(ServiceCollection);
+
+			return ServiceCollection.BuildServiceProvider();
+		}
 	}
 }
diff --git a/Engine/Source/Programs/Shared/EpicGames.Core/CommandHost.cs b/Engine/Source/Programs/Shared/EpicGames.Core/CommandHost.cs
--- a/Engine/Source/Programs/Shared/EpicGames.Core/CommandHost.cs
+++ b/Engine/Source/Programs/Shared/EpicGames.Core/CommandHost.cs
@@ -1,5 +1,6 @@
 // Copyright Epic Games, Inc. All Rights Reserved.
 
+using Microsoft.Extensions.DependencyInjection;
 using Microsoft.Extensions.Logging;
 using System;
 using System.Collections.Generic;
@@ -10,6 +11,53 @@ using System.Threading.Tasks;
 
 namespace EpicGames.Core
 {
+	/// <summary>
+	/// Interface for commands
+	/// </summary>
+	public interface ICommand
+	{
+		/// <summary>
+		/// Configure this object with the given command line arguments
+		/// </summary>
+		/// <param name="Arguments">Command line arguments</param>
+		void Configure(CommandLineArguments Arguments, ILogger Logger);
+
+		/// <summary>
+		/// Gets all command line parameters to show in help for this command
+		/// </summary>
+		/// <param name="Arguments">The command line arguments</param>
+		/// <returns>List of name/description pairs</returns>
+		List<KeyValuePair<string, string>> GetParameters(CommandLineArguments Arguments);
+
+		/// <summary>
+		/// Execute this command
+		/// </summary>
+		/// <param name="Logger">The logger to use for this command</param>
+		/// <returns>Exit code</returns>
+		Task<int> ExecuteAsync(ILogger Logger);
+	}
+
+	/// <summary>
+	/// Interface describing a command that can be exectued
+	/// </summary>
+	public interface ICommandFactory
+	{
+		/// <summary>
+		/// Names for this command
+		/// </summary>
+		public string[] Names { get; }
+
+		/// <summary>
+		/// Short description for the mode. Will be displayed in the help text.
+		/// </summary>
+		public string Description { get; }
+
+		/// <summary>
+		/// Create a command instance
+		/// </summary>
+		public ICommand CreateInstance(IServiceProvider ServiceProvider);
+	}
+
 	/// <summary>
 	/// Attribute used to specify names of program modes, and help text
 	/// </summary>
@@ -52,77 +100,87 @@ namespace EpicGames.Core
 	/// <summary>
 	/// Base class for all commands that can be executed by HordeAgent
 	/// </summary>
-	public abstract class Command
+	public abstract class Command : ICommand
 	{
-		/// <summary>
-		/// Configure this object with the given command line arguments
-		/// </summary>
-		/// <param name="Arguments">Command line arguments</param>
-		/// <param name="Logger">Logging output device</param>
+		/// <inheritdoc/>
 		public virtual void Configure(CommandLineArguments Arguments, ILogger Logger)
 		{
 			Arguments.ApplyTo(this, Logger);
 		}
 
-		/// <summary>
-		/// Gets all command line parameters to show in help for this command
-		/// </summary>
-		/// <param name="Arguments">The command line arguments</param>
-		/// <returns>List of name/description pairs</returns>
+		/// <inheritdoc/>
 		public virtual List<KeyValuePair<string, string>> GetParameters(CommandLineArguments Arguments)
 		{
 			return CommandLineArguments.GetParameters(GetType());
 		}
 
-		/// <summary>
-		/// Execute this command
-		/// </summary>
-		/// <param name="Logger">The logger to use for this command</param>
-		/// <returns>Exit code</returns>
+		/// <inheritdoc/>
 		public abstract Task<int> ExecuteAsync(ILogger Logger);
 	}
 
+	/// <summary>
+	/// Default implementation of a command factory
+	/// </summary>
+	class CommandFactory : ICommandFactory
+	{
+		public string[] Names { get; }
+		public string Description { get; }
+
+		public Type Type;
+
+		public CommandFactory(string[] Names, string Description, Type Type)
+		{
+			this.Names = Names;
+			this.Description = Description;
+			this.Type = Type;
+		}
+
+		public ICommand CreateInstance(IServiceProvider ServiceProvider) => (ICommand)ServiceProvider.GetRequiredService(Type);
+		public override string ToString() => String.Join(" ", Names);
+	}
+
 	/// <summary>
 	/// Entry point for dispatching commands
 	/// </summary>
 	public static class CommandHost
 	{
 		/// <summary>
-		/// Entry point for executing registered command types in a particular assembly
+		/// Adds services for executing the 
 		/// </summary>
-		/// <param name="Args">Command line arguments</param>
-		/// <param name="DefaultCommandType">The default command type</param>
-		/// <param name="Logger">Logging device to pass to the command object</param>
-		/// <returns>Return code from the command</returns>
-		public static Task<int> RunAsync(CommandLineArguments Args, Type? DefaultCommandType, ILogger Logger)
+		/// <param name="Services"></param>
+		/// <param name="Assembly"></param>
+		public static void AddCommandsFromAssembly(this IServiceCollection Services, Assembly Assembly)
 		{
-			return RunAsync(Args, Assembly.GetEntryAssembly()!, DefaultCommandType, Logger);
+			List<(CommandAttribute, Type)> Commands = new List<(CommandAttribute, Type)>();
+			foreach (Type Type in Assembly.GetTypes())
+			{
+				if (typeof(ICommand).IsAssignableFrom(Type) && !Type.IsAbstract)
+				{
+					CommandAttribute? Attribute = Type.GetCustomAttribute<CommandAttribute>();
+					if (Attribute != null)
+					{
+						Services.AddTransient(Type);
+						Services.AddTransient(typeof(ICommandFactory), SP => new CommandFactory(Attribute.Names, Attribute.Description, Type));
+					}
+				}
+			}
 		}
 
 		/// <summary>
 		/// Entry point for executing registered command types in a particular assembly
 		/// </summary>
 		/// <param name="Args">Command line arguments</param>
-		/// <param name="CommandAssembly">Assembly to scan for command types</param>
+		/// <param name="ServiceProvider">The service provider for the application</param>
 		/// <param name="DefaultCommandType">The default command type</param>
-		/// <param name="Logger">Logging device to pass to the command object</param>
 		/// <returns>Return code from the command</returns>
-		public static async Task<int> RunAsync(CommandLineArguments Args, Assembly CommandAssembly, Type? DefaultCommandType, ILogger Logger)
+		public static async Task<int> RunAsync(CommandLineArguments Args, IServiceProvider ServiceProvider, Type? DefaultCommandType)
 		{
 			// Find all the command types
-			List<(CommandAttribute, Type)> Commands = new List<(CommandAttribute, Type)>();
-			foreach (Type Type in CommandAssembly.GetTypes())
-			{
-				CommandAttribute? Attribute = Type.GetCustomAttribute<CommandAttribute>();
-				if (Attribute != null)
-				{
-					Commands.Add((Attribute, Type));
-				}
-			}
+			List<ICommandFactory> CommandFactories = ServiceProvider.GetServices<ICommandFactory>().ToList();
 
 			// Check if there's a matching command
-			Type? CommandType = null;
-			CommandAttribute? CommandAttribute = null;
+			ICommand? Command = null;
+			ICommandFactory? CommandFactory = null;
 
 			// Parse the positional arguments for the command name
 			string[] PositionalArgs = Args.GetPositionalArguments();
@@ -130,21 +188,12 @@ namespace EpicGames.Core
 			{
 				if (DefaultCommandType == null || Args.HasOption("-Help"))
 				{
-					Console.WriteLine(CommandAssembly.GetName().Name);
-					Console.WriteLine("");
-
-					AssemblyDescriptionAttribute? Description = CommandAssembly.GetCustomAttribute<AssemblyDescriptionAttribute>();
-					if (Description != null)
-					{
-						Console.WriteLine(Description.Description);
-					}
-
 					Console.WriteLine("Usage:");
 					Console.WriteLine("    [Command] [-Option1] [-Option2]...");
 					Console.WriteLine("");
 					Console.WriteLine("Commands:");
 
-					PrintCommands(Commands.Select(x => x.Item1));
+					PrintCommands(CommandFactories);
 
 					Console.WriteLine("");
 					Console.WriteLine("Specify \"<CommandName> -Help\" for command-specific help");
@@ -152,49 +201,47 @@ namespace EpicGames.Core
 				}
 				else 
 				{
-					CommandType = DefaultCommandType;
+					Command = (ICommand)ServiceProvider.GetService(DefaultCommandType);
 				}
 			}
 			else
 			{
-				foreach ((CommandAttribute Attribute, Type Type) in Commands.OrderBy(x => x.Item1.Names.Length))
+				foreach (ICommandFactory Factory in CommandFactories)
 				{
-					if (Attribute.Names.SequenceEqual(PositionalArgs, StringComparer.OrdinalIgnoreCase))
+					if (Factory.Names.SequenceEqual(PositionalArgs, StringComparer.OrdinalIgnoreCase))
 					{
-						CommandType = Type;
-						CommandAttribute = Attribute;
+						Command = Factory.CreateInstance(ServiceProvider);
+						CommandFactory = Factory;
 						break;
 					}
 				}
-				if (CommandType == null)
+				if (Command == null)
 				{
 					ConsoleUtils.WriteError($"Invalid command '{String.Join(" ", PositionalArgs)}'");
 					Console.WriteLine("");
 					Console.WriteLine("Available commands:");
 
-					PrintCommands(Commands.Select(x => x.Item1));
+					PrintCommands(CommandFactories);
 					return 1;
 				}
 			}
 
-			// Create the command instance
-			Command Command = (Command)Activator.CreateInstance(CommandType)!;
-
 			// If the help flag is specified, print the help info and exit immediately
 			if (Args.HasOption("-Help"))
 			{
-				if (CommandAttribute == null)
+				if (CommandFactory == null)
 				{
 					HelpUtils.PrintHelp(null, null, Command.GetParameters(Args));
 				}
 				else
 				{
-					HelpUtils.PrintHelp(String.Join(" ", CommandAttribute.Names), CommandAttribute.Description, Command.GetParameters(Args));
+					HelpUtils.PrintHelp(String.Join(" ", CommandFactory.Names), CommandFactory.Description, Command.GetParameters(Args));
 				}
 				return 1;
 			}
 
 			// Configure the command
+			ILogger Logger = ServiceProvider.GetRequiredService<ILoggerProvider>().CreateLogger("CommandHost");
 			try
 			{
 				Command.Configure(Args, Logger);
@@ -231,10 +278,10 @@ namespace EpicGames.Core
 		/// Print a formatted list of all the available commands
 		/// </summary>
 		/// <param name="Attributes">List of command attributes</param>
-		static void PrintCommands(IEnumerable<CommandAttribute> Attributes)
+		static void PrintCommands(IEnumerable<ICommandFactory> Attributes)
 		{
 			List<KeyValuePair<string, string>> Commands = new List<KeyValuePair<string, string>>();
-			foreach (CommandAttribute Attribute in Attributes)
+			foreach (ICommandFactory Attribute in Attributes)
 			{
 				Commands.Add(new KeyValuePair<string, string>(String.Join(" ", Attribute.Names), Attribute.Description));
 			}
