Horde: Simplify process for converting TreePack objects into refs, prior to adding to commit service.

diff --git a/Engine/Source/Programs/Horde/Horde.Agent/Commands/Archive/CreateArchiveCommand.cs b/Engine/Source/Programs/Horde/Horde.Agent/Commands/Archive/CreateArchiveCommand.cs
--- a/Engine/Source/Programs/Horde/Horde.Agent/Commands/Archive/CreateArchiveCommand.cs
+++ b/Engine/Source/Programs/Horde/Horde.Agent/Commands/Archive/CreateArchiveCommand.cs
@@ -30,37 +30,6 @@ using EpicGames.Perforce;
 
 namespace HordeAgent.Commands.Archive
 {
-	class FileSystemStorageClient : IBlobStorageClient
-	{
-		public DirectoryReference BaseDir;
-		ILogger Logger;
-
-		public FileSystemStorageClient(DirectoryReference BaseDir, ILogger Logger)
-		{
-			this.BaseDir = BaseDir;
-			this.Logger = Logger;
-
-			DirectoryReference.CreateDirectory(BaseDir);
-		}
-
-		public Task<Stream> ReadBlobAsync(NamespaceId NamespaceId, IoHash Hash, CancellationToken CancellationToken = default)
-		{
-			FileReference File = FileReference.Combine(BaseDir, $"{NamespaceId}-{Hash}.dat");
-			Logger.LogInformation("Reading {File} ({Size} bytes)", File, new FileInfo(File.FullName).Length);
-			return Task.FromResult<Stream>(FileReference.Open(File, FileMode.Open, FileAccess.Read, FileShare.Read));
-		}
-
-		public async Task WriteBlobAsync(NamespaceId NamespaceId, IoHash Hash, Stream Stream, CancellationToken CancellationToken = default)
-		{
-			FileReference File = FileReference.Combine(BaseDir, $"{NamespaceId}-{Hash}.dat");
-			Logger.LogInformation("Writing {File} ({Size} bytes)", File, Stream.Length);
-			using (Stream OutputStream = FileReference.Open(File, FileMode.Create, FileAccess.Write, FileShare.Read))
-			{
-				await Stream.CopyToAsync(OutputStream);
-			}
-		}
-	}
-
 	/// <summary>
 	/// Installs the agent as a service
 	/// </summary>
@@ -94,7 +63,7 @@ namespace HordeAgent.Commands.Archive
 
 			DirectoryReference TargetDir = TargetPath.Directory;
 			NamespaceId NamespaceId = new NamespaceId("test");
-			FileSystemStorageClient Storage = new FileSystemStorageClient(TargetDir, Logger);
+			FileStorageClient Storage = new FileStorageClient(TargetDir, Logger);
 
 			TreePack TargetPack = new TreePack(Storage, NamespaceId, Options);
 
@@ -120,7 +89,7 @@ namespace HordeAgent.Commands.Archive
 				await ReadFilesFromDiskAsync(SourcePath, TargetPack, TargetWriter);
 			}
 
-			IoHash NewHash = await TargetWriter.FinalizeAsync();
+			IoHash NewHash = await TargetWriter.FlushAsync();
 			TreePackObject Pack = await TargetPack.FlushAsync(NewHash, DateTime.UtcNow);
 			await WriteSummaryFiles(NamespaceId, Pack, new HashSet<IoHash>(), Storage, TargetDir);
 
@@ -136,7 +105,7 @@ namespace HordeAgent.Commands.Archive
 			return 0;
 		}
 
-		async Task WriteSummaryFiles(NamespaceId NamespaceId, TreePackObject RootObject, HashSet<IoHash> VisitedObjects, IBlobStorageClient StorageClient, DirectoryReference TargetDir)
+		async Task WriteSummaryFiles(NamespaceId NamespaceId, TreePackObject RootObject, HashSet<IoHash> VisitedObjects, IStorageClient StorageClient, DirectoryReference TargetDir)
 		{
 			foreach (TreePackObjectImport Import in RootObject.ObjectImports)
 			{
diff --git a/Engine/Source/Programs/Horde/Horde.Agent/Commands/Archive/ExtractArchiveCommand.cs b/Engine/Source/Programs/Horde/Horde.Agent/Commands/Archive/ExtractArchiveCommand.cs
--- a/Engine/Source/Programs/Horde/Horde.Agent/Commands/Archive/ExtractArchiveCommand.cs
+++ b/Engine/Source/Programs/Horde/Horde.Agent/Commands/Archive/ExtractArchiveCommand.cs
@@ -58,7 +58,7 @@ namespace HordeAgent.Commands.Archive
 
 			DirectoryReference Inputdir = InputPath.Directory;
 			NamespaceId NamespaceId = new NamespaceId("test");
-			FileSystemStorageClient Storage = new FileSystemStorageClient(Inputdir, Logger);
+			FileStorageClient Storage = new FileStorageClient(Inputdir, Logger);
 
 			TreePack TreePack = new TreePack(Storage, NamespaceId, Options);
 
@@ -131,7 +131,7 @@ namespace HordeAgent.Commands.Archive
 			}
 		}
 
-		async Task WriteSummaryFiles(NamespaceId NamespaceId, TreePackObject RootObject, HashSet<IoHash> VisitedObjects, IBlobStorageClient StorageClient, DirectoryReference TargetDir)
+		async Task WriteSummaryFiles(NamespaceId NamespaceId, TreePackObject RootObject, HashSet<IoHash> VisitedObjects, IStorageClient StorageClient, DirectoryReference TargetDir)
 		{
 			foreach (TreePackObjectImport Import in RootObject.ObjectImports)
 			{
diff --git a/Engine/Source/Programs/Shared/EpicGames.Horde/Storage/IStorageClient.cs b/Engine/Source/Programs/Shared/EpicGames.Horde/Storage/IStorageClient.cs
--- a/Engine/Source/Programs/Shared/EpicGames.Horde/Storage/IStorageClient.cs
+++ b/Engine/Source/Programs/Shared/EpicGames.Horde/Storage/IStorageClient.cs
@@ -154,8 +154,10 @@ namespace EpicGames.Horde.Storage
 	/// <summary>
 	/// Base interface for a storage client that only records blobs.
 	/// </summary>
-	public interface IBlobStorageClient
+	public interface IStorageClient
 	{
+		#region Blobs
+
 		/// <summary>
 		/// Opens a blob read stream
 		/// </summary>
@@ -173,13 +175,10 @@ namespace EpicGames.Horde.Storage
 		/// <param name="Stream">The stream to write</param>
 		/// <param name="CancellationToken">Cancellation token for the operation</param>
 		Task WriteBlobAsync(NamespaceId NamespaceId, IoHash Hash, Stream Stream, CancellationToken CancellationToken = default);
-	}
 
-	/// <summary>
-	/// Interface for a collection of ref documents
-	/// </summary>
-	public interface IStorageClient : IBlobStorageClient
-	{
+		#endregion
+		#region Refs
+
 		/// <summary>
 		/// Gets the given reference
 		/// </summary>
@@ -241,6 +240,8 @@ namespace EpicGames.Horde.Storage
 		/// <param name="CancellationToken">Cancellation token for the operation</param>
 		/// <returns>True if the ref was deleted, false if it did not exist</returns>
 		Task<bool> DeleteRefAsync(NamespaceId NamespaceId, BucketId BucketId, RefId RefId, CancellationToken CancellationToken = default);
+
+		#endregion
 	}
 
 	/// <summary>
@@ -258,7 +259,7 @@ namespace EpicGames.Horde.Storage
 		/// <param name="Hash">Hash of the blob to read</param>
 		/// <param name="MaxInMemoryBlobLength">Maximum allowed memory allocation to store the blob</param>
 		/// <returns>Data for the blob that was read. Throws an exception if the blob was not found.</returns>
-		public static async Task<byte[]> ReadBlobToMemoryAsync(this IBlobStorageClient StorageClient, NamespaceId NamespaceId, IoHash Hash, int MaxInMemoryBlobLength = DefaultMaxInMemoryBlobLength)
+		public static async Task<byte[]> ReadBlobToMemoryAsync(this IStorageClient StorageClient, NamespaceId NamespaceId, IoHash Hash, int MaxInMemoryBlobLength = DefaultMaxInMemoryBlobLength)
 		{
 			using Stream Stream = await StorageClient.ReadBlobAsync(NamespaceId, Hash);
 
@@ -290,7 +291,7 @@ namespace EpicGames.Horde.Storage
 		/// <param name="Hash">Hash of the blob to read</param>
 		/// <param name="MaxInMemoryBlobLength">Maximum allowed memory allocation to store the blob</param>
 		/// <returns>The decoded object</returns>
-		public static async Task<CbObject> ReadObjectAsync(this IBlobStorageClient StorageClient, NamespaceId NamespaceId, IoHash Hash, int MaxInMemoryBlobLength = DefaultMaxInMemoryBlobLength)
+		public static async Task<CbObject> ReadObjectAsync(this IStorageClient StorageClient, NamespaceId NamespaceId, IoHash Hash, int MaxInMemoryBlobLength = DefaultMaxInMemoryBlobLength)
 		{
 			return new CbObject(await ReadBlobToMemoryAsync(StorageClient, NamespaceId, Hash, MaxInMemoryBlobLength));
 		}
@@ -304,7 +305,7 @@ namespace EpicGames.Horde.Storage
 		/// <param name="Hash">Hash of the blob to read</param>
 		/// <param name="MaxInMemoryBlobLength">Maximum allowed memory allocation to store the blob</param>
 		/// <returns>The decoded object</returns>
-		public static async Task<T> ReadObjectAsync<T>(this IBlobStorageClient StorageClient, NamespaceId NamespaceId, IoHash Hash, int MaxInMemoryBlobLength = DefaultMaxInMemoryBlobLength)
+		public static async Task<T> ReadObjectAsync<T>(this IStorageClient StorageClient, NamespaceId NamespaceId, IoHash Hash, int MaxInMemoryBlobLength = DefaultMaxInMemoryBlobLength)
 		{
 			return CbSerializer.Deserialize<T>(await ReadBlobToMemoryAsync(StorageClient, NamespaceId, Hash, MaxInMemoryBlobLength));
 		}
@@ -316,7 +317,7 @@ namespace EpicGames.Horde.Storage
 		/// <param name="NamespaceId">Namespace containing the blob</param>
 		/// <param name="Data">Data to write</param>
 		/// <returns></returns>
-		public static async Task<IoHash> WriteBlobFromMemoryAsync(this IBlobStorageClient StorageClient, NamespaceId NamespaceId, ReadOnlyMemory<byte> Data)
+		public static async Task<IoHash> WriteBlobFromMemoryAsync(this IStorageClient StorageClient, NamespaceId NamespaceId, ReadOnlyMemory<byte> Data)
 		{
 			IoHash Hash = IoHash.Compute(Data.Span);
 			await WriteBlobFromMemoryAsync(StorageClient, NamespaceId, Hash, Data);
@@ -330,7 +331,7 @@ namespace EpicGames.Horde.Storage
 		/// <param name="NamespaceId">Namespace containing the blob</param>
 		/// <param name="Object">The object to be written</param>
 		/// <returns></returns>
-		public static Task<IoHash> WriteObjectAsync(this IBlobStorageClient StorageClient, NamespaceId NamespaceId, CbObject Object)
+		public static Task<IoHash> WriteObjectAsync(this IStorageClient StorageClient, NamespaceId NamespaceId, CbObject Object)
 		{
 			return WriteBlobFromMemoryAsync(StorageClient, NamespaceId, Object.GetView());
 		}
@@ -342,7 +343,7 @@ namespace EpicGames.Horde.Storage
 		/// <param name="NamespaceId">Namespace containing the blob</param>
 		/// <param name="Object">The object to be written</param>
 		/// <returns></returns>
-		public static Task<IoHash> WriteObjectAsync<T>(this IBlobStorageClient StorageClient, NamespaceId NamespaceId, T Object)
+		public static Task<IoHash> WriteObjectAsync<T>(this IStorageClient StorageClient, NamespaceId NamespaceId, T Object)
 		{
 			return WriteObjectAsync(StorageClient, NamespaceId, CbSerializer.Serialize<T>(Object));
 		}
@@ -355,7 +356,7 @@ namespace EpicGames.Horde.Storage
 		/// <param name="Hash">Hash of the data</param>
 		/// <param name="Data">The data to be written</param>
 		/// <returns></returns>
-		public static async Task WriteBlobFromMemoryAsync(this IBlobStorageClient StorageClient, NamespaceId NamespaceId, IoHash Hash, ReadOnlyMemory<byte> Data)
+		public static async Task WriteBlobFromMemoryAsync(this IStorageClient StorageClient, NamespaceId NamespaceId, IoHash Hash, ReadOnlyMemory<byte> Data)
 		{
 			using ReadOnlyMemoryStream Stream = new ReadOnlyMemoryStream(Data);
 			await StorageClient.WriteBlobAsync(NamespaceId, Hash, Stream);
diff --git a/Engine/Source/Programs/Shared/EpicGames.Horde/Storage/Impl/FileStorageClient.cs b/Engine/Source/Programs/Shared/EpicGames.Horde/Storage/Impl/FileStorageClient.cs
new file mode 100644
--- /dev/null
+++ b/Engine/Source/Programs/Shared/EpicGames.Horde/Storage/Impl/FileStorageClient.cs
@@ -0,0 +1,130 @@
+// Copyright Epic Games, Inc. All Rights Reserved.
+
+using EpicGames.Core;
+using EpicGames.Serialization;
+using Microsoft.Extensions.Logging;
+using System;
+using System.Collections.Generic;
+using System.IO;
+using System.Linq;
+using System.Text;
+using System.Threading;
+using System.Threading.Tasks;
+
+namespace EpicGames.Horde.Storage.Impl
+{
+	/// <summary>
+	/// Implementation of <see cref="IStorageClient"/> which writes directly to the local filesystem for testing. Not intended for production use.
+	/// </summary>
+	public class FileStorageClient : IStorageClient
+	{
+		class Ref : IRef
+		{
+			public NamespaceId NamespaceId { get; set; }
+			public BucketId BucketId { get; set; }
+			public RefId RefId { get; set; }
+			public CbObject Value { get; set; }
+
+			public Ref(NamespaceId NamespaceId, BucketId BucketId, RefId RefId, CbObject Value)
+			{
+				this.NamespaceId = NamespaceId;
+				this.BucketId = BucketId;
+				this.RefId = RefId;
+				this.Value = Value;
+			}
+		}
+
+		public DirectoryReference BaseDir { get; }
+		ILogger Logger;
+
+		/// <summary>
+		/// Constructor
+		/// </summary>
+		/// <param name="BaseDir"></param>
+		/// <param name="Logger"></param>
+		public FileStorageClient(DirectoryReference BaseDir, ILogger Logger)
+		{
+			this.BaseDir = BaseDir;
+			this.Logger = Logger;
+
+			DirectoryReference.CreateDirectory(BaseDir);
+		}
+
+		FileReference GetBlobFile(NamespaceId NamespaceId, IoHash Hash)
+		{
+			return FileReference.Combine(BaseDir, $"{NamespaceId}-{Hash}.blob");
+		}
+
+		/// <inheritdoc/>
+		public Task<Stream> ReadBlobAsync(NamespaceId NamespaceId, IoHash Hash, CancellationToken CancellationToken = default)
+		{
+			FileReference File = GetBlobFile(NamespaceId, Hash);
+			Logger.LogInformation("Reading {File} ({Size} bytes)", File, new FileInfo(File.FullName).Length);
+			return Task.FromResult<Stream>(FileReference.Open(File, FileMode.Open, FileAccess.Read, FileShare.Read));
+		}
+
+		/// <inheritdoc/>
+		public async Task WriteBlobAsync(NamespaceId NamespaceId, IoHash Hash, Stream Stream, CancellationToken CancellationToken = default)
+		{
+			FileReference File = GetBlobFile(NamespaceId, Hash);
+			Logger.LogInformation("Writing {File} ({Size} bytes)", File, Stream.Length);
+			using (Stream OutputStream = FileReference.Open(File, FileMode.Create, FileAccess.Write, FileShare.Read))
+			{
+				await Stream.CopyToAsync(OutputStream);
+			}
+		}
+
+		/// <inheritdoc/>
+		public Task<bool> DeleteRefAsync(NamespaceId NamespaceId, BucketId BucketId, RefId RefId, CancellationToken CancellationToken = default)
+		{
+			FileReference File = GetRefFile(NamespaceId, BucketId, RefId);
+			if (FileReference.Exists(File))
+			{
+				FileReference.Delete(File);
+				return Task.FromResult(true);
+			}
+			return Task.FromResult(false);
+		}
+
+		/// <inheritdoc/>
+		public Task<List<RefId>> FindMissingRefsAsync(NamespaceId NamespaceId, BucketId BucketId, List<RefId> RefIds, CancellationToken CancellationToken = default)
+		{
+			return Task.FromResult(RefIds.Where(x => !FileReference.Exists(GetRefFile(NamespaceId, BucketId, x))).ToList());
+		}
+
+		FileReference GetRefFile(NamespaceId NamespaceId, BucketId BucketId, RefId RefId)
+		{
+			return FileReference.Combine(BaseDir, $"{NamespaceId}-{BucketId}-{RefId}.ref");
+		}
+
+		/// <inheritdoc/>
+		public async Task<IRef> GetRefAsync(NamespaceId NamespaceId, BucketId BucketId, RefId RefId, CancellationToken CancellationToken = default)
+		{
+			FileReference File = GetRefFile(NamespaceId, BucketId, RefId);
+			Logger.LogInformation("Reading {File} ({Size} bytes)", File, new FileInfo(File.FullName).Length);
+			byte[] Data = await FileReference.ReadAllBytesAsync(File);
+			return new Ref(NamespaceId, BucketId, RefId, new CbObject(Data));
+		}
+
+		/// <inheritdoc/>
+		public Task<bool> HasRefAsync(NamespaceId NamespaceId, BucketId BucketId, RefId RefId, CancellationToken CancellationToken = default)
+		{
+			return Task.FromResult(FileReference.Exists(GetRefFile(NamespaceId, BucketId, RefId)));
+		}
+
+		/// <inheritdoc/>
+		public Task<List<IoHash>> TryFinalizeRefAsync(NamespaceId NamespaceId, BucketId BucketId, RefId RefId, IoHash Hash, CancellationToken CancellationToken = default)
+		{
+			throw new NotImplementedException();
+		}
+
+		/// <inheritdoc/>
+		public async Task<List<IoHash>> TrySetRefAsync(NamespaceId NamespaceId, BucketId BucketId, RefId RefId, CbObject Value, CancellationToken CancellationToken = default)
+		{
+			FileReference File = GetRefFile(NamespaceId, BucketId, RefId);
+			Logger.LogInformation("Writing {File} ({Size} bytes)", File, Value.GetView().Length);
+			await FileReference.WriteAllBytesAsync(File, Value.GetView().ToArray(), CancellationToken);
+			return new List<IoHash>();
+		}
+	}
+}
diff --git a/Engine/Source/Programs/Shared/EpicGames.Horde/Storage/TreePack.cs b/Engine/Source/Programs/Shared/EpicGames.Horde/Storage/TreePack.cs
--- a/Engine/Source/Programs/Shared/EpicGames.Horde/Storage/TreePack.cs
+++ b/Engine/Source/Programs/Shared/EpicGames.Horde/Storage/TreePack.cs
@@ -56,6 +56,41 @@ namespace EpicGames.Horde.Storage
 		/// </summary>
 		[CbField("data")]
 		public ReadOnlyMemory<byte> Data { get; set; }
+
+		/// <summary>
+		/// Parses an object from a ref
+		/// </summary>
+		/// <param name="Ref"></param>
+		/// <returns></returns>
+		public static TreePackObject Parse(IRef Ref) => Parse(Ref.Value.GetView());
+
+		/// <summary>
+		/// Parse an object from the given data
+		/// </summary>
+		/// <param name="Data"></param>
+		/// <returns></returns>
+		public static TreePackObject Parse(ReadOnlyMemory<byte> Data) => CbSerializer.Deserialize<TreePackObject>(Data);
+
+		/// <summary>
+		/// Gets the root hash for this object
+		/// </summary>
+		/// <returns></returns>
+		public IoHash GetRootHash()
+		{
+			TreePackExport Export = new TreePackExport(Exports.Span);
+			return Export.Hash;
+		}
+
+		/// <summary>
+		/// Gets the root node data for this object
+		/// </summary>
+		/// <returns></returns>
+		public ReadOnlyMemory<byte> GetRootNode()
+		{
+			TreePackExport Export = new TreePackExport(Exports.Span);
+			ReadOnlyMemory<byte> CompressedData = Data.Slice(Export.Offset, Export.CompressedLength);
+			return TreePack.UncompressData(CompressedData, Export.UncompressedLength);
+		}
 	}
 
 	/// <summary>
@@ -107,6 +142,84 @@ namespace EpicGames.Horde.Storage
 		public ReadOnlyMemory<byte> Exports { get; set; }
 	}
 
+	/// <summary>
+	/// Entry for a block of data exported from a blob
+	/// </summary>
+	public class TreePackExport
+	{
+		/// <summary>
+		/// Length of the export when encoded to a byte array
+		/// </summary>
+		public const int EncodedLength = IoHash.NumBytes + sizeof(int) + sizeof(int) + sizeof(int);
+
+		/// <summary>
+		/// Hash of the exported item
+		/// </summary>
+		public IoHash Hash { get; }
+
+		/// <summary>
+		/// Offset of the exported data
+		/// </summary>
+		public int Offset { get; }
+
+		/// <summary>
+		/// Compressed length of the exported data
+		/// </summary>
+		public int CompressedLength { get; }
+
+		/// <summary>
+		/// Length of the exported data when uncompressed
+		/// </summary>
+		public int UncompressedLength { get; }
+
+		/// <summary>
+		/// Constructor
+		/// </summary>
+		public TreePackExport(IoHash Hash, int Offset, int CompressedLength, int UncompressedLength)
+		{
+			this.Hash = Hash;
+			this.Offset = Offset;
+			this.CompressedLength = CompressedLength;
+			this.UncompressedLength = UncompressedLength;
+		}
+
+		/// <summary>
+		/// Constructor
+		/// </summary>
+		/// <param name="Span"></param>
+		public TreePackExport(ReadOnlySpan<byte> Span)
+		{
+			Hash = new IoHash(Span);
+			Span = Span.Slice(IoHash.NumBytes);
+
+			Offset = BinaryPrimitives.ReadInt32LittleEndian(Span);
+			Span = Span.Slice(sizeof(int));
+
+			CompressedLength = BinaryPrimitives.ReadInt32LittleEndian(Span);
+			Span = Span.Slice(sizeof(int));
+
+			UncompressedLength = BinaryPrimitives.ReadInt32LittleEndian(Span);
+		}
+
+		/// <summary>
+		/// Encode an export into a byte array 
+		/// </summary>
+		/// <param name="Output"></param>
+		public void CopyTo(Span<byte> Output)
+		{
+			Hash.CopyTo(Output);
+			Output = Output.Slice(IoHash.NumBytes);
+
+			BinaryPrimitives.WriteInt32LittleEndian(Output, Offset);
+			Output = Output.Slice(sizeof(int));
+
+			BinaryPrimitives.WriteInt32LittleEndian(Output, CompressedLength);
+			Output = Output.Slice(sizeof(int));
+
+			BinaryPrimitives.WriteInt32LittleEndian(Output, UncompressedLength);
+		}
+	}
+
 	#endregion
 
 	#region Node types
@@ -503,26 +616,7 @@ namespace EpicGames.Horde.Storage
 			public override string ToString() => IsLeaf ? $"{Hash} (Leaf)" : $"{Hash}";
 		}
 
-		/// <summary>
-		/// Entry for a block of data exported from a blob
-		/// </summary>
-		class ExportInfo
-		{
-			public IoHash Hash { get; }
-			public int Offset { get; }
-			public int CompressedLength { get; }
-			public int UncompressedLength { get; }
-
-			public ExportInfo(IoHash Hash, int Offset, int CompressedLength, int UncompressedLength)
-			{
-				this.Hash = Hash;
-				this.Offset = Offset;
-				this.CompressedLength = CompressedLength;
-				this.UncompressedLength = UncompressedLength;
-			}
-		}
-
-		IBlobStorageClient StorageClient { get; }
+		IStorageClient StorageClient { get; }
 		NamespaceId NamespaceId { get; }
 
 		public TreePackOptions Options { get; }
@@ -535,7 +629,7 @@ namespace EpicGames.Horde.Storage
 		/// <summary>
 		/// Constructor
 		/// </summary>
-		public TreePack(IBlobStorageClient StorageClient, NamespaceId NamespaceId)
+		public TreePack(IStorageClient StorageClient, NamespaceId NamespaceId)
 			: this(StorageClient, NamespaceId, new TreePackOptions())
 		{
 		}
@@ -543,7 +637,7 @@ namespace EpicGames.Horde.Storage
 		/// <summary>
 		/// Constructor
 		/// </summary>
-		public TreePack(IBlobStorageClient StorageClient, NamespaceId NamespaceId, TreePackOptions Options)
+		public TreePack(IStorageClient StorageClient, NamespaceId NamespaceId, TreePackOptions Options)
 		{
 			this.StorageClient = StorageClient;
 			this.NamespaceId = NamespaceId;
@@ -554,10 +648,43 @@ namespace EpicGames.Horde.Storage
 		/// Adds a root blob to the pack
 		/// </summary>
 		/// <param name="Data"></param>
-		public void AddRootBlob(ReadOnlyMemory<byte> Data)
+		/// <returns>Hash of the root node</returns>
+		public ReadOnlyMemory<byte> AddRootBlob(IRef Ref) => AddRootBlob(Ref.Value.GetView());
+
+		/// <summary>
+		/// Adds a root blob to the pack
+		/// </summary>
+		/// <param name="Data"></param>
+		/// <returns>Hash of the root node</returns>
+		public ReadOnlyMemory<byte> AddRootBlob(ReadOnlyMemory<byte> Data)
 		{
 			TreePackObject Object = CbSerializer.Deserialize<TreePackObject>(Data);
 			RegisterObject(Object, null);
+
+			TreePackExport Export = new TreePackExport(Object.Exports.Span);
+			return UncompressData(Object.Data.Slice(Export.Offset, Export.CompressedLength), Export.UncompressedLength);
+		}
+
+		/// <summary>
+		/// Adds a root blob to the pack
+		/// </summary>
+		/// <param name="Data"></param>
+		/// <returns>Hash of the root node</returns>
+		public void AddRootObject(TreePackObject Object)
+		{
+			RegisterObject(Object, null);
+		}
+
+		/// <summary>
+		/// Adds a root object to the pack
+		/// </summary>
+		/// <param name="Ref"></param>
+		/// <returns></returns>
+		public TreePackObject AddRootObject(IRef Ref)
+		{
+			TreePackObject Object = TreePackObject.Parse(Ref);
+			AddRootObject(Object);
+			return Object;
 		}
 
 		/// <summary>
@@ -627,16 +754,27 @@ namespace EpicGames.Horde.Storage
 			return Hash;
 		}
 
-		private static ReadOnlyMemory<byte> CompressData(ReadOnlyMemory<byte> Data)
+		/// <summary>
+		/// Compresses the given node data 
+		/// </summary>
+		/// <param name="Data"></param>
+		/// <returns></returns>
+		public static ReadOnlyMemory<byte> CompressData(ReadOnlyMemory<byte> Data)
 		{
 			byte[] Output = new byte[LZ4Codec.MaximumOutputSize(Data.Length)];
 			int Length = LZ4Codec.Encode(Data.Span, Output);
 			return Output.AsMemory(0, Length);
 		}
 
-		private static byte[] UncompressData(ReadOnlyMemory<byte> Data, int Length)
+		/// <summary>
+		/// Uncompresses the given node data
+		/// </summary>
+		/// <param name="Data"></param>
+		/// <param name="UncompressedLength">Size of the uncompressed data</param>
+		/// <returns></returns>
+		public static byte[] UncompressData(ReadOnlyMemory<byte> Data, int UncompressedLength)
 		{
-			byte[] Output = new byte[Length];
+			byte[] Output = new byte[UncompressedLength];
 			LZ4Codec.Decode(Data.Span, Output);
 			return Output;
 		}
@@ -709,6 +847,20 @@ namespace EpicGames.Horde.Storage
 			return UncompressData(Node.CompressedData, Node.UncompressedLength);
 		}
 
+		/// <summary>
+		/// Writes the pack file to storage, and set a reference to it
+		/// </summary>
+		/// <param name="BucketId"></param>
+		/// <param name="RefId"></param>
+		/// <param name="RootHash"></param>
+		/// <param name="UtcNow"></param>
+		/// <returns></returns>
+		public async Task WriteAsync(BucketId BucketId, RefId RefId, IoHash RootHash, DateTime UtcNow)
+		{
+			TreePackObject Object = await FlushAsync(RootHash, UtcNow);
+			await StorageClient.SetRefAsync(NamespaceId, BucketId, RefId, Object);
+		}
+
 		/// <summary>
 		/// Flush any pending blobs to disk
 		/// </summary>
@@ -856,7 +1008,7 @@ namespace EpicGames.Horde.Storage
 		void RegisterObject(TreePackObject Object, BlobInfo? BlobInfo)
 		{
 			// Register all the exports from the object
-			foreach (ExportInfo Export in ReadExports(Object.Exports))
+			foreach (TreePackExport Export in ReadExports(Object.Exports))
 			{
 				NodeInfo NodeInfo = FindOrAddNode(Export.Hash, Export.Offset, Export.CompressedLength, Export.UncompressedLength, BlobInfo);
 				NodeInfo.UncompressedLength = Export.UncompressedLength;
@@ -881,8 +1033,8 @@ namespace EpicGames.Horde.Storage
 			{
 				BlobInfo ImportedBlob = FindOrAddBlob(Import.Binary.Hash, Import.Time, Import.Length, true);
 
-				ExportInfo[] Exports = ReadExports(Import.Exports);
-				foreach (ExportInfo Export in Exports)
+				TreePackExport[] Exports = ReadExports(Import.Exports);
+				foreach (TreePackExport Export in Exports)
 				{
 					FindOrAddNode(Export.Hash, Export.Offset, Export.CompressedLength, Export.UncompressedLength, ImportedBlob);
 				}
@@ -919,7 +1071,7 @@ namespace EpicGames.Horde.Storage
 			TreePackObject Object = new TreePackObject();
 			Object.Time = UtcNow;
 			Object.Data = Data;
-			Object.Exports = WriteExports(Nodes.Select(x => new ExportInfo(x.Hash, x.Offset, x.CompressedLength, x.UncompressedLength)).ToArray());
+			Object.Exports = WriteExports(Nodes.Select(x => new TreePackExport(x.Hash, x.Offset, x.CompressedLength, x.UncompressedLength)).ToArray());
 
 			foreach (IGrouping<BlobInfo, NodeInfo> Group in References.GroupBy(x => x.Blob!))
 			{
@@ -930,7 +1082,7 @@ namespace EpicGames.Horde.Storage
 					BinaryImport.Binary = new CbBinaryAttachment(ImportBlob.Hash);
 					BinaryImport.Time = ImportBlob.Time;
 					BinaryImport.Length = ImportBlob.Length;
-					BinaryImport.Exports = WriteExports(Group.Select(x => new ExportInfo(x.Hash, x.Offset, x.CompressedLength, x.UncompressedLength)).ToArray());
+					BinaryImport.Exports = WriteExports(Group.Select(x => new TreePackExport(x.Hash, x.Offset, x.CompressedLength, x.UncompressedLength)).ToArray());
 					Object.BinaryImports.Add(BinaryImport);
 				}
 				else
@@ -1000,51 +1152,31 @@ namespace EpicGames.Horde.Storage
 			return Data;
 		}
 
-		static ExportInfo[] ReadExports(ReadOnlyMemory<byte> ExportData)
+		static TreePackExport[] ReadExports(ReadOnlyMemory<byte> ExportData)
 		{
 			int NumEntries = ExportData.Length / (IoHash.NumBytes + sizeof(int) + sizeof(int) + sizeof(int));
 
-			ExportInfo[] Exports = new ExportInfo[NumEntries];
+			TreePackExport[] Exports = new TreePackExport[NumEntries];
 
 			ReadOnlySpan<byte> Span = ExportData.Span;
 			for (int Idx = 0; Idx < NumEntries; Idx++)
 			{
-				IoHash Key = new IoHash(Span);
-				Span = Span.Slice(IoHash.NumBytes);
-
-				int Offset = BinaryPrimitives.ReadInt32LittleEndian(Span);
-				Span = Span.Slice(sizeof(int));
-
-				int CompressedLength = BinaryPrimitives.ReadInt32LittleEndian(Span);
-				Span = Span.Slice(sizeof(int));
-
-				int UncompressedLength = BinaryPrimitives.ReadInt32LittleEndian(Span);
-				Span = Span.Slice(sizeof(int));
-
-				Exports[Idx] = new ExportInfo(Key, Offset, CompressedLength, UncompressedLength);
+				Exports[Idx] = new TreePackExport(Span);
+				Span = Span.Slice(TreePackExport.EncodedLength);
 			}
 
 			return Exports;
 		}
 
-		static byte[] WriteExports(ExportInfo[] Exports)
+		static byte[] WriteExports(TreePackExport[] Exports)
 		{
 			byte[] Data = new byte[Exports.Length * (IoHash.NumBytes + sizeof(int) + sizeof(int) + sizeof(int))];
 
 			Span<byte> Output = Data;
-			foreach (ExportInfo Export in Exports.OrderBy(x => x.Offset))
+			foreach (TreePackExport Export in Exports.OrderBy(x => x.Offset))
 			{
-				Export.Hash.CopyTo(Output);
-				Output = Output.Slice(IoHash.NumBytes);
-
-				BinaryPrimitives.WriteInt32LittleEndian(Output, Export.Offset);
-				Output = Output.Slice(sizeof(int));
-
-				BinaryPrimitives.WriteInt32LittleEndian(Output, Export.CompressedLength);
-				Output = Output.Slice(sizeof(int));
-
-				BinaryPrimitives.WriteInt32LittleEndian(Output, Export.UncompressedLength);
-				Output = Output.Slice(sizeof(int));
+				Export.CopyTo(Output);
+				Output = Output.Slice(TreePackExport.EncodedLength);
 			}
 
 			return Data;
@@ -1144,7 +1276,7 @@ namespace EpicGames.Horde.Storage
 			}
 			foreach (TreePackBinaryImport BinaryImport in Object.BinaryImports)
 			{
-				foreach (ExportInfo Export in ReadExports(BinaryImport.Exports))
+				foreach (TreePackExport Export in ReadExports(BinaryImport.Exports))
 				{
 					HashToLocator[Export.Hash] = $"{BinaryImport.Binary} [{Export.Offset},{Export.CompressedLength}]";
 				}
@@ -1176,7 +1308,7 @@ namespace EpicGames.Horde.Storage
 				{
 					Writer.WriteLine($"- blob: {BinaryImport.Binary.Hash}");
 					Writer.WriteLine("  keys:");
-					foreach (ExportInfo Export in ReadExports(BinaryImport.Exports))
+					foreach (TreePackExport Export in ReadExports(BinaryImport.Exports))
 					{
 						Writer.WriteLine($"  - \"{Export.Hash}\"");
 					}
@@ -1186,8 +1318,8 @@ namespace EpicGames.Horde.Storage
 
 			Writer.WriteLine("exports:");
 
-			ExportInfo[] Exports = ReadExports(Object.Exports);
-			foreach (ExportInfo Export in Exports)
+			TreePackExport[] Exports = ReadExports(Object.Exports);
+			foreach (TreePackExport Export in Exports)
 			{
 				ReadOnlyMemory<byte> Data = UncompressData(Object.Data.Slice(Export.Offset, Export.CompressedLength), Export.UncompressedLength);
 				TreePackNodeType Type = (TreePackNodeType)Data.Span[0];
diff --git a/Engine/Source/Programs/Shared/EpicGames.Horde/Storage/TreePackDirWriter.cs b/Engine/Source/Programs/Shared/EpicGames.Horde/Storage/TreePackDirWriter.cs
--- a/Engine/Source/Programs/Shared/EpicGames.Horde/Storage/TreePackDirWriter.cs
+++ b/Engine/Source/Programs/Shared/EpicGames.Horde/Storage/TreePackDirWriter.cs
@@ -28,6 +28,16 @@ namespace EpicGames.Horde.Storage
 			this.TreePack = TreePack;
 		}
 
+		/// <summary>
+		/// Cosntructor. Initializes the contents of the tree to the root node in the given ref.
+		/// </summary>
+		/// <param name="TreePack"></param>
+		/// <param name="Ref"></param>
+		public TreePackDirWriter(TreePack TreePack, IRef Ref)
+			: this(TreePack, TreePackDirNode.Parse(TreePack.AddRootObject(Ref).GetRootNode()))
+		{
+		}
+
 		/// <summary>
 		/// Constructor. Initializes the contents of the tree to an existing directory node.
 		/// </summary>
@@ -121,11 +131,11 @@ namespace EpicGames.Horde.Storage
 		/// Flushes the current tree state and writes it to the tree pack
 		/// </summary>
 		/// <returns>Hash of the tree</returns>
-		public async Task<IoHash> FinalizeAsync()
+		public async Task<IoHash> FlushAsync()
 		{
 			foreach ((Utf8String Name, TreePackDirWriter Writer) in NameToSubDir)
 			{
-				IoHash Hash = await Writer.FinalizeAsync();
+				IoHash Hash = await Writer.FlushAsync();
 				NameToEntry[Name] = new TreePackDirEntry(TreePackDirEntryFlags.Directory, Name, Hash, Sha1Hash.Zero);
 			}
 			NameToSubDir.Clear();
@@ -135,6 +145,18 @@ namespace EpicGames.Horde.Storage
 			return await TreePack.AddNodeAsync(Node.ToByteArray());
 		}
 
+		/// <summary>
+		/// Writes the directory tree to the storage client with the 
+		/// </summary>
+		/// <param name="BucketId"></param>
+		/// <param name="RefId"></param>
+		/// <returns></returns>
+		public async Task WriteAsync(BucketId BucketId, RefId RefId)
+		{
+			IoHash RootHash = await FlushAsync();
+			await TreePack.WriteAsync(BucketId, RefId, RootHash, DateTime.UtcNow);
+		}
+
 		/// <summary>
 		/// Creates a writer from a particular directory hash
 		/// </summary>
