1) Adding new collection of UE types into C#
2) Enhancements to StringView and a few other helper classes
3) New UE style ignore case comparison
4) Moved shared UHT definitions into a common C# file.

#rnx

diff --git a/Engine/Source/Editor/BlueprintGraph/Private/EdGraphSchema_K2.cpp b/Engine/Source/Editor/BlueprintGraph/Private/EdGraphSchema_K2.cpp
--- a/Engine/Source/Editor/BlueprintGraph/Private/EdGraphSchema_K2.cpp
+++ b/Engine/Source/Editor/BlueprintGraph/Private/EdGraphSchema_K2.cpp
@@ -5188,7 +5188,7 @@ namespace FSetVariableByNameFunctionNames
 
 UFunction* UEdGraphSchema_K2::FindSetVariableByNameFunction(const FEdGraphPinType& PinType)
 {
-	//!!!! Keep this function synced with FExposeOnSpawnValidator::IsSupported !!!!
+	//!!!! Keep this function synced with FExposeOnSpawnValidator::IsSupported and Uht*Property.cs, CanExposeOnSpawn!!!!
 
 	struct FIsCustomStructureParamHelper
 	{
diff --git a/Engine/Source/Programs/Shared/EpicGames.Core/CityHash.cs b/Engine/Source/Programs/Shared/EpicGames.Core/CityHash.cs
new file mode 100644
--- /dev/null
+++ b/Engine/Source/Programs/Shared/EpicGames.Core/CityHash.cs
@@ -0,0 +1,458 @@
+﻿// Copyright (c) 2011 Google, Inc.
+//
+// Permission is hereby granted, free of charge, to any person obtaining a copy
+// of this software and associated documentation files (the "Software"), to deal
+// in the Software without restriction, including without limitation the rights
+// to use, copy, modify, merge, publish, distribute, sublicense, and/or sell
+// copies of the Software, and to permit persons to whom the Software is
+// furnished to do so, subject to the following conditions:
+//
+// The above copyright notice and this permission notice shall be included in
+// all copies or substantial portions of the Software.
+//
+// THE SOFTWARE IS PROVIDED "AS IS", WITHOUT WARRANTY OF ANY KIND, EXPRESS OR
+// IMPLIED, INCLUDING BUT NOT LIMITED TO THE WARRANTIES OF MERCHANTABILITY,
+// FITNESS FOR A PARTICULAR PURPOSE AND NONINFRINGEMENT. IN NO EVENT SHALL THE
+// AUTHORS OR COPYRIGHT HOLDERS BE LIABLE FOR ANY CLAIM, DAMAGES OR OTHER
+// LIABILITY, WHETHER IN AN ACTION OF CONTRACT, TORT OR OTHERWISE, ARISING FROM,
+// OUT OF OR IN CONNECTION WITH THE SOFTWARE OR THE USE OR OTHER DEALINGS IN
+// THE SOFTWARE.
+//
+// CityHash, by Geoff Pike and Jyrki Alakuijala
+//
+// This file provides CityHash64() and related functions.
+//
+// It's probably possible to create even faster hash functions by
+// writing a program that systematically explores some of the space of
+// possible hash functions, by using SIMD instructions, or by
+// compromising on hash quality.
+
+using System;
+
+namespace EpicGames.Core
+{
+	public class CityHash
+	{
+		// Some primes between 2^63 and 2^64 for various uses.
+		const ulong k0 = 0xc3a5c85c97cb3127;
+		const ulong k1 = 0xb492b66fbe98f273;
+		const ulong k2 = 0x9ae16a3b2f90404f;
+
+		// Magic numbers for 32-bit hashing.  Copied from Murmur3.
+		const uint c1 = 0xcc9e2d51;
+		const uint c2 = 0x1b873593;
+
+		public struct Uint128_64
+		{
+			public ulong lo;
+			public ulong hi;
+		}
+
+		private static uint bswap_32(uint x)
+		{
+			return (x >> 24) + ((x >> 8) & 0xff00) + ((x << 8) & 0xff0000) + (x << 24);
+		}
+
+		private static ulong bswap_64(ulong x)
+		{
+			return (x >> 56) + ((x >> 40) & 0xff00) + ((x >> 24) & 0xff0000) + ((x >> 8) & 0xff000000) + ((x << 8) & 0xff00000000L) + ((x << 24) & 0xff0000000000L) + ((x << 40) & 0xff000000000000L) + (x << 56);
+		}
+
+		private static unsafe uint Fetch32(byte* data)
+		{
+#if BIGENDIAN
+			return (uint)data[3] + ((uint)data[2] << 8) + ((uint)data[1] << 16) + ((uint)data[0] << 24);
+#else
+			return (uint)data[0] + ((uint)data[1] << 8) + ((uint)data[2] << 16) + ((uint)data[3] << 24);
+#endif
+		}
+
+		private static unsafe ulong Fetch64(byte* data)
+		{
+#if BIGENDIAN
+			return ((ulong)Fetch32(data) << 32) + (ulong)Fetch32(data + 4);
+#else
+			return (ulong)Fetch32(data) + ((ulong)Fetch32(data + 4) << 32);
+#endif
+		}
+
+		private static uint fmix(uint h)
+		{
+			h ^= h >> 16;
+			h *= 0x85ebca6b;
+			h ^= h >> 13;
+			h *= 0xc2b2ae35;
+			h ^= h >> 16;
+			return h;
+		}
+
+		private static uint Rotate32(uint val, int shift)
+		{
+			// Avoid shifting by 32: doing so yields an undefined result.
+			return shift == 0 ? val : ((val >> shift) | (val << (32 - shift)));
+		}
+
+		private static void SwapValues<T>(ref T a, ref T b)
+		{
+			T c = a;
+			a = b;
+			b = c;
+		}
+
+		private static void PERMUTE3<T>(ref T a, ref T b, ref T c)
+		{
+			//SwapValues(ref a, ref b);
+			//SwapValues(ref a, ref c);
+			T d = c;
+			c = b;
+			b = a;
+			a = d;
+		}
+
+		private static uint Mur(uint a, uint h)
+		{
+			// Helper from Murmur3 for combining two 32-bit values.
+			a *= c1;
+			a = Rotate32(a, 17);
+			a *= c2;
+			h ^= a;
+			h = Rotate32(h, 19);
+			return h * 5 + 0xe6546b64;
+		}
+
+		private static unsafe uint Hash32Len13to24(byte* s, uint len)
+		{
+			uint a = Fetch32(s - 4 + (len >> 1));
+			uint b = Fetch32(s + 4);
+			uint c = Fetch32(s + len - 8);
+			uint d = Fetch32(s + (len >> 1));
+			uint e = Fetch32(s);
+			uint f = Fetch32(s + len - 4);
+			uint h = len;
+
+			return fmix(Mur(f, Mur(e, Mur(d, Mur(c, Mur(b, Mur(a, h)))))));
+		}
+
+		private static unsafe uint Hash32Len0to4(byte* s, uint len)
+		{
+			uint b = 0;
+			uint c = 9;
+			for (uint i = 0; i < len; i++)
+			{
+				char v = (char)s[i];
+				b = b * c1 + v;
+				c ^= b;
+			}
+			return fmix(Mur(b, Mur(len, c)));
+		}
+
+		private static unsafe uint Hash32Len5to12(byte* s, uint len)
+		{
+			uint a = len, b = len * 5, c = 9, d = b;
+			a += Fetch32(s);
+			b += Fetch32(s + len - 4);
+			c += Fetch32(s + ((len >> 1) & 4));
+			return fmix(Mur(c, Mur(b, Mur(a, d))));
+		}
+
+		public static unsafe uint CityHash32(byte* s, uint len)
+		{
+
+			if (len <= 24)
+			{
+				return len <= 12 ?
+					(len <= 4 ? Hash32Len0to4(s, len) : Hash32Len5to12(s, len)) :
+					Hash32Len13to24(s, len);
+			}
+
+			// len > 24
+			uint h = len, g = c1 * len, f = g;
+			uint a0 = Rotate32(Fetch32(s + len - 4) * c1, 17) * c2;
+			uint a1 = Rotate32(Fetch32(s + len - 8) * c1, 17) * c2;
+			uint a2 = Rotate32(Fetch32(s + len - 16) * c1, 17) * c2;
+			uint a3 = Rotate32(Fetch32(s + len - 12) * c1, 17) * c2;
+			uint a4 = Rotate32(Fetch32(s + len - 20) * c1, 17) * c2;
+			h ^= a0;
+			h = Rotate32(h, 19);
+			h = h * 5 + 0xe6546b64;
+			h ^= a2;
+			h = Rotate32(h, 19);
+			h = h * 5 + 0xe6546b64;
+			g ^= a1;
+			g = Rotate32(g, 19);
+			g = g * 5 + 0xe6546b64;
+			g ^= a3;
+			g = Rotate32(g, 19);
+			g = g * 5 + 0xe6546b64;
+			f += a4;
+			f = Rotate32(f, 19);
+			f = f * 5 + 0xe6546b64;
+			uint iters = (len - 1) / 20;
+			do
+			{
+				uint _a0 = Rotate32(Fetch32(s) * c1, 17) * c2;
+				uint _a1 = Fetch32(s + 4);
+				uint _a2 = Rotate32(Fetch32(s + 8) * c1, 17) * c2;
+				uint _a3 = Rotate32(Fetch32(s + 12) * c1, 17) * c2;
+				uint _a4 = Fetch32(s + 16);
+				h ^= _a0;
+				h = Rotate32(h, 18);
+				h = h * 5 + 0xe6546b64;
+				f += _a1;
+				f = Rotate32(f, 19);
+				f = f * c1;
+				g += _a2;
+				g = Rotate32(g, 18);
+				g = g * 5 + 0xe6546b64;
+				h ^= _a3 + _a1;
+				h = Rotate32(h, 19);
+				h = h * 5 + 0xe6546b64;
+				g ^= _a4;
+				g = bswap_32(g) * 5;
+				h += _a4 * 5;
+				h = bswap_32(h);
+				f += _a0;
+				PERMUTE3(ref f, ref h, ref g);
+				s += 20;
+			} while (--iters != 0);
+			g = Rotate32(g, 11) * c1;
+			g = Rotate32(g, 17) * c1;
+			f = Rotate32(f, 11) * c1;
+			f = Rotate32(f, 17) * c1;
+			h = Rotate32(h + g, 19);
+			h = h * 5 + 0xe6546b64;
+			h = Rotate32(h, 17) * c1;
+			h = Rotate32(h + f, 19);
+			h = h * 5 + 0xe6546b64;
+			h = Rotate32(h, 17) * c1;
+			return h;
+		}
+
+		public static uint CityHash32(ReadOnlySpan<byte> Span)
+		{
+			unsafe
+			{
+				fixed (byte* s = Span)
+				{
+					return CityHash32(s, (uint)Span.Length);
+				}
+			}
+		}
+
+		// Bitwise right rotate.  Normally this will compile to a single
+		// instruction, especially if the shift is a manifest constant.
+		private static ulong Rotate(ulong val, int shift)
+		{
+			// Avoid shifting by 64: doing so yields an undefined result.
+			return shift == 0 ? val : ((val >> shift) | (val << (64 - shift)));
+		}
+
+		private static ulong ShiftMix(ulong val)
+		{
+			return val ^ (val >> 47);
+		}
+
+		private static ulong HashLen16(ulong u, ulong v)
+		{
+			return CityHash128to64(new Uint128_64 { lo = u, hi = v });
+		}
+
+		private static ulong HashLen16(ulong u, ulong v, ulong mul)
+		{
+			// Murmur-inspired hashing.
+			ulong a = (u ^ v) * mul;
+			a ^= (a >> 47);
+			ulong b = (v ^ a) * mul;
+			b ^= (b >> 47);
+			b *= mul;
+			return b;
+		}
+
+		private static unsafe ulong HashLen0to16(byte* s, uint len)
+		{
+			if (len >= 8)
+			{
+				ulong mul = k2 + len * 2;
+				ulong a = Fetch64(s) + k2;
+				ulong b = Fetch64(s + len - 8);
+				ulong c = Rotate(b, 37) * mul + a;
+				ulong d = (Rotate(a, 25) + b) * mul;
+				return HashLen16(c, d, mul);
+			}
+			if (len >= 4)
+			{
+				ulong mul = k2 + len * 2;
+				ulong a = Fetch32(s);
+				return HashLen16(len + (a << 3), Fetch32(s + len - 4), mul);
+			}
+			if (len > 0)
+			{
+				byte a = s[0];
+				byte b = s[len >> 1];
+				byte c = s[len - 1];
+				uint y = (uint)a + ((uint)b << 8);
+				uint z = len + ((uint)c << 2);
+				return ShiftMix(y * k2 ^ z * k0) * k2;
+			}
+			return k2;
+		}
+
+		// This probably works well for 16-byte strings as well, but it may be overkill
+		// in that case.
+		private static unsafe ulong HashLen17to32(byte* s, uint len)
+		{
+			ulong mul = k2 + len * 2;
+			ulong a = Fetch64(s) * k1;
+			ulong b = Fetch64(s + 8);
+			ulong c = Fetch64(s + len - 8) * mul;
+			ulong d = Fetch64(s + len - 16) * k2;
+			return HashLen16(Rotate(a + b, 43) + Rotate(c, 30) + d,
+				a + Rotate(b + k2, 18) + c, mul);
+		}
+
+		// Return a 16-byte hash for 48 bytes.  Quick and dirty.
+		// Callers do best to use "random-looking" values for a and b.
+		private static unsafe Uint128_64 WeakHashLen32WithSeeds(
+			ulong w, ulong x, ulong y, ulong z, ulong a, ulong b)
+		{
+			a += w;
+			b = Rotate(b + a + z, 21);
+			ulong c = a;
+			a += x;
+			a += y;
+			b += Rotate(a, 44);
+			return new Uint128_64 { lo = (a + z), hi = (b + c) };
+		}
+
+		// Return a 16-byte hash for s[0] ... s[31], a, and b.  Quick and dirty.
+		private static unsafe Uint128_64 WeakHashLen32WithSeeds(byte* s, ulong a, ulong b)
+		{
+			return WeakHashLen32WithSeeds(Fetch64(s),
+				Fetch64(s + 8),
+				Fetch64(s + 16),
+				Fetch64(s + 24),
+				a,
+				b);
+		}
+
+		// Return an 8-byte hash for 33 to 64 bytes.
+		private static unsafe ulong HashLen33to64(byte* s, uint len)
+		{
+			ulong mul = k2 + len * 2;
+			ulong a = Fetch64(s) * k2;
+			ulong b = Fetch64(s + 8);
+			ulong c = Fetch64(s + len - 24);
+			ulong d = Fetch64(s + len - 32);
+			ulong e = Fetch64(s + 16) * k2;
+			ulong f = Fetch64(s + 24) * 9;
+			ulong g = Fetch64(s + len - 8);
+			ulong h = Fetch64(s + len - 16) * mul;
+			ulong u = Rotate(a + g, 43) + (Rotate(b, 30) + c) * 9;
+			ulong v = ((a + g) ^ d) + f + 1;
+			ulong w = bswap_64((u + v) * mul) + h;
+			ulong x = Rotate(e + f, 42) + c;
+			ulong y = (bswap_64((v + w) * mul) + g) * mul;
+			ulong z = e + f + c;
+			a = bswap_64((x + z) * mul + y) + b;
+			b = ShiftMix((z + a) * mul + d + h) * mul;
+			return b + x;
+		}
+
+		public static unsafe ulong CityHash64(byte* s, uint len)
+		{
+			if (len <= 32)
+			{
+				if (len <= 16)
+				{
+					return HashLen0to16(s, len);
+				}
+				else
+				{
+					return HashLen17to32(s, len);
+				}
+			}
+			else if (len <= 64)
+			{
+				return HashLen33to64(s, len);
+			}
+
+			// For strings over 64 bytes we hash the end first, and then as we
+			// loop we keep 56 bytes of state: v, w, x, y, and z.
+			ulong x = Fetch64(s + len - 40);
+			ulong y = Fetch64(s + len - 16) + Fetch64(s + len - 56);
+			ulong z = HashLen16(Fetch64(s + len - 48) + len, Fetch64(s + len - 24));
+			Uint128_64 v = WeakHashLen32WithSeeds(s + len - 64, len, z);
+			Uint128_64 w = WeakHashLen32WithSeeds(s + len - 32, y + k1, x);
+			x = x * k1 + Fetch64(s);
+
+			// Decrease len to the nearest multiple of 64, and operate on 64-byte chunks.
+			len = (len - 1) & ~(uint)63;
+			do
+			{
+				x = Rotate(x + y + v.lo + Fetch64(s + 8), 37) * k1;
+				y = Rotate(y + v.hi + Fetch64(s + 48), 42) * k1;
+				x ^= w.hi;
+				y += v.lo + Fetch64(s + 40);
+				z = Rotate(z + w.lo, 33) * k1;
+				v = WeakHashLen32WithSeeds(s, v.hi * k1, x + w.lo);
+				w = WeakHashLen32WithSeeds(s + 32, z + w.hi, y + Fetch64(s + 16));
+				SwapValues(ref z, ref x);
+				s += 64;
+				len -= 64;
+			} while (len != 0);
+			return HashLen16(HashLen16(v.lo, w.lo) + ShiftMix(y) * k1 + z,
+				HashLen16(v.hi, w.hi) + x);
+		}
+
+		public static ulong CityHash64(ReadOnlySpan<byte> Span)
+		{
+			unsafe
+			{
+				fixed (byte* s = Span)
+				{
+					return CityHash64(s, (uint)Span.Length);
+				}
+			}
+		}
+
+		public static unsafe ulong CityHash64WithSeed(byte* s, uint len, ulong seed)
+		{
+			return CityHash64WithSeeds(s, len, k2, seed);
+		}
+
+		public ulong CityHash64WithSeed(ReadOnlySpan<byte> Span, ulong seed)
+		{
+			return CityHash64WithSeeds(Span, k2, seed);
+		}
+
+		public static unsafe ulong CityHash64WithSeeds(byte* s, uint len, ulong seed0, ulong seed1)
+		{
+			return HashLen16(CityHash64(s, len) - seed0, seed1);
+		}
+
+		public static ulong CityHash64WithSeeds(ReadOnlySpan<byte> Span, ulong seed0, ulong seed1)
+		{
+			unsafe
+			{
+				fixed (byte* s = Span)
+				{
+					return CityHash64WithSeeds(s, (uint)Span.Length, seed0, seed1);
+				}
+			}
+		}
+
+		// Hash 128 input bits down to 64 bits of output.
+		// This is intended to be a reasonably good hash function.
+		public static ulong CityHash128to64(Uint128_64 x)
+		{
+			// Murmur-inspired hashing.
+			const ulong kMul = 0x9ddfea08eb382d69L;
+			ulong a = (x.lo ^ x.hi) * kMul;
+			a ^= (a >> 47);
+			ulong b = (x.hi ^ a) * kMul;
+			b ^= (b >> 47);
+			b *= kMul;
+			return b;
+		}
+	}
+}
diff --git a/Engine/Source/Programs/Shared/EpicGames.Core/StringComparerUE.cs b/Engine/Source/Programs/Shared/EpicGames.Core/StringComparerUE.cs
new file mode 100644
--- /dev/null
+++ b/Engine/Source/Programs/Shared/EpicGames.Core/StringComparerUE.cs
@@ -0,0 +1,72 @@
+// Copyright Epic Games, Inc. All Rights Reserved.
+
+using System;
+using System.Collections.Generic;
+
+namespace EpicGames.Core
+{
+
+	/// <summary>
+	/// Comparer for String objects.  However, it implements UE style ignore case compare
+	/// </summary>
+	public class StringComparerUE : IComparer<string>, IEqualityComparer<string>
+	{
+		/// <summary>
+		/// Static instance of an ordinal String comparer
+		/// </summary>
+		public static StringComparerUE Ordinal = new StringComparerUE(StringComparison.Ordinal);
+
+		/// <summary>
+		/// Static instance of an ordinal String comparer which ignores case
+		/// </summary>
+		public static StringComparerUE OrdinalIgnoreCase = new StringComparerUE(StringComparison.OrdinalIgnoreCase);
+
+		/// <summary>
+		/// The comparison type
+		/// </summary>
+		public StringComparison ComparisonType { get; }
+
+		/// <summary>
+		/// Constructor
+		/// </summary>
+		/// <param name="ComparisonType">Type of comparison to perform</param>
+		public StringComparerUE(StringComparison ComparisonType)
+		{
+			this.ComparisonType = ComparisonType;
+		}
+
+		/// <inheritdoc/>
+		public bool Equals(string? X, string? Y)
+		{
+			return string.Equals(X, Y, ComparisonType);
+		}
+
+		/// <inheritdoc/>
+		public int GetHashCode(string obj)
+		{
+			return string.GetHashCode(obj, ComparisonType);
+		}
+
+		/// <inheritdoc/>
+		public int Compare(string? X, string? Y)
+		{
+			if (ComparisonType == StringComparison.OrdinalIgnoreCase)
+			{
+				if (X == null)
+				{
+					return Y == null ? 0 : -1;
+				}
+				else if (Y == null)
+				{
+					return 1;
+				}
+
+				return StringUtils.CompareIgnoreCaseUE(X.AsSpan(), Y.AsSpan());
+			}
+			else
+			{
+				return string.Compare(X, Y, ComparisonType);
+			}
+		}
+	}
+}
diff --git a/Engine/Source/Programs/Shared/EpicGames.Core/StringUtils.cs b/Engine/Source/Programs/Shared/EpicGames.Core/StringUtils.cs
--- a/Engine/Source/Programs/Shared/EpicGames.Core/StringUtils.cs
+++ b/Engine/Source/Programs/Shared/EpicGames.Core/StringUtils.cs
@@ -584,5 +584,55 @@ namespace EpicGames.Core
 			}
 			return Text;
 		}
+
+		/// <summary>
+		/// Compare two strings using UnrealEngine's ignore case algorithm
+		/// </summary>
+		/// <param name="X">First string to compare</param>
+		/// <param name="Y">Second string to compare</param>
+		/// <returns>Less than zero if X < Y, zero if X == Y, and greater than zero if X > y</returns>
+		public static int CompareIgnoreCaseUE(ReadOnlySpan<char> X, ReadOnlySpan<char> Y)
+		{
+			int Length = X.Length < Y.Length ? X.Length : Y.Length;
+
+			for (int Index = 0; Index < Length; ++Index)
+			{
+				char XC = X[Index];
+				char YC = Y[Index];
+				if (XC == YC)
+				{
+					continue;
+				}
+				else if (((XC | YC) & 0xffffff80) == 0) // if (BothAscii)
+				{
+					if (XC >= 'A' && XC <= 'Z')
+					{
+						XC += (char)32;
+					}
+					if (YC >= 'A' && YC <= 'Z')
+					{
+						YC += (char)32;
+					}
+					int Diff = XC - YC;
+					if (Diff != 0)
+					{
+						return Diff;
+					}
+				}
+				else
+				{
+					return XC - YC;
+				}
+			}
+
+			if (X.Length == Length)
+			{
+				return Y.Length == Length ? 0 : /* X[Length] */ -Y[Length];
+			}
+			else
+			{
+				return X[Length] /* - Y[Length] */;
+			}
+		}
 	}
 }
diff --git a/Engine/Source/Programs/Shared/EpicGames.Core/StringView.cs b/Engine/Source/Programs/Shared/EpicGames.Core/StringView.cs
--- a/Engine/Source/Programs/Shared/EpicGames.Core/StringView.cs
+++ b/Engine/Source/Programs/Shared/EpicGames.Core/StringView.cs
@@ -2,13 +2,11 @@
 
 using System;
 using System.Collections.Generic;
-using System.Diagnostics.CodeAnalysis;
-using System.Text;
 
 namespace EpicGames.Core
 {
 	/// <summary>
-	/// View of a character string. Allows comparing/manipulating substrings without unnecessary memory allocatinos.
+	/// View of a character string. Allows comparing/manipulating substrings without unnecessary memory allocations.
 	/// </summary>
 	public readonly struct StringView : IEquatable<StringView>
 	{
@@ -25,6 +23,44 @@ namespace EpicGames.Core
 			get { return Memory.Span; } 
 		}
 
+		/// <summary>
+		/// Length of the sequence of characters
+		/// </summary>
+		public int Length
+		{
+			get { return Memory.Length; }
+		}
+
+		/// <summary>
+		/// Constructor
+		/// </summary>
+		/// <param name="Text">String view to construct from</param>
+		public StringView(StringView Text)
+		{
+			this.Memory = Text.Memory;
+		}
+
+		/// <summary>
+		/// Constructor
+		/// </summary>
+		/// <param name="Text">String view to construct from</param>
+		/// <param name="Index">Offset within the string for this view</param>
+		public StringView(StringView Text, int Index)
+		{
+			this.Memory = Text.Memory.Slice(Index);
+		}
+
+		/// <summary>
+		/// Constructor
+		/// </summary>
+		/// <param name="Text">String view to construct from</param>
+		/// <param name="Index">Offset within the string for this view</param>
+		/// <param name="Count">Number of characters to include</param>
+		public StringView(StringView Text, int Index, int Count)
+		{
+			this.Memory = Text.Memory.Slice(Index, Count);
+		}
+
 		/// <summary>
 		/// Constructor
 		/// </summary>
@@ -34,6 +70,27 @@ namespace EpicGames.Core
 			this.Memory = Memory;
 		}
 
+		/// <summary>
+		/// Constructor
+		/// </summary>
+		/// <param name="Memory">The memory containing the characters</param>
+		/// <param name="Index">Offset within the string for this view</param>
+		public StringView(ReadOnlyMemory<char> Memory, int Index)
+		{
+			this.Memory = Memory.Slice(Index);
+		}
+
+		/// <summary>
+		/// Constructor
+		/// </summary>
+		/// <param name="Memory">The memory containing the characters</param>
+		/// <param name="Index">Offset within the string for this view</param>
+		/// <param name="Count">Number of characters to include</param>
+		public StringView(ReadOnlyMemory<char> Memory, int Index, int Count)
+		{
+			this.Memory = Memory.Slice(Index, Count);
+		}
+
 		/// <summary>
 		/// Constructor
 		/// </summary>
@@ -98,6 +155,12 @@ namespace EpicGames.Core
 			return String.GetHashCode(Memory.Span);
 		}
 
+		/// <inheritdoc/>
+		public int GetHashCode(StringComparison comparisonType)
+		{
+			return String.GetHashCode(Memory.Span, comparisonType);
+		}
+
 		/// <inheritdoc/>
 		public override string ToString()
 		{
@@ -161,7 +224,7 @@ namespace EpicGames.Core
 		/// <inheritdoc/>
 		public int GetHashCode(StringView obj)
 		{
-			return String.GetHashCode(obj.Span);
+			return String.GetHashCode(obj.Span, ComparisonType);
 		}
 
 		/// <inheritdoc/>
@@ -170,4 +233,59 @@ namespace EpicGames.Core
 			return X.Span.CompareTo(Y.Span, ComparisonType);
 		}
 	}
+
+	/// <summary>
+	/// Comparer for StringView objects.  However, it implements UE style ignore case compare
+	/// </summary>
+	public class StringViewComparerUE : IComparer<StringView>, IEqualityComparer<StringView>
+	{
+		/// <summary>
+		/// Static instance of an ordinal StringView comparer
+		/// </summary>
+		public static StringViewComparerUE Ordinal = new StringViewComparerUE(StringComparison.Ordinal);
+
+		/// <summary>
+		/// Static instance of an ordinal StringView comparer which ignores case
+		/// </summary>
+		public static StringViewComparerUE OrdinalIgnoreCase = new StringViewComparerUE(StringComparison.OrdinalIgnoreCase);
+
+		/// <summary>
+		/// The comparison type
+		/// </summary>
+		public StringComparison ComparisonType { get; }
+
+		/// <summary>
+		/// Constructor
+		/// </summary>
+		/// <param name="ComparisonType">Type of comparison to perform</param>
+		public StringViewComparerUE(StringComparison ComparisonType)
+		{
+			this.ComparisonType = ComparisonType;
+		}
+
+		/// <inheritdoc/>
+		public bool Equals(StringView X, StringView Y)
+		{
+			return X.Span.Equals(Y.Span, ComparisonType);
+		}
+
+		/// <inheritdoc/>
+		public int GetHashCode(StringView obj)
+		{
+			return String.GetHashCode(obj.Span, ComparisonType);
+		}
+
+		/// <inheritdoc/>
+		public int Compare(StringView X, StringView Y)
+		{
+			if (ComparisonType == StringComparison.OrdinalIgnoreCase)
+			{
+				return StringUtils.CompareIgnoreCaseUE(X.Span, Y.Span);
+			}
+			else
+			{
+				return X.Span.CompareTo(Y.Span, ComparisonType);
+			}
+		}
+	}
 }
diff --git a/Engine/Source/Programs/Shared/EpicGames.Core/UHTTypes.cs b/Engine/Source/Programs/Shared/EpicGames.Core/UHTTypes.cs
new file mode 100644
--- /dev/null
+++ b/Engine/Source/Programs/Shared/EpicGames.Core/UHTTypes.cs
@@ -0,0 +1,125 @@
+﻿// Copyright Epic Games, Inc. All Rights Reserved.
+
+using System.Collections.Generic;
+using System.Text.Json.Serialization;
+
+namespace EpicGames.Core
+{
+	/// <summary>
+	/// Defines the version of the code generation to be used.
+	/// 
+	/// This MUST be kept in sync with EGeneratedBodyVersion defined in 
+	/// Engine\Source\Programs\UnrealHeaderTool\Private\GeneratedCodeVersion.h.
+	/// </summary>
+	public enum EGeneratedCodeVersion
+	{
+		/// <summary>
+		/// Version not set or the default is to be used.
+		/// </summary>
+		None,
+
+		/// <summary>
+		/// 
+		/// </summary>
+		V1,
+
+		/// <summary>
+		/// 
+		/// </summary>
+		V2,
+
+		/// <summary>
+		/// 
+		/// </summary>
+		VLatest = V2
+	};
+
+	/// <summary>
+	/// This MUST be kept in sync with EPackageOverrideType defined in 
+	/// Engine\Source\Programs\UnrealHeaderTool\Public\IScriptGeneratorPluginInterface.h
+	/// </summary>
+	public enum EPackageOverrideType
+	{
+		None,
+		EditorOnly,
+		EngineDeveloper,
+		GameDeveloper,
+		EngineUncookedOnly,
+		GameUncookedOnly,
+	}
+
+	/// <summary>
+	/// Type of module. This should be sorted by the order in which we expect modules to be built.
+	/// 
+	/// This MUST be kept in sync with EBuildModuleType defined in 
+	/// Engine\Source\Programs\UnrealHeaderTool\Public\IScriptGeneratorPluginInterface.h
+	/// </summary>
+	public enum UHTModuleType
+	{
+		Program,
+		EngineRuntime,
+		EngineUncooked,
+		EngineDeveloper,
+		EngineEditor,
+		EngineThirdParty,
+		GameRuntime,
+		GameUncooked,
+		GameDeveloper,
+		GameEditor,
+		GameThirdParty,
+	}
+
+	public class UHTManifest
+	{
+		public class Module
+		{
+			public string Name { get; set; } = "";
+			[JsonConverter(typeof(JsonStringEnumConverter))]
+			public UHTModuleType ModuleType { get; set; } = UHTModuleType.Program;
+			[JsonConverter(typeof(JsonStringEnumConverter))]
+			public EPackageOverrideType OverrideModuleType { get; set; } = EPackageOverrideType.None;
+			public string BaseDirectory { get; set; } = string.Empty;
+			public string IncludeBase { get; set; } = string.Empty; // The include path which all UHT-generated includes should be relative to
+			public string OutputDirectory { get; set; } = string.Empty;
+			public List<string> ClassesHeaders { get; set; } = new List<string>();
+			public List<string> PublicHeaders { get; set; } = new List<string>();
+			public List<string> InternalHeaders { get; set; } = new List<string>();
+			public List<string> PrivateHeaders { get; set; } = new List<string>();
+			public string? GeneratedCPPFilenameBase { get; set; } = null;
+			public bool SaveExportedHeaders { get; set; } = false;
+			[JsonConverter(typeof(JsonStringEnumConverter))]
+			[JsonPropertyName("UHTGeneratedCodeVersion")]
+			public EGeneratedCodeVersion GeneratedCodeVersion { get; set; } = EGeneratedCodeVersion.None;
+
+			public override string ToString()
+			{
+				return Name;
+			}
+		}
+
+		/// <summary>
+		/// True if the current target is a game target
+		/// </summary>
+		public bool IsGameTarget { get; set; } = false;
+
+		/// <summary>
+		/// The engine path on the local machine
+		/// </summary>
+		public string RootLocalPath { get; set; } = string.Empty;
+
+		/// <summary>
+		/// Name of the target currently being compiled
+		/// </summary>
+		public string TargetName { get; set; } = string.Empty;
+
+		/// <summary>
+		/// File to contain additional dependencies that the generated code depends on
+		/// </summary>
+		public string ExternalDependenciesFile { get; set; } = string.Empty;
+
+		/// <summary>
+		/// List of modules
+		/// </summary>
+		public List<Module> Modules { get; set; } = new List<Module>();
+	}
+}
diff --git a/Engine/Source/Programs/Shared/EpicGames.Core/UnrealEngineTypes.cs b/Engine/Source/Programs/Shared/EpicGames.Core/UnrealEngineTypes.cs
new file mode 100644
--- /dev/null
+++ b/Engine/Source/Programs/Shared/EpicGames.Core/UnrealEngineTypes.cs
@@ -0,0 +1,1450 @@
+﻿// Copyright Epic Games, Inc. All Rights Reserved.
+
+using System;
+using System.Collections.Generic;
+using System.Reflection;
+
+namespace EpicGames.Core
+{
+	/// <summary>
+	/// Standard return codes used by UBT and UHT
+	/// 
+	/// This MUST be kept in sync with EGeneratedBodyVersion defined in 
+	/// Engine\Source\Runtime\Core\Public\Misc\ComplilationResult.h.
+	/// </summary>
+	public enum CompilationResult
+	{
+		/// <summary>
+		/// Compilation succeeded
+		/// </summary>
+		Succeeded = 0,
+
+		/// <summary>
+		/// Build was canceled, this is used on the engine side only
+		/// </summary>
+		Canceled = 1,
+
+		/// <summary>
+		/// All targets were up to date, used only with -canskiplink
+		/// </summary>
+		UpToDate = 2,
+
+		/// <summary>
+		/// The process has most likely crashed. This is what UE returns in case of an assert
+		/// </summary>
+		CrashOrAssert = 3,
+
+		/// <summary>
+		/// Compilation failed because generated code changed which was not supported
+		/// </summary>
+		FailedDueToHeaderChange = 4,
+
+		/// <summary>
+		/// Compilation failed due to the engine modules needing to be rebuilt
+		/// </summary>
+		FailedDueToEngineChange = 5,
+
+		/// <summary>
+		/// Compilation failed due to compilation errors
+		/// </summary>
+		OtherCompilationError = 6,
+
+		/// <summary>
+		/// Compilation failed due to live coding action limit being exceeded.
+		/// </summary>
+		LiveCodingLimitError = 7,
+
+		/// <summary>
+		/// Compilation is not supported in the current build
+		/// </summary>
+		Unsupported,
+
+		/// <summary>
+		/// Unknown error
+		/// </summary>
+		Unknown
+	}
+
+	/// <summary>
+	/// Helper extensions for CompilationResult
+	/// </summary>
+	public static class CompilationResultExtensions
+	{
+
+		/// <summary>
+		/// Test to see if the return code is a success
+		/// </summary>
+		/// <param name="Result">The result to test</param>
+		/// <returns>True if the return code is a success, false if it is not.</returns>
+		public static bool Succeeded(this CompilationResult Result)
+		{
+			return Result == CompilationResult.Succeeded || Result == CompilationResult.UpToDate;
+		}
+	}
+
+	/// <summary>
+	/// Flags describing a class.
+	///
+	/// This MUST be kept in sync with EClassFlags defined in 
+	/// Engine\Source\Runtime\CoreUObject\Public\UObject\ObjectMacros.h
+	/// </summary>
+	[Flags]
+	public enum EClassFlags : UInt32
+	{
+		/// <summary>
+		/// No Flags
+		/// </summary>
+		None = 0x00000000u,
+
+		/// <summary>
+		/// Class is abstract and can't be instantiated directly.
+		/// </summary>
+		Abstract = 0x00000001u,
+
+		/// <summary>
+		/// Save object configuration only to Default INIs, never to local INIs. Must be combined with "Config"
+		/// </summary>
+		DefaultConfig = 0x00000002u,
+
+		/// <summary>
+		/// Load object configuration at construction time.
+		/// </summary>
+		Config = 0x00000004u,
+
+		/// <summary>
+		/// This object type can't be saved; null it out at save time.
+		/// </summary>
+		Transient = 0x00000008u,
+
+		/// <summary>
+		/// This object type may not be available in certain context. (i.e. game runtime or in certain configuration). Optional class data is saved separately to other object types. (i.e. might use sidecar files)
+		/// </summary>
+		Optional = 0x00000010u,
+
+		/// <summary>
+		/// 
+		/// </summary>
+		MatchedSerializers = 0x00000020u,
+
+		/// <summary>
+		/// Indicates that the config settings for this class will be saved to Project/User*.ini (similar to "GlobalUserConfig")
+		/// </summary>
+		ProjectUserConfig = 0x00000040u,
+
+		/// <summary>
+		/// Class is a native class - native interfaces will have "Native" set, but not RF_MarkAsNative
+		/// </summary>
+		Native = 0x00000080u,
+
+		/// <summary>
+		/// Don't export to C++ header.
+		/// </summary>
+		NoExport = 0x00000100u,
+
+		/// <summary>
+		/// Do not allow users to create in the editor.
+		/// </summary>
+		NotPlaceable = 0x00000200u,
+
+		/// <summary>
+		/// Handle object configuration on a per-object basis, rather than per-class.
+		/// </summary>
+		PerObjectConfig = 0x00000400u,
+
+		/// <summary>
+		/// Whether SetUpRuntimeReplicationData still needs to be called for this class
+		/// </summary>
+		ReplicationDataIsSetUp = 0x00000800u,
+
+		/// <summary>
+		/// Class can be constructed from editinline New button.
+		/// </summary>
+		EditInlineNew = 0x00001000u,
+
+		/// <summary>
+		/// Display properties in the editor without using categories.
+		/// </summary>
+		CollapseCategories = 0x00002000u,
+
+		/// <summary>
+		/// Class is an interface
+		/// </summary>
+		Interface = 0x00004000u,
+
+		/// <summary>
+		/// Do not export a constructor for this class, assuming it is in the cpptext
+		/// </summary>
+		CustomConstructor = 0x00008000u,
+
+		/// <summary>
+		/// all properties and functions in this class are const and should be exported as const
+		/// </summary>
+		Const = 0x00010000u,
+
+		/// <summary>
+		/// Class flag indicating objects of this class need deferred dependency loading
+		/// </summary>
+		NeedsDeferredDependencyLoading = 0x00020000u,
+
+		/// <summary>
+		/// Indicates that the class was created from blueprint source material
+		/// </summary>
+		CompiledFromBlueprint = 0x00040000u,
+
+		/// <summary>
+		/// Indicates that only the bare minimum bits of this class should be DLL exported/imported
+		/// </summary>
+		MinimalAPI = 0x00080000u,
+
+		/// <summary>
+		/// Indicates this class must be DLL exported/imported (along with all of it's members)
+		/// </summary>
+		RequiredAPI = 0x00100000u,
+
+		/// <summary>
+		/// Indicates that references to this class default to instanced. Used to be subclasses of UComponent, but now can be any UObject
+		/// </summary>
+		DefaultToInstanced = 0x00200000u,
+
+		/// <summary>
+		/// Indicates that the parent token stream has been merged with ours.
+		/// </summary>
+		TokenStreamAssembled = 0x00400000u,
+
+		/// <summary>
+		/// Class has component properties.
+		/// </summary>
+		HasInstancedReference = 0x00800000u,
+
+		/// <summary>
+		/// Don't show this class in the editor class browser or edit inline new menus.
+		/// </summary>
+		Hidden = 0x01000000u,
+
+		/// <summary>
+		/// Don't save objects of this class when serializing
+		/// </summary>
+		Deprecated = 0x02000000u,
+
+		/// <summary>
+		/// Class not shown in editor drop down for class selection
+		/// </summary>
+		HideDropDown = 0x04000000u,
+
+		/// <summary>
+		/// Class settings are saved to <AppData>/..../Blah.ini (as opposed to "DefaultConfig")
+		/// </summary>
+		GlobalUserConfig = 0x08000000u,
+
+		/// <summary>
+		/// Class was declared directly in C++ and has no boilerplate generated by UnrealHeaderTool
+		/// </summary>
+		Intrinsic = 0x10000000u,
+
+		/// <summary>
+		/// Class has already been constructed (maybe in a previous DLL version before hot-reload).
+		/// </summary>
+		Constructed = 0x20000000u,
+
+		/// <summary>
+		/// Indicates that object configuration will not check against ini base/defaults when serialized
+		/// </summary>
+		ConfigDoNotCheckDefaults = 0x40000000u,
+
+		/// <summary>
+		/// Class has been consigned to oblivion as part of a blueprint recompile, and a newer version currently exists.
+		/// </summary>
+		NewerVersionExists = 0x80000000u,
+
+		/// <summary>
+		/// Flags to inherit from base class
+		/// </summary>
+		Inherit = Transient | Optional | DefaultConfig | Config | PerObjectConfig | ConfigDoNotCheckDefaults | NotPlaceable | Const | HasInstancedReference |
+			Deprecated | DefaultToInstanced | GlobalUserConfig | ProjectUserConfig | NeedsDeferredDependencyLoading,
+
+		/// <summary>
+		/// These flags will be cleared by the compiler when the class is parsed during script compilation
+		/// </summary>
+		RecompilerClear = Inherit | Abstract | NoExport | Native | Intrinsic | TokenStreamAssembled,
+
+		/// <summary>
+		/// These flags will be cleared by the compiler when the class is parsed during script compilation
+		/// </summary>
+		ShouldNeverBeLoaded = Native | Optional | Intrinsic | TokenStreamAssembled,
+
+		/// <summary>
+		/// These flags will be inherited from the base class only for non-intrinsic classes
+		/// </summary>
+		ScriptInherit = Inherit | EditInlineNew | CollapseCategories,
+
+		/// <summary>
+		/// This is used as a mask for the flags put into generated code for "compiled in" classes.
+		/// </summary>
+		SaveInCompiledInClasses = Abstract | DefaultConfig | GlobalUserConfig | ProjectUserConfig | Config | Transient | Optional | Native | NotPlaceable | PerObjectConfig |
+			ConfigDoNotCheckDefaults | EditInlineNew | CollapseCategories | Interface | DefaultToInstanced | HasInstancedReference | Hidden | Deprecated |
+			HideDropDown | Intrinsic | Const | MinimalAPI | RequiredAPI | MatchedSerializers | NeedsDeferredDependencyLoading,
+	};
+
+	/// <summary>
+	/// Extension message for EClassFlags
+	/// </summary>
+	public static class EClassFlagsExtensions
+	{
+		public static bool HasAnyFlags(this EClassFlags InFlags, EClassFlags TestFlags)
+		{
+			return (InFlags & TestFlags) != 0;
+		}
+
+		public static bool HasAllFlags(this EClassFlags InFlags, EClassFlags TestFlags)
+		{
+			return (InFlags & TestFlags) == TestFlags;
+		}
+
+		public static bool HasExactFlags(this EClassFlags InFlags, EClassFlags TestFlags, EClassFlags MatchFlags)
+		{
+			return (InFlags & TestFlags) == MatchFlags;
+		}
+	}
+
+	/// <summary>
+	/// Flags used for quickly casting classes of certain types; all class cast flags are inherited
+	///
+	/// This MUST be kept in sync with EClassCastFlags defined in 
+	/// Engine\Source\Runtime\CoreUObject\Public\UObject\ObjectMacros.h
+	/// </summary>
+	[Flags]
+	public enum EClassCastFlags : UInt64
+	{
+		None = 0x0000000000000000,
+
+		UField = 0x0000000000000001,
+		FInt8Property = 0x0000000000000002,
+		UEnum = 0x0000000000000004,
+		UStruct = 0x0000000000000008,
+		UScriptStruct = 0x0000000000000010,
+		UClass = 0x0000000000000020,
+		FByteProperty = 0x0000000000000040,
+		FIntProperty = 0x0000000000000080,
+		FFloatProperty = 0x0000000000000100,
+		FUInt64Property = 0x0000000000000200,
+		FClassProperty = 0x0000000000000400,
+		FUInt32Property = 0x0000000000000800,
+		FInterfaceProperty = 0x0000000000001000,
+		FNameProperty = 0x0000000000002000,
+		FStrProperty = 0x0000000000004000,
+		FProperty = 0x0000000000008000,
+		FObjectProperty = 0x0000000000010000,
+		FBoolProperty = 0x0000000000020000,
+		FUInt16Property = 0x0000000000040000,
+		UFunction = 0x0000000000080000,
+		FStructProperty = 0x0000000000100000,
+		FArrayProperty = 0x0000000000200000,
+		FInt64Property = 0x0000000000400000,
+		FDelegateProperty = 0x0000000000800000,
+		FNumericProperty = 0x0000000001000000,
+		FMulticastDelegateProperty = 0x0000000002000000,
+		FObjectPropertyBase = 0x0000000004000000,
+		FWeakObjectProperty = 0x0000000008000000,
+		FLazyObjectProperty = 0x0000000010000000,
+		FSoftObjectProperty = 0x0000000020000000,
+		FTextProperty = 0x0000000040000000,
+		FInt16Property = 0x0000000080000000,
+		FDoubleProperty = 0x0000000100000000,
+		FSoftClassProperty = 0x0000000200000000,
+		UPackage = 0x0000000400000000,
+		ULevel = 0x0000000800000000,
+		AActor = 0x0000001000000000,
+		APlayerController = 0x0000002000000000,
+		APawn = 0x0000004000000000,
+		USceneComponent = 0x0000008000000000,
+		UPrimitiveComponent = 0x0000010000000000,
+		USkinnedMeshComponent = 0x0000020000000000,
+		USkeletalMeshComponent = 0x0000040000000000,
+		UBlueprint = 0x0000080000000000,
+		UDelegateFunction = 0x0000100000000000,
+		UStaticMeshComponent = 0x0000200000000000,
+		FMapProperty = 0x0000400000000000,
+		FSetProperty = 0x0000800000000000,
+		FEnumProperty = 0x0001000000000000,
+		USparseDelegateFunction = 0x0002000000000000,
+		FMulticastInlineDelegateProperty = 0x0004000000000000,
+		FMulticastSparseDelegateProperty = 0x0008000000000000,
+		FFieldPathProperty = 0x0010000000000000,
+		FObjectPtrProperty = 0x0020000000000000,
+		FClassPtrProperty = 0x0040000000000000,
+		FLargeWorldCoordinatesRealProperty = 0x0080000000000000,
+		AllFlags = UInt64.MaxValue,
+	};
+
+	/// <summary>
+	/// Helper extensions for EClassCastFlags
+	/// </summary>
+	public static class EClassCastFlagsExtensions
+	{
+		public static bool HasAnyFlags(this EClassCastFlags InFlags, EClassCastFlags TestFlags)
+		{
+			return (InFlags & TestFlags) != 0;
+		}
+
+		public static bool HasAllFlags(this EClassCastFlags InFlags, EClassCastFlags TestFlags)
+		{
+			return (InFlags & TestFlags) == TestFlags;
+		}
+
+		public static bool HasExactFlags(this EClassCastFlags InFlags, EClassCastFlags TestFlags, EClassCastFlags MatchFlags)
+		{
+			return (InFlags & TestFlags) == MatchFlags;
+		}
+	}
+
+	/// <summary>
+	/// Flags describing a UEnum
+	///
+	/// This MUST be kept in sync with EEnumFlags defined in 
+	/// Engine\Source\Runtime\CoreUObject\Public\UObject\ObjectMacros.h
+	/// </summary>
+	[Flags]
+	public enum EEnumFlags : UInt32
+	{
+		None,
+		Flags = 1 << 0,
+		NewerVersionExists = 1 << 1,
+	}
+
+	/// <summary>
+	/// Helper extensions for EEnumFlags
+	/// </summary>
+	public static class EEnumFlagsExtensions
+	{
+		public static bool HasAnyFlags(this EEnumFlags InFlags, EEnumFlags TestFlags)
+		{
+			return (InFlags & TestFlags) != 0;
+		}
+
+		public static bool HasAllFlags(this EEnumFlags InFlags, EEnumFlags TestFlags)
+		{
+			return (InFlags & TestFlags) == TestFlags;
+		}
+
+		public static bool HasExactFlags(this EEnumFlags InFlags, EEnumFlags TestFlags, EEnumFlags MatchFlags)
+		{
+			return (InFlags & TestFlags) == MatchFlags;
+		}
+	}
+
+	/// <summary>
+	/// Function flags.
+	///
+	/// This MUST be kept in sync with EFunctionFlags defined in 
+	/// Engine\Source\Runtime\CoreUObject\Public\UObject\Script.h
+	/// </summary>
+	[Flags]
+	public enum EFunctionFlags : UInt32
+	{
+		None = 0x00000000,
+
+		/// <summary>
+		/// Function is final (prebindable, non-overridable function).
+		/// </summary>
+		Final = 0x00000001,
+
+		/// <summary>
+		/// Indicates this function is DLL exported/imported.
+		/// </summary>
+		RequiredAPI = 0x00000002,
+
+		/// <summary>
+		/// Function will only run if the object has network authority
+		/// </summary>
+		BlueprintAuthorityOnly = 0x00000004,
+
+		/// <summary>
+		/// Function is cosmetic in nature and should not be invoked on dedicated servers
+		/// </summary>
+		BlueprintCosmetic = 0x00000008,
+
+		// = 0x00000010,   // unused.
+		// = 0x00000020,   // unused.
+
+		/// <summary>
+		/// Function is network-replicated.
+		/// </summary>
+		Net = 0x00000040,
+
+		/// <summary>
+		/// Function should be sent reliably on the network.
+		/// </summary>
+		NetReliable = 0x00000080,
+
+		/// <summary>
+		/// Function is sent to a net service
+		/// </summary>
+		NetRequest = 0x00000100,
+
+		/// <summary>
+		/// Executable from command line.
+		/// </summary>
+		Exec = 0x00000200,
+
+		/// <summary>
+		/// Native function.
+		/// </summary>
+		Native = 0x00000400,
+
+		/// <summary>
+		/// Event function.
+		/// </summary>
+		Event = 0x00000800,
+
+		/// <summary>
+		/// Function response from a net service
+		/// </summary>
+		NetResponse = 0x00001000,
+
+		/// <summary>
+		/// Static function.
+		/// </summary>
+		Static = 0x00002000,
+
+		/// <summary>
+		/// Function is networked multicast Server -> All Clients
+		/// </summary>
+		NetMulticast = 0x00004000,
+
+		/// <summary>
+		/// Function is used as the merge 'ubergraph' for a blueprint, only assigned when using the persistent 'ubergraph' frame
+		/// </summary>
+		UbergraphFunction = 0x00008000,
+
+		/// <summary>
+		/// Function is a multi-cast delegate signature (also requires Delegate to be set!)
+		/// </summary>
+		MulticastDelegate = 0x00010000,
+
+		/// <summary>
+		/// Function is accessible in all classes (if overridden, parameters must remain unchanged).
+		/// </summary>
+		Public = 0x00020000,
+
+		/// <summary>
+		/// Function is accessible only in the class it is defined in (cannot be overridden, but function name may be reused in subclasses.  IOW: if overridden, parameters don't need to match, and Super.Func() cannot be accessed since it's private.)
+		/// </summary>
+		Private = 0x00040000,
+
+		/// <summary>
+		/// Function is accessible only in the class it is defined in and subclasses (if overridden, parameters much remain unchanged).
+		/// </summary>
+		Protected = 0x00080000,
+
+		/// <summary>
+		/// Function is delegate signature (either single-cast or multi-cast, depending on whether MulticastDelegate is set.)
+		/// </summary>
+		Delegate = 0x00100000,
+
+		/// <summary>
+		/// Function is executed on servers (set by replication code if passes check)
+		/// </summary>
+		NetServer = 0x00200000,
+
+		/// <summary>
+		/// function has out (pass by reference) parameters
+		/// </summary>
+		HasOutParms = 0x00400000,
+
+		/// <summary>
+		/// function has structs that contain defaults
+		/// </summary>
+		HasDefaults = 0x00800000,
+
+		/// <summary>
+		/// function is executed on clients
+		/// </summary>
+		NetClient = 0x01000000,
+
+		/// <summary>
+		/// function is imported from a DLL
+		/// </summary>
+		DLLImport = 0x02000000,
+
+		/// <summary>
+		/// function can be called from blueprint code
+		/// </summary>
+		BlueprintCallable = 0x04000000,
+
+		/// <summary>
+		/// function can be overridden/implemented from a blueprint
+		/// </summary>
+		BlueprintEvent = 0x08000000,
+
+		/// <summary>
+		/// function can be called from blueprint code, and is also pure (produces no side effects). If you set this, you should set BlueprintCallable as well.
+		/// </summary>
+		BlueprintPure = 0x10000000,
+
+		/// <summary>
+		/// function can only be called from an editor scrippt.
+		/// </summary>
+		EditorOnly = 0x20000000,
+
+		/// <summary>
+		/// function can be called from blueprint code, and only reads state (never writes state)
+		/// </summary>
+		Const = 0x40000000,
+
+		/// <summary>
+		/// function must supply a _Validate implementation
+		/// </summary>
+		NetValidate = 0x80000000,
+
+		AllFlags = 0xFFFFFFFF,
+
+		// Combinations of flags.
+		FuncInherit = Exec | Event | BlueprintCallable | BlueprintEvent | BlueprintAuthorityOnly | BlueprintCosmetic | Const,
+		FuncOverrideMatch = Exec | Final | Static | Public | Protected | Private,
+		NetFuncFlags = Net | NetReliable | NetServer | NetClient | NetMulticast,
+		AccessSpecifiers = Public | Private | Protected,
+	}
+
+	/// <summary>
+	/// Helper extensions for EFunctionFlags
+	/// </summary>
+	public static class EFunctionFlagsExtensions
+	{
+		public static bool HasAnyFlags(this EFunctionFlags InFlags, EFunctionFlags TestFlags)
+		{
+			return (InFlags & TestFlags) != 0;
+		}
+
+		public static bool HasAllFlags(this EFunctionFlags InFlags, EFunctionFlags TestFlags)
+		{
+			return (InFlags & TestFlags) == TestFlags;
+		}
+
+		public static bool HasExactFlags(this EFunctionFlags InFlags, EFunctionFlags TestFlags, EFunctionFlags MatchFlags)
+		{
+			return (InFlags & TestFlags) == MatchFlags;
+		}
+	}
+
+	/// <summary>
+	/// Objects flags for internal use(GC, low level UObject code)
+	/// 
+	/// This MUST be kept in sync with EInternalObjectFlags defined in 
+	/// Engine\Source\Runtime\CoreUObject\Public\UObject\ObjectMacros.h
+	/// </summary>
+	[Flags]
+	public enum EInternalObjectFlags : Int32
+	{
+		None = 0,
+		//~ All the other bits are reserved, DO NOT ADD NEW FLAGS HERE!
+
+		/// <summary>
+		/// Object is ready to be imported by another package during loading
+		/// </summary>
+		LoaderImport = 1 << 20,
+
+		/// <summary>
+		/// Garbage from logical point of view and should not be referenced. This flag is mirrored in EObjectFlags as RF_Garbage for performance
+		/// </summary>
+		Garbage = 1 << 21,
+
+		/// <summary>
+		/// Same as above but referenced through a persistent reference so it can't be GC'd
+		/// </summary>
+		PersistentGarbage = 1 << 22,
+
+		/// <summary>
+		/// External reference to object in cluster exists
+		/// </summary>
+		ReachableInCluster = 1 << 23,
+
+		/// <summary>
+		/// Root of a cluster
+		/// </summary>
+		ClusterRoot = 1 << 24,
+
+		/// <summary>
+		/// Native (UClass only).
+		/// </summary>
+		Native = 1 << 25,
+
+		/// <summary>
+		/// Object exists only on a different thread than the game thread.
+		/// </summary>
+		Async = 1 << 26,
+
+		/// <summary>
+		/// Object is being asynchronously loaded.
+		/// </summary>
+		AsyncLoading = 1 << 27,
+
+		/// <summary>
+		/// Object is not reachable on the object graph.
+		/// </summary>
+		Unreachable = 1 << 28,
+
+		/// <summary>
+		/// Objects that are pending destruction (invalid for gameplay but valid objects)
+		/// </summary>
+		PendingKill = 1 << 29,
+
+		/// <summary>
+		/// Object will not be garbage collected, even if unreferenced.
+		/// </summary>
+		RootSet = 1 << 30,
+
+		/// <summary>
+		/// Object didn't have its class constructor called yet (only the UObjectBase one to initialize its most basic members)
+		/// </summary>
+		PendingConstruction = 1 << 31,
+
+		GarbageCollectionKeepFlags = Native | Async | AsyncLoading | LoaderImport,
+
+		//~ Make sure this is up to date!
+		AllFlags = LoaderImport | Garbage | PersistentGarbage | ReachableInCluster | ClusterRoot | Native | Async | AsyncLoading | Unreachable | PendingKill | RootSet | PendingConstruction
+	};
+
+	/// <summary>
+	/// Helper extensions for EInternalObjectFlags
+	/// </summary>
+	public static class EInternalObjectFlagsExtensions
+	{
+		public static bool HasAnyFlags(this EInternalObjectFlags InFlags, EInternalObjectFlags TestFlags)
+		{
+			return (InFlags & TestFlags) != 0;
+		}
+
+		public static bool HasAllFlags(this EInternalObjectFlags InFlags, EInternalObjectFlags TestFlags)
+		{
+			return (InFlags & TestFlags) == TestFlags;
+		}
+
+		public static bool HasExactFlags(this EInternalObjectFlags InFlags, EInternalObjectFlags TestFlags, EInternalObjectFlags MatchFlags)
+		{
+			return (InFlags & TestFlags) == MatchFlags;
+		}
+	}
+
+	/// <summary>
+	/// Package flags, passed into UPackage::SetPackageFlags and related functions
+	/// 
+	/// This MUST be kept in sync with EPackageFlags defined in 
+	/// Engine\Source\Runtime\CoreUObject\Public\UObject\ObjectMacros.h
+	/// </summary>
+	[Flags]
+	public enum EPackageFlags : ulong
+	{
+		/// <summary>
+		/// No flags
+		/// </summary>
+		None = 0x00000000,
+
+		/// <summary>
+		/// Newly created package, not saved yet. In editor only.
+		/// </summary>
+		NewlyCreated = 0x00000001,
+
+		/// <summary>
+		/// Purely optional for clients.
+		/// </summary>
+		ClientOptional = 0x00000002,
+
+		/// <summary>
+		/// Only needed on the server side.
+		/// </summary>
+		ServerSideOnly = 0x00000004,
+
+		/// <summary>
+		/// This package is from "compiled in" classes.
+		/// </summary>
+		CompiledIn = 0x00000010,
+
+		/// <summary>
+		/// This package was loaded just for the purposes of diffing
+		/// </summary>
+		ForDiffing = 0x00000020,
+
+		/// <summary>
+		/// This is editor-only package (for example: editor module script package)
+		/// </summary>
+		EditorOnly = 0x00000040,
+
+		/// <summary>
+		/// Developer module
+		/// </summary>
+		Developer = 0x00000080,
+
+		/// <summary>
+		/// Loaded only in uncooked builds (i.e. runtime in editor)
+		/// </summary>
+		UncookedOnly = 0x00000100,
+
+		/// <summary>
+		/// Package is cooked
+		/// </summary>
+		Cooked = 0x00000200,
+
+		/// <summary>
+		/// Package doesn't contain any asset object (although asset tags can be present)
+		/// </summary>
+		ContainsNoAsset = 0x00000400,
+
+		/// <summary>
+		/// (Not Implemented) Objects in this package cannot be referenced in a different plugin or mount point (i.e /Game -> /Engine)
+		/// </summary>
+		NotExternallyReferenceable = 0x00000800,
+
+		// Unused = 0x00001000,
+
+		/// <summary>
+		/// Uses unversioned property serialization instead of versioned tagged property serialization
+		/// </summary>
+		UnversionedProperties = 0x00002000,
+
+		/// <summary>
+		/// Contains map data (UObjects only referenced by a single ULevel) but is stored in a different package
+		/// </summary>
+		ContainsMapData = 0x00004000,
+
+		/// <summary>
+		/// Temporarily set on a package while it is being saved.
+		/// </summary>
+		IsSaving = 0x00008000,
+
+		/// <summary>
+		/// package is currently being compiled
+		/// </summary>
+		Compiling = 0x00010000,
+
+		/// <summary>
+		/// Set if the package contains a ULevel/ UWorld object
+		/// </summary>
+		ContainsMap = 0x00020000,
+
+		/// <summary>
+		/// Set if the package contains any data to be gathered by localization
+		/// </summary>
+		RequiresLocalizationGather = 0x00040000,
+
+		// Unused = 0x00080000,
+
+		/// <summary>
+		/// Set if the package was created for the purpose of PIE
+		/// </summary>
+		PlayInEditor = 0x00100000,
+
+		/// <summary>
+		/// Package is allowed to contain UClass objects
+		/// </summary>
+		ContainsScript = 0x00200000,
+
+		/// <summary>
+		/// Editor should not export asset in this package
+		/// </summary>
+		DisallowExport = 0x00400000,
+
+		// Unused = 0x00800000,
+		// Unused = 0x01000000,	
+		// Unused = 0x02000000,	
+		// Unused = 0x04000000,
+		// Unused = 0x08000000,	
+
+		/// <summary>
+		/// This package should resolve dynamic imports from its export at runtime.
+		/// </summary>
+		DynamicImports = 0x10000000,
+
+		/// <summary>
+		/// This package contains elements that are runtime generated, and may not follow standard loading order rules
+		/// </summary>
+		RuntimeGenerated = 0x20000000,
+
+		/// <summary>
+		/// This package is reloading in the cooker, try to avoid getting data we will never need. We won't save this package.
+		/// </summary>
+		ReloadingForCooker = 0x40000000,
+
+		/// <summary>
+		/// Package has editor-only data filtered out
+		/// </summary>
+		FilterEditorOnly = 0x80000000,
+
+		/// <summary>
+		/// Transient Flags are cleared when serializing to or from PackageFileSummary
+		/// </summary>
+		TransientFlags = NewlyCreated | IsSaving | ReloadingForCooker,
+	}
+
+	/// <summary>
+	/// Helper extensions for EPackageFlags
+	/// </summary>
+	public static class EPackageFlagsExtensions
+	{
+		public static bool HasAnyFlags(this EPackageFlags InFlags, EPackageFlags TestFlags)
+		{
+			return (InFlags & TestFlags) != 0;
+		}
+
+		public static bool HasAllFlags(this EPackageFlags InFlags, EPackageFlags TestFlags)
+		{
+			return (InFlags & TestFlags) == TestFlags;
+		}
+
+		public static bool HasExactFlags(this EPackageFlags InFlags, EPackageFlags TestFlags, EPackageFlags MatchFlags)
+		{
+			return (InFlags & TestFlags) == MatchFlags;
+		}
+	}
+
+	/// <summary>
+	/// Flags associated with each property in a class, overriding the
+	/// property's default behavior.
+	/// 
+	/// This MUST be kept in sync with EPropertyFlags defined in 
+	/// Engine\Source\Runtime\CoreUObject\Public\UObject\ObjectMacros.h
+	/// </summary>
+	[Flags]
+	public enum EPropertyFlags : ulong
+	{
+		None = 0,
+
+		/// <summary>
+		/// Property is user-settable in the editor.
+		/// </summary>
+		Edit = 0x0000000000000001,
+
+		/// <summary>
+		/// This is a constant function parameter
+		/// </summary>
+		ConstParm = 0x0000000000000002,
+
+		/// <summary>
+		/// This property can be read by blueprint code
+		/// </summary>
+		BlueprintVisible = 0x0000000000000004,
+
+		/// <summary>
+		/// Object can be exported with actor.
+		/// </summary>
+		ExportObject = 0x0000000000000008,
+
+		/// <summary>
+		/// This property cannot be modified by blueprint code
+		/// </summary>
+		BlueprintReadOnly = 0x0000000000000010,
+
+		/// <summary>
+		/// Property is relevant to network replication.
+		/// </summary>
+		Net = 0x0000000000000020,
+
+		/// <summary>
+		/// Indicates that elements of an array can be modified, but its size cannot be changed.
+		/// </summary>
+		EditFixedSize = 0x0000000000000040,
+
+		/// <summary>
+		/// Function/When call parameter.
+		/// </summary>
+		Parm = 0x0000000000000080,
+
+		/// <summary>
+		/// Value is copied out after function call.
+		/// </summary>
+		OutParm = 0x0000000000000100,
+
+		/// <summary>
+		/// memset is fine for construction
+		/// </summary>
+		ZeroConstructor = 0x0000000000000200,
+
+		/// <summary>
+		/// Return value.
+		/// </summary>
+		ReturnParm = 0x0000000000000400,
+
+		/// <summary>
+		/// Disable editing of this property on an archetype/sub-blueprint
+		/// </summary>
+		DisableEditOnTemplate = 0x0000000000000800,
+
+		// Unused = 0x0000000000001000,
+
+		/// <summary>
+		/// Property is transient: shouldn't be saved or loaded, except for Blueprint CDOs.
+		/// </summary>
+		Transient = 0x0000000000002000,
+
+		/// <summary>
+		/// Property should be loaded/saved as permanent profile.
+		/// </summary>
+		Config = 0x0000000000004000,
+
+		// Unused = 0x0000000000008000,
+
+		/// <summary>
+		/// Disable editing on an instance of this class
+		/// </summary>
+		DisableEditOnInstance = 0x0000000000010000,
+
+		/// <summary>
+		/// Property is uneditable in the editor.
+		/// </summary>
+		EditConst = 0x0000000000020000,
+
+		/// <summary>
+		/// Load config from base class, not subclass.
+		/// </summary>
+		GlobalConfig = 0x0000000000040000,
+
+		/// <summary>
+		/// Property is a component references.
+		/// </summary>
+		InstancedReference = 0x0000000000080000,
+
+		// Unused = 0x0000000000100000,
+
+		/// <summary>
+		/// Property should always be reset to the default value during any type of duplication (copy/paste, binary duplication, etc.)
+		/// </summary>
+		DuplicateTransient = 0x0000000000200000,
+
+		// Unused = 0x0000000000400000,
+		// Unused = 0x0000000000800000,
+
+		/// <summary>
+		/// Property should be serialized for save games, this is only checked for game-specific archives with ArIsSaveGame
+		/// </summary>
+		SaveGame = 0x0000000001000000,
+
+		/// <summary>
+		/// Hide clear (and browse) button.
+		/// </summary>
+		NoClear = 0x0000000002000000,
+
+		// Unused = 0x0000000004000000,
+
+		/// <summary>
+		/// Value is passed by reference; "OutParam" and "Param" should also be set.
+		/// </summary>
+		ReferenceParm = 0x0000000008000000,
+
+		/// <summary>
+		/// MC Delegates only.  Property should be exposed for assigning in blueprint code
+		/// </summary>
+		BlueprintAssignable = 0x0000000010000000,
+
+		/// <summary>
+		/// Property is deprecated.  Read it from an archive, but don't save it.
+		/// </summary>
+		Deprecated = 0x0000000020000000,
+
+		/// <summary>
+		/// If this is set, then the property can be memcopied instead of CopyCompleteValue / CopySingleValue
+		/// </summary>
+		IsPlainOldData = 0x0000000040000000,
+
+		/// <summary>
+		/// Not replicated. For non replicated properties in replicated structs 
+		/// </summary>
+		RepSkip = 0x0000000080000000,
+
+		/// <summary>
+		/// Notify actors when a property is replicated
+		/// </summary>
+		RepNotify = 0x0000000100000000,
+
+		/// <summary>
+		/// interpolatable property for use with cinematics
+		/// </summary>
+		Interp = 0x0000000200000000,
+
+		/// <summary>
+		/// Property isn't transacted
+		/// </summary>
+		NonTransactional = 0x0000000400000000,
+
+		/// <summary>
+		/// Property should only be loaded in the editor
+		/// </summary>
+		EditorOnly = 0x0000000800000000,
+
+		/// <summary>
+		/// No destructor
+		/// </summary>
+		NoDestructor = 0x0000001000000000,
+
+		// Unused = 0x0000002000000000,
+
+		/// <summary>
+		/// Only used for weak pointers, means the export type is autoweak
+		/// </summary>
+		AutoWeak = 0x0000004000000000,
+
+		/// <summary>
+		/// Property contains component references.
+		/// </summary>
+		ContainsInstancedReference = 0x0000008000000000,
+
+		/// <summary>
+		/// asset instances will add properties with this flag to the asset registry automatically
+		/// </summary>
+		AssetRegistrySearchable = 0x0000010000000000,
+
+		/// <summary>
+		/// The property is visible by default in the editor details view
+		/// </summary>
+		SimpleDisplay = 0x0000020000000000,
+
+		/// <summary>
+		/// The property is advanced and not visible by default in the editor details view
+		/// </summary>
+		AdvancedDisplay = 0x0000040000000000,
+
+		/// <summary>
+		/// property is protected from the perspective of script
+		/// </summary>
+		Protected = 0x0000080000000000,
+
+		/// <summary>
+		/// MC Delegates only.  Property should be exposed for calling in blueprint code
+		/// </summary>
+		BlueprintCallable = 0x0000100000000000,
+
+		/// <summary>
+		/// MC Delegates only.  This delegate accepts (only in blueprint) only events with BlueprintAuthorityOnly.
+		/// </summary>
+		BlueprintAuthorityOnly = 0x0000200000000000,
+
+		/// <summary>
+		/// Property shouldn't be exported to text format (e.g. copy/paste)
+		/// </summary>
+		TextExportTransient = 0x0000400000000000,
+
+		/// <summary>
+		/// Property should only be copied in PIE
+		/// </summary>
+		NonPIEDuplicateTransient = 0x0000800000000000,
+
+		/// <summary>
+		/// Property is exposed on spawn
+		/// </summary>
+		ExposeOnSpawn = 0x0001000000000000,
+
+		/// <summary>
+		/// A object referenced by the property is duplicated like a component. (Each actor should have an own instance.)
+		/// </summary>
+		PersistentInstance = 0x0002000000000000,
+
+		/// <summary>
+		/// Property was parsed as a wrapper class like TSubclassOf<T>, FScriptInterface etc., rather than a USomething*
+		/// </summary>
+		UObjectWrapper = 0x0004000000000000,
+
+		/// <summary>
+		/// This property can generate a meaningful hash value.
+		/// </summary>
+		HasGetValueTypeHash = 0x0008000000000000,
+
+		/// <summary>
+		/// Public native access specifier
+		/// </summary>
+		NativeAccessSpecifierPublic = 0x0010000000000000,
+
+		/// <summary>
+		/// Protected native access specifier
+		/// </summary>
+		NativeAccessSpecifierProtected = 0x0020000000000000,
+
+		/// <summary>
+		/// Private native access specifier
+		/// </summary>
+		NativeAccessSpecifierPrivate = 0x0040000000000000,
+
+		/// <summary>
+		/// Property shouldn't be serialized, can still be exported to text
+		/// </summary>
+		SkipSerialization = 0x0080000000000000,
+
+		/// <summary>
+		/// All Native Access Specifier flags
+		/// </summary>
+		NativeAccessSpecifiers = NativeAccessSpecifierPublic | NativeAccessSpecifierProtected | NativeAccessSpecifierPrivate,
+
+		/// <summary>
+		/// All parameter flags
+		/// </summary>
+		ParmFlags = Parm | OutParm | ReturnParm | ReferenceParm | ConstParm,
+
+		/// <summary>
+		/// Flags that are propagated to properties inside array container
+		/// </summary>
+		PropagateToArrayInner =	ExportObject | PersistentInstance | InstancedReference | ContainsInstancedReference | Config | EditConst | Deprecated | EditorOnly | AutoWeak | UObjectWrapper,
+
+		/// <summary>
+		/// Flags that are propagated to value properties inside map container
+		/// </summary>
+		PropagateToMapValue = ExportObject | PersistentInstance | InstancedReference | ContainsInstancedReference | Config | EditConst | Deprecated | EditorOnly | AutoWeak | UObjectWrapper | Edit,
+
+		/// <summary>
+		/// Flags that are propagated to key properties inside map container
+		/// </summary>
+		PropagateToMapKey = ExportObject | PersistentInstance | InstancedReference | ContainsInstancedReference | Config | EditConst | Deprecated | EditorOnly | AutoWeak | UObjectWrapper | Edit,
+
+		/// <summary>
+		/// Flags that are propagated to properties inside set container
+		/// </summary>
+		PropagateToSetElement = ExportObject | PersistentInstance | InstancedReference | ContainsInstancedReference | Config | EditConst | Deprecated | EditorOnly | AutoWeak | UObjectWrapper | Edit,
+
+		/// <summary>
+		/// The flags that should never be set on interface properties
+		/// </summary>
+		InterfaceClearMask = ExportObject|InstancedReference|ContainsInstancedReference,
+
+		/// <summary>
+		/// All the properties that can be stripped for final release console builds
+		/// </summary>
+		DevelopmentAssets = EditorOnly,
+
+		/// <summary>
+		/// All the properties that should never be loaded or saved
+		/// </summary>
+		ComputedFlags = IsPlainOldData | NoDestructor | ZeroConstructor | HasGetValueTypeHash,
+
+		/// <summary>
+		/// Mask of all property flags
+		/// </summary>
+		AllFlags = UInt64.MaxValue,
+	};
+
+	/// <summary>
+	/// Helper extensions for EPropertyFlags
+	/// </summary>
+	public static class EPropertyFlagsExtensions
+	{
+		private static readonly Lazy<List<string>> BitNames = new Lazy<List<string>>(() =>
+		{
+			List<string> Out = new List<String>();
+			FieldInfo[] Fields = typeof(EPropertyFlags).GetFields();
+			for (int Bit = 0; Bit < 64; ++Bit)
+			{
+				bool bFound = false;
+				ulong Mask = (ulong)1 << Bit;
+				foreach (FieldInfo Field in Fields)
+				{
+					if (Field.IsSpecialName)
+					{
+						continue;
+					}
+					object? Value = Field.GetValue(null);
+					if (Value != null)
+					{
+						if (Mask == (ulong)Value)
+						{
+							Out.Add($"CPF_{Field.Name}");
+							bFound = true;
+							break;
+						}
+					}
+				}
+				if (!bFound)
+				{
+					Out.Add($"0x{Mask.ToString("X16")}");
+				}
+			}
+			return Out;
+		});
+
+		public static bool HasAnyFlags(this EPropertyFlags InFlags, EPropertyFlags TestFlags)
+		{
+			return (InFlags & TestFlags) != 0;
+		}
+
+		public static bool HasAllFlags(this EPropertyFlags InFlags, EPropertyFlags TestFlags)
+		{
+			return (InFlags & TestFlags) == TestFlags;
+		}
+
+		public static bool HasExactFlags(this EPropertyFlags InFlags, EPropertyFlags TestFlags, EPropertyFlags MatchFlags)
+		{
+			return (InFlags & TestFlags) == MatchFlags;
+		}
+
+		public static List<string> ToStringList(this EPropertyFlags InFlags, bool bIncludePrefix = true)
+		{
+			List<string> LocalBitNames = BitNames.Value;
+			ulong IntFlags = (UInt64)InFlags;
+			List<string> Out = new List<string>();
+			for (int Bit = 0; Bit < 64; ++Bit)
+			{
+				ulong Mask = (ulong)1 << Bit;
+				if (Mask > IntFlags)
+				{
+					break;
+				}
+				if ((Mask & IntFlags) != 0)
+				{
+					if (bIncludePrefix)
+					{
+						Out.Add(LocalBitNames[Bit]);
+					}
+					else
+					{
+						Out.Add(LocalBitNames[Bit].Substring(4));
+					}
+				}
+			}
+			return Out;
+		}
+	}
+
+	/// <summary>
+	/// Flags describing a struct
+	/// 
+	/// This MUST be kept in sync with EStructFlags defined in 
+	/// Engine\Source\Runtime\CoreUObject\Public\UObject\Class.h
+	/// </summary>
+	[Flags]
+	public enum EStructFlags
+	{
+		NoFlags = 0x00000000,
+		Native = 0x00000001,
+
+		/// <summary>
+		/// If set, this struct will be compared using native code
+		/// </summary>
+		IdenticalNative = 0x00000002,
+
+		HasInstancedReference = 0x00000004,
+
+		NoExport = 0x00000008,
+
+		/// <summary>
+		/// Indicates that this struct should always be serialized as a single unit
+		/// </summary>
+		Atomic = 0x00000010,
+
+		/// <summary>
+		/// Indicates that this struct uses binary serialization; it is unsafe to add/remove members from this struct without incrementing the package version
+		/// </summary>
+		Immutable = 0x00000020,
+
+		/// <summary>
+		/// If set, native code needs to be run to find referenced objects
+		/// </summary>
+		AddStructReferencedObjects = 0x00000040,
+
+		/// <summary>
+		/// Indicates that this struct should be exportable/importable at the DLL layer.  Base structs must also be exportable for this to work.
+		/// </summary>
+		RequiredAPI = 0x00000200,
+
+		/// <summary>
+		/// If set, this struct will be serialized using the CPP net serializer
+		/// </summary>
+		NetSerializeNative = 0x00000400,
+
+		/// <summary>
+		/// If set, this struct will be serialized using the CPP serializer
+		/// </summary>
+		SerializeNative = 0x00000800,
+
+		/// <summary>
+		/// If set, this struct will be copied using the CPP operator=
+		/// </summary>
+		CopyNative = 0x00001000,
+
+		/// <summary>
+		/// If set, this struct will be copied using memcpy
+		/// </summary>
+		IsPlainOldData = 0x00002000,
+
+		/// <summary>
+		/// If set, this struct has no destructor and non will be called. IsPlainOldData implies NoDestructor
+		/// </summary>
+		NoDestructor = 0x00004000,
+
+		/// <summary>
+		/// If set, this struct will not be constructed because it is assumed that memory is zero before construction.
+		/// </summary>
+		ZeroConstructor = 0x00008000,
+
+		/// <summary>
+		/// If set, native code will be used to export text
+		/// </summary>
+		ExportTextItemNative = 0x00010000,
+
+		/// <summary>
+		/// If set, native code will be used to export text
+		/// </summary>
+		ImportTextItemNative = 0x00020000,
+
+		/// <summary>
+		/// If set, this struct will have PostSerialize called on it after CPP serializer or tagged property serialization is complete
+		/// </summary>
+		PostSerializeNative = 0x00040000,
+
+		/// <summary>
+		/// If set, this struct will have SerializeFromMismatchedTag called on it if a mismatched tag is encountered.
+		/// </summary>
+		SerializeFromMismatchedTag = 0x00080000,
+
+		/// <summary>
+		/// If set, this struct will be serialized using the CPP net delta serializer
+		/// </summary>
+		NetDeltaSerializeNative = 0x00100000,
+
+		/// <summary>
+		/// If set, this struct will be have PostScriptConstruct called on it after a temporary object is constructed in a running blueprint
+		/// </summary>
+		PostScriptConstruct = 0x00200000,
+
+		/// <summary>
+		/// If set, this struct can share net serialization state across connections
+		/// </summary>
+		NetSharedSerialization = 0x00400000,
+
+		/// <summary>
+		/// If set, this struct has been cleaned and sanitized (trashed) and should not be used
+		/// </summary>
+		Trashed = 0x00800000,
+
+		/// <summary>
+		/// If set, this structure has been replaced via reinstancing
+		/// </summary>
+		NewerVersionExists = 0x01000000,
+
+		/// <summary>
+		/// Struct flags that are automatically inherited
+		/// </summary>
+		Inherit = HasInstancedReference | Atomic,
+
+		/// <summary>
+		/// Flags that are always computed, never loaded or done with code generation
+		/// </summary>
+		ComputedFlags = NetDeltaSerializeNative | NetSerializeNative | SerializeNative | PostSerializeNative | CopyNative | IsPlainOldData | 
+			NoDestructor | ZeroConstructor | IdenticalNative | AddStructReferencedObjects | ExportTextItemNative | ImportTextItemNative | 
+			SerializeFromMismatchedTag | PostScriptConstruct | NetSharedSerialization
+	}
+
+	/// <summary>
+	/// Helper extensions for EStructFlags
+	/// </summary>
+	public static class EStructFlagsExtensions
+	{
+		public static bool HasAnyFlags(this EStructFlags InFlags, EStructFlags TestFlags)
+		{
+			return (InFlags & TestFlags) != 0;
+		}
+
+		public static bool HasAllFlags(this EStructFlags InFlags, EStructFlags TestFlags)
+		{
+			return (InFlags & TestFlags) == TestFlags;
+		}
+
+		public static bool HasExactFlags(this EStructFlags InFlags, EStructFlags TestFlags, EStructFlags MatchFlags)
+		{
+			return (InFlags & TestFlags) == MatchFlags;
+		}
+	}
+}
diff --git a/Engine/Source/Programs/UnrealBuildTool/System/ConfigCache.cs b/Engine/Source/Programs/UnrealBuildTool/System/ConfigCache.cs
--- a/Engine/Source/Programs/UnrealBuildTool/System/ConfigCache.cs
+++ b/Engine/Source/Programs/UnrealBuildTool/System/ConfigCache.cs
@@ -171,12 +171,17 @@ namespace UnrealBuildTool
 		/// <param name="ProjectDir">The project directory to read the hierarchy for</param>
 		/// <param name="Platform">Which platform to read platform-specific config files for</param>
 		/// <param name="CustomConfig">Optional override config directory to search, for support of multiple target types</param>
+		/// <param name="CustomArgs">Optional list of command line arguments</param>
 		/// <returns>The requested config hierarchy</returns>
-		public static ConfigHierarchy ReadHierarchy(ConfigHierarchyType Type, DirectoryReference? ProjectDir, UnrealTargetPlatform Platform, string CustomConfig = "")
+		public static ConfigHierarchy ReadHierarchy(ConfigHierarchyType Type, DirectoryReference? ProjectDir, UnrealTargetPlatform Platform, string CustomConfig = "", string[]? CustomArgs = null)
 		{
 			// Handle command line overrides
 			List<String> OverrideStrings = new List<String>();
 			string[] CmdLine = Environment.GetCommandLineArgs();
+			if (CustomArgs != null)
+			{
+				CmdLine = CmdLine.Concat(CustomArgs).ToArray();
+			}
 			string IniConfigArgPrefix = "-ini:" + Enum.GetName(typeof(ConfigHierarchyType), Type) + ":";
 			string CustomConfigPrefix = "-CustomConfig=";
 			foreach (string CmdLineArg in CmdLine)
diff --git a/Engine/Source/Programs/UnrealBuildTool/System/ConfigHierarchy.cs b/Engine/Source/Programs/UnrealBuildTool/System/ConfigHierarchy.cs
--- a/Engine/Source/Programs/UnrealBuildTool/System/ConfigHierarchy.cs
+++ b/Engine/Source/Programs/UnrealBuildTool/System/ConfigHierarchy.cs
@@ -440,6 +440,25 @@ namespace UnrealBuildTool
 			return TryParse(Text, out Value);
 		}
 
+		/// <summary>
+		/// Gets an enumeration value associated with the specified key.
+		/// </summary>
+		/// <typeparam name="T"></typeparam>
+		/// <param name="SectionName">Section name</param>
+		/// <param name="KeyName">Key name</param>
+		/// <param name="Value">Value associated with the specified key. If the key has more than one value, only the first one is returned</param>
+		/// <returns>True if the key exists</returns>
+		public bool TryGetValue<T>(string SectionName, string KeyName, out T Value) where T : struct
+		{
+			string? Text;
+			if (TryGetValue(SectionName, KeyName, out Text))
+			{
+				Value = default(T);
+				return false;
+			}
+			return Enum.TryParse<T>(Text, out Value);
+		}
+
 		/// <summary>
 		/// Gets all values associated with the specified key
 		/// </summary>
diff --git a/Engine/Source/Programs/UnrealBuildTool/System/ExternalExecution.cs b/Engine/Source/Programs/UnrealBuildTool/System/ExternalExecution.cs
--- a/Engine/Source/Programs/UnrealBuildTool/System/ExternalExecution.cs
+++ b/Engine/Source/Programs/UnrealBuildTool/System/ExternalExecution.cs
@@ -21,70 +21,6 @@ using OpenTracing.Util;
 
 namespace UnrealBuildTool
 {
-	// This enum has to be compatible with the one defined in the
-	// Engine\Source\Runtime\Core\Public\Misc\ComplilationResult.h 
-	// to keep communication between UHT, UBT and Editor compiling
-	// processes valid.
-	enum CompilationResult
-	{
-		/// <summary>
-		/// Compilation succeeded
-		/// </summary>
-		Succeeded = 0,
-
-		/// <summary>
-		/// Build was canceled, this is used on the engine side only
-		/// </summary>
-		Canceled = 1,
-
-		/// <summary>
-		/// All targets were up to date, used only with -canskiplink
-		/// </summary>
-		UpToDate = 2,
-
-		/// <summary>
-		/// The process has most likely crashed. This is what UE returns in case of an assert
-		/// </summary>
-		CrashOrAssert = 3,
-
-		/// <summary>
-		/// Compilation failed because generated code changed which was not supported
-		/// </summary>
-		FailedDueToHeaderChange = 4,
-
-		/// <summary>
-		/// Compilation failed due to the engine modules needing to be rebuilt
-		/// </summary>
-		FailedDueToEngineChange = 5,
-
-		/// <summary>
-		/// Compilation failed due to compilation errors
-		/// </summary>
-		OtherCompilationError = 6,
-
-		/// <summary>
-		/// Compilation failed due to live coding action limit being exceeded.
-		/// </summary>
-		LiveCodingLimitError = 7,
-
-		/// <summary>
-		/// Compilation is not supported in the current build
-		/// </summary>
-		Unsupported,
-
-		/// <summary>
-		/// Unknown error
-		/// </summary>
-		Unknown
-	}
-
-	static class CompilationResultExtensions
-	{
-		public static bool Succeeded(this CompilationResult Result)
-		{
-			return Result == CompilationResult.Succeeded || Result == CompilationResult.UpToDate;
-		}
-	}
 
 	class CompilationResultException : BuildException
 	{
@@ -97,38 +33,8 @@ namespace UnrealBuildTool
 		}
 	}
 
-	/// <summary>
-	/// Type of module. Mirrored in UHT as EBuildModuleType.
-	/// This should be sorted by the order in which we expect modules to be built.
-	/// </summary>
-	enum UHTModuleType
-	{
-		Program,
-		EngineRuntime,
-		EngineUncooked,
-		EngineDeveloper,
-		EngineEditor,
-		EngineThirdParty,
-		GameRuntime,
-		GameUncooked,
-		GameDeveloper,
-		GameEditor,
-		GameThirdParty,
-	}
 	static class UHTModuleTypeExtensions
 	{
-		public static bool IsProgramModule(this UHTModuleType ModuleType)
-		{
-			return ModuleType == UHTModuleType.Program;
-		}
-		public static bool IsEngineModule(this UHTModuleType ModuleType)
-		{
-			return ModuleType == UHTModuleType.EngineRuntime || ModuleType == UHTModuleType.EngineDeveloper || ModuleType == UHTModuleType.EngineEditor || ModuleType == UHTModuleType.EngineThirdParty;
-		}
-		public static bool IsGameModule(this UHTModuleType ModuleType)
-		{
-			return ModuleType == UHTModuleType.GameRuntime || ModuleType == UHTModuleType.GameDeveloper || ModuleType == UHTModuleType.GameEditor || ModuleType == UHTModuleType.GameThirdParty;
-		}
 		public static UHTModuleType? EngineModuleTypeFromHostType(ModuleHostType ModuleType)
 		{
 			switch (ModuleType)
@@ -311,91 +217,6 @@ namespace UnrealBuildTool
 		}
 	}
 
-	/// <summary>
-	/// This MUST be kept in sync with EGeneratedBodyVersion enum and 
-	/// ToGeneratedBodyVersion function in UHT defined in GeneratedCodeVersion.h.
-	/// </summary>
-	public enum EGeneratedCodeVersion
-	{
-		/// <summary>
-		/// 
-		/// </summary>
-		None,
-
-		/// <summary>
-		/// 
-		/// </summary>
-		V1,
-
-		/// <summary>
-		/// 
-		/// </summary>
-		V2,
-
-		/// <summary>
-		/// 
-		/// </summary>
-		VLatest = V2
-	};
-
-	struct UHTManifest
-	{
-		public class Module
-		{
-			public string Name { get; set; }
-			public string ModuleType { get; set; }
-			public string OverrideModuleType { get; set; }
-			public string BaseDirectory { get; set; }
-			public string IncludeBase { get; set; }     // The include path which all UHT-generated includes should be relative to
-			public string OutputDirectory { get; set; }
-			public List<string> ClassesHeaders { get; set; }
-			public List<string> PublicHeaders { get; set; }
-			public List<string> InternalHeaders { get; set; }
-			public List<string> PrivateHeaders { get; set; }
-			public string? GeneratedCPPFilenameBase { get; set; }
-			public bool SaveExportedHeaders { get; set; }
-			[JsonConverter(typeof(JsonStringEnumConverter))]
-			public EGeneratedCodeVersion UHTGeneratedCodeVersion { get; set; }
-
-			public Module(UHTModuleInfo Info)
-			{
-				Name = Info.ModuleName;
-				ModuleType = Info.ModuleType;
-				OverrideModuleType = Info.OverrideModuleType;
-				BaseDirectory = Info.ModuleDirectories[0].FullName;
-				IncludeBase = Info.ModuleDirectories[0].ParentDirectory!.FullName;
-				OutputDirectory = Path.GetDirectoryName(Info.GeneratedCPPFilenameBase)!;
-				ClassesHeaders = Info.PublicUObjectClassesHeaders.Select((Header) => Header.AbsolutePath).ToList();
-				PublicHeaders = Info.PublicUObjectHeaders.Select((Header) => Header.AbsolutePath).ToList();
-				InternalHeaders = Info.InternalUObjectHeaders.Select((Header) => Header.AbsolutePath).ToList();
-				PrivateHeaders = Info.PrivateUObjectHeaders.Select((Header) => Header.AbsolutePath).ToList();
-				GeneratedCPPFilenameBase = Info.GeneratedCPPFilenameBase;
-				SaveExportedHeaders = !Info.bIsReadOnly;
-				UHTGeneratedCodeVersion = Info.GeneratedCodeVersion;
-			}
-
-			public override string ToString()
-			{
-				return Name;
-			}
-		}
-
-		public bool IsGameTarget { get; set; }     // True if the current target is a game target
-		public string RootLocalPath { get; set; }    // The engine path on the local machine
-		public string TargetName { get; set; }       // Name of the target currently being compiled
-		public string ExternalDependenciesFile { get; set; } // File to contain additional dependencies that the generated code depends on
-		public List<Module> Modules { get; set; }
-
-		public UHTManifest(string InTargetName, TargetType InTargetType, string InRootLocalPath, string InExternalDependenciesFile, List<Module> InModules)
-		{
-			IsGameTarget = (InTargetType != TargetType.Program);
-			RootLocalPath = InRootLocalPath;
-			TargetName = InTargetName;
-			ExternalDependenciesFile = InExternalDependenciesFile;
-			Modules = InModules;
-		}
-	}
-
 	class UHTModuleHeaderInfo
 	{
 		public DirectoryItem SourceFolder;
@@ -1151,9 +972,32 @@ namespace UnrealBuildTool
 				List<UHTManifest.Module> Modules = new List<UHTManifest.Module>();
 				foreach(UHTModuleInfo UObjectModule in Makefile.UObjectModules)
 				{
-					Modules.Add(new UHTManifest.Module(UObjectModule));
+					Modules.Add(
+						new UHTManifest.Module
+						{
+							Name = UObjectModule.ModuleName,
+							ModuleType = (UHTModuleType)Enum.Parse(typeof(UHTModuleType), UObjectModule.ModuleType),
+							OverrideModuleType = (EPackageOverrideType)Enum.Parse(typeof(EPackageOverrideType), UObjectModule.OverrideModuleType),
+							BaseDirectory = UObjectModule.ModuleDirectories[0].FullName,
+							IncludeBase = UObjectModule.ModuleDirectories[0].ParentDirectory!.FullName,
+							OutputDirectory = Path.GetDirectoryName(UObjectModule.GeneratedCPPFilenameBase)!,
+							ClassesHeaders = UObjectModule.PublicUObjectClassesHeaders.Select((Header) => Header.AbsolutePath).ToList(),
+							PublicHeaders = UObjectModule.PublicUObjectHeaders.Select((Header) => Header.AbsolutePath).ToList(),
+							InternalHeaders = UObjectModule.InternalUObjectHeaders.Select((Header) => Header.AbsolutePath).ToList(),
+							PrivateHeaders = UObjectModule.PrivateUObjectHeaders.Select((Header) => Header.AbsolutePath).ToList(),
+							GeneratedCPPFilenameBase = UObjectModule.GeneratedCPPFilenameBase,
+							SaveExportedHeaders = !UObjectModule.bIsReadOnly,
+							GeneratedCodeVersion = UObjectModule.GeneratedCodeVersion,
+						});
 				}
-				UHTManifest Manifest = new UHTManifest(TargetName, Makefile.TargetType, RootLocalPath, ExternalDependenciesFile.FullName, Modules);
+				UHTManifest Manifest = new UHTManifest
+				{
+					TargetName = TargetName,
+					IsGameTarget = Makefile.TargetType != TargetType.Program,
+					RootLocalPath = RootLocalPath,
+					ExternalDependenciesFile = ExternalDependenciesFile.FullName,
+					Modules = Modules,
+				};
 
 				if (!bIsBuildingUHT && bUHTNeedsToRun)
 				{
diff --git a/Engine/Source/Runtime/CoreUObject/Public/UObject/Class.h b/Engine/Source/Runtime/CoreUObject/Public/UObject/Class.h
--- a/Engine/Source/Runtime/CoreUObject/Public/UObject/Class.h
+++ b/Engine/Source/Runtime/CoreUObject/Public/UObject/Class.h
@@ -630,6 +630,12 @@ private:
 	void SerializeVersionedTaggedProperties(FStructuredArchive::FSlot Slot, uint8* Data, UStruct* DefaultsStruct, uint8* Defaults, const UObject* BreakRecursionIfFullyLoad) const;
 };
 
+/**
+ * Flags describing a struct
+ * 
+ * This MUST be kept in sync with EStructFlags defined in
+ * Engine\Source\Programs\Shared\EpicGames.Core\UnrealEngineTypes.cs
+ */
 enum EStructFlags
 {
 	// State flags.
diff --git a/Engine/Source/Runtime/CoreUObject/Public/UObject/ObjectMacros.h b/Engine/Source/Runtime/CoreUObject/Public/UObject/ObjectMacros.h
--- a/Engine/Source/Runtime/CoreUObject/Public/UObject/ObjectMacros.h
+++ b/Engine/Source/Runtime/CoreUObject/Public/UObject/ObjectMacros.h
@@ -102,7 +102,12 @@ enum ESaveFlags
 	SAVE_Unversioned = SAVE_Unversioned_Native | SAVE_Unversioned_Properties,
 };
 
-/** Package flags, passed into UPackage::SetPackageFlags and related functions */
+/** 
+ * Package flags, passed into UPackage::SetPackageFlags and related functions
+ *
+ * This MUST be kept in sync with EPackageFlags defined in
+ * Engine\Source\Programs\Shared\EpicGames.Core\UnrealEngineTypes.cs
+ */
 enum EPackageFlags
 {
 	PKG_None						= 0x00000000,	///< No flags
@@ -169,6 +174,9 @@ public:
 
 /**
  * Flags describing a class.
+ *
+ * This MUST be kept in sync with EClassFlags defined in
+ * Engine\Source\Programs\Shared\EpicGames.Core\UnrealEngineTypes.cs
  */
 enum EClassFlags
 {
@@ -297,6 +305,9 @@ ENUM_CLASS_FLAGS(EClassFlags);
 
 /**
  * Flags used for quickly casting classes of certain types; all class cast flags are inherited
+ *
+ * This MUST be kept in sync with EClassCastFlags defined in
+ * Engine\Source\Programs\Shared\EpicGames.Core\UnrealEngineTypes.cs
  */
 enum EClassCastFlags : uint64
 {
@@ -369,6 +380,9 @@ ENUM_CLASS_FLAGS(EClassCastFlags)
  * Flags associated with each property in a class, overriding the
  * property's default behavior.
  * @warning When adding one here, please update ParsePropertyFlags()
+ * 
+ * This MUST be kept in sync with EPackageFlags defined in
+ * Engine\Source\Programs\Shared\EpicGames.Core\UnrealEngineTypes.cs
  */
 enum EPropertyFlags : uint64
 {
@@ -554,7 +568,12 @@ PRAGMA_ENABLE_DEPRECATION_WARNINGS
 
 ENUM_CLASS_FLAGS(EObjectFlags);
 
-/** Objects flags for internal use (GC, low level UObject code) */
+/** 
+ * Objects flags for internal use (GC, low level UObject code)
+ *
+ * This MUST be kept in sync with EInternalObjectFlags defined in
+ * Engine\Source\Programs\Shared\EpicGames.Core\UnrealEngineTypes.cs
+ */
 enum class EInternalObjectFlags : int32
 {
 	None = 0,
@@ -582,7 +601,12 @@ enum class EInternalObjectFlags : int32
 };
 ENUM_CLASS_FLAGS(EInternalObjectFlags);
 
-/** Flags describing a UEnum */
+/**
+ * Flags describing a UEnum 
+ * 
+ * This MUST be kept in sync with EEnumFlags defined in
+ * Engine\Source\Programs\Shared\EpicGames.Core\UnrealEngineTypes.cs
+ */
 enum class EEnumFlags
 {
 	None,
diff --git a/Engine/Source/Runtime/CoreUObject/Public/UObject/Script.h b/Engine/Source/Runtime/CoreUObject/Public/UObject/Script.h
--- a/Engine/Source/Runtime/CoreUObject/Public/UObject/Script.h
+++ b/Engine/Source/Runtime/CoreUObject/Public/UObject/Script.h
@@ -102,6 +102,9 @@ namespace FunctionCallspace
 // Function flags.
 //
 // Note: Please keep ParseFunctionFlags in sync when this enum is modified.
+//
+// This MUST be kept in sync with EEnumFlags defined in
+// Engine\Source\Programs\Shared\EpicGames.Core\UnrealEngineTypes.cs
 enum EFunctionFlags : uint32
 {
 	// Function flags.
