ManagedProcess: Process.StartTime & ExitTime can't be trusted on Mac

#rnx

diff --git a/Engine/Source/Programs/Shared/EpicGames.Core/ManagedProcess.cs b/Engine/Source/Programs/Shared/EpicGames.Core/ManagedProcess.cs
--- a/Engine/Source/Programs/Shared/EpicGames.Core/ManagedProcess.cs
+++ b/Engine/Source/Programs/Shared/EpicGames.Core/ManagedProcess.cs
@@ -845,6 +845,10 @@ namespace EpicGames.Core
 			FrameworkProcess.StartInfo.RedirectStandardError = true;
 			FrameworkProcess.StartInfo.UseShellExecute = false;
 			FrameworkProcess.StartInfo.CreateNoWindow = true;
+			FrameworkProcess.Exited += (sender, e) =>
+			{
+				FrameworkExitTime = DateTime.Now;
+			};
 
 			if (Environment != null)
 			{
@@ -854,17 +858,9 @@ namespace EpicGames.Core
 				}
 			}
 
-			FrameworkStartTime = DateTime.Now;
+			FrameworkStartTime = FrameworkExitTime = DateTime.Now;
 			FrameworkProcess.Start();
 
-			try
-			{
-				FrameworkStartTime = FrameworkProcess.StartTime;
-			}
-			catch
-			{
-			}
-
 			try
 			{
 				FrameworkProcess.PriorityClass = Priority;
@@ -1271,6 +1267,7 @@ namespace EpicGames.Core
 		}
 
 		private DateTime FrameworkStartTime = DateTime.MinValue;
+		private DateTime FrameworkExitTime = DateTime.MinValue;
 
 		/// <summary>
 		/// The creation time of the process.
@@ -1313,7 +1310,7 @@ namespace EpicGames.Core
 				}
 				else
 				{
-					return FrameworkProcess.ExitTime;
+					return FrameworkExitTime;
 				}
 			}
 		}
