C# Compact Binary Serialization - Added more tests from Horde Storage (Jupiter Tests) and fixed issues exposed by them.

diff --git a/Engine/Source/Programs/Shared/EpicGames.Serialization.Tests/JupiterTests.cs b/Engine/Source/Programs/Shared/EpicGames.Serialization.Tests/JupiterTests.cs
--- a/Engine/Source/Programs/Shared/EpicGames.Serialization.Tests/JupiterTests.cs
+++ b/Engine/Source/Programs/Shared/EpicGames.Serialization.Tests/JupiterTests.cs
@@ -99,6 +99,21 @@ namespace EpicGames.Serialization.Tests
 			Assert.AreEqual(IoHash.Parse("c7a03f83c08cdca882110ecf2b5654ee3b09b11e"), payloadFields[2]["RawHash"]!.AsHash());
 		}
 
+		[TestMethod]
+		public void compact_binary()
+		{
+			byte[] bytes = File.ReadAllBytes("CompactBinaryObjects/compact_binary");
+
+			CbObject o = new CbObject(bytes);
+			Assert.AreEqual("", o.AsField().Name);
+			List<CbField> buildActionFields = o.ToList();
+			Assert.AreEqual(3, buildActionFields.Count);
+			CbField payloads = buildActionFields[0];
+			List<CbField> payloadFields = payloads.ToList();
+			Assert.AreEqual(2, payloadFields.Count);
+
+			Assert.AreEqual("{\"Key\":{\"Bucket\":\"EditorDomainPackage\",\"Hash\":\"37dbaa409ef30ba67f18c8fc2faaf606636cb915\"},\"Meta\":{\"FileSize\":24789},\"Attachments\":[{\"Id\":\"000000000000000000000001\",\"RawHash\":\"da6fc57e4b9f91377c9509ea0ad567bacb3796c5\",\"RawSize\":24789}]}", o.ToJson());
+		}
 
 		[TestMethod]
 		public void WriteArray()
diff --git a/Engine/Source/Programs/Shared/EpicGames.Serialization/CbField.cs b/Engine/Source/Programs/Shared/EpicGames.Serialization/CbField.cs
--- a/Engine/Source/Programs/Shared/EpicGames.Serialization/CbField.cs
+++ b/Engine/Source/Programs/Shared/EpicGames.Serialization/CbField.cs
@@ -575,6 +575,17 @@ namespace EpicGames.Serialization
 			return GetType(type) == CbFieldType.TimeSpan;
 		}
 
+		/// <summary>
+		/// Tests if the given field type is a object id
+		/// </summary>
+		/// <param name="Type">Type to check</param>
+		/// <returns>True if the field is a object id</returns>
+		public static bool IsObjectId(CbFieldType type)
+		{
+			return GetType(type) == CbFieldType.ObjectId;
+		}
+
+
 		/// <summary>
 		/// Tests if the given field type has fields
 		/// </summary>
@@ -820,6 +831,8 @@ namespace EpicGames.Serialization
 						return AsDateTime();
 					case CbFieldType.TimeSpan:
 						return AsTimeSpan();
+					case CbFieldType.ObjectId:
+						return AsObjectId();
 					default:
 						throw new NotImplementedException($"Unknown field type ({FieldType})");
 				}
@@ -1372,6 +1385,27 @@ namespace EpicGames.Serialization
 		/// <returns></returns>
 		public TimeSpan AsTimeSpan(TimeSpan Default = default) => new TimeSpan(AsTimeSpanTicks(Default.Ticks));
 
+
+		/// <summary>
+		/// Access the field as a object id
+		/// </summary>
+		/// <param name="Default"></param>
+		/// <returns></returns>
+		public ReadOnlyMemory<byte> AsObjectId(ReadOnlyMemory<byte> Default = default)
+		{
+			if (CbFieldUtils.IsObjectId(TypeWithFlags))
+			{
+				Error = CbFieldError.None;
+
+				return Payload;
+			}
+			else
+			{
+				Error = CbFieldError.TypeError;
+				return Default;
+			}
+		}
+
 		/// <inheritdoc cref="CbFieldUtils.HasFieldName(CbFieldType)"/>
 		public bool HasName() => CbFieldUtils.HasFieldName(TypeWithFlags);
 
@@ -1420,6 +1454,10 @@ namespace EpicGames.Serialization
 		/// <inheritdoc cref="CbFieldUtils.IsTimeSpan(CbFieldType)"/>
 		public bool IsTimeSpan() => CbFieldUtils.IsTimeSpan(TypeWithFlags);
 
+		
+		/// <inheritdoc cref="CbFieldUtils.IsObjectId(CbFieldType)"/>
+		public bool IsObjectId() => CbFieldUtils.IsObjectId(TypeWithFlags);
+
 		/// <summary>
 		/// Whether the field has a value
 		/// </summary>
@@ -2631,7 +2669,10 @@ namespace EpicGames.Serialization
 		{
 			if (Field.IsObject())
 			{
-				Writer.WriteStartObject();
+				if (Field.NameLen != 0)
+					Writer.WriteStartObject(Field.Name.Span);
+				else
+					Writer.WriteStartObject();
 				CbObject Object = Field.AsObject();
 				foreach (CbField ObjectField in Object.InnerField)
 				{
@@ -2641,7 +2682,12 @@ namespace EpicGames.Serialization
 			}
 			else if (Field.IsArray())
 			{
-				Writer.WriteStartArray();
+				Writer.WriteStartArray(Field.Name.Span);
+				CbArray Array = Field.AsArray();
+				foreach (CbField ObjectField in Array)
+				{
+					WriteField(ObjectField, Writer);
+				}
 				Writer.WriteEndArray();
 			}
 			else if (Field.IsInteger())
@@ -2675,6 +2721,10 @@ namespace EpicGames.Serialization
 			{
 				Writer.WriteString(Field.Name.Span, Field.AsUtf8String().Span);
 			}
+			else if (Field.IsObjectId())
+			{
+				Writer.WriteString(Field.Name.Span, StringUtils.FormatUtf8HexString(Field.AsObjectId().Span).Span);
+			}
 			else
 			{
 				throw new NotImplementedException($"Unhandled type {Field.GetType()} when attempting to convert to json");
