EpicGames.Serialization - Added support for C# arrays (and not just lists) when serializing C# properties to compact binary
Also fixed issue were empty lists would not be serialized (field skipped)

diff --git a/Engine/Source/Programs/Shared/EpicGames.Serialization/CbConverter.cs b/Engine/Source/Programs/Shared/EpicGames.Serialization/CbConverter.cs
--- a/Engine/Source/Programs/Shared/EpicGames.Serialization/CbConverter.cs
+++ b/Engine/Source/Programs/Shared/EpicGames.Serialization/CbConverter.cs
@@ -240,6 +240,7 @@ namespace EpicGames.Serialization
 			new CbClassConverterFactory(),
 			new CbEnumConverterFactory(),
 			new CbListConverterFactory(),
+			new CbArrayConverterFactory(),
 			new CbDictionaryConverterFactory(),
 			new CbNullableConverterFactory()
 		};
diff --git a/Engine/Source/Programs/Shared/EpicGames.Serialization/Converters/CbArrayConverter.cs b/Engine/Source/Programs/Shared/EpicGames.Serialization/Converters/CbArrayConverter.cs
new file mode 100644
--- /dev/null
+++ b/Engine/Source/Programs/Shared/EpicGames.Serialization/Converters/CbArrayConverter.cs
@@ -0,0 +1,73 @@
+// Copyright Epic Games, Inc. All Rights Reserved.
+
+using EpicGames.Core;
+using System;
+using System.Collections;
+using System.Collections.Concurrent;
+using System.Collections.Generic;
+using System.Diagnostics.CodeAnalysis;
+using System.Linq;
+using System.Linq.Expressions;
+using System.Reflection;
+using System.Reflection.Emit;
+
+namespace EpicGames.Serialization.Converters
+{
+	/// <summary>
+	/// Converter for array types
+	/// </summary>
+	/// <typeparam name="T"></typeparam>
+	class CbArrayConverter<T> : CbConverterBase<T[]>
+	{
+		/// <inheritdoc/>
+		public override T[] Read(CbField Field)
+		{
+			List<T> List = new List<T>();
+			foreach (CbField ElementField in Field)
+			{
+				List.Add(CbSerializer.Deserialize<T>(ElementField));
+			}
+			return List.ToArray();
+		}
+
+		/// <inheritdoc/>
+		public override void Write(CbWriter Writer, T[] List)
+		{
+			Writer.BeginArray();
+			foreach (T Element in List)
+			{
+				CbSerializer.Serialize<T>(Writer, Element);
+			}
+			Writer.EndArray();
+		}
+
+		/// <inheritdoc/>
+		public override void WriteNamed(CbWriter Writer, Utf8String Name, T[] Array)
+		{
+			Writer.BeginArray(Name);
+			foreach (T Element in Array)
+			{
+				CbSerializer.Serialize<T>(Writer, Element);
+			}
+			Writer.EndArray();
+		}
+	}
+	
+	/// <summary>
+	/// Factory for CbListConverter
+	/// </summary>
+	class CbArrayConverterFactory : CbConverterFactory
+	{
+		/// <inheritdoc/>
+		public override ICbConverter? CreateConverter(Type Type)
+		{
+			if (Type.IsArray)
+			{
+				Type ElementType = Type.GetElementType()!;
+				Type ConverterType = typeof(CbArrayConverter<>).MakeGenericType(ElementType);
+				return (ICbConverter)Activator.CreateInstance(ConverterType)!;
+			}
+			return null;
+		}
+	}
+}
diff --git a/Engine/Source/Programs/Shared/EpicGames.Serialization/Converters/CbListConverter.cs b/Engine/Source/Programs/Shared/EpicGames.Serialization/Converters/CbListConverter.cs
--- a/Engine/Source/Programs/Shared/EpicGames.Serialization/Converters/CbListConverter.cs
+++ b/Engine/Source/Programs/Shared/EpicGames.Serialization/Converters/CbListConverter.cs
@@ -44,15 +44,12 @@ namespace EpicGames.Serialization.Converters
 		/// <inheritdoc/>
 		public override void WriteNamed(CbWriter Writer, Utf8String Name, List<T> List)
 		{
-			if (List.Count > 0)
+			Writer.BeginArray(Name);
+			foreach (T Element in List)
 			{
-				Writer.BeginArray(Name);
-				foreach (T Element in List)
-				{
-					CbSerializer.Serialize<T>(Writer, Element);
-				}
-				Writer.EndArray();
+				CbSerializer.Serialize<T>(Writer, Element);
 			}
+			Writer.EndArray();
 		}
 	}
 
