Moved Compact Binary formatters into a EpicGames.AspNet shared library, this shared library can be used for any common dependencies that also requires asp.net

diff --git a/Engine/Source/Programs/Horde/Horde.Build/Horde.Build.csproj b/Engine/Source/Programs/Horde/Horde.Build/Horde.Build.csproj
--- a/Engine/Source/Programs/Horde/Horde.Build/Horde.Build.csproj
+++ b/Engine/Source/Programs/Horde/Horde.Build/Horde.Build.csproj
@@ -95,6 +95,7 @@
   </ItemGroup>
 
   <ItemGroup>
+    <ProjectReference Include="..\..\Shared\EpicGames.AspNet\EpicGames.AspNet.csproj" />
     <ProjectReference Include="..\..\Shared\EpicGames.BuildGraph\EpicGames.BuildGraph.csproj" />
     <ProjectReference Include="..\..\Shared\EpicGames.Core\EpicGames.Core.csproj" />
     <ProjectReference Include="..\..\Shared\EpicGames.Horde\EpicGames.Horde.csproj" />
diff --git a/Engine/Source/Programs/Horde/Horde.Build/Startup.cs b/Engine/Source/Programs/Horde/Horde.Build/Startup.cs
--- a/Engine/Source/Programs/Horde/Horde.Build/Startup.cs
+++ b/Engine/Source/Programs/Horde/Horde.Build/Startup.cs
@@ -89,6 +89,7 @@ using Microsoft.AspNetCore.Mvc.ModelBinding;
 using Horde.Build.Storage.Services;
 using EpicGames.Horde.Storage;
 using System.Net.Http;
+using EpicGames.AspNet;
 using EpicGames.Horde.Storage.Impl;
 
 namespace HordeServer
diff --git a/Engine/Source/Programs/Horde/Horde.Build/Utilities/CompactBinaryFormatter.cs b/Engine/Source/Programs/Horde/Horde.Build/Utilities/CompactBinaryFormatter.cs
deleted file mode 100644
diff --git a/Engine/Source/Programs/Shared/EpicGames.AspNet/CompactBinaryFormatter.cs b/Engine/Source/Programs/Shared/EpicGames.AspNet/CompactBinaryFormatter.cs
new file mode 100644
--- /dev/null
+++ b/Engine/Source/Programs/Shared/EpicGames.AspNet/CompactBinaryFormatter.cs
@@ -0,0 +1,142 @@
+// Copyright Epic Games, Inc. All Rights Reserved.
+
+using Microsoft.AspNetCore.Mvc.Formatters;
+using System;
+using Microsoft.Net.Http.Headers;
+using System.Threading.Tasks;
+using System.IO;
+using EpicGames.Serialization;
+using EpicGames.Core;
+using Microsoft.Extensions.DependencyInjection;
+using Microsoft.Extensions.Primitives;
+using Microsoft.Extensions.Logging;
+
+namespace EpicGames.AspNet
+{
+	/// <summary>
+	/// Global constants
+	/// </summary>
+	public static partial class CustomMediaTypeNames
+	{
+		/// <summary>
+		/// Media type for compact binary
+		/// </summary>
+		public const string UnrealCompactBinary = "application/x-ue-cb";
+	}
+
+	/// <summary>
+	/// Converter to allow reading compact binary objects as request bodies
+	/// </summary>
+	public class CbInputFormatter : InputFormatter
+	{
+		/// <summary>
+		/// Constructor
+		/// </summary>
+		public CbInputFormatter()
+		{
+			SupportedMediaTypes.Add(MediaTypeHeaderValue.Parse(CustomMediaTypeNames.UnrealCompactBinary));
+		}
+
+		/// <inheritdoc/>
+		protected override bool CanReadType(Type Type)
+		{
+			return true;
+		}
+
+		/// <inheritdoc/>
+		public override async Task<InputFormatterResult> ReadRequestBodyAsync(InputFormatterContext Context)
+		{
+			// Buffer the data into an array
+			byte[] Data;
+			using (MemoryStream Stream = new MemoryStream())
+			{
+				await Context.HttpContext.Request.Body.CopyToAsync(Stream);
+				Data = Stream.ToArray();
+			}
+
+			// Serialize the object
+			CbField Field;
+			try
+			{
+				Field = new CbField(Data);
+			}
+			catch (Exception Ex)
+			{
+				ILogger<CbInputFormatter> Logger = Context.HttpContext.RequestServices.GetRequiredService<ILogger<CbInputFormatter>>();
+				Logger.LogError("Unable to parse compact binary: {Dump}", FormatHexDump(Data, 256));
+				foreach ((string Name, StringValues Values) in Context.HttpContext.Request.Headers)
+				{
+					foreach (string Value in Values)
+					{
+						Logger.LogInformation("Header {Name}: {Value}", Name, Value);
+					}
+				}
+				throw new Exception($"Unable to parse compact binary request: {FormatHexDump(Data, 256)}", Ex);
+			}
+			return await InputFormatterResult.SuccessAsync(CbSerializer.Deserialize(new CbField(Data), Context.ModelType)!);
+		}
+
+		static string FormatHexDump(byte[] Data, int MaxLength)
+		{
+			ReadOnlySpan<byte> Span = Data.AsSpan(0, Math.Min(Data.Length, MaxLength));
+			string HexString = StringUtils.FormatHexString(Span);
+
+			char[] HexDump = new char[Span.Length * 3];
+			for (int Idx = 0; Idx < Span.Length; Idx++)
+			{
+				HexDump[(Idx * 3) + 0] = ((Idx & 15) == 0) ? '\n' : ' ';
+				HexDump[(Idx * 3) + 1] = HexString[(Idx * 2) + 0];
+				HexDump[(Idx * 3) + 2] = HexString[(Idx * 2) + 1];
+			}
+			return new string(HexDump);
+		}
+	}
+
+	/// <summary>
+	/// Converter to allow writing compact binary objects as responses
+	/// </summary>
+	public class CbOutputFormatter : OutputFormatter
+	{
+		/// <summary>
+		/// Constructor
+		/// </summary>
+		public CbOutputFormatter()
+		{
+			SupportedMediaTypes.Add(MediaTypeHeaderValue.Parse(CustomMediaTypeNames.UnrealCompactBinary));
+		}
+
+		/// <inheritdoc/>
+		protected override bool CanWriteType(Type? Type)
+		{
+			return true;
+		}
+
+		/// <inheritdoc/>
+		public override async Task WriteResponseBodyAsync(OutputFormatterWriteContext Context)
+		{
+			ReadOnlyMemory<byte> Data;
+			if (Context.Object is CbObject Object)
+			{
+				Data = Object.GetView();
+			}
+			else
+			{
+				Data = CbSerializer.Serialize(Context.ObjectType!, Context.Object!).GetView();
+			}
+			await Context.HttpContext.Response.BodyWriter.WriteAsync(Data);
+		}
+	}
+
+	/// <summary>
+	/// Special version of <see cref="CbOutputFormatter"/> which returns native CbObject encoded data. Can be 
+	/// inserted at a high priority in the output formatter list to prevent transcoding to json.
+	/// </summary>
+	public class CbPreferredOutputFormatter : CbOutputFormatter
+	{
+		/// <inheritdoc/>
+		protected override bool CanWriteType(Type? Type)
+		{
+			return Type == typeof(CbObject);
+		}
+	}
+}
diff --git a/Engine/Source/Programs/Shared/EpicGames.AspNet/EpicGames.AspNet.csproj b/Engine/Source/Programs/Shared/EpicGames.AspNet/EpicGames.AspNet.csproj
new file mode 100644
--- /dev/null
+++ b/Engine/Source/Programs/Shared/EpicGames.AspNet/EpicGames.AspNet.csproj
@@ -0,0 +1,42 @@
+<Project Sdk="Microsoft.NET.Sdk.Web">
+  <PropertyGroup>
+    <TargetFramework>netcoreapp3.1</TargetFramework>
+    <Configuration Condition=" '$(Configuration)' == '' ">Development</Configuration>
+    <OutputType>Library</OutputType>
+    <RootNamespace>EpicGames.AspNet</RootNamespace>
+    <GenerateAssemblyInfo>false</GenerateAssemblyInfo>
+    <GenerateTargetFrameworkAttribute>false</GenerateTargetFrameworkAttribute>
+    <Configurations>Debug;Release;Development</Configurations>
+    <Nullable>enable</Nullable>
+    <DebugType>pdbonly</DebugType>
+  </PropertyGroup>
+  <PropertyGroup Condition=" '$(Configuration)|$(Platform)' == 'Debug|AnyCPU' ">
+    <DocumentationFile></DocumentationFile>
+  </PropertyGroup>
+  <PropertyGroup Condition="'$(Configuration)|$(Platform)' == 'Development|AnyCPU'">
+    <DocumentationFile></DocumentationFile>
+    <Optimize>true</Optimize>
+  </PropertyGroup>
+  <PropertyGroup Condition="'$(Configuration)|$(Platform)' == 'Release|AnyCPU'">
+    <DocumentationFile></DocumentationFile>
+    <Optimize>true</Optimize>
+  </PropertyGroup>
+  
+  <ItemGroup>
+    <Compile Include="..\MetaData.cs">
+      <Link>Properties\MetaData.cs</Link>
+    </Compile>
+  </ItemGroup>
+  <ItemGroup>
+    <ProjectReference Include="..\..\Shared\EpicGames.Core\EpicGames.Core.csproj" />
+    <ProjectReference Include="..\EpicGames.Serialization\EpicGames.Serialization.csproj" />
+  </ItemGroup>
+  <ItemGroup>
+    <PackageReference Include="Microsoft.CSharp" Version="4.7.0" />
+    <PackageReference Include="Microsoft.Win32.Registry" Version="4.7.0" />
+    <PackageReference Include="System.Data.DataSetExtensions" Version="4.5.0" />
+
+    
+  
+  </ItemGroup>
+</Project>
\ No newline at end of file
