Missing file from change 19111703 (more tests for EpicGames.Serialization)

diff --git a/Engine/Source/Programs/Shared/EpicGames.Serialization.Tests/EpicGames.Serialization.Tests.csproj b/Engine/Source/Programs/Shared/EpicGames.Serialization.Tests/EpicGames.Serialization.Tests.csproj
--- a/Engine/Source/Programs/Shared/EpicGames.Serialization.Tests/EpicGames.Serialization.Tests.csproj
+++ b/Engine/Source/Programs/Shared/EpicGames.Serialization.Tests/EpicGames.Serialization.Tests.csproj
@@ -36,6 +36,9 @@
     <None Update="CompactBinaryObjects\build">
       <CopyToOutputDirectory>PreserveNewest</CopyToOutputDirectory>
     </None>
+    <None Update="CompactBinaryObjects\compact_binary">
+      <CopyToOutputDirectory>PreserveNewest</CopyToOutputDirectory>
+    </None>
     <None Update="CompactBinaryObjects\ReferenceOutput">
       <CopyToOutputDirectory>PreserveNewest</CopyToOutputDirectory>
     </None>
