Horde: Fix documentation comments for EpicGames.Horde.

diff --git a/Engine/Source/Programs/Shared/EpicGames.Horde/Auth/OAuthHandler.cs b/Engine/Source/Programs/Shared/EpicGames.Horde/Auth/OAuthHandler.cs
--- a/Engine/Source/Programs/Shared/EpicGames.Horde/Auth/OAuthHandler.cs
+++ b/Engine/Source/Programs/Shared/EpicGames.Horde/Auth/OAuthHandler.cs
@@ -20,6 +20,9 @@ namespace EpicGames.Horde.Auth
 	/// </summary>
 	public class AuthenticationException : Exception
 	{
+		/// <summary>
+		/// Constructor
+		/// </summary>
 		public AuthenticationException(string Message, Exception? InnerException)
 			: base(Message, InnerException)
 		{
@@ -75,6 +78,11 @@ namespace EpicGames.Horde.Auth
 		string CachedAccessToken = String.Empty;
 		DateTime ExpiresAt = DateTime.MinValue;
 
+		/// <summary>
+		/// Constructor
+		/// </summary>
+		/// <param name="Client"></param>
+		/// <param name="Options"></param>
 		public OAuthHandler(HttpClient Client, IOAuthOptions Options)
 		{
 			this.Client = Client;
diff --git a/Engine/Source/Programs/Shared/EpicGames.Horde/Auth/TokenHandler.cs b/Engine/Source/Programs/Shared/EpicGames.Horde/Auth/TokenHandler.cs
--- a/Engine/Source/Programs/Shared/EpicGames.Horde/Auth/TokenHandler.cs
+++ b/Engine/Source/Programs/Shared/EpicGames.Horde/Auth/TokenHandler.cs
@@ -33,6 +33,10 @@ namespace EpicGames.Horde.Auth
 	{
 		ITokenAuthOptions Options;
 
+		/// <summary>
+		/// Constructor
+		/// </summary>
+		/// <param name="Options"></param>
 		public TokenHandler(ITokenAuthOptions Options)
 		{
 			this.Options = Options;
diff --git a/Engine/Source/Programs/Shared/EpicGames.Horde/Common/Condition.cs b/Engine/Source/Programs/Shared/EpicGames.Horde/Common/Condition.cs
--- a/Engine/Source/Programs/Shared/EpicGames.Horde/Common/Condition.cs
+++ b/Engine/Source/Programs/Shared/EpicGames.Horde/Common/Condition.cs
@@ -325,7 +325,14 @@ namespace EpicGames.Horde.Common
 			}
 		}
 
+		/// <summary>
+		/// The condition text
+		/// </summary>
 		public string Text { get; }
+
+		/// <summary>
+		/// Error produced when parsing the condition
+		/// </summary>
 		public string? Error { get; private set; }
 		List<Token> Tokens = new List<Token>();
 		List<string> Strings = new List<string>();
@@ -350,10 +357,22 @@ namespace EpicGames.Horde.Common
 			}
 		}
 
+		/// <summary>
+		/// Determines if the condition is empty
+		/// </summary>
+		/// <returns>True if the condition is empty</returns>
 		public bool IsEmpty() => Tokens.Count == 0 && IsValid();
 
+		/// <summary>
+		/// Checks if the condition has been parsed correctly
+		/// </summary>
 		public bool IsValid() => Error == null;
 
+		/// <summary>
+		/// Parse the given text as a condition
+		/// </summary>
+		/// <param name="Text">Condition text to parse</param>
+		/// <returns>The new condition object</returns>
 		public static Condition Parse(string Text)
 		{
 			Condition Condition = new Condition(Text);
@@ -364,6 +383,11 @@ namespace EpicGames.Horde.Common
 			return Condition;
 		}
 
+		/// <summary>
+		/// Attempts to parse the given text as a condition
+		/// </summary>
+		/// <param name="Text">Condition to parse</param>
+		/// <returns>The parsed condition. Does not validate whether the parse completed successfully; call <see cref="Condition.IsValid"/> to verify.</returns>
 		public static Condition TryParse(string Text)
 		{
 			return new Condition(Text);
@@ -487,6 +511,11 @@ namespace EpicGames.Horde.Common
 			}
 		}
 
+		/// <summary>
+		/// Evaluate the condition using the given callback to retreive property values
+		/// </summary>
+		/// <param name="GetPropertyValues"></param>
+		/// <returns></returns>
 		public bool Evaluate(Func<string, IEnumerable<string>> GetPropertyValues)
 		{
 			if (IsEmpty())
@@ -581,6 +610,10 @@ namespace EpicGames.Horde.Common
 			}
 		}
 
+		/// <summary>
+		/// Implicit conversion from string to conditions
+		/// </summary>
+		/// <param name="Text"></param>
 		[return: NotNullIfNotNull("Text")]
 		public static implicit operator Condition?(string? Text)
 		{
@@ -594,6 +627,7 @@ namespace EpicGames.Horde.Common
 			}
 		}
 
+		/// <inheritdoc/>
 		public override string ToString() => (Error != null) ? $"[Error] {Text}" : Text;
 	}
 
diff --git a/Engine/Source/Programs/Shared/EpicGames.Horde/Compute/IComputeClient.cs b/Engine/Source/Programs/Shared/EpicGames.Horde/Compute/IComputeClient.cs
--- a/Engine/Source/Programs/Shared/EpicGames.Horde/Compute/IComputeClient.cs
+++ b/Engine/Source/Programs/Shared/EpicGames.Horde/Compute/IComputeClient.cs
@@ -196,6 +196,7 @@ namespace EpicGames.Horde.Compute
 		/// <param name="TaskRefId">Ref describing the request</param>
 		/// <param name="RequirementsHash">Requirements of the task to execute</param>
 		/// <param name="SkipCacheLookup">Whether to skip the lookup of cached output for this task</param>
+		/// <param name="CancellationToken">Cancellation token for the operation</param>
 		public static Task AddTaskAsync(this IComputeClient ComputeClient, ClusterId ClusterId, ChannelId ChannelId, RefId TaskRefId, IoHash RequirementsHash, bool SkipCacheLookup, CancellationToken CancellationToken = default)
 		{
 			return ComputeClient.AddTasksAsync(ClusterId, ChannelId, new[] { TaskRefId }, RequirementsHash, SkipCacheLookup, CancellationToken);
diff --git a/Engine/Source/Programs/Shared/EpicGames.Horde/Compute/Impl/HttpComputeClient.cs b/Engine/Source/Programs/Shared/EpicGames.Horde/Compute/Impl/HttpComputeClient.cs
--- a/Engine/Source/Programs/Shared/EpicGames.Horde/Compute/Impl/HttpComputeClient.cs
+++ b/Engine/Source/Programs/Shared/EpicGames.Horde/Compute/Impl/HttpComputeClient.cs
@@ -18,6 +18,9 @@ namespace EpicGames.Horde.Compute.Impl
 {
 	#region Requests / Responses
 
+	/// <summary>
+	/// Information about a compute cluster
+	/// </summary>
 	public class GetComputeClusterInfo : IComputeClusterInfo
 	{
 		/// <inheritdoc/>
@@ -76,19 +79,27 @@ namespace EpicGames.Horde.Compute.Impl
 	/// </summary>
 	public class AddTasksRequest
 	{
-		/// <inheritdoc cref="AddTasksRpcRequest.ChannelId"/>
+		/// <summary>
+		/// Channel to post the new tasks to. This should be a unique identifier (eg. a GUID) synthesized by the client to distinguish it from other clients, and used for querying status later.
+		/// </summary>
 		[CbField("c")]
 		public ChannelId ChannelId { get; set; }
 
-		/// <inheritdoc cref="AddTasksRpcRequest.TaskHashes"/>
+		/// <summary>
+		/// Refs for tasks to be executed
+		/// </summary>
 		[Required, CbField("t")]
 		public List<RefId> TaskRefIds { get; set; } = new List<RefId>();
 
-		/// <inheritdoc cref="AddTasksRpcRequest.RequirementsHash"/>
+		/// <summary>
+		/// Requirements for agents executing the tasks
+		/// </summary>
 		[CbField("r")]
 		public CbObjectAttachment RequirementsHash { get; set; }
 
-		/// <inheritdoc cref="AddTasksRpcRequest.DoNotCache"/>
+		/// <summary>
+		/// If set, prevents data being cached.
+		/// </summary>
 		[CbField("nc")]
 		public bool DoNotCache { get; set; }
 	}
@@ -107,10 +118,17 @@ namespace EpicGames.Horde.Compute.Impl
 
 	#endregion
 
+	/// <summary>
+	/// Implementation of <see cref="IComputeClient"/> which uses HTTP to communicate with a Horde Storage server
+	/// </summary>
 	public class HttpComputeClient : IComputeClient
 	{
 		HttpClient HttpClient;
 
+		/// <summary>
+		/// Constructor
+		/// </summary>
+		/// <param name="HttpClient"></param>
 		public HttpComputeClient(HttpClient HttpClient)
 		{
 			this.HttpClient = HttpClient;
diff --git a/Engine/Source/Programs/Shared/EpicGames.Horde/Compute/Requirements.cs b/Engine/Source/Programs/Shared/EpicGames.Horde/Compute/Requirements.cs
--- a/Engine/Source/Programs/Shared/EpicGames.Horde/Compute/Requirements.cs
+++ b/Engine/Source/Programs/Shared/EpicGames.Horde/Compute/Requirements.cs
@@ -15,7 +15,7 @@ namespace EpicGames.Horde.Compute
 	public class Requirements
 	{
 		/// <summary>
-		/// Condition string to be evaluated against the machine spec, eg. cpu-cores >= 10 && ram.mb >= 200 && pool == 'worker'
+		/// Condition string to be evaluated against the machine spec, eg. cpu-cores >= 10 &amp;&amp; ram.mb >= 200 &amp;&amp; pool == 'worker'
 		/// </summary>
 		[CbField("c")]
 		public Condition? Condition { get; set; }
@@ -48,6 +48,7 @@ namespace EpicGames.Horde.Compute
 			this.Condition = Condition;
 		}
 
+		/// <inheritdoc/>
 		public override string ToString()
 		{
 			List<string> List = new List<string>();
diff --git a/Engine/Source/Programs/Shared/EpicGames.Horde/EpicGames.Horde.csproj b/Engine/Source/Programs/Shared/EpicGames.Horde/EpicGames.Horde.csproj
--- a/Engine/Source/Programs/Shared/EpicGames.Horde/EpicGames.Horde.csproj
+++ b/Engine/Source/Programs/Shared/EpicGames.Horde/EpicGames.Horde.csproj
@@ -10,6 +10,7 @@
     <Configurations>Debug;Release;Development</Configurations>
     <Nullable>enable</Nullable>
     <DebugType>pdbonly</DebugType>
+    <GenerateDocumentationFile>True</GenerateDocumentationFile>
   </PropertyGroup>
   <PropertyGroup Condition=" '$(Configuration)|$(Platform)' == 'Debug|AnyCPU' ">
     <DocumentationFile></DocumentationFile>
diff --git a/Engine/Source/Programs/Shared/EpicGames.Horde/Storage/Impl/FileStorageClient.cs b/Engine/Source/Programs/Shared/EpicGames.Horde/Storage/Impl/FileStorageClient.cs
--- a/Engine/Source/Programs/Shared/EpicGames.Horde/Storage/Impl/FileStorageClient.cs
+++ b/Engine/Source/Programs/Shared/EpicGames.Horde/Storage/Impl/FileStorageClient.cs
@@ -34,6 +34,9 @@ namespace EpicGames.Horde.Storage.Impl
 			}
 		}
 
+		/// <summary>
+		/// Base directory for storing files
+		/// </summary>
 		public DirectoryReference BaseDir { get; }
 		ILogger Logger;
 
diff --git a/Engine/Source/Programs/Shared/EpicGames.Horde/Storage/TreePack.cs b/Engine/Source/Programs/Shared/EpicGames.Horde/Storage/TreePack.cs
--- a/Engine/Source/Programs/Shared/EpicGames.Horde/Storage/TreePack.cs
+++ b/Engine/Source/Programs/Shared/EpicGames.Horde/Storage/TreePack.cs
@@ -105,7 +105,7 @@ namespace EpicGames.Horde.Storage
 	public class TreePackObjectImport
 	{
 		/// <summary>
-		/// Imported object. The <see cref="Nodes"/> field contains a flat array of hashes imported from the object.
+		/// Imported object. The <see cref="Keys"/> field contains a flat array of hashes imported from the object.
 		/// </summary>
 		[CbField("object")]
 		public CbObjectAttachment Object { get; set; }
@@ -123,8 +123,7 @@ namespace EpicGames.Horde.Storage
 	public class TreePackBinaryImport
 	{
 		/// <summary>
-		/// Imported binary. Binary objects do not have a compact binary header, and effectively consist solely of the <see cref="TreePack.Data"/> field. 
-		/// The <see cref="Nodes"/> field contains encoded a list of encoded entry objects.
+		/// Imported binary.
 		/// </summary>
 		[CbField("binary")]
 		public CbBinaryAttachment Binary { get; set; }
@@ -630,6 +629,9 @@ namespace EpicGames.Horde.Storage
 		IStorageClient StorageClient { get; }
 		NamespaceId NamespaceId { get; }
 
+		/// <summary>
+		/// 
+		/// </summary>
 		public TreePackOptions Options { get; }
 		Dictionary<IoHash, NodeInfo> Nodes = new Dictionary<IoHash, NodeInfo>();
 		Dictionary<IoHash, BlobInfo> Blobs = new Dictionary<IoHash, BlobInfo>();
@@ -658,7 +660,7 @@ namespace EpicGames.Horde.Storage
 		/// <summary>
 		/// Adds a root blob to the pack
 		/// </summary>
-		/// <param name="Data"></param>
+		/// <param name="Ref"></param>
 		/// <returns>Hash of the root node</returns>
 		public ReadOnlyMemory<byte> AddRootBlob(IRef Ref) => AddRootBlob(Ref.Value.GetView());
 
@@ -679,7 +681,7 @@ namespace EpicGames.Horde.Storage
 		/// <summary>
 		/// Adds a root blob to the pack
 		/// </summary>
-		/// <param name="Data"></param>
+		/// <param name="Object"></param>
 		/// <returns>Hash of the root node</returns>
 		public void AddRootObject(TreePackObject Object)
 		{
@@ -748,7 +750,7 @@ namespace EpicGames.Horde.Storage
 		/// <summary>
 		/// Adds information about an object to the store
 		/// </summary>
-		/// <param name="Object"></param>
+		/// <param name="RootObject"></param>
 		public void AddNodes(TreePackObject RootObject)
 		{
 			RegisterObject(RootObject, null);
@@ -821,7 +823,7 @@ namespace EpicGames.Horde.Storage
 		/// <summary>
 		/// Gets data for the given node
 		/// </summary>
-		/// <param name="Key"></param>
+		/// <param name="Hash"></param>
 		/// <returns></returns>
 		public ValueTask<ReadOnlyMemory<byte>> GetDataAsync(IoHash Hash)
 		{
@@ -832,7 +834,7 @@ namespace EpicGames.Horde.Storage
 		/// <summary>
 		/// Gets data for the given node
 		/// </summary>
-		/// <param name="Key"></param>
+		/// <param name="Node"></param>
 		/// <returns></returns>
 		private async ValueTask<ReadOnlyMemory<byte>> GetDataAsync(NodeInfo Node)
 		{
@@ -1196,7 +1198,6 @@ namespace EpicGames.Horde.Storage
 		}
 
 		/// <inheritdoc cref="GetCostHeuristic(int, TimeSpan)"/>
-		/// <param name="Index">Index to calculate the heuristic for</param>
 		double GetCostHeuristic(BlobInfo Info, DateTime UtcNow) => GetCostHeuristic(Info.Length, UtcNow - Info.Time);
 
 		/// <summary>
@@ -1204,6 +1205,7 @@ namespace EpicGames.Horde.Storage
 		/// </summary>
 		/// <param name="Nodes">The full list of nodes</param>
 		/// <param name="CollapseBlobs">Set of blobs that should be collapsed</param>
+		/// <param name="UtcNow">The current time</param>
 		/// <returns></returns>
 		double GetNewCostHeuristic(List<NodeInfo> Nodes, HashSet<BlobInfo> CollapseBlobs, DateTime UtcNow)
 		{
@@ -1277,6 +1279,13 @@ namespace EpicGames.Horde.Storage
 			return Probability * (DownloadInit + (Size / DownloadRate));
 		}
 
+		/// <summary>
+		/// 
+		/// </summary>
+		/// <param name="RootHash"></param>
+		/// <param name="Directory"></param>
+		/// <param name="Logger"></param>
+		/// <returns></returns>
 		public async Task WriteTreeSummaryAsync(IoHash RootHash, DirectoryReference Directory, ILogger Logger)
 		{
 			DirectoryReference.CreateDirectory(Directory);
@@ -1313,6 +1322,11 @@ namespace EpicGames.Horde.Storage
 			}
 		}
 
+		/// <summary>
+		/// 
+		/// </summary>
+		/// <param name="File"></param>
+		/// <param name="Object"></param>
 		public static void WriteSummary(FileReference File, TreePackObject Object)
 		{
 			Dictionary<IoHash, string> HashToLocator = new Dictionary<IoHash, string>();
diff --git a/Engine/Source/Programs/Shared/EpicGames.Horde/Storage/TreePackDirWriter.cs b/Engine/Source/Programs/Shared/EpicGames.Horde/Storage/TreePackDirWriter.cs
--- a/Engine/Source/Programs/Shared/EpicGames.Horde/Storage/TreePackDirWriter.cs
+++ b/Engine/Source/Programs/Shared/EpicGames.Horde/Storage/TreePackDirWriter.cs
@@ -17,6 +17,9 @@ namespace EpicGames.Horde.Storage
 		Dictionary<Utf8String, TreePackDirEntry> NameToEntry = new Dictionary<Utf8String, TreePackDirEntry>();
 		Dictionary<Utf8String, TreePackDirWriter> NameToSubDir = new Dictionary<Utf8String, TreePackDirWriter>();
 
+		/// <summary>
+		/// 
+		/// </summary>
 		public IEnumerable<TreePackDirEntry> Entries => NameToEntry.Values;
 
 		/// <summary>
diff --git a/Engine/Source/Programs/Shared/EpicGames.Horde/Storage/TreePackFileWriter.cs b/Engine/Source/Programs/Shared/EpicGames.Horde/Storage/TreePackFileWriter.cs
--- a/Engine/Source/Programs/Shared/EpicGames.Horde/Storage/TreePackFileWriter.cs
+++ b/Engine/Source/Programs/Shared/EpicGames.Horde/Storage/TreePackFileWriter.cs
@@ -13,7 +13,11 @@ namespace EpicGames.Horde.Storage
 	/// </summary>
 	public class TreePackFileWriter
 	{
+		/// <summary>
+		/// 
+		/// </summary>
 		public TreePack TreePack { get; }
+
 		int MinBlockSize;
 		int MaxBlockSize;
 		uint WindowHashThreshold;
