Horde: Refactored system for creating bundles of tree objects.

- Tree can now be modified without having to have entire working set in memory.
- Multiple nodes can be compressed/decompressed together, and an optional IMemoryCache allows caching decompressed blocks.
- Added support for specifying custom compact-binary serializers on individual properties.
- Added test suite to test bundling functionality.
- Fixed issue where VarInts of negative values would not serialize correctly.

diff --git a/Engine/Source/Programs/Shared/EpicGames.Core/BuzHash.cs b/Engine/Source/Programs/Shared/EpicGames.Core/BuzHash.cs
--- a/Engine/Source/Programs/Shared/EpicGames.Core/BuzHash.cs
+++ b/Engine/Source/Programs/Shared/EpicGames.Core/BuzHash.cs
@@ -39,8 +39,27 @@ namespace EpicGames.Core
 			0xf9c18d66, 0x593ade65, 0xd95ddf11,
 		};
 
-		ulong Count = 0;
-		uint State = 0;
+		int Count;
+		uint State;
+
+		/// <summary>
+		/// Constructor
+		/// </summary>
+		public BuzHash()
+			: this(0, 0)
+		{
+		}
+
+		/// <summary>
+		/// Constructor
+		/// </summary>
+		/// <param name="Count"></param>
+		/// <param name="State"></param>
+		public BuzHash(int Count, uint State)
+		{
+			this.Count = Count;
+			this.State = State;
+		}
 
 		/// <summary>
 		/// Resets the state of the 
@@ -57,10 +76,46 @@ namespace EpicGames.Core
 		/// <param name="Value">New value to append to the window</param>
 		public void Add(byte Value)
 		{
-			State = Rol32(State, 1) ^ Table[Value];
+			State = Add(State, Value);
 			Count++;
 		}
 
+		/// <summary>
+		/// Add a byte to the window and update the hash
+		/// </summary>
+		/// <param name="Value">New value to append to the window</param>
+		public void Add(ReadOnlySpan<byte> Span)
+		{
+			State = Add(State, Span);
+			Count += Span.Length;
+		}
+
+		/// <summary>
+		/// Static method for updating a hash given a particular byte value
+		/// </summary>
+		/// <param name="State">The current hash value</param>
+		/// <param name="Value">Byte to add to the hash</param>
+		/// <returns>New hash value</returns>
+		public static uint Add(uint State, byte Value)
+		{
+			return Rol32(State, 1) ^ Table[Value];
+		}
+
+		/// <summary>
+		/// Static method for appending a range of data to a hash value
+		/// </summary>
+		/// <param name="State">The current hash value</param>
+		/// <param name="Data">Data to append to the hash</param>
+		/// <returns>New hash value</returns>
+		public static uint Add(uint State, ReadOnlySpan<byte> Data)
+		{
+			for (int Idx = 0; Idx < Data.Length; Idx++)
+			{
+				State = Add(State, Data[Idx]);
+			}
+			return State;
+		}
+
 		/// <summary>
 		/// Removes a byte from the start of the hash window
 		/// </summary>
diff --git a/Engine/Source/Programs/Shared/EpicGames.Horde.Tests/BundleTests.cs b/Engine/Source/Programs/Shared/EpicGames.Horde.Tests/BundleTests.cs
new file mode 100644
--- /dev/null
+++ b/Engine/Source/Programs/Shared/EpicGames.Horde.Tests/BundleTests.cs
@@ -0,0 +1,328 @@
+// Copyright Epic Games, Inc. All Rights Reserved.
+
+using EpicGames.Core;
+using EpicGames.Horde.Bundles;
+using EpicGames.Horde.Bundles.Nodes;
+using EpicGames.Horde.Storage;
+using EpicGames.Horde.Storage.Impl;
+using EpicGames.Serialization;
+using Microsoft.VisualStudio.TestTools.UnitTesting;
+using System;
+using System.Buffers.Binary;
+using System.Collections.Generic;
+using System.Security.Cryptography;
+using System.Text;
+using System.Threading.Tasks;
+
+namespace HordeServerTests
+{
+	[TestClass]
+	public class BundleTests
+	{
+		MemoryStorageClient StorageClient = new MemoryStorageClient();
+		NamespaceId NamespaceId = new NamespaceId("namespace");
+		BucketId BucketId = new BucketId("bucket");
+
+		[TestMethod]
+		public void SerializationTests()
+		{
+			DateTime UtcNow = DateTime.UtcNow;
+
+			List<BundleImport> Imports = new List<BundleImport>();
+			Imports.Add(new BundleImport(IoHash.Compute(Encoding.UTF8.GetBytes("import1")), 0, 123));
+			Imports.Add(new BundleImport(IoHash.Compute(Encoding.UTF8.GetBytes("import2")), 1, 456));
+
+			BundleObject OldObject = new BundleObject();
+			OldObject.CreationTimeUtc = DateTime.UtcNow;
+			OldObject.ImportObjects.Add(new BundleImportObject(UtcNow - TimeSpan.FromMinutes(5.0), new CbObjectAttachment(IoHash.Compute(new byte[] { 1, 2, 3 })), 12345, Imports));
+			OldObject.Exports.Add(new BundleExport(IoHash.Compute(Encoding.UTF8.GetBytes("export1")), 2, 5, new BundleCompressionPacket(0) { EncodedLength = 20, DecodedLength = 40 }, 0, 40, new int[] { 1, 2 }));
+			OldObject.Exports.Add(new BundleExport(IoHash.Compute(Encoding.UTF8.GetBytes("export2")), 3, 6, new BundleCompressionPacket(20) { EncodedLength = 10, DecodedLength = 20 }, 0, 20, new int[] { -1 }));
+
+			CbObject SerializedObject = CbSerializer.Serialize(OldObject);
+			ReadOnlyMemory<byte> SerializedData = SerializedObject.GetView();
+
+			BundleObject NewObject = CbSerializer.Deserialize<BundleObject>(SerializedData);
+
+			Assert.AreEqual(OldObject.CreationTimeUtc, NewObject.CreationTimeUtc);
+
+			Assert.AreEqual(OldObject.ImportObjects.Count, NewObject.ImportObjects.Count);
+			for (int Idx = 0; Idx < OldObject.ImportObjects.Count; Idx++)
+			{
+				BundleImportObject OldObjectImport = OldObject.ImportObjects[Idx];
+				BundleImportObject NewObjectImport = NewObject.ImportObjects[Idx];
+
+				Assert.AreEqual(OldObjectImport.CreationTimeUtc, NewObjectImport.CreationTimeUtc);
+				Assert.AreEqual(OldObjectImport.Object, NewObjectImport.Object);
+				Assert.AreEqual(OldObjectImport.TotalCost, NewObjectImport.TotalCost);
+				Assert.AreEqual(OldObjectImport.Imports.Count, NewObjectImport.Imports.Count);
+
+				for (int ImportIdx = 0; ImportIdx < OldObjectImport.Imports.Count; ImportIdx++)
+				{
+					BundleImport OldImport = OldObjectImport.Imports[ImportIdx];
+					BundleImport NewImport = NewObjectImport.Imports[ImportIdx];
+
+					Assert.AreEqual(OldImport.Hash, NewImport.Hash);
+					Assert.AreEqual(OldImport.Rank, NewImport.Rank);
+					Assert.AreEqual(OldImport.Cost, NewImport.Cost);
+				}
+			}
+
+			Assert.AreEqual(OldObject.Exports.Count, NewObject.Exports.Count);
+			for (int Idx = 0; Idx < OldObject.Exports.Count; Idx++)
+			{
+				BundleExport OldExport = OldObject.Exports[Idx];
+				BundleExport NewExport = NewObject.Exports[Idx];
+
+				Assert.AreEqual(OldExport.Hash, NewExport.Hash);
+				Assert.AreEqual(OldExport.Rank, NewExport.Rank);
+				Assert.AreEqual(OldExport.Cost, NewExport.Cost);
+				Assert.AreEqual(OldExport.Packet.Offset, NewExport.Packet.Offset);
+				Assert.AreEqual(OldExport.Packet.EncodedLength, NewExport.Packet.EncodedLength);
+				Assert.AreEqual(OldExport.Packet.DecodedLength, NewExport.Packet.DecodedLength);
+				Assert.AreEqual(OldExport.Offset, NewExport.Offset);
+				Assert.AreEqual(OldExport.Length, NewExport.Length);
+				Assert.IsTrue(OldExport.References.AsSpan().SequenceEqual(NewExport.References.AsSpan()));
+			}
+		}
+
+		[TestMethod]
+		public async Task BasicTestDirectory()
+		{
+			Bundle<DirectoryNode> Bundle = new Bundle<DirectoryNode>(StorageClient, NamespaceId, new BundleOptions(), null);
+			DirectoryNode Node = await Bundle.Root.AddDirectoryAsync("hello");
+			DirectoryNode Node2 = await Node.AddDirectoryAsync("world");
+
+			RefId RefId = new RefId("testref");
+			await Bundle.WriteAsync(BucketId, RefId, false, DateTime.MinValue);
+
+			// Should be stored inline
+			Assert.AreEqual(1, StorageClient.Refs.Count);
+			Assert.AreEqual(0, StorageClient.Blobs.Count);
+
+			// Check the ref
+			BundleObject Root = await StorageClient.GetRefAsync<BundleObject>(NamespaceId, BucketId, RefId);
+			Assert.AreEqual(3, Root.Exports.Count);
+			Assert.AreEqual(0, Root.Exports[0].Rank);
+			Assert.AreEqual(1, Root.Exports[1].Rank);
+			Assert.AreEqual(2, Root.Exports[2].Rank);
+
+			// Create a new bundle and read it back in again
+			Bundle = new Bundle<DirectoryNode>(StorageClient, NamespaceId, new BundleOptions(), null);
+			await Bundle.ReadAsync(BucketId, RefId);
+
+			Assert.AreEqual(1, Bundle.Root.Entries.Count);
+			DirectoryNode? OutputNode = await Bundle.Root.FindDirectoryAsync("hello");
+			Assert.IsNotNull(OutputNode);
+
+			Assert.AreEqual(1, OutputNode!.Entries.Count);
+			DirectoryNode? OutputNode2 = await OutputNode!.FindDirectoryAsync("world");
+			Assert.IsNotNull(OutputNode2);
+
+			Assert.AreEqual(0, OutputNode2!.Entries.Count);
+		}
+
+
+
+		// caching
+		// hierarchical chunking
+
+		[TestMethod]
+		public async Task DedupTests()
+		{
+			BundleOptions Options = new BundleOptions();
+			Options.MaxBlobSize = 1;
+			Options.MaxInlineBlobSize = 1;
+
+			Bundle<DirectoryNode> Bundle = new Bundle<DirectoryNode>(StorageClient, NamespaceId, Options, null);
+
+			await Bundle.Root.AddDirectoryAsync("node1");
+			await Bundle.Root.AddDirectoryAsync("node2");
+			await Bundle.Root.AddDirectoryAsync("node3");
+
+			RefId RefId = new RefId("ref");
+			await Bundle.WriteAsync(BucketId, RefId, false, DateTime.UtcNow);
+
+			Assert.AreEqual(1, StorageClient.Refs.Count);
+			Assert.AreEqual(1, StorageClient.Blobs.Count);
+		}
+
+		[TestMethod]
+		public async Task ReloadTests()
+		{
+			BundleOptions Options = new BundleOptions();
+			Options.MaxBlobSize = 1;
+			Options.MaxInlineBlobSize = 1;
+
+			Bundle<DirectoryNode> Bundle = new Bundle<DirectoryNode>(StorageClient, NamespaceId, Options, null);
+
+			DirectoryNode Node1 = await Bundle.Root.AddDirectoryAsync("node1");
+			DirectoryNode Node2 = await Node1.AddDirectoryAsync("node2");
+			DirectoryNode Node3 = await Node2.AddDirectoryAsync("node3");
+			DirectoryNode Node4 = await Node3.AddDirectoryAsync("node4");
+
+			RefId RefId = new RefId("ref");
+			await Bundle.WriteAsync(BucketId, RefId, false, DateTime.UtcNow);
+
+			Assert.AreEqual(1, StorageClient.Refs.Count);
+			Assert.AreEqual(4, StorageClient.Blobs.Count);
+
+			Bundle<DirectoryNode> NewBundle = new Bundle<DirectoryNode>(StorageClient, NamespaceId, Options, null);
+			await NewBundle.ReadAsync(BucketId, RefId);
+
+			DirectoryNode? NewNode1 = await NewBundle.Root.FindDirectoryAsync("node1");
+			Assert.IsNotNull(NewNode1);
+
+			DirectoryNode? NewNode2 = await NewNode1!.FindDirectoryAsync("node2");
+			Assert.IsNotNull(NewNode2);
+
+			DirectoryNode? NewNode3 = await NewNode2!.FindDirectoryAsync("node3");
+			Assert.IsNotNull(NewNode3);
+
+			DirectoryNode? NewNode4 = await NewNode3!.FindDirectoryAsync("node4");
+			Assert.IsNotNull(NewNode4);
+		}
+
+		[TestMethod]
+		public async Task CompactTest()
+		{
+			DateTime Time = DateTime.UtcNow;
+
+			BundleOptions Options = new BundleOptions();
+			Options.MaxBlobSize = 1024 * 1024;
+			Options.MaxInlineBlobSize = 1;
+
+			Bundle<DirectoryNode> Bundle = new Bundle<DirectoryNode>(StorageClient, NamespaceId, Options, null);
+
+			DirectoryNode Node1 = await Bundle.Root.AddDirectoryAsync("node1");
+			DirectoryNode Node2 = await Node1.AddDirectoryAsync("node2");
+			DirectoryNode Node3 = await Node2.AddDirectoryAsync("node3");
+			DirectoryNode Node4 = await Bundle.Root.AddDirectoryAsync("node4"); // same contents as node 3
+
+			RefId RefId1 = new RefId("ref1");
+			await Bundle.WriteAsync(BucketId, RefId1, false, Time);
+
+			Assert.AreEqual(1, StorageClient.Refs.Count);
+			Assert.AreEqual(1, StorageClient.Blobs.Count);
+
+			IRef Ref1 = StorageClient.Refs[(NamespaceId, BucketId, RefId1)];
+
+			IoHash RootHash1 = Bundle.RootHash;
+			BundleObject RootObject1 = CbSerializer.Deserialize<BundleObject>(Ref1.Value);
+			Assert.AreEqual(1, RootObject1.Exports.Count);
+			Assert.AreEqual(1, RootObject1.ImportObjects.Count);
+			Assert.AreEqual(RootHash1, RootObject1.Exports[0].Hash);
+
+			IoHash LeafHash1 = RootObject1.ImportObjects[0].Object.Hash;
+			BundleObject LeafObject1 = CbSerializer.Deserialize<BundleObject>(StorageClient.Blobs[(NamespaceId, LeafHash1)]);
+			Assert.AreEqual(3, LeafObject1.Exports.Count); // node4 == node3
+			Assert.AreEqual(0, LeafObject1.ImportObjects.Count);
+
+			// Remove one of the nodes from the root without compacting. the existing blob should be reused.
+			Bundle.Root.Delete("node1");
+
+			RefId RefId2 = new RefId("ref2");
+			await Bundle.WriteAsync(BucketId, RefId2, false, Time);
+
+			IRef Ref2 = StorageClient.Refs[(NamespaceId, BucketId, RefId2)];
+
+			IoHash RootHash2 = Bundle.RootHash;
+			BundleObject RootObject2 = CbSerializer.Deserialize<BundleObject>(Ref2.Value);
+			Assert.AreEqual(1, RootObject2.Exports.Count);
+			Assert.AreEqual(1, RootObject2.ImportObjects.Count);
+			Assert.AreEqual(RootHash2, RootObject2.Exports[0].Hash);
+
+			IoHash LeafHash2 = RootObject2.ImportObjects[0].Object.Hash;
+			Assert.AreEqual(LeafHash1, LeafHash2);
+
+			// Repack it and check that we make a new object
+			RefId RefId3 = new RefId("ref3");
+			await Bundle.WriteAsync(BucketId, RefId3, true, Time + TimeSpan.FromDays(7));
+
+			IRef Ref3 = StorageClient.Refs[(NamespaceId, BucketId, RefId3)];
+
+			IoHash RootHash3 = Bundle.RootHash;
+			BundleObject RootObject3 = CbSerializer.Deserialize<BundleObject>(Ref3.Value);
+			Assert.AreEqual(1, RootObject3.Exports.Count);
+			Assert.AreEqual(1, RootObject3.ImportObjects.Count);
+			Assert.AreEqual(RootHash3, RootObject3.Exports[0].Hash);
+
+			IoHash LeafHash3 = RootObject3.ImportObjects[0].Object.Hash;
+			Assert.AreNotEqual(LeafHash1, LeafHash3);
+
+			BundleObject LeafObject3 = CbSerializer.Deserialize<BundleObject>(StorageClient.Blobs[(NamespaceId, LeafHash3)]);
+			Assert.AreEqual(1, LeafObject3.Exports.Count);
+			Assert.AreEqual(0, LeafObject3.ImportObjects.Count);
+		}
+
+		[TestMethod]
+		public async Task FixedSizeChunkingTests()
+		{
+			ChunkOptions Options = new ChunkOptions();
+			Options.LeafOptions = new TypedChunkOptions(64, 64, 64);
+			Options.InteriorOptions = new TypedChunkOptions(IoHash.NumBytes * 4, IoHash.NumBytes * 4, IoHash.NumBytes * 4);
+
+			await ChunkingTests(Options);
+		}
+
+		[TestMethod]
+		public async Task VariableSizeChunkingTests()
+		{
+			ChunkOptions Options = new ChunkOptions();
+			Options.LeafOptions = new TypedChunkOptions(32, 64, 96);
+			Options.InteriorOptions = new TypedChunkOptions(IoHash.NumBytes * 1, IoHash.NumBytes * 4, IoHash.NumBytes * 12);
+
+			await ChunkingTests(Options);
+		}
+
+		async Task ChunkingTests(ChunkOptions Options)
+		{
+			byte[] Data = new byte[4096];
+			new Random(0).NextBytes(Data);
+
+			for (int Idx = 0; Idx < Data.Length; Idx++)
+			{
+				Data[Idx] = (byte)Idx;
+			}
+
+			Bundle<ChunkNode> Bundle = new Bundle<ChunkNode>(StorageClient, NamespaceId, new BundleOptions(), null);
+
+			const int NumIterations = 100;
+			for (int Idx = 0; Idx < NumIterations; Idx++)
+			{
+				Bundle.Root.Append(Data, Options);
+			}
+
+			byte[] Result = await Bundle.Root.ToByteArrayAsync();
+			Assert.AreEqual(NumIterations * Data.Length, Result.Length);
+
+			for (int Idx = 0; Idx < NumIterations; Idx++)
+			{
+				ReadOnlyMemory<byte> SpanData = Result.AsMemory(Idx * Data.Length, Data.Length);
+				Assert.IsTrue(SpanData.Span.SequenceEqual(Data));
+			}
+
+			await CheckSizes(Bundle.Root, Options, true);
+		}
+
+		async Task CheckSizes(ChunkNode Node, ChunkOptions Options, bool Rightmost)
+		{
+			if (Node.Depth == 0)
+			{
+				Assert.IsTrue(Rightmost || Node.Payload.Length >= Options.LeafOptions.MinSize);
+				Assert.IsTrue(Node.Payload.Length <= Options.LeafOptions.MaxSize);
+			}
+			else
+			{
+				Assert.IsTrue(Rightmost || Node.Payload.Length >= Options.InteriorOptions.MinSize);
+				Assert.IsTrue(Node.Payload.Length <= Options.InteriorOptions.MaxSize);
+
+				int ChildCount = Node.GetChildCount();
+				for (int Idx = 0; Idx < ChildCount; Idx++)
+				{
+					ChunkNode ChildNode = await Node.GetChildNodeAsync(Idx);
+					await CheckSizes(ChildNode, Options, Idx == ChildCount - 1);
+				}
+			}
+		}
+	}
+}
diff --git a/Engine/Source/Programs/Shared/EpicGames.Horde.Tests/EpicGames.Horde.Tests.csproj b/Engine/Source/Programs/Shared/EpicGames.Horde.Tests/EpicGames.Horde.Tests.csproj
new file mode 100644
--- /dev/null
+++ b/Engine/Source/Programs/Shared/EpicGames.Horde.Tests/EpicGames.Horde.Tests.csproj
@@ -0,0 +1,22 @@
+<Project Sdk="Microsoft.NET.Sdk">
+
+  <PropertyGroup>
+    <TargetFramework>netcoreapp3.1</TargetFramework>
+    <Configurations>Debug;Release;Development</Configurations>
+    <Nullable>enable</Nullable>
+    <IsPackable>false</IsPackable>
+  </PropertyGroup>
+
+  <ItemGroup>
+    <PackageReference Include="Microsoft.NET.Test.Sdk" Version="16.5.0" />
+    <PackageReference Include="MSTest.TestAdapter" Version="2.1.0" />
+    <PackageReference Include="MSTest.TestFramework" Version="2.1.0" />
+    <PackageReference Include="coverlet.collector" Version="1.2.0" />
+  </ItemGroup>
+
+  <ItemGroup>
+    <ProjectReference Include="..\EpicGames.Core\EpicGames.Core.csproj" />
+    <ProjectReference Include="..\EpicGames.Horde\EpicGames.Horde.csproj" />
+  </ItemGroup>
+
+</Project>
diff --git a/Engine/Source/Programs/Shared/EpicGames.Horde/Bundles/Bundle.cs b/Engine/Source/Programs/Shared/EpicGames.Horde/Bundles/Bundle.cs
new file mode 100644
--- /dev/null
+++ b/Engine/Source/Programs/Shared/EpicGames.Horde/Bundles/Bundle.cs
@@ -0,0 +1,772 @@
+// Copyright Epic Games, Inc. All Rights Reserved.
+
+using EpicGames.Core;
+using EpicGames.Horde.Bundles.Nodes;
+using EpicGames.Horde.Storage;
+using EpicGames.Serialization;
+using K4os.Compression.LZ4;
+using Microsoft.Extensions.Caching.Memory;
+using Microsoft.Extensions.Logging;
+using System;
+using System.Buffers.Binary;
+using System.Collections;
+using System.Collections.Generic;
+using System.Diagnostics;
+using System.Diagnostics.CodeAnalysis;
+using System.IO;
+using System.IO.Compression;
+using System.Linq;
+using System.Reflection;
+using System.Security.Cryptography;
+using System.Text;
+using System.Text.Json;
+using System.Threading.Tasks;
+
+namespace EpicGames.Horde.Bundles
+{
+	/// <summary>
+	/// Options for configuring a bundle serializer
+	/// </summary>
+	public class BundleOptions
+	{
+		/// <summary>
+		/// Maximum payload size fo a blob
+		/// </summary>
+		public int MaxBlobSize { get; set; } = 10 * 1024 * 1024;
+
+		/// <summary>
+		/// Maximum payload size for the root blob
+		/// </summary>
+		public int MaxInlineBlobSize { get; set; } = 1024 * 64;
+
+		/// <summary>
+		/// Minimum size of a block to be compressed
+		/// </summary>
+		public int MinCompressionPacketSize { get; set; } = 16 * 1024;
+	}
+
+	/// <summary>
+	/// Base class for manipulating bundles
+	/// </summary>
+	public class Bundle
+	{
+		/// <summary>
+		/// Information about a node within a bundle.
+		/// </summary>
+		[DebuggerDisplay("{Hash}")]
+		sealed class NodeInfo
+		{
+			public IoHash Hash { get; }
+			public int Rank { get; }
+			public int Cost { get; }
+
+			public BlobInfo? Blob;
+			public BundleExport? Export;
+			public ReadOnlyMemory<byte> Data;
+			public NodeInfo[]? References;
+
+			public NodeInfo(IoHash Hash, int Rank, int Cost)
+			{
+				this.Hash = Hash;
+				this.Rank = Rank;
+				this.Cost = Cost;
+			}
+
+			public NodeInfo(BlobInfo Blob, BundleImport Import)
+				: this(Import.Hash, Import.Rank, Import.Cost)
+			{
+				this.Blob = Blob;
+			}
+
+			void Set(BlobInfo? Blob, BundleExport? Export, ReadOnlyMemory<byte> Data, NodeInfo[]? References)
+			{
+				this.Blob = Blob;
+				this.Export = Export;
+				this.Data = Data;
+				this.References = References;
+			}
+
+			public void SetImport(BlobInfo Blob) => Set(Blob, null, ReadOnlyMemory<byte>.Empty, null);
+			public void SetExport(BlobInfo? Blob, BundleExport Export, ReadOnlyMemory<byte> Data, NodeInfo[] References) => Set(Blob, Export, Data, References);
+			public void SetStandalone(ReadOnlyMemory<byte> Data, NodeInfo[] References) => Set(null, null, Data, References);
+		}
+
+		/// <summary>
+		/// Information about a blob
+		/// </summary>
+		[DebuggerDisplay("{Hash}")]
+		class BlobInfo
+		{
+			public IoHash Hash { get; }
+			public DateTime CreationTimeUtc { get; }
+			public int TotalCost { get; }
+
+			public bool Mounted { get; set; }
+			public HashSet<BlobInfo> ReferencedBy { get; } = new HashSet<BlobInfo>(); // Used for repacking
+
+			public BlobInfo(IoHash Hash, DateTime CreationTimeUtc, int TotalCost)
+			{
+				this.Hash = Hash;
+				this.CreationTimeUtc = CreationTimeUtc;
+				this.TotalCost = TotalCost;
+			}
+ 		}
+
+		IStorageClient StorageClient;
+		NamespaceId NamespaceId;
+
+		/// <summary>
+		/// Hash of the root object
+		/// </summary>
+		public IoHash RootHash { get; protected set; }
+		object RootCacheKey = new object();
+
+		Dictionary<IoHash, NodeInfo> HashToNode = new Dictionary<IoHash, NodeInfo>();
+		Dictionary<IoHash, BlobInfo> HashToBlob = new Dictionary<IoHash, BlobInfo>();
+		IMemoryCache? Cache;
+
+		/// <summary>
+		/// Options for the bundle
+		/// </summary>
+		public BundleOptions Options { get; }
+
+		byte[] BlockBuffer = Array.Empty<byte>();
+		byte[] EncodedBuffer = Array.Empty<byte>();
+
+		/// <summary>
+		/// Constructor
+		/// </summary>
+		/// <param name="StorageClient">Client interface for the storage backend</param>
+		/// <param name="NamespaceId">Namespace for storing blobs in the storage system</param>
+		/// <param name="Options">Options for controlling how things are serialized</param>
+		/// <param name="Cache">Cache for storing decompressed objects, serialized blobs.</param>
+		public Bundle(IStorageClient StorageClient, NamespaceId NamespaceId, BundleOptions Options, IMemoryCache? Cache)
+		{
+			this.StorageClient = StorageClient;
+			this.NamespaceId = NamespaceId;
+			this.Options = Options;
+			this.Cache = Cache;
+			this.HashToNode = new Dictionary<IoHash, NodeInfo>();
+		}
+
+		/// <summary>
+		/// Gets data for the node with the given hash
+		/// </summary>
+		/// <param name="Hash">Hash of the node to return data for</param>
+		/// <returns>The node data</returns>
+		public ValueTask<ReadOnlyMemory<byte>> GetDataAsync(IoHash Hash)
+		{
+			return GetDataAsync(HashToNode[Hash]);
+		}
+
+		/// <summary>
+		/// Gets the data for a given node
+		/// </summary>
+		/// <param name="Node">The node to get the data for</param>
+		/// <returns>The node data</returns>
+		async ValueTask<ReadOnlyMemory<byte>> GetDataAsync(NodeInfo Node)
+		{
+			if (Node.Blob == null)
+			{
+				if (Node.Export == null)
+				{
+					return Node.Data;
+				}
+				else
+				{
+					return DecodePacket(RootCacheKey, Node.Data, Node.Export.Packet).Slice(Node.Export.Offset, Node.Export.Length);
+				}
+			}
+			else
+			{
+				await MountBlobAsync(Node.Blob);
+				Debug.Assert(Node.Export != null);
+				BundleObject Object = await GetObjectAsync(Node.Blob.Hash);
+				return DecodePacket(Node.Blob.Hash, Object.Data, Node.Export.Packet).Slice(Node.Export.Offset, Node.Export.Length);
+			}
+		}
+
+		/// <summary>
+		/// Write a node's data
+		/// </summary>
+		/// <param name="Data">The node data</param>
+		/// <param name="References">Hashes for referenced nodes</param>
+		/// <returns>Hash of the data</returns>
+		public IoHash WriteNode(ReadOnlyMemory<byte> Data, IEnumerable<IoHash> References)
+		{
+			IoHash Hash = IoHash.Compute(Data.Span);
+			NodeInfo[] NodeReferences = References.Select(x => HashToNode[x]).Distinct().OrderBy(x => x.Hash).ToArray();
+
+			NodeInfo? Node;
+			if (!HashToNode.TryGetValue(Hash, out Node))
+			{
+				int Rank = 0;
+				foreach (NodeInfo Reference in NodeReferences)
+				{
+					Rank = Math.Max(Rank, Reference.Rank + 1);
+				}
+
+				Node = new NodeInfo(Hash, Rank, Data.Length);
+				HashToNode[Hash] = Node;
+			}
+
+			Node.SetStandalone(Data, NodeReferences);
+			return Hash;
+		}
+
+		async ValueTask MountBlobAsync(BlobInfo Blob)
+		{
+			if(!Blob.Mounted)
+			{
+				BundleObject Object = await StorageClient.ReadObjectAsync<BundleObject>(NamespaceId, Blob.Hash);
+				RegisterObject(Blob, Object, ReadOnlyMemory<byte>.Empty);
+				Blob.Mounted = true;
+			}
+		}
+
+		/// <summary>
+		/// Reads an object data from the store
+		/// </summary>
+		/// <param name="Hash">Hash of the object</param>
+		/// <returns>The parsed object</returns>
+		async ValueTask<BundleObject> GetObjectAsync(IoHash Hash)
+		{
+			BundleObject Object;
+			if (Cache == null || !Cache.TryGetValue(Hash, out Object))
+			{
+				ReadOnlyMemory<byte> Data = await StorageClient.ReadBlobToMemoryAsync(NamespaceId, Hash);
+				Object = CbSerializer.Deserialize<BundleObject>(Data);
+
+				if (Cache != null)
+				{
+					using (ICacheEntry Entry = Cache.CreateEntry(Hash))
+					{
+						Entry.SetValue(Object);
+						Entry.SetSize(Data.Length);
+					}
+				}
+			}
+			return Object;
+		}
+
+		/// <summary>
+		/// Gets a decoded block from the store
+		/// </summary>
+		/// <param name="BlobKey">Key for the blob</param>
+		/// <param name="Data">Raw blob data</param>
+		/// <param name="Packet">The decoded block location and size</param>
+		/// <returns>The decoded data</returns>
+		ReadOnlyMemory<byte> DecodePacket(object BlobKey, ReadOnlyMemory<byte> Data, BundleCompressionPacket Packet)
+		{
+			object CacheKey = (BlobKey, Packet.Offset);
+
+			ReadOnlyMemory<byte> DecodedData;
+			if (Cache == null || !Cache.TryGetValue(CacheKey, out DecodedData))
+			{
+				byte[] DecodedBuffer = new byte[Packet.DecodedLength];
+				LZ4Codec.Decode(Data.Span.Slice(Packet.Offset, Packet.EncodedLength), DecodedBuffer);
+
+				DecodedData = DecodedBuffer;
+
+				if (Cache != null)
+				{
+					using (ICacheEntry Entry = Cache.CreateEntry(CacheKey))
+					{
+						Entry.SetValue(DecodedData);
+						Entry.SetSize(DecodedData.Length);
+					}
+				}
+			}
+
+			return DecodedData;
+		}
+
+		/// <summary>
+		/// Reads a ref into the bundle
+		/// </summary>
+		/// <param name="BucketId">Bucket containing the ref</param>
+		/// <param name="RefId">Reference id</param>
+		public virtual async Task ReadAsync(BucketId BucketId, RefId RefId)
+		{
+			BundleObject Object = await StorageClient.GetRefAsync<BundleObject>(NamespaceId, BucketId, RefId);
+			RegisterObject(null, Object, Object.Data);
+			RootHash = Object.Exports[Object.Exports.Count - 1].Hash;
+		}
+
+		NodeInfo FindOrAddNode(IoHash Hash, int Rank, int Cost)
+		{
+			NodeInfo? Node;
+			if (!HashToNode.TryGetValue(Hash, out Node))
+			{
+				Node = new NodeInfo(Hash, Rank, Cost);
+				HashToNode[Hash] = Node;
+			}
+			return Node;
+		}
+
+		void RegisterObject(BlobInfo? Blob, BundleObject Object, ReadOnlyMemory<byte> Data)
+		{
+			List<NodeInfo> Nodes = new List<NodeInfo>();
+			foreach (BundleExport Export in Object.Exports)
+			{
+				NodeInfo Node = FindOrAddNode(Export.Hash, Export.Rank, Export.Cost);
+				Nodes.Add(Node);
+			}
+			foreach (BundleImportObject ImportObject in Object.ImportObjects)
+			{
+				BlobInfo? ImportBlob;
+				if (!HashToBlob.TryGetValue(ImportObject.Object.Hash, out ImportBlob))
+				{
+					ImportBlob = new BlobInfo(ImportObject.Object.Hash, ImportObject.CreationTimeUtc, ImportObject.TotalCost);
+					HashToBlob[ImportObject.Object.Hash] = ImportBlob;
+				}
+				foreach (BundleImport Import in ImportObject.Imports)
+				{
+					NodeInfo Node = FindOrAddNode(Import.Hash, Import.Rank, Import.Cost);
+					if(Node.Blob == null)
+					{
+						Node.SetImport(ImportBlob);
+					}
+					Nodes.Add(Node);
+				}
+			}
+			for(int Idx = 0; Idx < Object.Exports.Count; Idx++)
+			{
+				NodeInfo Node = Nodes[Idx];
+				Node.SetExport(Blob, Object.Exports[Idx], Data, Object.Exports[Idx].References.Select(x => Nodes[x]).ToArray());
+			}
+		}
+
+		/// <summary>
+		/// Persist the bundle to storage
+		/// </summary>
+		/// <param name="BucketId"></param>
+		/// <param name="RefId"></param>
+		/// <param name="Compact"></param>
+		/// <param name="UtcNow"></param>
+		public virtual async Task WriteAsync(BucketId BucketId, RefId RefId, bool Compact, DateTime UtcNow)
+		{
+			NodeInfo RootNode = HashToNode[RootHash];
+
+			// If we're compacting the output, see if there are any other blocks we should merge
+			if (Compact)
+			{
+				// Find the live set of objects
+				HashSet<NodeInfo> LiveNodes = new HashSet<NodeInfo>();
+				await FindFullLiveSetAsync(RootNode, LiveNodes);
+
+				// Find the live set of blobs
+				HashSet<BlobInfo> LiveBlobs = new HashSet<BlobInfo>();
+				foreach (NodeInfo LiveNode in LiveNodes)
+				{
+					if (LiveNode.Blob != null)
+					{
+						LiveBlobs.Add(LiveNode.Blob);
+					}
+				}
+
+				// Clear the reference list on each one
+				foreach (BlobInfo LiveBlob in LiveBlobs)
+				{
+					LiveBlob.ReferencedBy.Clear();
+				}
+
+				// Find the first reference from each node to other blobs
+				foreach (NodeInfo NodeInfo in LiveNodes)
+				{
+					if (NodeInfo.Rank > 0 && NodeInfo.Blob != null)
+					{
+						foreach (NodeInfo ReferencedNode in NodeInfo.References!)
+						{
+							BlobInfo? ReferencedBlob = ReferencedNode.Blob;
+							if (ReferencedBlob != null)
+							{
+								ReferencedBlob.ReferencedBy.Add(NodeInfo.Blob);
+							}
+						}
+					}
+				}
+
+				// Clear out the reference data on each blob
+				HashSet<BlobInfo> VisitedBlobs = new HashSet<BlobInfo>();
+				foreach (BlobInfo LiveBlob in LiveBlobs)
+				{
+					ExpandReferences(LiveBlob, VisitedBlobs);
+				}
+
+				// Find the total cost of all the current blobs, then loop through the blobs trying to find a more optimal arrangement
+				List<NodeInfo> SortedLiveNodes = LiveNodes.OrderByDescending(x => x.Rank).ThenByDescending(x => x.Hash).ToList();
+				for (; ; )
+				{
+					BlobInfo? MergeBlob = null;
+					double MergeCost = GetNewCostHeuristic(SortedLiveNodes, new HashSet<BlobInfo>(), UtcNow);
+
+					foreach (BlobInfo LiveBlob in LiveBlobs)
+					{
+						double PotentialCost = GetNewCostHeuristic(SortedLiveNodes, LiveBlob.ReferencedBy, UtcNow);
+						if (PotentialCost < MergeCost)
+						{
+							MergeBlob = LiveBlob;
+							MergeCost = PotentialCost;
+						}
+					}
+
+					// Bail out if we didn't find anything to merge
+					if (MergeBlob == null)
+					{
+						break;
+					}
+
+					// Make sure this blob is resident (it may not have been loaded above if it's a leaf blob)
+					await MountBlobAsync(MergeBlob);
+
+					// Find all the live nodes in this blob
+					List<NodeInfo> MergeNodes = LiveNodes.Where(x => x.Blob == MergeBlob).ToList();
+					foreach (NodeInfo MergeNode in MergeNodes)
+					{
+						ReadOnlyMemory<byte> Data = await GetDataAsync(MergeNode);
+						MergeNode.SetStandalone(Data, MergeNode.References!);
+					}
+
+					// Remove the merged blob from the live list
+					LiveBlobs.Remove(MergeBlob);
+				}
+			}
+
+			// Find all the modified nodes
+			HashSet<NodeInfo> NewNodes = new HashSet<NodeInfo>();
+			FindModifiedLiveSet(RootNode, NewNodes);
+
+			// Write all the new blobs. Sort by rank and go from end to start, updating the list as we go so that we can figure out imports correctly.
+			NodeInfo[] WriteNodes = NewNodes.Where(x => x.Blob == null).OrderBy(x => x.Rank).ThenBy(x => x.Hash).ToArray();
+
+			int MinIdx = 0;
+			int NextBlobCost = WriteNodes[0].Cost;
+
+			for (int MaxIdx = 1; MaxIdx < WriteNodes.Length; MaxIdx++)
+			{
+				NextBlobCost += WriteNodes[MaxIdx].Cost;
+				if (NextBlobCost > Options.MaxBlobSize)
+				{
+					await WriteObjectAsync(WriteNodes.Slice(MinIdx, MaxIdx - MinIdx), UtcNow);
+					MinIdx = MaxIdx;
+					NextBlobCost = WriteNodes[MaxIdx].Cost;
+				}
+			}
+
+			if (NextBlobCost > Options.MaxInlineBlobSize && MinIdx + 1 < WriteNodes.Length)
+			{
+				await WriteObjectAsync(WriteNodes.Slice(MinIdx, WriteNodes.Length - 1 - MinIdx), UtcNow);
+				MinIdx = WriteNodes.Length - 1;
+			}
+
+			// Write the final ref
+			await WriteRefAsync(WriteNodes.Slice(MinIdx), BucketId, RefId, UtcNow);
+		}
+
+		private void FindModifiedLiveSet(NodeInfo Node, HashSet<NodeInfo> Nodes)
+		{
+			if (Nodes.Add(Node) && Node.Rank > 0)
+			{
+				foreach (NodeInfo Reference in Node.References!)
+				{
+					if (Reference.Blob == null)
+					{
+						FindModifiedLiveSet(Reference, Nodes);
+					}
+				}
+			}
+		}
+
+		async Task FindFullLiveSetAsync(NodeInfo Node, HashSet<NodeInfo> LiveNodes)
+		{
+			if (LiveNodes.Add(Node) && Node.Rank > 0)
+			{
+				if (Node.Blob != null)
+				{
+					await MountBlobAsync(Node.Blob);
+				}
+				foreach (NodeInfo Reference in Node.References!)
+				{
+					await FindFullLiveSetAsync(Reference, LiveNodes);
+				}
+			}
+		}
+
+		void ExpandReferences(BlobInfo Blob, HashSet<BlobInfo> VisitedBlobs)
+		{
+			if (VisitedBlobs.Add(Blob))
+			{
+				BlobInfo[] ReferencedBy = Blob.ReferencedBy.ToArray();
+				foreach (BlobInfo ReferencedBlob in ReferencedBy)
+				{
+					ExpandReferences(ReferencedBlob, VisitedBlobs);
+					Blob.ReferencedBy.UnionWith(ReferencedBlob.ReferencedBy);
+				}
+				Blob.ReferencedBy.Add(Blob);
+			}
+		}
+
+		async Task WriteRefAsync(IReadOnlyList<NodeInfo> Nodes, BucketId BucketId, RefId RefId, DateTime UtcNow)
+		{
+			BundleObject Object = await CreateObjectAsync(Nodes, UtcNow);
+			await StorageClient.SetRefAsync(NamespaceId, BucketId, RefId, Object);
+
+			foreach (NodeInfo Node in Nodes)
+			{
+				Node.Blob = null;
+				Node.Data = Object.Data;
+			}
+
+			RootCacheKey = new object();
+		}
+
+		async Task WriteObjectAsync(IReadOnlyList<NodeInfo> Nodes, DateTime UtcNow)
+		{
+			BundleObject Object = await CreateObjectAsync(Nodes, UtcNow);
+			IoHash Hash = await StorageClient.WriteObjectAsync(NamespaceId, Object);
+
+			BlobInfo Blob = new BlobInfo(Hash, Object.CreationTimeUtc, Object.Exports.Sum(x => x.Cost));
+			foreach (NodeInfo Node in Nodes)
+			{
+				Node.Blob = Blob;
+				Node.Data = ReadOnlyMemory<byte>.Empty;
+			}
+		}
+
+		async Task<BundleObject> CreateObjectAsync(IReadOnlyList<NodeInfo> Nodes, DateTime UtcNow)
+		{
+			// Find node indices for all the export
+			Dictionary<NodeInfo, int> NodeToIndex = new Dictionary<NodeInfo, int>();
+			foreach (NodeInfo Node in Nodes)
+			{
+				NodeToIndex[Node] = NodeToIndex.Count;
+			}
+
+			// Create the new object
+			BundleObject Object = new BundleObject();
+			Object.CreationTimeUtc = UtcNow;
+
+			// Find all the imported nodes
+			HashSet<NodeInfo> ImportedNodes = new HashSet<NodeInfo>();
+			foreach (NodeInfo Node in Nodes)
+			{
+				ImportedNodes.UnionWith(Node.References!);
+			}
+			ImportedNodes.ExceptWith(Nodes);
+
+			// Find all the imports
+			foreach (IGrouping<BlobInfo, NodeInfo> ImportGroup in ImportedNodes.GroupBy(x => x.Blob!))
+			{
+				BlobInfo Blob = ImportGroup.Key;
+
+				NodeInfo[] ImportNodes = ImportGroup.ToArray();
+				foreach (NodeInfo ImportNode in ImportNodes)
+				{
+					NodeToIndex[ImportNode] = NodeToIndex.Count;
+				}
+
+				List<BundleImport> Imports = ImportGroup.Select(x => new BundleImport(x.Hash, x.Rank, x.Cost)).ToList();
+				BundleImportObject ImportObject = new BundleImportObject(Blob.CreationTimeUtc, new CbObjectAttachment(Blob.Hash), Blob.TotalCost, Imports);
+				Object.ImportObjects.Add(ImportObject);
+			}
+
+			// Compress all the nodes into the encoded buffer
+			BundleCompressionPacket Packet = new BundleCompressionPacket(0);
+			foreach (NodeInfo Node in Nodes)
+			{
+				ReadOnlyMemory<byte> NodeData = await GetDataAsync(Node);
+
+				// If we can't fit this data into the current block, flush the contents of it first
+				if (Packet.DecodedLength + NodeData.Length > Options.MinCompressionPacketSize)
+				{
+					FlushPacket(Packet);
+					Packet = new BundleCompressionPacket(Packet.Offset + Packet.EncodedLength);
+				}
+
+				// Create the export for this node
+				int[] References = Node.References.Select(x => NodeToIndex[x]).ToArray();
+				Node.Export = new BundleExport(Node.Hash, Node.Rank, Node.Cost, Packet, Packet.DecodedLength, NodeData.Length, References);
+				Object.Exports.Add(Node.Export);
+
+				// Write out the new block
+				int Offset = Packet.EncodedLength;
+				if (NodeData.Length < Options.MinCompressionPacketSize)
+				{
+					int RequiredSize = Math.Max(Packet.DecodedLength + NodeData.Length, (int)(Options.MaxBlobSize * 1.2));
+					CreateFreeSpace(ref BlockBuffer, Packet.DecodedLength, RequiredSize);
+					NodeData.CopyTo(BlockBuffer.AsMemory(Packet.DecodedLength));
+					Packet.DecodedLength += NodeData.Length;
+				}
+				else
+				{
+					FlushPacket(Packet);
+					Packet = new BundleCompressionPacket(Packet.Offset + Packet.EncodedLength);
+				}
+			}
+			FlushPacket(Packet);
+
+			// Flush the data
+			Object.Data = EncodedBuffer.AsSpan(0, Packet.Offset + Packet.EncodedLength).ToArray();
+
+			return Object;
+		}
+
+		void FlushPacket(BundleCompressionPacket Packet)
+		{
+			if (Packet.DecodedLength > 0)
+			{
+				int MinFreeSpace = LZ4Codec.MaximumOutputSize(Packet.DecodedLength);
+				CreateFreeSpace(ref EncodedBuffer, Packet.Offset, Packet.Offset + MinFreeSpace);
+
+				ReadOnlySpan<byte> InputSpan = BlockBuffer.AsSpan(0, Packet.DecodedLength);
+				Span<byte> OutputSpan = EncodedBuffer.AsSpan(Packet.Offset);
+
+				Packet.EncodedLength = LZ4Codec.Encode(InputSpan, OutputSpan);
+			}
+		}
+
+		void CreateFreeSpace(ref byte[] Buffer, int UsedSize, int RequiredSize)
+		{
+			if (RequiredSize > Buffer.Length)
+			{
+				byte[] NewBuffer = new byte[RequiredSize];
+				Buffer.AsSpan(0, UsedSize).CopyTo(NewBuffer);
+				Buffer = NewBuffer;
+			}
+		}
+
+		/// <inheritdoc cref="GetCostHeuristic(int, TimeSpan)"/>
+		double GetCostHeuristic(BlobInfo Info, DateTime UtcNow) => GetCostHeuristic(Info.TotalCost, UtcNow - Info.CreationTimeUtc);
+
+		/// <summary>
+		/// Gets the cost of packaging the given set of nodes which are either not in a blob or are in a blob listed to be collapsed
+		/// </summary>
+		/// <param name="Nodes">The full list of nodes</param>
+		/// <param name="CollapseBlobs">Set of blobs that should be collapsed</param>
+		/// <param name="UtcNow">The current time</param>
+		/// <returns></returns>
+		double GetNewCostHeuristic(List<NodeInfo> Nodes, HashSet<BlobInfo> CollapseBlobs, DateTime UtcNow)
+		{
+			double Cost = 0.0;
+			foreach (BlobInfo CollapseBlob in CollapseBlobs)
+			{
+				Cost -= GetCostHeuristic(CollapseBlob, UtcNow);
+			}
+
+			int BlobSize = 0;
+			for (int Idx = Nodes.Count - 1; Idx >= 0; Idx--)
+			{
+				NodeInfo Node = Nodes[Idx];
+				if (Node.Blob == null || CollapseBlobs.Contains(Node.Blob))
+				{
+					if (BlobSize > 0 && BlobSize + Node.Data.Length > Options.MaxBlobSize)
+					{
+						Cost += GetCostHeuristic(BlobSize, TimeSpan.Zero);
+						BlobSize = 0;
+					}
+					BlobSize += Node.Data.Length;
+				}
+			}
+			Cost += GetCostHeuristic(BlobSize, TimeSpan.Zero);
+			return Cost;
+		}
+
+		/// <summary>
+		/// Heuristic which estimates the cost of a particular blob. This is used to compare scenarios of merging blobs to reduce download
+		/// size against keeping older blobs which a lot of agents already have.
+		/// </summary>
+		/// <param name="Size">Size of the blob</param>
+		/// <param name="Age">Age of the blob</param>
+		/// <returns>Heuristic for the cost of a blob</returns>
+		static double GetCostHeuristic(int Size, TimeSpan Age)
+		{
+			// Time overhead to starting a download
+			const double DownloadInit = 0.1;
+
+			// Download speed for agents, in bytes/sec
+			const double DownloadRate = 1024 * 1024;
+
+			// Probability of an agent having to download everything. Prevents bias against keeping a large number of files.
+			const double CleanSyncProbability = 0.2;
+
+			// Average length of time between agents having to update
+			TimeSpan AverageCoherence = TimeSpan.FromHours(4.0);
+
+			// Scale the age into a -1.0 -> 1.0 range around AverageCoherence
+			double ScaledAge = (AverageCoherence - Age).TotalSeconds / AverageCoherence.TotalSeconds;
+
+			// Get the probability of agents having to sync this blob based on its age. This is modeled as a logistic function (1 / (1 + e^-x))
+			// with value of 0.5 at AverageCoherence, and MaxInterval at zero.
+
+			// Find the scale factor for the 95% interval
+			//    1 / (1 + e^-x) = MaxInterval
+			//    e^-x = (1 / MaxInterval) - 1
+			//    x = -ln((1 / MaxInterval) - 1)
+			const double MaxInterval = 0.95;
+			double SigmoidScale = -Math.Log((1.0 / MaxInterval) - 1.0);
+
+			// Find the probability of having to sync this 
+			double Param = ScaledAge * SigmoidScale;
+			double Probability = 1.0 / (1.0 + Math.Exp(-Param));
+
+			// Scale the probability against having to do a full sync
+			Probability = CleanSyncProbability + (Probability * (1.0 - CleanSyncProbability));
+
+			// Compute the final cost estimate; the amount of time we expect agents to spend downloading the file
+			return Probability * (DownloadInit + (Size / DownloadRate));
+		}
+	}
+
+	/// <summary>
+	/// Derived version of <see cref="Bundle"/> supporting a concrete root node type
+	/// </summary>
+	/// <typeparam name="T">The root node type</typeparam>
+	public class Bundle<T> : Bundle where T : BundleNode
+	{
+		static BundleNodeFactory<T> Factory = CreateFactory();
+
+		/// <summary>
+		/// Object at the root of the bundle
+		/// </summary>
+		public T Root { get; private set; }
+
+		/// <summary>
+		/// Constructor
+		/// </summary>
+		public Bundle(IStorageClient StorageClient, NamespaceId NamespaceId, BundleOptions Options, IMemoryCache? Cache)
+			: base(StorageClient, NamespaceId, Options, Cache)
+		{
+			Root = Factory.CreateRoot(this);
+		}
+
+		/// <summary>
+		/// Create a factory for manipulating T instances
+		/// </summary>
+		/// <returns></returns>
+		static BundleNodeFactory<T> CreateFactory()
+		{
+			BundleNodeFactoryAttribute? Attribute = typeof(T).GetCustomAttribute<BundleNodeFactoryAttribute>();
+			if (Attribute == null)
+			{
+				throw new InvalidOperationException("Nodes used as the root of a bundle must have a BundleNodeFactoryAttribute");
+			}
+			return (BundleNodeFactory<T>)Activator.CreateInstance(Attribute.Type)!;
+		}
+
+		/// <inheritdoc/>
+		public override async Task ReadAsync(BucketId BucketId, RefId RefId)
+		{
+			await base.ReadAsync(BucketId, RefId);
+			ReadOnlyMemory<byte> Data = await GetDataAsync(RootHash);
+			Root = Factory.ParseRoot(this, RootHash, Data);
+		}
+
+		/// <inheritdoc/>
+		public override async Task WriteAsync(BucketId BucketId, RefId RefId, bool Compact, DateTime UtcNow)
+		{
+			RootHash = Root.Serialize();
+			await base.WriteAsync(BucketId, RefId, Compact, UtcNow);
+		}
+	}
+}
diff --git a/Engine/Source/Programs/Shared/EpicGames.Horde/Bundles/BundleNode.cs b/Engine/Source/Programs/Shared/EpicGames.Horde/Bundles/BundleNode.cs
new file mode 100644
--- /dev/null
+++ b/Engine/Source/Programs/Shared/EpicGames.Horde/Bundles/BundleNode.cs
@@ -0,0 +1,127 @@
+// Copyright Epic Games, Inc. All Rights Reserved.
+
+using EpicGames.Core;
+using System;
+using System.Collections.Generic;
+using System.Text;
+using System.Threading.Tasks;
+
+namespace EpicGames.Horde.Bundles.Nodes
+{
+	/// <summary>
+	/// Base class for nodes stored in a bundle
+	/// </summary>
+	public abstract class BundleNode
+	{
+		/// <summary>
+		/// The bundle that owns this node
+		/// </summary>
+		public Bundle Owner { get; }
+
+		/// <summary>
+		/// The node's parent. Each expanded instance of the node in the tree will have a unique object, so each will have a singular parent.
+		/// </summary>
+		protected BundleNode? Parent { get; set; }
+
+		/// <summary>
+		/// Hash of this data. Set to zero once modified.
+		/// </summary>
+		private IoHash Hash { get; set; }
+
+		/// <summary>
+		/// Constructor
+		/// </summary>
+		public BundleNode(Bundle Owner, BundleNode? Parent)
+			: this(Owner, Parent, IoHash.Zero)
+		{
+		}
+
+		/// <summary>
+		/// Constructor
+		/// </summary>
+		public BundleNode(Bundle Owner, BundleNode? Parent, IoHash Hash)
+		{
+			this.Owner = Owner;
+			this.Parent = Parent;
+			this.Hash = Hash;
+		}
+
+		/// <summary>
+		/// Checks if the current node is dirty
+		/// </summary>
+		/// <returns>True if the node is dirty</returns>
+		public bool IsDirty() => Hash == IoHash.Zero;
+
+		/// <summary>
+		/// Mark this node as dirty
+		/// </summary>
+		protected void MarkAsDirty()
+		{
+			if (!IsDirty())
+			{
+				Hash = IoHash.Zero;
+				Parent?.MarkAsDirty();
+			}
+		}
+
+		/// <summary>
+		/// Serialize this node to a buffer
+		/// </summary>
+		/// <returns>Hash of the serialized node</returns>
+		public IoHash Serialize()
+		{
+			if (IsDirty())
+			{
+				Hash = SerializeDirty();
+			}
+			return Hash;
+		}
+
+		/// <summary>
+		/// Serialize this node to the bundle
+		/// </summary>
+		/// <returns>Hash of the serialized node</returns>
+		protected abstract IoHash SerializeDirty();
+	}
+
+	/// <summary>
+	/// Attribute used to define a factory for a particular node type
+	/// </summary>
+	public class BundleNodeFactoryAttribute : Attribute
+	{
+		/// <summary>
+		/// The factory type. Should be derived from <see cref="BundleNodeFactory{T}"/>
+		/// </summary>
+		public Type Type { get; }
+
+		/// <summary>
+		/// Constructor
+		/// </summary>
+		public BundleNodeFactoryAttribute(Type Type) => this.Type = Type;
+	}
+
+	/// <summary>
+	/// Factory class for creating node types
+	/// </summary>
+	/// <typeparam name="T">The type of node returned</typeparam>
+	public abstract class BundleNodeFactory<T> where T : BundleNode
+	{
+		/// <summary>
+		/// Creates a default value for the factory
+		/// </summary>
+		/// <param name="Bundle"></param>
+		/// <returns>New instance of the node type</returns>
+		public abstract T CreateRoot(Bundle Bundle);
+
+		/// <summary>
+		/// Parse an object from data
+		/// </summary>
+		/// <param name="Bundle">The owning bundle</param>
+		/// <param name="Hash">Hash of the data</param>
+		/// <param name="Data"></param>
+		/// <returns>New instance of the node type</returns>
+		public abstract T ParseRoot(Bundle Bundle, IoHash Hash, ReadOnlyMemory<byte> Data);
+	}
+}
+
+
diff --git a/Engine/Source/Programs/Shared/EpicGames.Horde/Bundles/BundleObject.cs b/Engine/Source/Programs/Shared/EpicGames.Horde/Bundles/BundleObject.cs
new file mode 100644
--- /dev/null
+++ b/Engine/Source/Programs/Shared/EpicGames.Horde/Bundles/BundleObject.cs
@@ -0,0 +1,396 @@
+// Copyright Epic Games, Inc. All Rights Reserved.
+
+using EpicGames.Core;
+using EpicGames.Horde.Bundles.Nodes;
+using EpicGames.Serialization;
+using K4os.Compression.LZ4;
+using Microsoft.Extensions.Logging;
+using System;
+using System.Buffers.Binary;
+using System.Collections;
+using System.Collections.Generic;
+using System.Diagnostics;
+using System.Diagnostics.CodeAnalysis;
+using System.IO;
+using System.IO.Compression;
+using System.Linq;
+using System.Security.Cryptography;
+using System.Text;
+using System.Text.Json;
+using System.Threading.Tasks;
+
+namespace EpicGames.Horde.Bundles
+{
+	/// <summary>
+	/// Blob within the bundle.
+	/// </summary>
+	public class BundleObject
+	{
+		/// <summary>
+		/// Current schema version
+		/// </summary>
+		public const int CurrentSchema = 1;
+
+		/// <summary>
+		/// Schema version
+		/// </summary>
+		[CbField("schema")]
+		public int Schema { get; set; } = CurrentSchema;
+
+		/// <summary>
+		/// Time that this object was minted. Used to evaluate how likely it will be that a client will already have it.
+		/// </summary>
+		[CbField("time")]
+		public DateTime CreationTimeUtc { get; set; }
+
+		/// <summary>
+		/// Other objects that are referenced.
+		/// </summary>
+		[CbField("imports")]
+		public List<BundleImportObject> ImportObjects { get; set; } = new List<BundleImportObject>();
+
+		/// <summary>
+		/// Exported items in the data 
+		/// </summary>
+		[CbField("exports")]
+		[CbConverter(typeof(BundleExportListConverter))]
+		public List<BundleExport> Exports { get; set; } = new List<BundleExport>();
+
+		/// <summary>
+		/// Raw data for this exported entries in this blob.
+		/// </summary>
+		[CbField("data")]
+		public ReadOnlyMemory<byte> Data { get; set; }
+	}
+
+	/// <summary>
+	/// Reference to another tree pack object
+	/// </summary>	
+	[DebuggerDisplay("{Object}")]
+	public class BundleImportObject
+	{
+		/// <summary>
+		/// Time that the blob was created.
+		/// </summary>
+		[CbField("time")]
+		public DateTime CreationTimeUtc { get; set; }
+
+		/// <summary>
+		/// Imported object.
+		/// </summary>
+		[CbField("object")]
+		public CbObjectAttachment Object { get; set; }
+
+		/// <summary>
+		/// Sum of the costs of nodes strored in this import.
+		/// </summary>
+		[CbField("cost")]
+		public int TotalCost { get; set; }
+
+		/// <summary>
+		/// List of nodes imported from this object, and their sizes.
+		/// </summary>
+		[CbField("imports")]
+		[CbConverter(typeof(BundleImportListConverter))]
+		public List<BundleImport> Imports { get; set; } = new List<BundleImport>();
+
+		/// <summary>
+		/// Default constructor, for serialization
+		/// </summary>
+		private BundleImportObject()
+		{
+		}
+
+		/// <summary>
+		/// Constructor
+		/// </summary>
+		public BundleImportObject(DateTime CreationTimeUtc, CbObjectAttachment Object, int TotalCost, List<BundleImport> Imports)
+		{
+			this.CreationTimeUtc = CreationTimeUtc;
+			this.Object = Object;
+			this.TotalCost = TotalCost;
+			this.Imports = Imports;
+		}
+	}
+
+	/// <summary>
+	/// Entry for a node imported from another object
+	/// </summary>
+	public class BundleImport
+	{
+		/// <summary>
+		/// Hash of the node
+		/// </summary>
+		public IoHash Hash { get; }
+
+		/// <summary>
+		/// The node's rank, ie. its distance from a leaf node (where leaf nodes have a value of zero).
+		/// </summary>
+		public int Rank { get; }
+
+		/// <summary>
+		/// Decompressed size of this node
+		/// </summary>
+		public int Cost { get; }
+
+		/// <summary>
+		/// Constructor
+		/// </summary>
+		public BundleImport(IoHash Hash, int Rank, int Cost)
+		{
+			this.Hash = Hash;
+			this.Rank = Rank;
+			this.Cost = Cost;
+		}
+	}
+
+	class BundleImportListConverter : CbConverterBase<List<BundleImport>>
+	{
+		public override List<BundleImport> Read(CbField Field) => Deserialize(Field.AsBinary().Span);
+		public override void Write(CbWriter Writer, List<BundleImport> Value) => Writer.WriteBinaryArrayValue(Serialize(Value));
+		public override void WriteNamed(CbWriter Writer, Utf8String Name, List<BundleImport> Value) => Writer.WriteBinaryArray(Name, Serialize(Value));
+
+		static List<BundleImport> Deserialize(ReadOnlySpan<byte> Span)
+		{
+			List<BundleImport> Imports = new List<BundleImport>();
+			while (Span.Length > 0)
+			{
+				IoHash Hash = new IoHash(Span);
+				Span = Span.Slice(IoHash.NumBytes);
+
+				int Rank = (int)VarInt.Read(Span, out int RankBytes);
+				Span = Span.Slice(RankBytes);
+
+				int Cost = (int)VarInt.Read(Span, out int CostBytes);
+				Span = Span.Slice(CostBytes);
+
+				Imports.Add(new BundleImport(Hash, Rank, Cost));
+			}
+			return Imports;
+		}
+
+		static byte[] Serialize(List<BundleImport> Imports)
+		{
+			byte[] Data = new byte[Imports.Sum(x => IoHash.NumBytes + VarInt.Measure(x.Rank) + VarInt.Measure(x.Cost))];
+
+			Span<byte> Span = Data;
+			foreach (BundleImport Import in Imports)
+			{
+				Import.Hash.CopyTo(Span);
+				Span = Span.Slice(IoHash.NumBytes);
+
+				int RankBytes = VarInt.Write(Span, Import.Rank);
+				Span = Span.Slice(RankBytes);
+
+				int CostBytes = VarInt.Write(Span, Import.Cost);
+				Span = Span.Slice(CostBytes);
+			}
+
+			return Data;
+		}
+	}
+
+	/// <summary>
+	/// Descriptor for a compression packet
+	/// </summary>
+	public class BundleCompressionPacket
+	{
+		/// <summary>
+		/// Offset of the packet within the <see cref="BundleObject"/>'s data array
+		/// </summary>
+		public int Offset { get; }
+
+		/// <summary>
+		/// Encoded length of the packet
+		/// </summary>
+		public int EncodedLength { get; set; }
+
+		/// <summary>
+		/// Decoded length of the packet
+		/// </summary>
+		public int DecodedLength { get; set; }
+
+		/// <summary>
+		/// Constructor
+		/// </summary>
+		/// <param name="Offset"></param>
+		public BundleCompressionPacket(int Offset)
+		{
+			this.Offset = Offset;
+		}
+	}
+
+	/// <summary>
+	/// Entry for a node exported from an object
+	/// </summary>
+	public class BundleExport : BundleImport
+	{
+		/// <summary>
+		/// Packet containing the node
+		/// </summary>
+		public BundleCompressionPacket Packet { get; }
+
+		/// <summary>
+		/// Offset of the node within the decompressed packet
+		/// </summary>
+		public int Offset { get; }
+
+		/// <summary>
+		/// Length of the node within the decompressed packet
+		/// </summary>
+		public int Length { get; }
+
+		/// <summary>
+		/// References to other nodes.
+		/// </summary>
+		public int[] References { get; }
+
+		/// <summary>
+		/// Constructor
+		/// </summary>
+		public BundleExport(IoHash Hash, int Rank, int Cost, BundleCompressionPacket Packet, int Offset, int Length, int[] References)
+			: base(Hash, Rank, Cost)
+		{
+			this.Packet = Packet;
+			this.Offset = Offset;
+			this.Length = Length;
+			this.References = References;
+		}
+	}
+
+	class BundleExportListConverter : CbConverterBase<List<BundleExport>>
+	{
+		public override List<BundleExport> Read(CbField Field) => Deserialize(Field.AsBinary().Span);
+		public override void Write(CbWriter Writer, List<BundleExport> Value) => Writer.WriteBinaryArrayValue(Serialize(Value));
+		public override void WriteNamed(CbWriter Writer, Utf8String Name, List<BundleExport> Value) => Writer.WriteBinaryArray(Name, Serialize(Value));
+
+		static List<BundleExport> Deserialize(ReadOnlySpan<byte> Span)
+		{
+			List<BundleExport> Exports = new List<BundleExport>();
+
+			int PacketOffset = 0;
+			while (Span.Length > 0)
+			{
+				BundleCompressionPacket Packet = new BundleCompressionPacket(PacketOffset);
+
+				Packet.EncodedLength = (int)VarInt.Read(Span, out int EncodedLengthBytes);
+				Span = Span.Slice(EncodedLengthBytes);
+
+				Packet.DecodedLength = (int)VarInt.Read(Span, out int DecodedLengthBytes);
+				Span = Span.Slice(DecodedLengthBytes);
+
+				PacketOffset += Packet.EncodedLength;
+
+				int Offset = 0;
+				while (Offset < Packet.DecodedLength)
+				{
+					IoHash Hash = new IoHash(Span);
+					Span = Span.Slice(IoHash.NumBytes);
+
+					int Rank = (int)VarInt.Read(Span, out int RankBytes);
+					Span = Span.Slice(RankBytes);
+
+					int Cost = (int)VarInt.Read(Span, out int CostBytes);
+					Span = Span.Slice(CostBytes);
+
+					int Length = (int)VarInt.Read(Span, out int LengthBytes);
+					Span = Span.Slice(LengthBytes);
+
+					int NumReferences = (int)VarInt.Read(Span, out int NumReferencesBytes);
+					Span = Span.Slice(NumReferencesBytes);
+
+					int[] References = new int[NumReferences];
+					for (int ReferenceIdx = 0; ReferenceIdx < NumReferences; ReferenceIdx++)
+					{
+						References[ReferenceIdx] = (int)VarInt.Read(Span, out int ReferenceBytes);
+						Span = Span.Slice(ReferenceBytes);
+					}
+
+					Exports.Add(new BundleExport(Hash, Rank, Cost, Packet, Offset, Length, References));
+					Offset += Length;
+				}
+			}
+
+			return Exports;
+		}
+
+		static byte[] Serialize(List<BundleExport> Exports)
+		{
+			int Length = Measure(Exports);
+
+			byte[] Data = new byte[Length];
+			Span<byte> Span = Data;
+
+			BundleCompressionPacket? PrevPacket = null;
+			foreach (BundleExport Export in Exports)
+			{
+				BundleCompressionPacket Packet = Export.Packet;
+				if (PrevPacket == null || Packet.Offset != PrevPacket.Offset)
+				{
+					CheckPacketSequence(PrevPacket, Packet);
+
+					int EncodedLengthBytes = VarInt.Write(Span, Packet.EncodedLength);
+					Span = Span.Slice(EncodedLengthBytes);
+
+					int DecodedLengthBytes = VarInt.Write(Span, Packet.DecodedLength);
+					Span = Span.Slice(DecodedLengthBytes);
+
+					PrevPacket = Packet;
+				}
+
+				Export.Hash.CopyTo(Span);
+				Span = Span.Slice(IoHash.NumBytes);
+
+				int RankBytes = VarInt.Write(Span, Export.Rank);
+				Span = Span.Slice(RankBytes);
+
+				int CostBytes = VarInt.Write(Span, Export.Cost);
+				Span = Span.Slice(CostBytes);
+
+				int LengthBytes = VarInt.Write(Span, Export.Length);
+				Span = Span.Slice(LengthBytes);
+
+				int NumReferencesBytes = VarInt.Write(Span, Export.References.Length);
+				Span = Span.Slice(NumReferencesBytes);
+
+				foreach (int Reference in Export.References)
+				{
+					int ReferenceBytes = VarInt.Write(Span, Reference);
+					Span = Span.Slice(ReferenceBytes);
+				}
+			}
+
+			Debug.Assert(Span.Length == 0);
+			return Data;
+		}
+
+		static int Measure(List<BundleExport> Exports)
+		{
+			int Length = 0;
+
+			BundleCompressionPacket? PrevPacket = null;
+			foreach (BundleExport Export in Exports)
+			{
+				BundleCompressionPacket Packet = Export.Packet;
+				if (PrevPacket == null || Packet.Offset != PrevPacket.Offset)
+				{
+					CheckPacketSequence(PrevPacket, Packet);
+					Length += VarInt.Measure(Packet.EncodedLength) + VarInt.Measure(Packet.DecodedLength);
+					PrevPacket = Packet;
+				}
+				Length += IoHash.NumBytes + VarInt.Measure(Export.Rank) + VarInt.Measure(Export.Cost) + VarInt.Measure(Export.Length) + VarInt.Measure(Export.References.Length) + Export.References.Sum(x => VarInt.Measure(x));
+			}
+
+			return Length;
+		}
+
+		static void CheckPacketSequence(BundleCompressionPacket? PrevPacket, BundleCompressionPacket Packet)
+		{
+			int ExpectedOffset = (PrevPacket == null) ? 0 : (PrevPacket.Offset + PrevPacket.EncodedLength);
+			if (Packet.Offset != ExpectedOffset)
+			{
+				throw new InvalidOperationException("Bundle compression packets are not sequential");
+			}
+		}
+	}
+}
diff --git a/Engine/Source/Programs/Shared/EpicGames.Horde/Bundles/Nodes/ChunkNode.cs b/Engine/Source/Programs/Shared/EpicGames.Horde/Bundles/Nodes/ChunkNode.cs
new file mode 100644
--- /dev/null
+++ b/Engine/Source/Programs/Shared/EpicGames.Horde/Bundles/Nodes/ChunkNode.cs
@@ -0,0 +1,432 @@
+// Copyright Epic Games, Inc. All Rights Reserved.
+
+using EpicGames.Core;
+using EpicGames.Serialization;
+using System;
+using System.Collections.Generic;
+using System.Diagnostics;
+using System.IO;
+using System.Linq;
+using System.Text;
+using System.Threading.Tasks;
+
+namespace EpicGames.Horde.Bundles.Nodes
+{
+	/// <summary>
+	/// Options for creating chunk nodes
+	/// </summary>
+	public class ChunkOptions
+	{
+		/// <summary>
+		/// Options for creating leaf nodes
+		/// </summary>
+		public TypedChunkOptions LeafOptions { get; set; }
+
+		/// <summary>
+		/// Options for creating interior nodes
+		/// </summary>
+		public TypedChunkOptions InteriorOptions { get; set; }
+
+		/// <summary>
+		/// Constructor
+		/// </summary>
+		public ChunkOptions()
+		{
+			LeafOptions = new TypedChunkOptions(32 * 1024, 256 * 1024, 64 * 1024);
+			InteriorOptions = new TypedChunkOptions(32 * 1024, 256 * 1024, 64 * 1024);
+		}
+	}
+
+	/// <summary>
+	/// Options for creating a specific type of chunk nodes
+	/// </summary>
+	public class TypedChunkOptions
+	{
+		/// <summary>
+		/// Minimum chunk size
+		/// </summary>
+		public int MinSize { get; set; }
+
+		/// <summary>
+		/// Maximum chunk size. Chunks will be split on this boundary if another match is not found.
+		/// </summary>
+		public int MaxSize { get; set; }
+
+		/// <summary>
+		/// Target chunk size for content-slicing
+		/// </summary>
+		public int TargetSize { get; set; }
+
+		/// <summary>
+		/// Default constructor
+		/// </summary>
+		public TypedChunkOptions(int MinSize, int MaxSize, int TargetSize)
+		{
+			this.MinSize = MinSize;
+			this.MaxSize = MaxSize;
+			this.TargetSize = TargetSize;
+		}
+	}
+
+	/// <summary>
+	/// Representation of a data stream, split into chunks along content-aware boundaries using a rolling hash (<see cref="BuzHash"/>).
+	/// Chunks are pushed into a tree hierarchy as data is appended to the root, with nodes of the tree also split along content-aware boundaries with <see cref="IoHash.NumBytes"/> granularity.
+	/// Once a chunk has been written to storage, it is treated as immutable.
+	/// </summary>
+	[BundleNodeFactory(typeof(ChunkNodeFactory))]
+	public sealed class ChunkNode : BundleNode
+	{
+		const byte TypeId = (byte)'c';
+
+		internal int Depth { get; private set; }
+		internal ReadOnlyMemory<byte> Payload { get; private set; }
+		ReadOnlyMemory<byte> Data; // Full serialized data, including type id and header fields.
+
+		// In-memory state
+		uint Hash;
+		byte[]? WriteBuffer; // Null if the node is complete and does not support writes
+		List<ChunkNode?>? ChildNodes;
+
+		/// <summary>
+		/// Constructor
+		/// </summary>
+		public ChunkNode(Bundle Owner, ChunkOptions Options)
+			: this(Owner, null, 0, Options)
+		{
+		}
+
+		/// <summary>
+		/// Create a writable chunk node
+		/// </summary>
+		private ChunkNode(Bundle Owner, BundleNode? Parent, int Depth, ChunkOptions Options)
+			: base(Owner, Parent)
+		{
+			MakeWritable(Depth, Options);
+		}
+
+		/// <summary>
+		/// Create a chunk node from deserialized data
+		/// </summary>
+		public ChunkNode(Bundle Owner, BundleNode? Parent, IoHash Hash, ReadOnlyMemory<byte> Data)
+			: base(Owner, Parent, Hash)
+		{
+			ReadOnlySpan<byte> Span = Data.Span;
+			if (Span[0] != TypeId)
+			{
+				throw new InvalidDataException("Invalid type id");
+			}
+
+			Depth = (int)VarInt.Read(Span.Slice(1), out int DepthBytes);
+			int HeaderLength = 1 + DepthBytes;
+
+			Payload = Data.Slice(HeaderLength);
+			this.Data = Data;
+
+			if (Depth > 0)
+			{
+				ChildNodes = new List<ChunkNode?>(Payload.Length / IoHash.NumBytes);
+				ChildNodes.AddRange(Enumerable.Repeat<ChunkNode?>(null, Payload.Length / IoHash.NumBytes));
+			}
+		}
+
+		private void MakeWritable(int Depth, ChunkOptions Options)
+		{
+			int WriteBufferSize = 1 + VarInt.Measure(Depth);
+			if (Depth == 0)
+			{
+				WriteBufferSize += Options.LeafOptions.MaxSize;
+			}
+			else
+			{
+				WriteBufferSize += Options.InteriorOptions.MaxSize;
+			}
+
+			byte[] WriteBuffer = new byte[WriteBufferSize];
+			WriteBuffer[0] = TypeId;
+			int DepthBytes = VarInt.Write(WriteBuffer.AsSpan(1), Depth);
+			int HeaderLength = 1 + DepthBytes;
+
+			this.Depth = Depth;
+			this.Payload = ReadOnlyMemory<byte>.Empty;
+			this.Data = WriteBuffer.AsMemory(0, HeaderLength);
+			this.Hash = 0;
+			this.WriteBuffer = WriteBuffer;
+			this.ChildNodes = (Depth > 0) ? new List<ChunkNode?>() : null;
+		}
+
+		/// <summary>
+		/// Returns the number of child nodes to this chunk
+		/// </summary>
+		public int GetChildCount() => (ChildNodes == null)? 0 : ChildNodes.Count;
+
+		/// <summary>
+		/// Gets a child node at the given index
+		/// </summary>
+		/// <param name="Index">Index of the child node</param>
+		/// <returns>Node containing data for the given child</returns>
+		public async ValueTask<ChunkNode> GetChildNodeAsync(int Index)
+		{
+			if(ChildNodes == null)
+			{
+				throw new InvalidOperationException("Node does not contain any children");
+			}
+
+			ChunkNode? ChildNode = ChildNodes[Index];
+			if (ChildNode == null)
+			{
+				IoHash Hash = new IoHash(Payload.Span.Slice(Index * IoHash.NumBytes));
+				ReadOnlyMemory<byte> Data = await Owner.GetDataAsync(Hash);
+				ChildNode = new ChunkNode(Owner, this, Hash, Data);
+				ChildNodes[Index] = ChildNode;
+			}
+			return ChildNode;
+		}
+
+		/// <summary>
+		/// Enumerate the children from this node
+		/// </summary>
+		/// <returns></returns>
+		public async IAsyncEnumerable<ChunkNode> GetChildNodesAsync()
+		{
+			for (int Idx = 0; Idx < GetChildCount(); Idx++)
+			{
+				yield return await GetChildNodeAsync(Idx);
+			}
+		}
+
+		/// <summary>
+		/// Serialize this node and its children into a byte array
+		/// </summary>
+		/// <returns>Array of data stored by the tree</returns>
+		public async Task<byte[]> ToByteArrayAsync()
+		{
+			using (MemoryStream Stream = new MemoryStream())
+			{
+				await CopyToAsync(Stream);
+				return Stream.ToArray();
+			}
+		}
+
+		/// <summary>
+		/// Copies the contents of this node and its children to the given output stream
+		/// </summary>
+		/// <param name="OutputStream">The output stream to receive the data</param>
+		public async Task CopyToAsync(Stream OutputStream)
+		{
+			if (Depth > 0)
+			{
+				for (int Idx = 0; Idx < GetChildCount(); Idx++)
+				{
+					ChunkNode Node = await GetChildNodeAsync(Idx);
+					await Node.CopyToAsync(OutputStream);
+				}
+			}
+			else
+			{
+				await OutputStream.WriteAsync(Payload);
+			}
+		}
+
+		/// <summary>
+		/// Append data to this chunk. Must only be called on the root node in a chunk tree.
+		/// </summary>
+		/// <param name="Input">The data to write</param>
+		/// <param name="Options">Settings for chunking the data</param>
+		public void Append(ReadOnlyMemory<byte> Input, ChunkOptions Options)
+		{
+			if (Parent is ChunkNode)
+			{
+				throw new InvalidOperationException("Data may only be appended to the root of a tree of chunk nodes.");
+			}
+
+			for (; ; )
+			{
+				// Append as much data as possible to the existing tree
+				Input = AppendToNode(Input, Options);
+				if (Input.IsEmpty)
+				{
+					break;
+				}
+
+				// Increase the height of the tree by pushing the contents of this node into a new child node
+				ChunkNode NewNode = new ChunkNode(Owner, this, IoHash.Zero, Data);
+				if (ChildNodes != null)
+				{
+					for (int Idx = 0; Idx < ChildNodes.Count; Idx++)
+					{
+						ChunkNode? ChildNode = ChildNodes[Idx];
+						if (ChildNode != null)
+						{
+							ChildNode.Parent = NewNode;
+							NewNode.ChildNodes![Idx] = ChildNode;
+						}
+					}
+				}
+
+				// Reinitialize the current node
+				MakeWritable(Depth + 1, Options);
+				AddChildNode(NewNode);
+			}
+		}
+
+		private ReadOnlyMemory<byte> AppendToNode(ReadOnlyMemory<byte> Data, ChunkOptions Options)
+		{
+			if (WriteBuffer == null || Data.Length == 0)
+			{
+				return Data;
+			}
+
+			if (Depth == 0)
+			{
+				return AppendToLeafNode(Data, Options);
+			}
+			else
+			{
+				return AppendToInteriorNode(Data, Options);
+			}
+		}
+
+		private ReadOnlyMemory<byte> AppendToLeafNode(ReadOnlyMemory<byte> NewData, ChunkOptions Options)
+		{
+			ReadOnlySpan<byte> NewSpan = NewData.Span;
+
+			uint HashThreshold = (uint)((1L << 32) / Options.LeafOptions.TargetSize);
+			for (int Length = 0; Length < NewData.Length; Length++)
+			{
+				Hash = BuzHash.Add(Hash, NewSpan[Length]);
+
+				int BlockSize = Payload.Length + Length;
+				if ((BlockSize >= Options.LeafOptions.MinSize && Hash < HashThreshold) || (BlockSize >= Options.LeafOptions.MaxSize))
+				{
+					AppendData(NewSpan.Slice(0, Length));
+					WriteBuffer = null;
+					return NewData.Slice(Length);
+				}
+			}
+
+			AppendData(NewSpan);
+			return ReadOnlyMemory<byte>.Empty;
+		}
+
+		private ReadOnlyMemory<byte> AppendToInteriorNode(ReadOnlyMemory<byte> NewData, ChunkOptions Options)
+		{
+			for (; ; )
+			{
+				Debug.Assert(ChildNodes != null);
+
+				// Try to write to the last node
+				if (ChildNodes.Count > 0)
+				{
+					ChunkNode? LastNode = ChildNodes[^1];
+					if (LastNode != null)
+					{
+						NewData = LastNode.AppendToNode(NewData, Options);
+						if (LastNode.WriteBuffer == null)
+						{
+							Span<byte> LastHashSpan = WriteBuffer.AsSpan(Data.Length - IoHash.NumBytes);
+
+							// Add it to the write buffer
+							IoHash LastHash = LastNode.Serialize();
+							LastHash.CopyTo(LastHashSpan);
+
+							// Update the hash
+							Hash = BuzHash.Add(Hash, LastHashSpan);
+
+							// Check if it's time to finish this chunk
+							uint HashThreshold = (uint)(((1L << 32) * IoHash.NumBytes) / Options.LeafOptions.TargetSize);
+							if ((Payload.Length >= Options.InteriorOptions.MinSize && Hash < HashThreshold) || (Payload.Length >= Options.InteriorOptions.MaxSize))
+							{
+								WriteBuffer = null;
+								return NewData;
+							}
+						}
+						if (NewData.Length == 0)
+						{
+							return NewData;
+						}
+					}
+				}
+
+				// Create a new child node
+				ChunkNode NewNode = new ChunkNode(Owner, this, Depth - 1, Options);
+				AddChildNode(NewNode);
+			}
+		}
+
+		void AddChildNode(ChunkNode NewNode)
+		{
+			Debug.Assert(ChildNodes != null);
+			ChildNodes.Add(NewNode);
+			ExpandData(IoHash.NumBytes);
+		}
+
+		void ExpandData(int Size)
+		{
+			Debug.Assert(WriteBuffer != null);
+
+			int NewDataLength = Data.Length + Size;
+			if (NewDataLength > WriteBuffer.Length)
+			{
+				Array.Resize(ref WriteBuffer, NewDataLength + 4096);
+			}
+
+			int HeaderLength = Data.Length - Payload.Length;
+			Data = WriteBuffer.AsMemory(0, NewDataLength);
+			Payload = Data.Slice(HeaderLength);
+		}
+
+		void AppendData(ReadOnlySpan<byte> NewData)
+		{
+			ExpandData(NewData.Length);
+			NewData.CopyTo(WriteBuffer.AsSpan(Data.Length - NewData.Length));
+		}
+
+		/// <inheritdoc/>
+		protected override IoHash SerializeDirty()
+		{
+			// If we're still writing to the last node, flush it to the write buffer
+			if (ChildNodes != null && ChildNodes.Count > 0)
+			{
+				ChunkNode? LastNode = ChildNodes[^1];
+				if (LastNode != null && LastNode.WriteBuffer != null)
+				{
+					Span<byte> LastHashSpan = WriteBuffer.AsSpan(Payload.Length - IoHash.NumBytes);
+					IoHash LastHash = LastNode.Serialize();
+					LastHash.CopyTo(LastHashSpan);
+				}
+			}
+
+			// Do not allow this node to be modified from this point out
+			WriteBuffer = null;
+
+			// Write the current data
+			List<IoHash> References = new List<IoHash>();
+			if (Depth > 0)
+			{
+				for (ReadOnlySpan<byte> Span = Payload.Span; Span.Length > 0; Span = Span.Slice(IoHash.NumBytes))
+				{
+					References.Add(new IoHash(Span));
+				}
+			}
+			return Owner.WriteNode(Data, References);
+		}
+	}
+
+	/// <summary>
+	/// Factory class for chunk nodes
+	/// </summary>
+	public class ChunkNodeFactory : BundleNodeFactory<ChunkNode>
+	{
+		/// <inheritdoc/>
+		public override ChunkNode CreateRoot(Bundle Bundle)
+		{
+			return new ChunkNode(Bundle, new ChunkOptions());
+		}
+
+		/// <inheritdoc/>
+		public override ChunkNode ParseRoot(Bundle Bundle, IoHash Hash, ReadOnlyMemory<byte> Data)
+		{
+			return new ChunkNode(Bundle, null, Hash, Data);
+		}
+	}
+}
diff --git a/Engine/Source/Programs/Shared/EpicGames.Horde/Bundles/Nodes/DirectoryNode.cs b/Engine/Source/Programs/Shared/EpicGames.Horde/Bundles/Nodes/DirectoryNode.cs
new file mode 100644
--- /dev/null
+++ b/Engine/Source/Programs/Shared/EpicGames.Horde/Bundles/Nodes/DirectoryNode.cs
@@ -0,0 +1,333 @@
+// Copyright Epic Games, Inc. All Rights Reserved.
+
+using EpicGames.Core;
+using EpicGames.Serialization;
+using System;
+using System.Collections.Generic;
+using System.Diagnostics.CodeAnalysis;
+using System.Linq;
+using System.Text;
+using System.Threading.Tasks;
+
+namespace EpicGames.Horde.Bundles.Nodes
+{
+	/// <summary>
+	/// Flags for a directory entry
+	/// </summary>
+	public enum DirectoryEntryFlags : byte
+	{
+		/// <summary>
+		/// This item is a directory
+		/// </summary>
+		Directory = 1,
+
+		/// <summary>
+		/// This item is a file
+		/// </summary>
+		File = 2,
+
+		/// <summary>
+		/// Indicates that the referenced file is executable
+		/// </summary>
+		Executable = 4,
+
+		/// <summary>
+		/// File should be stored as read-only
+		/// </summary>
+		ReadOnly = 8,
+
+		/// <summary>
+		/// File contents are utf-8 encoded text. Client may want to replace line-endings with OS-specific format.
+		/// </summary>
+		Text = 16,
+
+		/// <summary>
+		/// The attached entry includes a Git SHA1 of the corresponding blob/tree contents.
+		/// </summary>
+		HasGitSha1 = 32,
+
+		/// <summary>
+		/// The data for this entry is a Perforce depot path and revision rather than the actual file contents.
+		/// </summary>
+		PerforceDepotPathAndRevision = 64,
+	}
+
+	/// <summary>
+	/// Entry within a directory
+	/// </summary>
+	public class DirectoryEntry
+	{
+		/// <summary>
+		/// Directory that owns this entry
+		/// </summary>
+		public DirectoryNode Directory { get; }
+
+		/// <summary>
+		/// Name of the entry
+		/// </summary>
+		public Utf8String Name { get; }
+
+		/// <summary>
+		/// Flags for this entry
+		/// </summary>
+		public DirectoryEntryFlags Flags { get; }
+
+		/// <summary>
+		/// Hash of this entry
+		/// </summary>
+		IoHash Hash;
+
+		/// <summary>
+		/// Cached logical representation of this object
+		/// </summary>
+		BundleNode? Node;
+
+		/// <summary>
+		/// Constructor
+		/// </summary>
+		internal DirectoryEntry(DirectoryNode Directory, Utf8String Name, DirectoryEntryFlags Flags, IoHash Hash)
+		{
+			this.Directory = Directory;
+			this.Name = Name;
+			this.Flags = Flags;
+			this.Hash = Hash;
+		}
+
+		/// <summary>
+		/// Constructor
+		/// </summary>
+		internal DirectoryEntry(DirectoryNode Directory, Utf8String Name, DirectoryEntryFlags Flags, BundleNode Node)
+		{
+			this.Directory = Directory;
+			this.Name = Name;
+			this.Flags = Flags;
+			this.Node = Node;
+		}
+
+		internal IoHash Serialize()
+		{
+			if (Node != null)
+			{
+				Hash = Node.Serialize();
+			}
+			return Hash;
+		}
+
+		/// <summary>
+		/// Gets the file associated with this entry
+		/// </summary>
+		/// <returns>File object associated with this entry</returns>
+		public async ValueTask<ChunkNode> GetFileAsync()
+		{
+			if ((Flags & DirectoryEntryFlags.File) == 0)
+			{
+				throw new InvalidCastException();
+			}
+			if (Node == null)
+			{
+				ReadOnlyMemory<byte> Data = await Directory.Owner.GetDataAsync(Hash);
+				Node = new ChunkNode(Directory.Owner, Directory, Hash, Data);
+			}
+			return (ChunkNode)Node;
+		}
+
+		/// <summary>
+		/// Gets the directory associated with this entry
+		/// </summary>
+		/// <returns>Directory object associated with this entry</returns>
+		public async ValueTask<DirectoryNode> GetDirectoryAsync()
+		{
+			if ((Flags & DirectoryEntryFlags.Directory) == 0)
+			{
+				throw new InvalidCastException();
+			}
+			if (Node == null)
+			{
+				ReadOnlyMemory<byte> Data = await Directory.Owner.GetDataAsync(Hash);
+				Node = DirectoryNode.Deserialize(Directory.Owner, Directory, Hash, Data.Span);
+			}
+			return (DirectoryNode)Node;
+		}
+	}
+
+	/// <summary>
+	/// A directory node
+	/// </summary>
+	[BundleNodeFactory(typeof(DirectoryNodeFactory))]
+	public class DirectoryNode : BundleNode
+	{
+		internal const byte TypeId = (byte)'d';
+
+		Dictionary<Utf8String, DirectoryEntry> NameToEntry = new Dictionary<Utf8String, DirectoryEntry>();
+
+		/// <summary>
+		/// Entries within this directory
+		/// </summary>
+		public IReadOnlyCollection<DirectoryEntry> Entries => NameToEntry.Values;
+
+		internal DirectoryNode(Bundle Owner, BundleNode? Parent)
+			: base(Owner, Parent)
+		{
+		}
+
+		internal DirectoryNode(Bundle Owner, BundleNode? Parent, IoHash Hash)
+			: base(Owner, Parent, Hash)
+		{
+		}
+
+		/// <summary>
+		/// Clear the contents of this directory
+		/// </summary>
+		public void Clear()
+		{
+			NameToEntry.Clear();
+			MarkAsDirty();
+		}
+
+		/// <summary>
+		/// Attempt to get a directory entry with the given name
+		/// </summary>
+		/// <param name="Name">Name of the entry to find</param>
+		/// <param name="Entry">Receives the entry on success, or null on failure</param>
+		/// <returns>True if a matching entry was found</returns>
+		public bool TryGetEntry(Utf8String Name, [NotNullWhen(true)] out DirectoryEntry? Entry) => NameToEntry.TryGetValue(Name, out Entry);
+
+		/// <summary>
+		/// Adds a new directory with the given name
+		/// </summary>
+		/// <param name="Name">Name of the new directory</param>
+		/// <returns>The new directory object</returns>
+		public ValueTask<DirectoryNode> AddDirectoryAsync(Utf8String Name)
+		{
+			DirectoryNode NewNode = new DirectoryNode(Owner, this);
+
+			DirectoryEntry Entry = new DirectoryEntry(this, Name, DirectoryEntryFlags.Directory, NewNode);
+			NameToEntry.Add(Name, Entry);
+			MarkAsDirty();
+
+			return new ValueTask<DirectoryNode>(NewNode);
+		}
+
+		/// <summary>
+		/// Tries to get a directory with the given name
+		/// </summary>
+		/// <param name="Name">Name of the new directory</param>
+		/// <returns>The new directory object</returns>
+		public async ValueTask<DirectoryNode?> FindDirectoryAsync(Utf8String Name)
+		{
+			if (TryGetEntry(Name, out DirectoryEntry? Entry))
+			{
+				return await Entry.GetDirectoryAsync();
+			}
+			else
+			{
+				return null;
+			}
+		}
+
+		/// <summary>
+		/// Tries to get a directory with the given name, or adds it if not present
+		/// </summary>
+		/// <param name="Name">Name of the new directory</param>
+		/// <returns>The new directory object</returns>
+		public async ValueTask<DirectoryNode> FindOrAddDirectoryAsync(Utf8String Name)
+		{
+			DirectoryNode? Node = await FindDirectoryAsync(Name);
+			return Node ?? await AddDirectoryAsync(Name);
+		}
+
+		/// <summary>
+		/// Deletes the file entry with the given name
+		/// </summary>
+		/// <param name="Name">Name of the entry to delete</param>
+		/// <returns>True if the entry was found, false otherwise</returns>
+		public bool Delete(Utf8String Name)
+		{
+			if(NameToEntry.Remove(Name))
+			{
+				MarkAsDirty();
+				return true;
+			}
+			return false;
+		}
+
+		/// <inheritdoc/>
+		protected override IoHash SerializeDirty()
+		{
+			int Size = 1;
+			foreach (DirectoryEntry Entry in NameToEntry.Values)
+			{
+				Size += 1 + (Entry.Name.Length + 1) + IoHash.NumBytes;
+			}
+
+			byte[] Data = new byte[Size];
+			Span<byte> Span = Data.AsSpan();
+
+			Span[0] = TypeId;
+			Span = Span.Slice(1);
+
+			List<IoHash> References = new List<IoHash>();
+			foreach (DirectoryEntry Entry in NameToEntry.Values.OrderBy(x => x.Name))
+			{
+				IoHash EntryHash = Entry.Serialize();
+				References.Add(EntryHash);
+
+				Span[0] = (byte)Entry.Flags;
+				Span = Span.Slice(1);
+
+				Entry.Name.Span.CopyTo(Span);
+				Span = Span.Slice(Entry.Name.Length + 1);
+
+				EntryHash.CopyTo(Span);
+				Span = Span.Slice(IoHash.NumBytes);
+			}
+
+			return Owner.WriteNode(Data, References);
+		}
+
+		internal static DirectoryNode Deserialize(Bundle Owner, BundleNode? Parent, IoHash Hash, ReadOnlySpan<byte> Span)
+		{
+			if(Span[0] != TypeId)
+			{
+				throw new InvalidOperationException("Invalid signature byte for directory");
+			}
+
+			Span = Span.Slice(1);
+
+			DirectoryNode Node = new DirectoryNode(Owner, Parent);
+			while (Span.Length > 0)
+			{
+				DirectoryEntryFlags Flags = (DirectoryEntryFlags)Span[0];
+				Span = Span.Slice(1);
+
+				int Length = Span.IndexOf((byte)0);
+				Utf8String Name = new Utf8String(Span.Slice(0, Length).ToArray());
+				Span = Span.Slice(Length + 1);
+
+				IoHash EntryHash = new IoHash(Span);
+				Span = Span.Slice(IoHash.NumBytes);
+
+				Node.NameToEntry[Name] = new DirectoryEntry(Node, Name, Flags, EntryHash);
+			}
+			return Node;
+		}
+	}
+
+	/// <summary>
+	/// Factory for creating and serializing <see cref="DirectoryNode"/> objects
+	/// </summary>
+	public class DirectoryNodeFactory : BundleNodeFactory<DirectoryNode>
+	{
+		/// <inheritdoc/>
+		public override DirectoryNode CreateRoot(Bundle Bundle)
+		{
+			return new DirectoryNode(Bundle, null);
+		}
+
+		/// <inheritdoc/>
+		public override DirectoryNode ParseRoot(Bundle Bundle, IoHash Hash, ReadOnlyMemory<byte> Data)
+		{
+			return DirectoryNode.Deserialize(Bundle, null, Hash, Data.Span);
+		}
+	}
+}
diff --git a/Engine/Source/Programs/Shared/EpicGames.Horde/EpicGames.Horde.csproj b/Engine/Source/Programs/Shared/EpicGames.Horde/EpicGames.Horde.csproj
--- a/Engine/Source/Programs/Shared/EpicGames.Horde/EpicGames.Horde.csproj
+++ b/Engine/Source/Programs/Shared/EpicGames.Horde/EpicGames.Horde.csproj
@@ -31,6 +31,7 @@
   <ItemGroup>
     <PackageReference Include="K4os.Compression.LZ4" Version="1.2.6" />
     <PackageReference Include="Microsoft.CodeAnalysis.CSharp.Workspaces" Version="3.10.0" />
+    <PackageReference Include="Microsoft.Extensions.Caching.Abstractions" Version="6.0.0" />
     <PackageReference Include="Microsoft.Extensions.DependencyInjection" Version="6.0.0" />
     <PackageReference Include="Microsoft.Extensions.Http" Version="3.1.22" />
     <PackageReference Include="Microsoft.Extensions.Logging" Version="5.0.0" />
diff --git a/Engine/Source/Programs/Shared/EpicGames.Horde/Properties/AssemblyInfo.cs b/Engine/Source/Programs/Shared/EpicGames.Horde/Properties/AssemblyInfo.cs
new file mode 100644
--- /dev/null
+++ b/Engine/Source/Programs/Shared/EpicGames.Horde/Properties/AssemblyInfo.cs
@@ -0,0 +1,5 @@
+// Copyright Epic Games, Inc. All Rights Reserved.
+
+using System.Runtime.CompilerServices;
+
+[assembly: InternalsVisibleTo("EpicGames.Horde.Tests")]
diff --git a/Engine/Source/Programs/Shared/EpicGames.Horde/Storage/Impl/MemoryStorageClient.cs b/Engine/Source/Programs/Shared/EpicGames.Horde/Storage/Impl/MemoryStorageClient.cs
new file mode 100644
--- /dev/null
+++ b/Engine/Source/Programs/Shared/EpicGames.Horde/Storage/Impl/MemoryStorageClient.cs
@@ -0,0 +1,101 @@
+// Copyright Epic Games, Inc. All Rights Reserved.
+
+using EpicGames.Core;
+using EpicGames.Serialization;
+using Microsoft.Extensions.Logging;
+using System;
+using System.Collections.Generic;
+using System.IO;
+using System.Linq;
+using System.Text;
+using System.Threading;
+using System.Threading.Tasks;
+
+namespace EpicGames.Horde.Storage.Impl
+{
+	/// <summary>
+	/// Implementation of <see cref="IStorageClient"/> which writes directly to the local filesystem for testing. Not intended for production use.
+	/// </summary>
+	public class MemoryStorageClient : IStorageClient
+	{
+		class Ref : IRef
+		{
+			public NamespaceId NamespaceId { get; set; }
+			public BucketId BucketId { get; set; }
+			public RefId RefId { get; set; }
+			public CbObject Value { get; set; }
+
+			public Ref(NamespaceId NamespaceId, BucketId BucketId, RefId RefId, CbObject Value)
+			{
+				this.NamespaceId = NamespaceId;
+				this.BucketId = BucketId;
+				this.RefId = RefId;
+				this.Value = Value;
+			}
+		}
+
+		/// <summary>
+		/// Blobs that have been written to storage
+		/// </summary>
+		public Dictionary<(NamespaceId, IoHash), ReadOnlyMemory<byte>> Blobs { get; } = new Dictionary<(NamespaceId, IoHash), ReadOnlyMemory<byte>>();
+
+		/// <summary>
+		/// Refs that have been written to storage
+		/// </summary>
+		public Dictionary<(NamespaceId, BucketId, RefId), IRef> Refs { get; } = new Dictionary<(NamespaceId, BucketId, RefId), IRef>();
+
+		/// <inheritdoc/>
+		public Task<Stream> ReadBlobAsync(NamespaceId NamespaceId, IoHash Hash, CancellationToken CancellationToken = default)
+		{
+			ReadOnlyMemory<byte> Data = Blobs[(NamespaceId, Hash)];
+			return Task.FromResult<Stream>(new ReadOnlyMemoryStream(Data));
+		}
+
+		/// <inheritdoc/>
+		public async Task WriteBlobAsync(NamespaceId NamespaceId, IoHash Hash, Stream Stream, CancellationToken CancellationToken = default)
+		{
+			using (MemoryStream MemoryStream = new MemoryStream())
+			{
+				await Stream.CopyToAsync(MemoryStream);
+				Blobs[(NamespaceId, Hash)] = MemoryStream.ToArray();
+			}
+		}
+
+		/// <inheritdoc/>
+		public Task<bool> DeleteRefAsync(NamespaceId NamespaceId, BucketId BucketId, RefId RefId, CancellationToken CancellationToken = default)
+		{
+			return Task.FromResult(Refs.Remove((NamespaceId, BucketId, RefId)));
+		}
+
+		/// <inheritdoc/>
+		public Task<List<RefId>> FindMissingRefsAsync(NamespaceId NamespaceId, BucketId BucketId, List<RefId> RefIds, CancellationToken CancellationToken = default)
+		{
+			return Task.FromResult(RefIds.Where(x => !Refs.ContainsKey((NamespaceId, BucketId, x))).ToList());
+		}
+
+		/// <inheritdoc/>
+		public Task<IRef> GetRefAsync(NamespaceId NamespaceId, BucketId BucketId, RefId RefId, CancellationToken CancellationToken = default)
+		{
+			return Task.FromResult<IRef>(Refs[(NamespaceId, BucketId, RefId)]);
+		}
+
+		/// <inheritdoc/>
+		public Task<bool> HasRefAsync(NamespaceId NamespaceId, BucketId BucketId, RefId RefId, CancellationToken CancellationToken = default)
+		{
+			return Task.FromResult(Refs.ContainsKey((NamespaceId, BucketId, RefId)));
+		}
+
+		/// <inheritdoc/>
+		public Task<List<IoHash>> TryFinalizeRefAsync(NamespaceId NamespaceId, BucketId BucketId, RefId RefId, IoHash Hash, CancellationToken CancellationToken = default)
+		{
+			throw new NotImplementedException();
+		}
+
+		/// <inheritdoc/>
+		public Task<List<IoHash>> TrySetRefAsync(NamespaceId NamespaceId, BucketId BucketId, RefId RefId, CbObject Value, CancellationToken CancellationToken = default)
+		{
+			Refs[(NamespaceId, BucketId, RefId)] = new Ref(NamespaceId, BucketId, RefId, Value);
+			return Task.FromResult(new List<IoHash>());
+		}
+	}
+}
diff --git a/Engine/Source/Programs/Shared/EpicGames.Serialization.Tests/VarIntTests.cs b/Engine/Source/Programs/Shared/EpicGames.Serialization.Tests/VarIntTests.cs
new file mode 100644
--- /dev/null
+++ b/Engine/Source/Programs/Shared/EpicGames.Serialization.Tests/VarIntTests.cs
@@ -0,0 +1,32 @@
+// Copyright Epic Games, Inc. All Rights Reserved.
+
+using System;
+using System.Collections.Generic;
+using System.IO;
+using System.Linq;
+using System.Numerics;
+using EpicGames.Core;
+using Microsoft.VisualStudio.TestTools.UnitTesting;
+
+namespace EpicGames.Serialization.Tests
+{
+	[TestClass]
+	public class VarIntTests
+	{
+		[TestMethod]
+		public void TestVarInt()
+		{
+			byte[] Buffer = new byte[20];
+
+			int Length = VarInt.Write(Buffer, -1);
+			Assert.AreEqual(9, Length);
+			Assert.AreEqual(9, VarInt.Measure(-1));
+
+			Assert.AreEqual(9, VarInt.Measure(Buffer));
+			int Value = (int)(long)VarInt.Read(Buffer, out int BytesRead);
+			Assert.AreEqual(9, BytesRead);
+
+			Assert.AreEqual(-1, Value);
+		}
+	}
+}
diff --git a/Engine/Source/Programs/Shared/EpicGames.Serialization/CbConverter.cs b/Engine/Source/Programs/Shared/EpicGames.Serialization/CbConverter.cs
--- a/Engine/Source/Programs/Shared/EpicGames.Serialization/CbConverter.cs
+++ b/Engine/Source/Programs/Shared/EpicGames.Serialization/CbConverter.cs
@@ -14,7 +14,7 @@ namespace EpicGames.Serialization
 	/// <summary>
 	/// Attribute declaring the converter to use for a type
 	/// </summary>
-	[AttributeUsage(AttributeTargets.Class | AttributeTargets.Struct)]
+	[AttributeUsage(AttributeTargets.Class | AttributeTargets.Struct | AttributeTargets.Property)]
 	public class CbConverterAttribute : Attribute
 	{
 		/// <summary>
@@ -114,9 +114,9 @@ namespace EpicGames.Serialization
 	/// </summary>
 	public static class CbConverterMethods
 	{
-		class CbConverterMethodsWrapper<T> : ICbConverterMethods
+		class CbConverterMethodsWrapper<TConverter, T> : ICbConverterMethods where TConverter : class, ICbConverter<T>
 		{
-			static ICbConverter<T> StaticConverter = null!;
+			static TConverter StaticConverter = null!;
 
 			static T Read(CbField Field) => StaticConverter.Read(Field);
 			static void Write(CbWriter Writer, T Value) => StaticConverter.Write(Writer, Value);
@@ -126,7 +126,7 @@ namespace EpicGames.Serialization
 			public MethodInfo WriteMethod { get; } = GetMethodInfo(() => Write(null!, default!));
 			public MethodInfo WriteNamedMethod { get; } = GetMethodInfo(() => WriteNamed(null!, null!, default!));
 
-			public CbConverterMethodsWrapper(ICbConverter<T> Converter)
+			public CbConverterMethodsWrapper(TConverter Converter)
 			{
 				StaticConverter = Converter;
 			}
@@ -137,12 +137,30 @@ namespace EpicGames.Serialization
 			}
 		}
 
+		static Dictionary<PropertyInfo, ICbConverterMethods> PropertyToMethods = new Dictionary<PropertyInfo, ICbConverterMethods>();
 		static Dictionary<Type, ICbConverterMethods> TypeToMethods = new Dictionary<Type, ICbConverterMethods>();
 
 		static ICbConverterMethods CreateWrapper(Type Type, ICbConverter Converter)
 		{
-			Type ConverterType = typeof(CbConverterMethodsWrapper<>).MakeGenericType(Type);
-			return (ICbConverterMethods)Activator.CreateInstance(ConverterType, new object[] { Converter })!;
+			Type ConverterMethodsType = typeof(CbConverterMethodsWrapper<,>).MakeGenericType(Converter.GetType(), Type);
+			return (ICbConverterMethods)Activator.CreateInstance(ConverterMethodsType, new object[] { Converter })!;
+		}
+
+		/// <summary>
+		/// Gets a <see cref="ICbConverterMethods"/> interface for the given type
+		/// </summary>
+		/// <param name="Type"></param>
+		/// <returns></returns>
+		public static ICbConverterMethods Get(PropertyInfo Property)
+		{
+			ICbConverterMethods? Methods;
+			if (!PropertyToMethods.TryGetValue(Property, out Methods))
+			{
+				ICbConverter Converter = CbConverter.GetConverter(Property);
+				Methods = (Converter as ICbConverterMethods) ?? CreateWrapper(Property.PropertyType, Converter);
+				PropertyToMethods.Add(Property, Methods);
+			}
+			return Methods;
 		}
 
 		/// <summary>
@@ -211,6 +229,11 @@ namespace EpicGames.Serialization
 		/// </summary>
 		static object LockObject = new object();
 
+		/// <summary>
+		/// Cache of property to converter
+		/// </summary>
+		static Dictionary<PropertyInfo, ICbConverter> PropertyToConverter = new Dictionary<PropertyInfo, ICbConverter>();
+
 		/// <summary>
 		/// Cache of type to converter
 		/// </summary>
@@ -286,6 +309,31 @@ namespace EpicGames.Serialization
 			}
 		}
 
+		/// <summary>
+		/// Gets the converter for a particular property
+		/// </summary>
+		/// <param name="Property"></param>
+		/// <returns></returns>
+		public static ICbConverter GetConverter(PropertyInfo Property)
+		{
+			CbConverterAttribute? ConverterAttribute = Property.GetCustomAttribute<CbConverterAttribute>();
+			if (ConverterAttribute != null)
+			{
+				Type ConverterType = ConverterAttribute.ConverterType;
+				lock (LockObject)
+				{
+					ICbConverter? Converter;
+					if (!PropertyToConverter.TryGetValue(Property, out Converter))
+					{
+						Converter = (ICbConverter?)Activator.CreateInstance(ConverterType)!;
+						PropertyToConverter.Add(Property, Converter);
+					}
+					return Converter;
+				}
+			}
+			return GetConverter(Property.PropertyType);
+		}
+
 		/// <summary>
 		/// Gets the converter for a particular type
 		/// </summary>
diff --git a/Engine/Source/Programs/Shared/EpicGames.Serialization/Converters/CbClassConverter.cs b/Engine/Source/Programs/Shared/EpicGames.Serialization/Converters/CbClassConverter.cs
--- a/Engine/Source/Programs/Shared/EpicGames.Serialization/Converters/CbClassConverter.cs
+++ b/Engine/Source/Programs/Shared/EpicGames.Serialization/Converters/CbClassConverter.cs
@@ -38,6 +38,7 @@ namespace EpicGames.Serialization.Converters
 		public DynamicMethod ReadConcreteMethod { get; }
 		public DynamicMethod WriteConcreteContentsMethod { get; }
 
+		static Dictionary<PropertyInfo, CbClassConverterMethods> PropertyToMethods = new Dictionary<PropertyInfo, CbClassConverterMethods>();
 		static Dictionary<Type, CbClassConverterMethods> TypeToMethods = new Dictionary<Type, CbClassConverterMethods>();
 
 		public static CbClassConverterMethods Create(Type ClassType)
@@ -252,7 +253,7 @@ namespace EpicGames.Serialization.Converters
 				}
 				else
 				{
-					ICbConverterMethods Methods = CbConverterMethods.Get(PropertyType);
+					ICbConverterMethods Methods = CbConverterMethods.Get(Property);
 					WriteMethod = Methods.WriteNamedMethod;
 				}
 
@@ -370,7 +371,7 @@ namespace EpicGames.Serialization.Converters
 				}
 				else
 				{
-					ReadMethod = CbConverterMethods.Get(Property.PropertyType).ReadMethod;
+					ReadMethod = CbConverterMethods.Get(Property).ReadMethod;
 				}
 
 				// if(!CbReflectedTypeInfo<Type>.MatchName(Stack(0), Idx)) goto SkipPropertyLabel
diff --git a/Engine/Source/Programs/Shared/EpicGames.Serialization/VarInt.cs b/Engine/Source/Programs/Shared/EpicGames.Serialization/VarInt.cs
--- a/Engine/Source/Programs/Shared/EpicGames.Serialization/VarInt.cs
+++ b/Engine/Source/Programs/Shared/EpicGames.Serialization/VarInt.cs
@@ -60,7 +60,7 @@ namespace EpicGames.Serialization
 		/// <returns></returns>
 		public static int Measure(int Value)
 		{
-			return Measure((uint)Value);
+			return Measure((ulong)(long)Value);
 		}
 
 		/// <summary>
