Compact Binary Serialization C# - Added support for serializing uint and ulong

diff --git a/Engine/Source/Programs/Shared/EpicGames.Serialization/CbConverter.cs b/Engine/Source/Programs/Shared/EpicGames.Serialization/CbConverter.cs
--- a/Engine/Source/Programs/Shared/EpicGames.Serialization/CbConverter.cs
+++ b/Engine/Source/Programs/Shared/EpicGames.Serialization/CbConverter.cs
@@ -241,7 +241,9 @@ namespace EpicGames.Serialization
 		{
 			[typeof(bool)] = new CbPrimitiveConverter<bool>(x => x.AsBool(), (w, v) => w.WriteBoolValue(v), (w, n, v) => w.WriteBool(n, v)),
 			[typeof(int)] = new CbPrimitiveConverter<int>(x => x.AsInt32(), (w, v) => w.WriteIntegerValue(v), (w, n, v) => w.WriteInteger(n, v)),
+			[typeof(uint)] = new CbPrimitiveConverter<uint>(x => x.AsUInt32(), (w, v) => w.WriteIntegerValue(v), (w, n, v) => w.WriteInteger(n, v)),
 			[typeof(long)] = new CbPrimitiveConverter<long>(x => x.AsInt64(), (w, v) => w.WriteIntegerValue(v), (w, n, v) => w.WriteInteger(n, v)),
+			[typeof(ulong)] = new CbPrimitiveConverter<ulong>(x => x.AsUInt64(), (w, v) => w.WriteIntegerValue(v), (w, n, v) => w.WriteInteger(n, v)),
 			[typeof(double)] = new CbPrimitiveConverter<double>(x => x.AsDouble(), (w, v) => w.WriteDoubleValue(v), (w, n, v) => w.WriteDouble(n, v)),
 			[typeof(Utf8String)] = new CbPrimitiveConverter<Utf8String>(x => x.AsUtf8String(), (w, v) => w.WriteUtf8StringValue(v), (w, n, v) => w.WriteUtf8String(n, v)),
 			[typeof(IoHash)] = new CbPrimitiveConverter<IoHash>(x => x.AsHash(), (w, v) => w.WriteHashValue(v), (w, n, v) => w.WriteHash(n, v)),
diff --git a/Engine/Source/Programs/Shared/EpicGames.Serialization/CbField.cs b/Engine/Source/Programs/Shared/EpicGames.Serialization/CbField.cs
--- a/Engine/Source/Programs/Shared/EpicGames.Serialization/CbField.cs
+++ b/Engine/Source/Programs/Shared/EpicGames.Serialization/CbField.cs
@@ -1041,7 +1041,15 @@ namespace EpicGames.Serialization
 		/// <summary>
 		/// Access the field as an int32. Returns the provided default on error.
 		/// </summary>
-		public uint AsUInt32(uint Default = 0)
+		public uint AsUInt32()
+		{
+			return AsUInt32(0);
+		}
+
+		/// <summary>
+		/// Access the field as an int32. Returns the provided default on error.
+		/// </summary>
+		public uint AsUInt32(uint Default)
 		{
 			return (uint)AsInteger(Default, 32, false);
 		}
@@ -1049,7 +1057,15 @@ namespace EpicGames.Serialization
 		/// <summary>
 		/// Access the field as an int64. Returns the provided default on error.
 		/// </summary>
-		public ulong AsUInt64(ulong Default = 0)
+		public ulong AsUInt64()
+		{
+			return AsUInt64(0);
+		}
+
+		/// <summary>
+		/// Access the field as an int64. Returns the provided default on error.
+		/// </summary>
+		public ulong AsUInt64(ulong Default)
 		{
 			return (ulong)AsInteger(Default, 64, false);
 		}
