Horde: Rename ChunkNode to FileNode.

diff --git a/Engine/Source/Programs/Horde/Horde.Agent/Commands/Bundles/CreateCommand.cs b/Engine/Source/Programs/Horde/Horde.Agent/Commands/Bundles/CreateCommand.cs
--- a/Engine/Source/Programs/Horde/Horde.Agent/Commands/Bundles/CreateCommand.cs
+++ b/Engine/Source/Programs/Horde/Horde.Agent/Commands/Bundles/CreateCommand.cs
@@ -71,7 +71,7 @@ namespace HordeAgent.Commands.Bundles
 			IStorageClient StorageClient = base.CreateStorageClient(Logger);
 
 			Bundle<DirectoryNode> Bundle = new Bundle<DirectoryNode>(StorageClient, NamespaceId, new BundleOptions(), null);
-			await Bundle.Root.CopyFromDirectoryAsync(InputDir.ToDirectoryInfo(), new ChunkOptions(), Logger);
+			await Bundle.Root.CopyFromDirectoryAsync(InputDir.ToDirectoryInfo(), new ChunkingOptions(), Logger);
 			await Bundle.WriteAsync(BucketId, RefId, false, DateTime.UtcNow);
 
 			return 0;
diff --git a/Engine/Source/Programs/Horde/Horde.Build/Commands/Bundles/CreateCommand.cs b/Engine/Source/Programs/Horde/Horde.Build/Commands/Bundles/CreateCommand.cs
--- a/Engine/Source/Programs/Horde/Horde.Build/Commands/Bundles/CreateCommand.cs
+++ b/Engine/Source/Programs/Horde/Horde.Build/Commands/Bundles/CreateCommand.cs
@@ -49,7 +49,7 @@ namespace HordeServer.Commands.Bundles
 			IStorageClient StorageClient = base.CreateStorageClient(Logger);
 
 			Bundle<DirectoryNode> Bundle = new Bundle<DirectoryNode>(StorageClient, NamespaceId, new BundleOptions(), null);
-			await Bundle.Root.CopyFromDirectoryAsync(InputDir.ToDirectoryInfo(), new ChunkOptions(), Logger);
+			await Bundle.Root.CopyFromDirectoryAsync(InputDir.ToDirectoryInfo(), new ChunkingOptions(), Logger);
 			await Bundle.WriteAsync(BucketId, RefId, false, DateTime.UtcNow);
 
 			return 0;
diff --git a/Engine/Source/Programs/Horde/Horde.Build/Perforce/Impl/CommitService.cs b/Engine/Source/Programs/Horde/Horde.Build/Perforce/Impl/CommitService.cs
--- a/Engine/Source/Programs/Horde/Horde.Build/Perforce/Impl/CommitService.cs
+++ b/Engine/Source/Programs/Horde/Horde.Build/Perforce/Impl/CommitService.cs
@@ -87,7 +87,7 @@ namespace HordeServer.Commits.Impl
 		/// <summary>
 		/// Options for how objects are sliced
 		/// </summary>
-		public ChunkOptions Chunking { get; set; } = new ChunkOptions();
+		public ChunkingOptions Chunking { get; set; } = new ChunkingOptions();
 	}
 
 	/// <summary>
@@ -909,7 +909,7 @@ namespace HordeServer.Commits.Impl
 					}
 
 					string Path = SyncRecord.Path.Substring(ClientInfo.Client.Root.Length).Replace('\\', '/');
-					ChunkNode File = await Contents.CreateFileByPathAsync(Path, DirectoryEntryFlags.PerforceDepotPathAndRevision);
+					FileNode File = await Contents.CreateFileByPathAsync(Path, DirectoryEntryFlags.PerforceDepotPathAndRevision);
 
 					byte[] Data = Encoding.UTF8.GetBytes($"{SyncRecord.DepotFile}#{SyncRecord.Revision}");
 					File.Append(Data, Options.Chunking);
@@ -917,7 +917,7 @@ namespace HordeServer.Commits.Impl
 			}
 			else
 			{
-				Dictionary<int, ChunkNode> Files = new Dictionary<int, ChunkNode>();
+				Dictionary<int, FileNode> Files = new Dictionary<int, FileNode>();
 				await foreach (PerforceResponse Response in Perforce.StreamCommandAsync("sync", Array.Empty<string>(), new string[] { $"{QueryPath}@{Change}" }, null, typeof(SyncRecord), true, default))
 				{
 					PerforceIo? Io = Response.Io;
@@ -930,7 +930,7 @@ namespace HordeServer.Commits.Impl
 						}
 						else if (Io.Command == PerforceIoCommand.Write)
 						{
-							ChunkNode File = Files[Io.File];
+							FileNode File = Files[Io.File];
 							File.Append(Io.Payload, Options.Chunking);
 						}
 						else if (Io.Command == PerforceIoCommand.Close)
diff --git a/Engine/Source/Programs/Shared/EpicGames.Horde.Tests/BundleTests.cs b/Engine/Source/Programs/Shared/EpicGames.Horde.Tests/BundleTests.cs
--- a/Engine/Source/Programs/Shared/EpicGames.Horde.Tests/BundleTests.cs
+++ b/Engine/Source/Programs/Shared/EpicGames.Horde.Tests/BundleTests.cs
@@ -252,9 +252,9 @@ namespace HordeServerTests
 		[TestMethod]
 		public async Task FixedSizeChunkingTests()
 		{
-			ChunkOptions Options = new ChunkOptions();
-			Options.LeafOptions = new TypedChunkOptions(64, 64, 64);
-			Options.InteriorOptions = new TypedChunkOptions(IoHash.NumBytes * 4, IoHash.NumBytes * 4, IoHash.NumBytes * 4);
+			ChunkingOptions Options = new ChunkingOptions();
+			Options.LeafOptions = new ChunkingOptionsForNodeType(64, 64, 64);
+			Options.InteriorOptions = new ChunkingOptionsForNodeType(IoHash.NumBytes * 4, IoHash.NumBytes * 4, IoHash.NumBytes * 4);
 
 			await ChunkingTests(Options);
 		}
@@ -262,14 +262,14 @@ namespace HordeServerTests
 		[TestMethod]
 		public async Task VariableSizeChunkingTests()
 		{
-			ChunkOptions Options = new ChunkOptions();
-			Options.LeafOptions = new TypedChunkOptions(32, 64, 96);
-			Options.InteriorOptions = new TypedChunkOptions(IoHash.NumBytes * 1, IoHash.NumBytes * 4, IoHash.NumBytes * 12);
+			ChunkingOptions Options = new ChunkingOptions();
+			Options.LeafOptions = new ChunkingOptionsForNodeType(32, 64, 96);
+			Options.InteriorOptions = new ChunkingOptionsForNodeType(IoHash.NumBytes * 1, IoHash.NumBytes * 4, IoHash.NumBytes * 12);
 
 			await ChunkingTests(Options);
 		}
 
-		async Task ChunkingTests(ChunkOptions Options)
+		async Task ChunkingTests(ChunkingOptions Options)
 		{
 			byte[] Data = new byte[4096];
 			new Random(0).NextBytes(Data);
@@ -279,7 +279,7 @@ namespace HordeServerTests
 				Data[Idx] = (byte)Idx;
 			}
 
-			Bundle<ChunkNode> Bundle = new Bundle<ChunkNode>(StorageClient, NamespaceId, new BundleOptions(), null);
+			Bundle<FileNode> Bundle = new Bundle<FileNode>(StorageClient, NamespaceId, new BundleOptions(), null);
 
 			const int NumIterations = 100;
 			for (int Idx = 0; Idx < NumIterations; Idx++)
@@ -299,7 +299,7 @@ namespace HordeServerTests
 			await CheckSizes(Bundle.Root, Options, true);
 		}
 
-		async Task CheckSizes(ChunkNode Node, ChunkOptions Options, bool Rightmost)
+		async Task CheckSizes(FileNode Node, ChunkingOptions Options, bool Rightmost)
 		{
 			if (Node.Depth == 0)
 			{
@@ -314,7 +314,7 @@ namespace HordeServerTests
 				int ChildCount = Node.GetChildCount();
 				for (int Idx = 0; Idx < ChildCount; Idx++)
 				{
-					ChunkNode ChildNode = await Node.GetChildNodeAsync(Idx);
+					FileNode ChildNode = await Node.GetChildNodeAsync(Idx);
 					await CheckSizes(ChildNode, Options, Idx == ChildCount - 1);
 				}
 			}
diff --git a/Engine/Source/Programs/Shared/EpicGames.Horde/Bundles/Nodes/ChunkNode.cs b/Engine/Source/Programs/Shared/EpicGames.Horde/Bundles/Nodes/ChunkNode.cs
deleted file mode 100644
diff --git a/Engine/Source/Programs/Shared/EpicGames.Horde/Bundles/Nodes/DirectoryNode.cs b/Engine/Source/Programs/Shared/EpicGames.Horde/Bundles/Nodes/DirectoryNode.cs
--- a/Engine/Source/Programs/Shared/EpicGames.Horde/Bundles/Nodes/DirectoryNode.cs
+++ b/Engine/Source/Programs/Shared/EpicGames.Horde/Bundles/Nodes/DirectoryNode.cs
@@ -119,7 +119,7 @@ namespace EpicGames.Horde.Bundles.Nodes
 		/// Gets the file associated with this entry
 		/// </summary>
 		/// <returns>File object associated with this entry</returns>
-		public async ValueTask<ChunkNode> GetFileAsync()
+		public async ValueTask<FileNode> GetFileAsync()
 		{
 			if ((Flags & DirectoryEntryFlags.File) == 0)
 			{
@@ -128,9 +128,9 @@ namespace EpicGames.Horde.Bundles.Nodes
 			if (Node == null)
 			{
 				ReadOnlyMemory<byte> Data = await Directory.Owner.GetDataAsync(Hash);
-				Node = new ChunkNode(Directory.Owner, Directory, Hash, Data);
+				Node = new FileNode(Directory.Owner, Directory, Hash, Data);
 			}
-			return (ChunkNode)Node;
+			return (FileNode)Node;
 		}
 
 		/// <summary>
@@ -200,14 +200,14 @@ namespace EpicGames.Horde.Bundles.Nodes
 		/// <param name="Name">Name of the new directory</param>
 		/// <param name="Flags">Flags for the new file</param>
 		/// <returns>The new directory object</returns>
-		public ChunkNode CreateFile(Utf8String Name, DirectoryEntryFlags Flags)
+		public FileNode CreateFile(Utf8String Name, DirectoryEntryFlags Flags)
 		{
 			if((Flags & DirectoryEntryFlags.Directory) != 0)
 			{
 				throw new ArgumentException("Directory flag cannot be specified for new file");
 			}
 
-			ChunkNode NewNode = new ChunkNode(Owner, this);
+			FileNode NewNode = new FileNode(Owner, this);
 
 			DirectoryEntry Entry = new DirectoryEntry(this, Name, Flags | DirectoryEntryFlags.File, NewNode);
 			NameToEntry[Name] = Entry;
@@ -222,7 +222,7 @@ namespace EpicGames.Horde.Bundles.Nodes
 		/// <param name="Path">Path to the file</param>
 		/// <param name="Flags">Flags for the new file</param>
 		/// <returns>The new directory object</returns>
-		public async ValueTask<ChunkNode> CreateFileByPathAsync(Utf8String Path, DirectoryEntryFlags Flags)
+		public async ValueTask<FileNode> CreateFileByPathAsync(Utf8String Path, DirectoryEntryFlags Flags)
 		{
 			DirectoryNode Directory = this;
 
@@ -395,7 +395,7 @@ namespace EpicGames.Horde.Bundles.Nodes
 		/// <param name="Options">Options for chunking file content</param>
 		/// <param name="Logger"></param>
 		/// <returns></returns>
-		public async Task CopyFromDirectoryAsync(DirectoryInfo DirectoryInfo, ChunkOptions Options, ILogger Logger)
+		public async Task CopyFromDirectoryAsync(DirectoryInfo DirectoryInfo, ChunkingOptions Options, ILogger Logger)
 		{
 			foreach (DirectoryInfo SubDirectoryInfo in DirectoryInfo.EnumerateDirectories())
 			{
@@ -406,7 +406,7 @@ namespace EpicGames.Horde.Bundles.Nodes
 			foreach (FileInfo FileInfo in DirectoryInfo.EnumerateFiles())
 			{
 				Logger.LogInformation("Adding {File}", FileInfo.FullName);
-				ChunkNode FileNode = CreateFile(FileInfo.Name, DirectoryEntryFlags.File);
+				FileNode FileNode = CreateFile(FileInfo.Name, DirectoryEntryFlags.File);
  				await FileNode.CopyFromFileAsync(FileInfo, Options);
 			}
 		}
@@ -433,7 +433,7 @@ namespace EpicGames.Horde.Bundles.Nodes
 				else
 				{
 					FileInfo FileInfo = new FileInfo(Path.Combine(DirectoryInfo.FullName, Entry.Name.ToString()));
-					ChunkNode FileNode = await Entry.GetFileAsync();
+					FileNode FileNode = await Entry.GetFileAsync();
 					Tasks.Add(Task.Run(() => FileNode.CopyToFileAsync(FileInfo)));
 				}
 			}
diff --git a/Engine/Source/Programs/Shared/EpicGames.Horde/Bundles/Nodes/FileNode.cs b/Engine/Source/Programs/Shared/EpicGames.Horde/Bundles/Nodes/FileNode.cs
new file mode 100644
--- /dev/null
+++ b/Engine/Source/Programs/Shared/EpicGames.Horde/Bundles/Nodes/FileNode.cs
@@ -0,0 +1,479 @@
+// Copyright Epic Games, Inc. All Rights Reserved.
+
+using EpicGames.Core;
+using EpicGames.Serialization;
+using System;
+using System.Collections.Generic;
+using System.Diagnostics;
+using System.IO;
+using System.Linq;
+using System.Text;
+using System.Threading.Tasks;
+
+namespace EpicGames.Horde.Bundles.Nodes
+{
+	/// <summary>
+	/// Options for creating file nodes
+	/// </summary>
+	public class ChunkingOptions
+	{
+		/// <summary>
+		/// Options for creating leaf nodes
+		/// </summary>
+		public ChunkingOptionsForNodeType LeafOptions { get; set; }
+
+		/// <summary>
+		/// Options for creating interior nodes
+		/// </summary>
+		public ChunkingOptionsForNodeType InteriorOptions { get; set; }
+
+		/// <summary>
+		/// Constructor
+		/// </summary>
+		public ChunkingOptions()
+		{
+			LeafOptions = new ChunkingOptionsForNodeType(32 * 1024, 256 * 1024, 64 * 1024);
+			InteriorOptions = new ChunkingOptionsForNodeType(32 * 1024, 256 * 1024, 64 * 1024);
+		}
+	}
+
+	/// <summary>
+	/// Options for creating a specific type of file nodes
+	/// </summary>
+	public class ChunkingOptionsForNodeType
+	{
+		/// <summary>
+		/// Minimum chunk size
+		/// </summary>
+		public int MinSize { get; set; }
+
+		/// <summary>
+		/// Maximum chunk size. Chunks will be split on this boundary if another match is not found.
+		/// </summary>
+		public int MaxSize { get; set; }
+
+		/// <summary>
+		/// Target chunk size for content-slicing
+		/// </summary>
+		public int TargetSize { get; set; }
+
+		/// <summary>
+		/// Default constructor
+		/// </summary>
+		public ChunkingOptionsForNodeType(int MinSize, int MaxSize, int TargetSize)
+		{
+			this.MinSize = MinSize;
+			this.MaxSize = MaxSize;
+			this.TargetSize = TargetSize;
+		}
+	}
+
+	/// <summary>
+	/// Representation of a data stream, split into chunks along content-aware boundaries using a rolling hash (<see cref="BuzHash"/>).
+	/// Chunks are pushed into a tree hierarchy as data is appended to the root, with nodes of the tree also split along content-aware boundaries with <see cref="IoHash.NumBytes"/> granularity.
+	/// Once a chunk has been written to storage, it is treated as immutable.
+	/// </summary>
+	[BundleNodeFactory(typeof(FileNodeFactory))]
+	public sealed class FileNode : BundleNode
+	{
+		const byte TypeId = (byte)'c';
+
+		internal int Depth { get; private set; }
+		internal ReadOnlyMemory<byte> Payload { get; private set; }
+		ReadOnlyMemory<byte> Data; // Full serialized data, including type id and header fields.
+
+		// In-memory state
+		uint Hash;
+		byte[]? WriteBuffer; // Null if the node is complete and does not support writes
+		List<FileNode?>? ChildNodes;
+
+		/// <summary>
+		/// Constructor
+		/// </summary>
+		public FileNode(Bundle Owner, BundleNode? Parent)
+			: this(Owner, Parent, 0, new ChunkingOptions())
+		{
+		}
+
+		/// <summary>
+		/// Create a writable file node
+		/// </summary>
+		private FileNode(Bundle Owner, BundleNode? Parent, int Depth, ChunkingOptions Options)
+			: base(Owner, Parent)
+		{
+			MakeWritable(Depth, Options);
+		}
+
+		/// <summary>
+		/// Create a file node from deserialized data
+		/// </summary>
+		public FileNode(Bundle Owner, BundleNode? Parent, IoHash Hash, ReadOnlyMemory<byte> Data)
+			: base(Owner, Parent, Hash)
+		{
+			ReadOnlySpan<byte> Span = Data.Span;
+			if (Span[0] != TypeId)
+			{
+				throw new InvalidDataException("Invalid type id");
+			}
+
+			Depth = (int)VarInt.Read(Span.Slice(1), out int DepthBytes);
+			int HeaderLength = 1 + DepthBytes;
+
+			Payload = Data.Slice(HeaderLength);
+			this.Data = Data;
+
+			if (Depth > 0)
+			{
+				ChildNodes = new List<FileNode?>(Payload.Length / IoHash.NumBytes);
+				ChildNodes.AddRange(Enumerable.Repeat<FileNode?>(null, Payload.Length / IoHash.NumBytes));
+			}
+		}
+
+		private void MakeWritable(int Depth, ChunkingOptions Options)
+		{
+			int WriteBufferSize = 1 + VarInt.Measure(Depth);
+			if (Depth == 0)
+			{
+				WriteBufferSize += Options.LeafOptions.MaxSize;
+			}
+			else
+			{
+				WriteBufferSize += Options.InteriorOptions.MaxSize;
+			}
+
+			byte[] WriteBuffer = new byte[WriteBufferSize];
+			WriteBuffer[0] = TypeId;
+			int DepthBytes = VarInt.Write(WriteBuffer.AsSpan(1), Depth);
+			int HeaderLength = 1 + DepthBytes;
+
+			this.Depth = Depth;
+			this.Payload = ReadOnlyMemory<byte>.Empty;
+			this.Data = WriteBuffer.AsMemory(0, HeaderLength);
+			this.Hash = 0;
+			this.WriteBuffer = WriteBuffer;
+			this.ChildNodes = (Depth > 0) ? new List<FileNode?>() : null;
+		}
+
+		/// <summary>
+		/// Returns the number of child nodes to this chunk
+		/// </summary>
+		public int GetChildCount() => (ChildNodes == null)? 0 : ChildNodes.Count;
+
+		/// <summary>
+		/// Gets a child node at the given index
+		/// </summary>
+		/// <param name="Index">Index of the child node</param>
+		/// <returns>Node containing data for the given child</returns>
+		public async ValueTask<FileNode> GetChildNodeAsync(int Index)
+		{
+			if(ChildNodes == null)
+			{
+				throw new InvalidOperationException("Node does not contain any children");
+			}
+
+			FileNode? ChildNode = ChildNodes[Index];
+			if (ChildNode == null)
+			{
+				IoHash Hash = new IoHash(Payload.Span.Slice(Index * IoHash.NumBytes));
+				ReadOnlyMemory<byte> Data = await Owner.GetDataAsync(Hash);
+				ChildNode = new FileNode(Owner, this, Hash, Data);
+				ChildNodes[Index] = ChildNode;
+			}
+			return ChildNode;
+		}
+
+		/// <summary>
+		/// Enumerate the children from this node
+		/// </summary>
+		/// <returns></returns>
+		public async IAsyncEnumerable<FileNode> GetChildNodesAsync()
+		{
+			for (int Idx = 0; Idx < GetChildCount(); Idx++)
+			{
+				yield return await GetChildNodeAsync(Idx);
+			}
+		}
+
+		/// <summary>
+		/// Serialize this node and its children into a byte array
+		/// </summary>
+		/// <returns>Array of data stored by the tree</returns>
+		public async Task<byte[]> ToByteArrayAsync()
+		{
+			using (MemoryStream Stream = new MemoryStream())
+			{
+				await CopyToStreamAsync(Stream);
+				return Stream.ToArray();
+			}
+		}
+
+		/// <summary>
+		/// Copy data from the given stream into this file node
+		/// </summary>
+		/// <param name="InputStream"></param>
+		/// <param name="Options"></param>
+		/// <returns></returns>
+		public async Task CopyFromStreamAsync(Stream InputStream, ChunkingOptions Options)
+		{
+			byte[] Buffer = new byte[64 * 1024];
+			for (; ; )
+			{
+				int ReadSize = await InputStream.ReadAsync(Buffer);
+				if(ReadSize == 0)
+				{
+					break;
+				}
+				Append(Buffer.AsMemory(0, ReadSize), Options);
+			}
+		}
+
+		/// <summary>
+		/// Copies the contents of this file from disk
+		/// </summary>
+		/// <param name="File"></param>
+		/// <param name="Options"></param>
+		/// <returns></returns>
+		public async Task CopyFromFileAsync(FileInfo File, ChunkingOptions Options)
+		{
+			using (FileStream Stream = File.OpenRead())
+			{
+				await CopyFromStreamAsync(Stream, Options);
+			}
+		}
+
+		/// <summary>
+		/// Copies the contents of this node and its children to the given output stream
+		/// </summary>
+		/// <param name="OutputStream">The output stream to receive the data</param>
+		public async Task CopyToStreamAsync(Stream OutputStream)
+		{
+			if (Depth > 0)
+			{
+				for (int Idx = 0; Idx < GetChildCount(); Idx++)
+				{
+					FileNode Node = await GetChildNodeAsync(Idx);
+					await Node.CopyToStreamAsync(OutputStream);
+				}
+			}
+			else
+			{
+				await OutputStream.WriteAsync(Payload);
+			}
+		}
+
+		/// <summary>
+		/// Extracts the contents of this node to a file
+		/// </summary>
+		/// <param name="File">File to write with the contents of this node</param>
+		/// <returns></returns>
+		public async Task CopyToFileAsync(FileInfo File)
+		{
+			using (FileStream Stream = File.OpenWrite())
+			{
+				await CopyToStreamAsync(Stream);
+			}
+		}
+
+		/// <summary>
+		/// Append data to this chunk. Must only be called on the root node in a chunk tree.
+		/// </summary>
+		/// <param name="Input">The data to write</param>
+		/// <param name="Options">Settings for chunking the data</param>
+		public void Append(ReadOnlyMemory<byte> Input, ChunkingOptions Options)
+		{
+			if (Parent is FileNode)
+			{
+				throw new InvalidOperationException("Data may only be appended to the root of a tree of file nodes.");
+			}
+
+			for (; ; )
+			{
+				// Append as much data as possible to the existing tree
+				Input = AppendToNode(Input, Options);
+				if (Input.IsEmpty)
+				{
+					break;
+				}
+
+				// Increase the height of the tree by pushing the contents of this node into a new child node
+				FileNode NewNode = new FileNode(Owner, this, IoHash.Zero, Data);
+				if (ChildNodes != null)
+				{
+					for (int Idx = 0; Idx < ChildNodes.Count; Idx++)
+					{
+						FileNode? ChildNode = ChildNodes[Idx];
+						if (ChildNode != null)
+						{
+							ChildNode.Parent = NewNode;
+							NewNode.ChildNodes![Idx] = ChildNode;
+						}
+					}
+				}
+
+				// Reinitialize the current node
+				MakeWritable(Depth + 1, Options);
+				AddChildNode(NewNode);
+			}
+		}
+
+		private ReadOnlyMemory<byte> AppendToNode(ReadOnlyMemory<byte> Data, ChunkingOptions Options)
+		{
+			if (WriteBuffer == null || Data.Length == 0)
+			{
+				return Data;
+			}
+
+			if (Depth == 0)
+			{
+				return AppendToLeafNode(Data, Options);
+			}
+			else
+			{
+				return AppendToInteriorNode(Data, Options);
+			}
+		}
+
+		private ReadOnlyMemory<byte> AppendToLeafNode(ReadOnlyMemory<byte> NewData, ChunkingOptions Options)
+		{
+			ReadOnlySpan<byte> NewSpan = NewData.Span;
+
+			uint HashThreshold = (uint)((1L << 32) / Options.LeafOptions.TargetSize);
+			for (int Length = 0; Length < NewData.Length; Length++)
+			{
+				Hash = BuzHash.Add(Hash, NewSpan[Length]);
+
+				int BlockSize = Payload.Length + Length;
+				if ((BlockSize >= Options.LeafOptions.MinSize && Hash < HashThreshold) || (BlockSize >= Options.LeafOptions.MaxSize))
+				{
+					AppendData(NewSpan.Slice(0, Length));
+					WriteBuffer = null;
+					return NewData.Slice(Length);
+				}
+			}
+
+			AppendData(NewSpan);
+			return ReadOnlyMemory<byte>.Empty;
+		}
+
+		private ReadOnlyMemory<byte> AppendToInteriorNode(ReadOnlyMemory<byte> NewData, ChunkingOptions Options)
+		{
+			for (; ; )
+			{
+				Debug.Assert(ChildNodes != null);
+
+				// Try to write to the last node
+				if (ChildNodes.Count > 0)
+				{
+					FileNode? LastNode = ChildNodes[^1];
+					if (LastNode != null)
+					{
+						NewData = LastNode.AppendToNode(NewData, Options);
+						if (LastNode.WriteBuffer == null)
+						{
+							Span<byte> LastHashSpan = WriteBuffer.AsSpan(Data.Length - IoHash.NumBytes);
+
+							// Add it to the write buffer
+							IoHash LastHash = LastNode.Serialize();
+							LastHash.CopyTo(LastHashSpan);
+
+							// Update the hash
+							Hash = BuzHash.Add(Hash, LastHashSpan);
+
+							// Check if it's time to finish this chunk
+							uint HashThreshold = (uint)(((1L << 32) * IoHash.NumBytes) / Options.LeafOptions.TargetSize);
+							if ((Payload.Length >= Options.InteriorOptions.MinSize && Hash < HashThreshold) || (Payload.Length >= Options.InteriorOptions.MaxSize))
+							{
+								WriteBuffer = null;
+								return NewData;
+							}
+						}
+						if (NewData.Length == 0)
+						{
+							return NewData;
+						}
+					}
+				}
+
+				// Create a new child node
+				FileNode NewNode = new FileNode(Owner, this, Depth - 1, Options);
+				AddChildNode(NewNode);
+			}
+		}
+
+		void AddChildNode(FileNode NewNode)
+		{
+			Debug.Assert(ChildNodes != null);
+			ChildNodes.Add(NewNode);
+			ExpandData(IoHash.NumBytes);
+		}
+
+		void ExpandData(int Size)
+		{
+			Debug.Assert(WriteBuffer != null);
+
+			int NewDataLength = Data.Length + Size;
+			if (NewDataLength > WriteBuffer.Length)
+			{
+				Array.Resize(ref WriteBuffer, NewDataLength + 4096);
+			}
+
+			int HeaderLength = Data.Length - Payload.Length;
+			Data = WriteBuffer.AsMemory(0, NewDataLength);
+			Payload = Data.Slice(HeaderLength);
+		}
+
+		void AppendData(ReadOnlySpan<byte> NewData)
+		{
+			ExpandData(NewData.Length);
+			NewData.CopyTo(WriteBuffer.AsSpan(Data.Length - NewData.Length));
+		}
+
+		/// <inheritdoc/>
+		protected override IoHash SerializeDirty()
+		{
+			// If we're still writing to the last node, flush it to the write buffer
+			if (ChildNodes != null && ChildNodes.Count > 0)
+			{
+				FileNode? LastNode = ChildNodes[^1];
+				if (LastNode != null && LastNode.WriteBuffer != null)
+				{
+					Span<byte> LastHashSpan = WriteBuffer.AsSpan(Data.Length - IoHash.NumBytes);
+					IoHash LastHash = LastNode.Serialize();
+					LastHash.CopyTo(LastHashSpan);
+				}
+			}
+
+			// Do not allow this node to be modified from this point out
+			WriteBuffer = null;
+
+			// Write the current data
+			List<IoHash> References = new List<IoHash>();
+			if (Depth > 0)
+			{
+				for (ReadOnlySpan<byte> Span = Payload.Span; Span.Length > 0; Span = Span.Slice(IoHash.NumBytes))
+				{
+					References.Add(new IoHash(Span));
+				}
+			}
+			return Owner.WriteNode(Data, References);
+		}
+	}
+
+	/// <summary>
+	/// Factory class for file nodes
+	/// </summary>
+	public class FileNodeFactory : BundleNodeFactory<FileNode>
+	{
+		/// <inheritdoc/>
+		public override FileNode CreateRoot(Bundle Bundle)
+		{
+			return new FileNode(Bundle, null);
+		}
+
+		/// <inheritdoc/>
+		public override FileNode ParseRoot(Bundle Bundle, IoHash Hash, ReadOnlyMemory<byte> Data)
+		{
+			return new FileNode(Bundle, null, Hash, Data);
+		}
+	}
+}
