EpicGames.Core: Fix potential exception reading past valid end of buffer when spawning child processes.

diff --git a/Engine/Source/Programs/Shared/EpicGames.Core/ManagedProcess.cs b/Engine/Source/Programs/Shared/EpicGames.Core/ManagedProcess.cs
--- a/Engine/Source/Programs/Shared/EpicGames.Core/ManagedProcess.cs
+++ b/Engine/Source/Programs/Shared/EpicGames.Core/ManagedProcess.cs
@@ -1064,7 +1064,7 @@ namespace EpicGames.Core
 			for (; ; )
 			{
 				// Skip any '\n' that's part of a '\r\n' sequence
-				if (bCarriageReturn && BufferPos < Buffer.Length)
+				if (bCarriageReturn && BufferPos < BufferLen)
 				{
 					if (Buffer[BufferPos] == '\n')
 					{
