UHT in UBT initial version.  Currently disabled.

diff --git a/Engine/Source/Programs/Shared/EpicGames.UHT/Exporters/CodeGen/UhtHeaderCodeGeneratorCppFile.cs b/Engine/Source/Programs/Shared/EpicGames.UHT/Exporters/CodeGen/UhtHeaderCodeGeneratorCppFile.cs
new file mode 100644
--- /dev/null
+++ b/Engine/Source/Programs/Shared/EpicGames.UHT/Exporters/CodeGen/UhtHeaderCodeGeneratorCppFile.cs
@@ -0,0 +1,1799 @@
+// Copyright Epic Games, Inc. All Rights Reserved.
+
+using EpicGames.Core;
+using EpicGames.UHT.Types;
+using EpicGames.UHT.Utils;
+using System;
+using System.Collections.Generic;
+using System.Linq;
+using System.Text;
+
+namespace EpicGames.UHT.Exporters.CodeGen
+{
+	internal class UhtHeaderCodeGeneratorCppFile : UhtHeaderCodeGenerator
+	{
+
+		/// <summary>
+		/// Construct an instance of this generator object
+		/// </summary>
+		/// <param name="CodeGenerator">The base code generator</param>
+		/// <param name="Package">Package being generated</param>
+		/// <param name="HeaderFile">Header file being generated</param>
+		public UhtHeaderCodeGeneratorCppFile(UhtCodeGenerator CodeGenerator, UhtPackage Package, UhtHeaderFile HeaderFile)
+			: base(CodeGenerator, Package, HeaderFile)
+		{
+		}
+
+		/// <summary>
+		/// For a given UE header file, generated the generated H file
+		/// </summary>
+		/// <param name="CppOutput">Output object</param>
+		public void Generate(IUhtExportOutput CppOutput)
+		{
+			ref UhtCodeGenerator.HeaderInfo HeaderInfo = ref this.HeaderInfos[this.HeaderFile.HeaderFileTypeIndex];
+			using (BorrowStringBuilder Borrower = new BorrowStringBuilder(StringBuilderCache.Big))
+			{
+				StringBuilder Builder = Borrower.StringBuilder;
+
+				Builder.Append(HeaderCopyright);
+				Builder.Append(RequiredCPPIncludes);
+				Builder.Append("#include \"").Append(HeaderInfo.IncludePath).Append("\"\r\n");
+
+				bool bAddedStructuredArchiveFromArchiveHeader = false;
+				bool bAddedArchiveUObjectFromStructuredArchiveHeader = false;
+				HashSet<UhtHeaderFile> AddedIncludes = new HashSet<UhtHeaderFile>();
+				AddedIncludes.Add(this.HeaderFile);
+				foreach (UhtType Type in this.HeaderFile.Children)
+				{
+					if (Type is UhtClass Class)
+					{
+						if (Class.ClassWithin != this.Session.UObject && !Class.ClassWithin.HeaderFile.bIsNoExportTypes)
+						{
+							if (AddedIncludes.Add(Class.ClassWithin.HeaderFile))
+							{
+								Builder.Append("#include \"").Append(this.HeaderInfos[Class.ClassWithin.HeaderFile.HeaderFileTypeIndex].IncludePath).Append("\"\r\n");
+							}
+						}
+
+						switch (Class.SerializerArchiveType)
+						{
+							case UhtSerializerArchiveType.None:
+								break;
+
+							case UhtSerializerArchiveType.Archive:
+								if (!bAddedArchiveUObjectFromStructuredArchiveHeader)
+								{
+									Builder.Append("#include \"Serialization/ArchiveUObjectFromStructuredArchive.h\"\r\n");
+									bAddedArchiveUObjectFromStructuredArchiveHeader = true;
+								}
+								break;
+
+							case UhtSerializerArchiveType.StructuredArchiveRecord:
+								if (!bAddedStructuredArchiveFromArchiveHeader)
+								{
+									Builder.Append("#include \"Serialization/StructuredArchive.h\"\r\n");
+									bAddedStructuredArchiveFromArchiveHeader = true;
+								}
+								break;
+						}
+					}
+				}
+
+				Builder.Append(DisableDeprecationWarnings).Append("\r\n");
+				string CleanFileName = this.HeaderFile.FileNameWithoutExtension.Replace('.', '_');
+				Builder.Append("void EmptyLinkFunctionForGeneratedCode").Append(CleanFileName).Append("() {}\r\n");
+
+				if (this.HeaderFile.References.CrossModule.References.Count > 0)
+				{
+					ReadOnlyMemory<string> Sorted = this.HeaderFile.References.CrossModule.GetSortedReferences(
+						(int ObjectIndex, bool bRegistered) => GetCrossReference(ObjectIndex, bRegistered));
+					Builder.Append("// Cross Module References\r\n");
+					foreach (string CrossReference in Sorted.Span)
+					{
+						Builder.Append(CrossReference);
+					}
+					Builder.Append("// End Cross Module References\r\n");
+				}
+
+				int GeneratedBodyStart = Builder.Length;
+
+				bool bHasRegisteredEnums = false;
+				bool bHasRegisteredScriptStructs = false;
+				bool bHasRegisteredClasses = false;
+				bool bAllEnumsEditorOnly = true;
+				foreach (UhtField Field in this.HeaderFile.References.ExportTypes)
+				{
+					if (Field is UhtEnum Enum)
+					{
+						AppendEnum(Builder, Enum);
+						bHasRegisteredEnums = true;
+						bAllEnumsEditorOnly &= Enum.bIsEditorOnly;
+					}
+					else if (Field is UhtScriptStruct ScriptStruct)
+					{
+						AppendScriptStruct(Builder, ScriptStruct);
+						if (ScriptStruct.ScriptStructFlags.HasAnyFlags(EStructFlags.Native))
+						{
+							bHasRegisteredScriptStructs = true;
+						}
+					}
+					else if (Field is UhtFunction Function)
+					{
+						AppendDelegate(Builder, Function);
+					}
+					else if (Field is UhtClass Class)
+					{
+						if (!Class.ClassFlags.HasAnyFlags(EClassFlags.Intrinsic))
+						{
+							AppendClass(Builder, Class);
+							bHasRegisteredClasses = true;
+						}
+					}
+				}
+
+				if (bHasRegisteredEnums || bHasRegisteredScriptStructs || bHasRegisteredClasses)
+				{
+					string Name = $"Z_CompiledInDeferFile_{HeaderInfo.FileId}";
+					string StaticsName = $"{Name}_Statics";
+
+					Builder.Append("\tstruct ").Append(StaticsName).Append("\r\n");
+					Builder.Append("\t{\r\n");
+
+					if (bHasRegisteredEnums)
+					{
+						if (bAllEnumsEditorOnly)
+						{
+							Builder.Append("#if WITH_EDITORONLY_DATA\r\n");
+						}
+						Builder.Append("\t\tstatic const FEnumRegisterCompiledInInfo EnumInfo[];\r\n");
+						if (bAllEnumsEditorOnly)
+						{
+							Builder.Append("#endif\r\n");
+						}
+					}
+					if (bHasRegisteredScriptStructs)
+					{
+						Builder.Append("\t\tstatic const FStructRegisterCompiledInInfo ScriptStructInfo[];\r\n");
+					}
+					if (bHasRegisteredClasses)
+					{
+						Builder.Append("\t\tstatic const FClassRegisterCompiledInInfo ClassInfo[];\r\n");
+					}
+
+					Builder.Append("\t};\r\n");
+
+					uint CombinedHash = uint.MaxValue;
+
+					if (bHasRegisteredEnums)
+					{
+						if (bAllEnumsEditorOnly)
+						{
+							Builder.Append("#if WITH_EDITORONLY_DATA\r\n");
+						}
+						Builder.Append("\tconst FEnumRegisterCompiledInInfo ").Append(StaticsName).Append("::EnumInfo[] = {\r\n");
+						foreach (UhtObject Object in this.HeaderFile.References.ExportTypes)
+						{
+							if (Object is UhtEnum Enum)
+							{
+								if (!bAllEnumsEditorOnly && Enum.bIsEditorOnly)
+								{
+									Builder.Append("#if WITH_EDITORONLY_DATA\r\n");
+								}
+								uint Hash = this.ObjectInfos[Enum.ObjectTypeIndex].Hash;
+								Builder
+									.Append("\t\t{ ")
+									.Append(Enum.SourceName)
+									.Append("_StaticEnum, TEXT(\"")
+									.Append(Enum.EngineName)
+									.Append("\"), &Z_Registration_Info_UEnum_")
+									.Append(Enum.EngineName)
+									.Append($", CONSTRUCT_RELOAD_VERSION_INFO(FEnumReloadVersionInfo, {Hash}U) }},\r\n");
+								if (!bAllEnumsEditorOnly && Enum.bIsEditorOnly)
+								{
+									Builder.Append("#endif\r\n");
+								}
+								CombinedHash = HashCombine(CombinedHash, Hash);
+							}
+						}
+						Builder.Append("\t};\r\n");
+						if (bAllEnumsEditorOnly)
+						{
+							Builder.Append("#endif\r\n");
+						}
+					}
+
+					if (bHasRegisteredScriptStructs)
+					{
+						Builder.Append("\tconst FStructRegisterCompiledInInfo ").Append(StaticsName).Append("::ScriptStructInfo[] = {\r\n");
+						foreach (UhtObject Object in this.HeaderFile.References.ExportTypes)
+						{
+							if (Object is UhtScriptStruct ScriptStruct)
+							{
+								if (ScriptStruct.ScriptStructFlags.HasAnyFlags(EStructFlags.Native))
+								{
+									uint Hash = this.ObjectInfos[ScriptStruct.ObjectTypeIndex].Hash;
+									Builder
+										.Append("\t\t{ ")
+										.Append(ScriptStruct.SourceName)
+										.Append("::StaticStruct, Z_Construct_UScriptStruct_")
+										.Append(ScriptStruct.SourceName)
+										.Append("_Statics::NewStructOps, TEXT(\"")
+										.Append(ScriptStruct.EngineName)
+										.Append("\"), &Z_Registration_Info_UScriptStruct_")
+										.Append(ScriptStruct.EngineName)
+										.Append(", CONSTRUCT_RELOAD_VERSION_INFO(FStructReloadVersionInfo, sizeof(")
+										.Append(ScriptStruct.SourceName)
+										.Append($"), {Hash}U) }},\r\n");
+									CombinedHash = HashCombine(CombinedHash, Hash);
+								}
+							}
+						}
+						Builder.Append("\t};\r\n");
+					}
+
+					if (bHasRegisteredClasses)
+					{
+						Builder.Append("\tconst FClassRegisterCompiledInInfo ").Append(StaticsName).Append("::ClassInfo[] = {\r\n");
+						foreach (UhtObject Object in this.HeaderFile.References.ExportTypes)
+						{
+							if (Object is UhtClass Class)
+							{
+								if (!Class.ClassFlags.HasAnyFlags(EClassFlags.Intrinsic))
+								{
+									uint Hash = this.ObjectInfos[Class.ObjectTypeIndex].Hash;
+									Builder
+										.Append("\t\t{ Z_Construct_UClass_")
+										.Append(Class.SourceName)
+										.Append(", ")
+										.Append(Class.SourceName)
+										.Append("::StaticClass, TEXT(\"")
+										.Append(Class.SourceName)
+										.Append("\"), &Z_Registration_Info_UClass_")
+										.Append(Class.SourceName)
+										.Append(", CONSTRUCT_RELOAD_VERSION_INFO(FClassReloadVersionInfo, sizeof(")
+										.Append(Class.SourceName)
+										.Append($"), {Hash}U) }},\r\n");
+									CombinedHash = HashCombine(CombinedHash, Hash);
+								}
+							}
+						}
+						Builder.Append("\t};\r\n");
+					}
+
+					Builder
+						.Append("\tstatic FRegisterCompiledInInfo ")
+						.Append(Name)
+						.Append($"_{CombinedHash}(TEXT(\"")
+						.Append(Package.EngineName)
+						.Append("\"),\r\n");
+
+					Builder.Append("\t\t").AppendArray(!bHasRegisteredClasses, false, StaticsName, "ClassInfo").Append(",\r\n");
+					Builder.Append("\t\t").AppendArray(!bHasRegisteredScriptStructs, false, StaticsName, "ScriptStructInfo").Append(",\r\n");
+					Builder.Append("\t\t").AppendArray(!bHasRegisteredEnums, bAllEnumsEditorOnly, StaticsName, "EnumInfo").Append(");\r\n");
+				}
+
+				if (this.Session.bIncludeDebugOutput)
+				{
+					Builder.Append("#if 0\r\n");
+					ReadOnlyMemory<string> Sorted = this.HeaderFile.References.Declaration.GetSortedReferences(
+						(int ObjectIndex, bool bRegistered) => GetExternalDecl(ObjectIndex, bRegistered));
+					foreach (string Declaration in Sorted.Span)
+					{
+						Builder.Append(Declaration);
+					}
+					Builder.Append("#endif\r\n");
+				}
+
+				int GeneratedBodyEnd = Builder.Length;
+
+				Builder.Append(EnableDeprecationWarnings).Append("\r\n");
+				StringView GeneratedBody = CppOutput.CommitOutput(Builder);
+
+				// Save the hash of the generated body 
+				this.HeaderInfos[this.HeaderFile.HeaderFileTypeIndex].BodyHash = UhtHash.GenenerateTextHash(GeneratedBody.Span.Slice(GeneratedBodyStart, GeneratedBodyEnd - GeneratedBodyStart));
+			}
+		}
+
+		private StringBuilder AppendEnum(StringBuilder Builder, UhtEnum Enum)
+		{
+			const string MetaDataParamsName = "Enum_MetaDataParams";
+			const string ObjectFlags = "RF_Public|RF_Transient|RF_MarkAsNative";
+			string SingletonName = GetSingletonName(Enum, true);
+			string StaticsName = SingletonName + "_Statics";
+			string RegistrationName = $"Z_Registration_Info_UEnum_{Enum.SourceName}";
+
+			string EnumDisplayNameFn = Enum.MetaData.GetValueOrDefault(UhtNames.EnumDisplayNameFn);
+			if (EnumDisplayNameFn.Length == 0)
+			{
+				EnumDisplayNameFn = "nullptr";
+			}
+
+			using (UhtMacroBlockEmitter MacroBlockEmitter = new UhtMacroBlockEmitter(Builder, "WITH_EDITORONLY_DATA", Enum.bIsEditorOnly))
+			{
+
+				// If we don't have a zero 0 then we emit a static assert to verify we have one
+				if (!Enum.IsValidEnumValue(0) && Enum.MetaData.ContainsKey(UhtNames.BlueprintType))
+				{
+					bool bHasUnparsedValue = Enum.EnumValues.Exists(x => x.Value == -1);
+					if (bHasUnparsedValue)
+					{
+						Builder.Append("\tstatic_assert(");
+						bool bDoneFirst = false;
+						foreach (UhtEnumValue Value in Enum.EnumValues)
+						{
+							if (Value.Value == -1)
+							{
+								if (bDoneFirst)
+								{
+									Builder.Append("||");
+								}
+								bDoneFirst = true;
+								Builder.Append("!int64(").Append(Value.Name).Append(")");
+							}
+						}
+						Builder.Append(", \"'").Append(Enum.SourceName).Append("' does not have a 0 entry!(This is a problem when the enum is initalized by default)\");\r\n");
+					}
+				}
+
+				Builder.Append("\tstatic FEnumRegistrationInfo ").Append(RegistrationName).Append(";\r\n");
+				Builder.Append("\tstatic UEnum* ").Append(Enum.SourceName).Append("_StaticEnum()\r\n");
+				Builder.Append("\t{\r\n");
+
+				Builder.Append("\t\tif (!").Append(RegistrationName).Append(".OuterSingleton)\r\n");
+				Builder.Append("\t\t{\r\n");
+				Builder.Append("\t\t\t").Append(RegistrationName).Append(".OuterSingleton = GetStaticEnum(").Append(SingletonName).Append(", ")
+					.Append(this.PackageSingletonName).Append("(), TEXT(\"").Append(Enum.SourceName).Append("\"));\r\n");
+				Builder.Append("\t\t}\r\n");
+				Builder.Append("\t\treturn ").Append(RegistrationName).Append(".OuterSingleton;\r\n");
+
+				Builder.Append("\t}\r\n");
+
+				Builder.Append("\ttemplate<> ").Append(this.PackageApi).Append("UEnum* StaticEnum<").Append(Enum.CppType).Append(">()\r\n");
+				Builder.Append("\t{\r\n");
+				Builder.Append("\t\treturn ").Append(Enum.SourceName).Append("_StaticEnum();\r\n");
+				Builder.Append("\t}\r\n");
+
+				// Everything from this point on will be part of the definition hash
+				int HashCodeBlockStart = Builder.Length;
+
+				// Statics declaration
+				Builder.Append("\tstruct ").Append(StaticsName).Append("\r\n");
+				Builder.Append("\t{\r\n");
+				Builder.Append("\t\tstatic const UECodeGen_Private::FEnumeratorParam Enumerators[];\r\n");
+				Builder.AppendMetaDataDecl(Enum, MetaDataParamsName, 2);
+				Builder.Append("\t\tstatic const UECodeGen_Private::FEnumParams EnumParams;\r\n");
+				Builder.Append("\t};\r\n");
+
+				// Enumerators
+				Builder.Append("\tconst UECodeGen_Private::FEnumeratorParam ").Append(StaticsName).Append("::Enumerators[] = {\r\n");
+				int EnumIndex = 0;
+				foreach (UhtEnumValue Value in Enum.EnumValues)
+				{
+					string? KeyName;
+					if (!Enum.MetaData.TryGetValue("OverrideName", EnumIndex, out KeyName))
+					{
+						KeyName = Value.Name.ToString();
+					}
+					Builder.Append("\t\t{ ").AppendUTF8LiteralString(KeyName).Append(", (int64)").Append(Value.Name).Append(" },\r\n");
+					++EnumIndex;
+				}
+				Builder.Append("\t};\r\n");
+
+				// Meta data
+				Builder.AppendMetaDataDef(Enum, StaticsName, MetaDataParamsName, 1);
+
+				// Singleton parameters
+				Builder.Append("\tconst UECodeGen_Private::FEnumParams ").Append(StaticsName).Append("::EnumParams = {\r\n");
+				Builder.Append("\t\t(UObject*(*)())").Append(this.PackageSingletonName).Append(",\r\n");
+				Builder.Append("\t\t").Append(EnumDisplayNameFn).Append(",\r\n");
+				Builder.Append("\t\t").AppendUTF8LiteralString(Enum.SourceName).Append(",\r\n");
+				Builder.Append("\t\t").AppendUTF8LiteralString(Enum.CppType).Append(",\r\n");
+				Builder.Append("\t\t").Append(StaticsName).Append("::Enumerators,\r\n");
+				Builder.Append("\t\tUE_ARRAY_COUNT(").Append(StaticsName).Append("::Enumerators),\r\n");
+				Builder.Append("\t\t").Append(ObjectFlags).Append(",\r\n");
+				Builder.Append("\t\t").Append(Enum.EnumFlags.HasAnyFlags(EEnumFlags.Flags) ? "EEnumFlags::Flags" : "EEnumFlags::None").Append(",\r\n");
+				Builder.Append("\t\t(uint8)UEnum::ECppForm::").Append(Enum.CppForm.ToString()).Append(",\r\n");
+				Builder.Append("\t\t").AppendMetaDataParams(Enum, StaticsName, MetaDataParamsName).Append("\r\n");
+				Builder.Append("\t};\r\n");
+
+				// Registration singleton
+				Builder.Append("\tUEnum* ").Append(SingletonName).Append("()\r\n");
+				Builder.Append("\t{\r\n");
+				Builder.Append("\t\tif (!").Append(RegistrationName).Append(".InnerSingleton)\r\n");
+				Builder.Append("\t\t{\r\n");
+				Builder.Append("\t\t\tUECodeGen_Private::ConstructUEnum(").Append(RegistrationName).Append(".InnerSingleton, ").Append(StaticsName).Append("::EnumParams);\r\n");
+				Builder.Append("\t\t}\r\n");
+				Builder.Append("\t\treturn ").Append(RegistrationName).Append(".InnerSingleton;\r\n");
+				Builder.Append("\t}\r\n");
+
+				using (UhtBorrowBuffer BorrowBuffer = new UhtBorrowBuffer(Builder, HashCodeBlockStart, Builder.Length - HashCodeBlockStart))
+				{
+					this.ObjectInfos[Enum.ObjectTypeIndex].Hash = UhtHash.GenenerateTextHash(BorrowBuffer.Buffer.Memory.Span);
+				}
+			}
+			return Builder;
+		}
+
+		private StringBuilder AppendScriptStruct(StringBuilder Builder, UhtScriptStruct ScriptStruct)
+		{
+			const string MetaDataParamsName = "Struct_MetaDataParams";
+			string SingletonName = GetSingletonName(ScriptStruct, true);
+			string StaticsName = SingletonName + "_Statics";
+			string RegistrationName = $"Z_Registration_Info_UScriptStruct_{ScriptStruct.EngineName}";
+			List<UhtScriptStruct> NoExportStructs = FindNoExportStructs(ScriptStruct);
+
+			if (ScriptStruct.ScriptStructFlags.HasAnyFlags(EStructFlags.Native))
+			{
+				// Inject static assert to verify that we do not add vtable
+				if (ScriptStruct.SuperScriptStruct != null)
+				{
+					Builder.Append("\r\n");
+					Builder.Append("static_assert(std::is_polymorphic<")
+						.Append(ScriptStruct.SourceName)
+						.Append(">() == std::is_polymorphic<")
+						.Append(ScriptStruct.SuperScriptStruct.SourceName)
+						.Append(">(), \"USTRUCT ")
+						.Append(ScriptStruct.SourceName)
+						.Append(" cannot be polymorphic unless super ")
+						.Append(ScriptStruct.SuperScriptStruct.SourceName)
+						.Append(" is polymorphic\");\r\n");
+					Builder.Append("\r\n");
+				}
+
+				// Outer singleton
+				Builder.Append("\tstatic FStructRegistrationInfo ").Append(RegistrationName).Append(";\r\n");
+				Builder.Append("class UScriptStruct* ").Append(ScriptStruct.SourceName).Append("::StaticStruct()\r\n");
+				Builder.Append("{\r\n");
+				Builder.Append("\tif (!").Append(RegistrationName).Append(".OuterSingleton)\r\n");
+				Builder.Append("\t{\r\n");
+				Builder.Append("\t\t")
+					.Append(RegistrationName)
+					.Append(".OuterSingleton = GetStaticStruct(")
+					.Append(SingletonName)
+					.Append(", ")
+					.Append(this.PackageSingletonName)
+					.Append("(), TEXT(\"")
+					.Append(ScriptStruct.EngineName)
+					.Append("\"));\r\n");
+
+				// if this struct has RigVM methods - we need to register the method to our central
+				// registry on construction of the static struct
+				if (ScriptStruct.RigVMStructInfo != null)
+				{
+					foreach (UhtRigVMMethodInfo MethodInfo in ScriptStruct.RigVMStructInfo.Methods)
+					{
+						Builder
+							.Append("\t\tFRigVMRegistry::Get().Register(TEXT(\"")
+							.Append(ScriptStruct.SourceName)
+							.Append("::")
+							.Append(MethodInfo.Name)
+							.Append("\"), &")
+							.Append(ScriptStruct.SourceName)
+							.Append("::RigVM")
+							.Append(MethodInfo.Name)
+							.Append(", ")
+							.Append(RegistrationName)
+							.Append(".OuterSingleton);\r\n");
+					}
+				}
+
+				Builder.Append("\t}\r\n");
+				Builder.Append("\treturn ").Append(RegistrationName).Append(".OuterSingleton;\r\n");
+				Builder.Append("}\r\n");
+
+				// Generate the StaticStruct specialization
+				Builder.Append("template<> ").Append(this.PackageApi).Append("UScriptStruct* StaticStruct<").Append(ScriptStruct.SourceName).Append(">()\r\n");
+				Builder.Append("{\r\n");
+				Builder.Append("\treturn ").Append(ScriptStruct.SourceName).Append("::StaticStruct();\r\n");
+				Builder.Append("}\r\n");
+			}
+
+			// Everything from this point on will be part of the definition hash
+			int HashCodeBlockStart = Builder.Length;
+
+			// Declare the statics structure
+			{
+				Builder.Append("\tstruct ").Append(StaticsName).Append("\r\n");
+				Builder.Append("\t{\r\n");
+
+				foreach (UhtScriptStruct NoExportStruct in NoExportStructs)
+				{
+					AppendMirrorsForNoexportStruct(Builder, NoExportStruct, 2);
+				}
+
+				// Meta data
+				Builder.AppendMetaDataDecl(ScriptStruct, MetaDataParamsName, 2);
+
+				// New struct ops
+				if (ScriptStruct.ScriptStructFlags.HasAnyFlags(EStructFlags.Native))
+				{
+					Builder.Append("\t\tstatic void* NewStructOps();\r\n");
+				}
+
+				AppendPropertiesDecl(Builder, ScriptStruct, ScriptStruct.SourceName, StaticsName, 2);
+
+				Builder.Append("\t\tstatic const UECodeGen_Private::FStructParams ReturnStructParams;\r\n");
+				Builder.Append("\t};\r\n");
+			}
+
+			// Populate the elements of the static structure
+			{
+				Builder.AppendMetaDataDef(ScriptStruct, StaticsName, MetaDataParamsName, 1);
+
+				if (ScriptStruct.ScriptStructFlags.HasAnyFlags(EStructFlags.Native))
+				{
+					Builder.Append("\tvoid* ").Append(StaticsName).Append("::NewStructOps()\r\n");
+					Builder.Append("\t{\r\n");
+					Builder.Append("\t\treturn (UScriptStruct::ICppStructOps*)new UScriptStruct::TCppStructOps<").Append(ScriptStruct.SourceName).Append(">();\r\n");
+					Builder.Append("\t}\r\n");
+				}
+
+				AppendPropertiesDefs(Builder, ScriptStruct, ScriptStruct.SourceName, StaticsName, 1);
+
+				Builder.Append("\tconst UECodeGen_Private::FStructParams ").Append(StaticsName).Append("::ReturnStructParams = {\r\n");
+				Builder.Append("\t\t(UObject* (*)())").Append(this.PackageSingletonName).Append(",\r\n");
+				Builder.Append("\t\t").Append(GetSingletonName(ScriptStruct.SuperScriptStruct, true)).Append(",\r\n");
+				if (ScriptStruct.ScriptStructFlags.HasAnyFlags(EStructFlags.Native))
+				{
+					Builder.Append("\t\t&NewStructOps,\r\n");
+				}
+				else
+				{
+					Builder.Append("\t\tnullptr,\r\n");
+				}
+				Builder.Append("\t\t").AppendUTF8LiteralString(ScriptStruct.EngineName).Append(",\r\n");
+				Builder.Append("\t\tsizeof(").Append(ScriptStruct.SourceName).Append("),\r\n");
+				Builder.Append("\t\talignof(").Append(ScriptStruct.SourceName).Append("),\r\n");
+				Builder.AppendPropertiesParams(ScriptStruct, StaticsName, 2, "\r\n");
+				Builder.Append("\t\tRF_Public|RF_Transient|RF_MarkAsNative,\r\n");
+				Builder.Append($"\t\tEStructFlags(0x{(uint)(ScriptStruct.ScriptStructFlags & ~EStructFlags.ComputedFlags):X8}),\r\n");
+				Builder.Append("\t\t").AppendMetaDataParams(ScriptStruct, StaticsName, MetaDataParamsName).Append("\r\n");
+				Builder.Append("\t};\r\n");
+			}
+
+			// Generate the registration function
+			{
+				Builder.Append("\tUScriptStruct* ").Append(SingletonName).Append("()\r\n");
+				Builder.Append("\t{\r\n");
+				string InnerSingletonName;
+				if (ScriptStruct.ScriptStructFlags.HasAnyFlags(EStructFlags.Native))
+				{
+					InnerSingletonName = $"{RegistrationName}.InnerSingleton";
+				}
+				else
+				{
+					Builder.Append("\t\tstatic UScriptStruct* ReturnStruct = nullptr;\r\n");
+					InnerSingletonName = "ReturnStruct";
+				}
+				Builder.Append("\t\tif (!").Append(InnerSingletonName).Append(")\r\n");
+				Builder.Append("\t\t{\r\n");
+				Builder.Append("\t\t\tUECodeGen_Private::ConstructUScriptStruct(").Append(InnerSingletonName).Append(", ").Append(StaticsName).Append("::ReturnStructParams);\r\n");
+				Builder.Append("\t\t}\r\n");
+				Builder.Append("\t\treturn ").Append(InnerSingletonName).Append(";\r\n");
+				Builder.Append("\t}\r\n");
+			}
+
+			using (UhtBorrowBuffer BorrowBuffer = new UhtBorrowBuffer(Builder, HashCodeBlockStart, Builder.Length - HashCodeBlockStart))
+			{
+				this.ObjectInfos[ScriptStruct.ObjectTypeIndex].Hash = UhtHash.GenenerateTextHash(BorrowBuffer.Buffer.Memory.Span);
+			}
+
+			// if this struct has RigVM methods we need to implement both the 
+			// virtual function as well as the stub method here.
+			// The static method is implemented by the user using a macro.
+			if (ScriptStruct.RigVMStructInfo != null)
+			{
+				foreach (UhtRigVMMethodInfo MethodInfo in ScriptStruct.RigVMStructInfo.Methods)
+				{
+					Builder.Append("\r\n");
+
+					Builder
+						.Append(MethodInfo.ReturnType)
+						.Append(' ')
+						.Append(ScriptStruct.SourceName)
+						.Append("::")
+						.Append(MethodInfo.Name)
+						.Append('(')
+						.AppendParameterDecls(MethodInfo.Parameters, false, ",\r\n\t\t")
+						.Append(")\r\n");
+					Builder.Append("{\r\n");
+					Builder.Append("\tFRigVMExecuteContext RigVMExecuteContext;\r\n");
+
+					bool bWroteLine = false;
+					foreach (UhtRigVMParameter Parameter in ScriptStruct.RigVMStructInfo.Members)
+					{
+						if (!Parameter.RequiresCast())
+						{
+							continue;
+						}
+						bWroteLine = true;
+						Builder.Append("\t").Append(Parameter.CastType).Append(' ').Append(Parameter.CastName).Append('(').Append(Parameter.Name).Append(");\r\n");
+					}
+					if (bWroteLine)
+					{
+						//COMPATIBILITY-TODO - Remove the tab
+						Builder.Append("\t\r\n");
+					}
+
+					//COMPATIBILITY-TODO - Replace spaces with \t
+					Builder.Append("\t").Append(MethodInfo.ReturnPrefix()).Append("Static").Append(MethodInfo.Name).Append("(\r\n");
+					Builder.Append("\t\tRigVMExecuteContext");
+					Builder.AppendParameterNames(ScriptStruct.RigVMStructInfo.Members, true, ",\r\n\t\t", true);
+					Builder.AppendParameterNames(MethodInfo.Parameters, true, ",\r\n\t\t");
+					Builder.Append("\r\n");
+					Builder.Append("\t);\r\n");
+					Builder.Append("}\r\n");
+				}
+				Builder.Append("\r\n");
+			}
+			return Builder;
+		}
+
+		private StringBuilder AppendDelegate(StringBuilder Builder, UhtFunction Function)
+		{
+			AppendFunction(Builder, Function, false);
+			return Builder;
+		}
+
+		private StringBuilder AppendFunction(StringBuilder Builder, UhtFunction Function, bool bIsNoExport)
+		{
+			const string MetaDataParamsName = "Function_MetaDataParams";
+			string SingletonName = GetSingletonName(Function, true);
+			string StaticsName = SingletonName + "_Statics";
+			bool bParamsInStatic = bIsNoExport || !Function.FunctionFlags.HasAnyFlags(EFunctionFlags.Event);
+			bool bIsNet = Function.FunctionFlags.HasAnyFlags(EFunctionFlags.NetRequest | EFunctionFlags.NetResponse);
+
+			string StrippedFunctionName = Function.StrippedFunctionName;
+			string EventParameters = GetEventStructParametersName(Function.Outer, StrippedFunctionName);
+
+			// Everything from this point on will be part of the definition hash
+			int HashCodeBlockStart = Builder.Length;
+
+			Builder.AppendBeginEditorOnlyGuard(Function.FunctionFlags.HasAnyFlags(EFunctionFlags.EditorOnly));
+
+			// Statics declaration
+			{
+				Builder.Append("\tstruct ").Append(StaticsName).Append("\r\n");
+				Builder.Append("\t{\r\n");
+
+				if (bParamsInStatic)
+				{
+					List<UhtScriptStruct> NoExportStructs = FindNoExportStructs(Function);
+					foreach (UhtScriptStruct NoExportStruct in NoExportStructs)
+					{
+						AppendMirrorsForNoexportStruct(Builder, NoExportStruct, 2);
+					}
+					AppendEventParameter(Builder, Function, StrippedFunctionName, UhtPropertyTextType.EventParameterFunctionMember, false, 2, "\r\n");
+				}
+
+				AppendPropertiesDecl(Builder, Function, EventParameters, StaticsName, 2);
+
+				Builder.AppendMetaDataDecl(Function, MetaDataParamsName, 2);
+
+				Builder.Append("\t\tstatic const UECodeGen_Private::FFunctionParams FuncParams;\r\n");
+
+				Builder.Append("\t};\r\n");
+			}
+
+			// Statics definition
+			{
+				AppendPropertiesDefs(Builder, Function, EventParameters, StaticsName, 1);
+
+				Builder.AppendMetaDataDef(Function, StaticsName, MetaDataParamsName, 1);
+
+				Builder
+					.Append("\tconst UECodeGen_Private::FFunctionParams ")
+					.Append(StaticsName).Append("::FuncParams = { ")
+					.Append("(UObject*(*)())").Append(GetFunctionOuterFunc(Function)).Append(", ")
+					.Append(GetSingletonName(Function.SuperFunction, true)).Append(", ")
+					.AppendUTF8LiteralString(Function.EngineName).Append(", ")
+					.AppendUTF8LiteralString(Function.FunctionType == UhtFunctionType.SparseDelegate, Function.SparseOwningClassName).Append(", ")
+					.AppendUTF8LiteralString(Function.FunctionType == UhtFunctionType.SparseDelegate, Function.SparseDelegateName).Append(", ");
+
+				if (Function.Children.Count > 0)
+				{
+					UhtFunction TempFunction = Function;
+					while (TempFunction.SuperFunction != null)
+					{
+						TempFunction = TempFunction.SuperFunction;
+					}
+
+					if (bParamsInStatic)
+					{
+						Builder.Append("sizeof(").Append(StaticsName).Append("::").Append(GetEventStructParametersName(TempFunction.Outer, TempFunction.StrippedFunctionName)).Append(")");
+					}
+					else
+					{
+						Builder.Append("sizeof(").Append(GetEventStructParametersName(TempFunction.Outer, TempFunction.StrippedFunctionName)).Append(")");
+					}
+				}
+				else
+				{
+					Builder.Append("0");
+				}
+				Builder.Append(", ");
+
+				Builder
+					.AppendPropertiesParams(Function, StaticsName, 0, " ")
+					.Append("RF_Public|RF_Transient|RF_MarkAsNative, ")
+					.Append($"(EFunctionFlags)0x{(uint)Function.FunctionFlags:X8}, ")
+					.Append(bIsNet ? Function.RPCId : 0).Append(", ")
+					.Append(bIsNet ? Function.RPCResponseId : 0).Append(", ")
+					.AppendMetaDataParams(Function, StaticsName, MetaDataParamsName)
+					.Append(" };\r\n");
+			}
+
+			// Registration function
+			{
+				Builder.Append("\tUFunction* ").Append(SingletonName).Append("()\r\n");
+				Builder.Append("\t{\r\n");
+				Builder.Append("\t\tstatic UFunction* ReturnFunction = nullptr;\r\n");
+				Builder.Append("\t\tif (!ReturnFunction)\r\n");
+				Builder.Append("\t\t{\r\n");
+				Builder.Append("\t\t\tUECodeGen_Private::ConstructUFunction(&ReturnFunction, ").Append(StaticsName).Append("::FuncParams);\r\n");
+				Builder.Append("\t\t}\r\n");
+				Builder.Append("\t\treturn ReturnFunction;\r\n");
+				Builder.Append("\t}\r\n");
+			}
+
+			Builder.AppendEndEditorOnlyGuard(Function.FunctionFlags.HasAnyFlags(EFunctionFlags.EditorOnly));
+
+			using (UhtBorrowBuffer BorrowBuffer = new UhtBorrowBuffer(Builder, HashCodeBlockStart, Builder.Length - HashCodeBlockStart))
+			{
+				this.ObjectInfos[Function.ObjectTypeIndex].Hash = UhtHash.GenenerateTextHash(BorrowBuffer.Buffer.Memory.Span);
+			}
+			return Builder;
+		}
+
+		private string GetFunctionOuterFunc(UhtFunction Function)
+		{
+			if (Function.Outer == null)
+			{
+				return "nullptr";
+			}
+			else if (Function.Outer is UhtHeaderFile HeaderFile)
+			{
+				return GetSingletonName(HeaderFile.Package, true);
+			}
+			else
+			{
+				return GetSingletonName((UhtObject)Function.Outer, true);
+			}
+		}
+
+		private StringBuilder AppendMirrorsForNoexportStruct(StringBuilder Builder, UhtScriptStruct NoExportStruct, int Tabs)
+		{
+			Builder.AppendTabs(Tabs).Append("struct ").Append(NoExportStruct.SourceName);
+			if (NoExportStruct.SuperScriptStruct != null)
+			{
+				Builder.Append(" : public ").Append(NoExportStruct.SuperScriptStruct.SourceName);
+			}
+			Builder.Append("\r\n");
+			Builder.AppendTabs(Tabs).Append("{\r\n");
+
+			// Export the struct's CPP properties
+			AppendExportProperties(Builder, NoExportStruct, Tabs + 1);
+
+			Builder.AppendTabs(Tabs).Append("};\r\n");
+			Builder.Append("\r\n");
+			return Builder;
+		}
+
+		private StringBuilder AppendExportProperties(StringBuilder Builder, UhtScriptStruct ScriptStruct, int Tabs)
+		{
+			using (UhtMacroBlockEmitter Emitter = new UhtMacroBlockEmitter(Builder, "WITH_EDITORONLY_DATA"))
+			{
+				foreach (UhtProperty Property in ScriptStruct.Properties)
+				{
+					Emitter.Set(Property.bIsEditorOnlyProperty);
+					Builder.AppendTabs(Tabs).AppendFullDecl(Property, UhtPropertyTextType.ExportMember, false).Append(";\r\n");
+				}
+			}
+			return Builder;
+		}
+
+		private StringBuilder AppendPropertiesDecl(StringBuilder Builder, UhtStruct Struct, string StructSourceName, string StaticsName, int Tabs)
+		{
+			if (!Struct.Children.Any(x => x is UhtProperty))
+			{
+				return Builder;
+			}
+
+			PropertyMemberContextImpl Context = new PropertyMemberContextImpl(this.CodeGenerator, Struct, StructSourceName, StaticsName);
+
+			using (UhtMacroBlockEmitter Emitter = new UhtMacroBlockEmitter(Builder, "WITH_EDITORONLY_DATA"))
+			{
+				bool bHasAllEditorOnlyDataProperties = true;
+				foreach (UhtType Type in Struct.Children)
+				{
+					if (Type is UhtProperty Property)
+					{
+						Emitter.Set(Property.bIsEditorOnlyProperty);
+						bHasAllEditorOnlyDataProperties &= Property.bIsEditorOnlyProperty;
+						Builder.AppendMemberDecl(Property, Context, Property.EngineName, "", Tabs);
+					}
+				}
+
+				// This will force it off if the last one was editor only but we has some that weren't
+				Emitter.Set(bHasAllEditorOnlyDataProperties);
+
+				Builder.AppendTabs(Tabs).Append("static const UECodeGen_Private::FPropertyParamsBase* const PropPointers[];\r\n");
+
+				foreach (UhtType Type in Struct.Children)
+				{
+					if (Type is UhtProperty Property)
+					{
+						Emitter.Set(Property.bIsEditorOnlyProperty);
+					}
+				}
+				Emitter.Set(bHasAllEditorOnlyDataProperties);
+			}
+			return Builder;
+		}
+
+		private StringBuilder AppendPropertiesDefs(StringBuilder Builder, UhtStruct Struct, string StructSourceName, string StaticsName, int Tabs)
+		{
+			if (!Struct.Children.Any(x => x is UhtProperty))
+			{
+				return Builder;
+			}
+
+			PropertyMemberContextImpl Context = new PropertyMemberContextImpl(this.CodeGenerator, Struct, StructSourceName, StaticsName);
+
+			using (UhtMacroBlockEmitter Emitter = new UhtMacroBlockEmitter(Builder, "WITH_EDITORONLY_DATA"))
+			{
+				bool bHasAllEditorOnlyDataProperties = true;
+				foreach (UhtType Type in Struct.Children)
+				{
+					if (Type is UhtProperty Property)
+					{
+						Emitter.Set(Property.bIsEditorOnlyProperty);
+						bHasAllEditorOnlyDataProperties &= Property.bIsEditorOnlyProperty;
+						Builder.AppendMemberDef(Property, Context, Property.EngineName, "", null, Tabs);
+					}
+				}
+
+				// This will force it off if the last one was editor only but we has some that weren't
+				Emitter.Set(bHasAllEditorOnlyDataProperties);
+
+				Builder.AppendTabs(Tabs).Append("const UECodeGen_Private::FPropertyParamsBase* const ").Append(StaticsName).Append("::PropPointers[] = {\r\n");
+				foreach (UhtType Type in Struct.Children)
+				{
+					if (Type is UhtProperty Property)
+					{
+						Emitter.Set(Property.bIsEditorOnlyProperty);
+						Builder.AppendMemberPtr(Property, Context, Property.EngineName, "", Tabs + 1);
+					}
+				}
+
+				// This will force it off if the last one was editor only but we has some that weren't
+				Emitter.Set(bHasAllEditorOnlyDataProperties);
+
+				Builder.AppendTabs(Tabs).Append("};\r\n");
+			}
+			return Builder;
+		}
+
+		private StringBuilder AppendClass(StringBuilder Builder, UhtClass Class)
+		{
+			// Collect the functions in reversed order
+			List<UhtFunction> ReversedFunctions = new List<UhtFunction>(Class.Functions.Where(x => IsRpcFunction(x) && ShouldExportFunction(x)));
+			ReversedFunctions.Reverse();
+
+			// Check to see if we have any RPC functions for the editor
+			bool bHasEditorRpc = ReversedFunctions.Any(x => x.FunctionFlags.HasAnyFlags(EFunctionFlags.EditorOnly));
+
+			// Output the RPC methods
+			AppendRpcFunctions(Builder, Class, ReversedFunctions, false);
+			if (bHasEditorRpc)
+			{
+				Builder.AppendBeginEditorOnlyGuard();
+				AppendRpcFunctions(Builder, Class, ReversedFunctions, true);
+				Builder.AppendEndEditorOnlyGuard();
+			}
+
+			// Add the accessors
+			AppendPropertyAccessors(Builder, Class);
+
+			// Collect the callback function and sort by name to make the order stable
+			List<UhtFunction> CallbackFunctions = new List<UhtFunction>(Class.Functions.Where(x => x.FunctionFlags.HasAnyFlags(EFunctionFlags.Event) && x.SuperFunction == null));
+			CallbackFunctions.Sort((x, y) => StringComparerUE.OrdinalIgnoreCase.Compare(x.EngineName, y.EngineName));
+
+			// VM -> C++ proxies (events and delegates).
+			if (!Class.ClassFlags.HasAnyFlags(EClassFlags.NoExport))
+			{
+				AppendCallbackFunctions(Builder, Class, CallbackFunctions);
+			}
+
+			if (!Class.ClassFlags.HasAnyFlags(EClassFlags.NoExport))
+			{
+				AppendNatives(Builder, Class);
+			}
+
+			AppendNativeGeneratedInitCode(Builder, Class);
+
+			if (!Class.ClassExportFlags.HasAnyFlags(UhtClassExportFlags.HasCustomVTableHelperConstructor))
+			{
+				Builder.Append("\tDEFINE_VTABLE_PTR_HELPER_CTOR(").Append(Class.SourceName).Append(");\r\n");
+			}
+
+			// Only write out adapters if the user has provided one or the other of the Serialize overloads
+			if (Class.SerializerArchiveType != UhtSerializerArchiveType.None && Class.SerializerArchiveType != UhtSerializerArchiveType.All)
+			{
+				AppendSerializer(Builder, Class, UhtSerializerArchiveType.Archive, "IMPLEMENT_FARCHIVE_SERIALIZER");
+				AppendSerializer(Builder, Class, UhtSerializerArchiveType.StructuredArchiveRecord, "IMPLEMENT_FSTRUCTUREDARCHIVE_SERIALIZER");
+			}
+
+			if (Class.ClassFlags.HasAnyFlags(EClassFlags.Interface))
+			{
+				AppendInterfaceCallFunctions(Builder, Class, CallbackFunctions);
+			}
+			return Builder;
+		}
+
+		private StringBuilder AppendPropertyAccessors(StringBuilder Builder, UhtClass Class)
+		{
+			foreach (UhtType Type in Class.Children)
+			{
+				if (Type is UhtProperty Property)
+				{
+					bool bEditorOnlyProperty = Property.bIsEditorOnlyProperty;
+					if (Property.PropertyExportFlags.HasAnyFlags(UhtPropertyExportFlags.GetterFound))
+					{
+						Builder.Append("\tvoid ").Append(Class.SourceName).Append("::").AppendPropertyGetterWrapperName(Property).Append("(const void* Object, void* OutValue)\r\n");
+						Builder.Append("\t{\r\n");
+						if (bEditorOnlyProperty)
+						{
+							Builder.Append("\t#if WITH_EDITORONLY_DATA\r\n");
+						}
+						Builder.Append("\t\tconst ").Append(Class.SourceName).Append("* Obj = (const ").Append(Class.SourceName).Append("*)Object;\r\n");
+						Builder
+							.Append("\t\t")
+							.AppendPropertyText(Property, UhtPropertyTextType.GetterSetterArg)
+							.Append("& Result = *(")
+							.AppendPropertyText(Property, UhtPropertyTextType.GetterSetterArg)
+							.Append("*)OutValue;\r\n");
+						Builder
+							.Append("\t\tResult = (")
+							.AppendPropertyText(Property, UhtPropertyTextType.GetterSetterArg)
+							.Append(")Obj->")
+							.Append(Property.Getter!)
+							.Append("();\r\n");
+						if (bEditorOnlyProperty)
+						{
+							Builder.Append("\t#endif // WITH_EDITORONLY_DATA\r\n");
+						}
+						Builder.Append("\t}\r\n");
+					}
+					if (Property.PropertyExportFlags.HasAnyFlags(UhtPropertyExportFlags.SetterFound))
+					{
+						Builder.Append("\tvoid ").Append(Class.SourceName).Append("::").AppendPropertySetterWrapperName(Property).Append("(void* Object, const void* InValue)\r\n");
+						Builder.Append("\t{\r\n");
+						if (bEditorOnlyProperty)
+						{
+							Builder.Append("\t#if WITH_EDITORONLY_DATA\r\n");
+						}
+						Builder.Append("\t\t").Append(Class.SourceName).Append("* Obj = (").Append(Class.SourceName).Append("*)Object;\r\n");
+						Builder
+							.Append("\t\t")
+							.AppendPropertyText(Property, UhtPropertyTextType.GetterSetterArg)
+							.Append("& Value = *(")
+							.AppendPropertyText(Property, UhtPropertyTextType.GetterSetterArg)
+							.Append("*)InValue;\r\n");
+						Builder
+							.Append("\t\tObj->")
+							.Append(Property.Setter!)
+							.Append("(Value);\r\n");
+						if (bEditorOnlyProperty)
+						{
+							Builder.Append("\t#endif // WITH_EDITORONLY_DATA\r\n");
+						}
+						Builder.Append("\t}\r\n");
+					}
+				}
+			}
+			return Builder;
+		}
+
+		private StringBuilder AppendInterfaceCallFunctions(StringBuilder Builder, UhtClass Class, List<UhtFunction> CallbackFunctions)
+		{
+			foreach (UhtFunction Function in CallbackFunctions)
+			{
+				Builder.Append("\tstatic FName NAME_").Append(Function.Outer?.SourceName).Append("_").Append(Function.SourceName).Append(" = FName(TEXT(\"").Append(Function.EngineName).Append("\"));\r\n");
+				string ExtraParameter = Function.FunctionFlags.HasAnyFlags(EFunctionFlags.Const) ? "const UObject* O" : "UObject* O";
+				AppendNativeFunctionHeader(Builder, Function, UhtPropertyTextType.InterfaceFunction, false, null, ExtraParameter, 0, "\r\n");
+				Builder.Append("\t{\r\n");
+				Builder.Append("\t\tcheck(O != NULL);\r\n");
+				Builder.Append("\t\tcheck(O->GetClass()->ImplementsInterface(").Append(Class.SourceName).Append("::StaticClass()));\r\n");
+				if (Function.Children.Count > 0)
+				{
+					Builder.Append("\t\t").Append(GetEventStructParametersName(Class, Function.StrippedFunctionName)).Append(" Parms;\r\n");
+				}
+				Builder.Append("\t\tUFunction* const Func = O->FindFunction(NAME_").Append(Function.Outer?.SourceName).Append("_").Append(Function.SourceName).Append(");\r\n");
+				Builder.Append("\t\tif (Func)\r\n");
+				Builder.Append("\t\t{\r\n");
+				foreach (UhtProperty Property in Function.ParameterProperties.Span)
+				{
+					Builder.Append("\t\t\tParms.").Append(Property.SourceName).Append("=").Append(Property.SourceName).Append(";\r\n");
+				}
+				Builder
+					.Append("\t\t\t")
+					.Append(Function.FunctionFlags.HasAnyFlags(EFunctionFlags.Const) ? "const_cast<UObject*>(O)" : "O")
+					.Append("->ProcessEvent(Func, ")
+					.Append(Function.Children.Count > 0 ? "&Parms" : "NULL")
+					.Append(");\r\n");
+				foreach (UhtProperty Property in Function.ParameterProperties.Span)
+				{
+					if (Property.PropertyFlags.HasExactFlags(EPropertyFlags.OutParm | EPropertyFlags.ConstParm, EPropertyFlags.OutParm))
+					{
+						Builder.Append("\t\t\t").Append(Property.SourceName).Append("=Parms.").Append(Property.SourceName).Append(";\r\n");
+					}
+				}
+				Builder.Append("\t\t}\r\n");
+
+				// else clause to call back into native if it's a BlueprintNativeEvent
+				if (Function.FunctionFlags.HasAnyFlags(EFunctionFlags.Native))
+				{
+					Builder
+						.Append("\t\telse if (auto I = (")
+						.Append(Function.FunctionFlags.HasAnyFlags(EFunctionFlags.Const) ? "const I" : "I")
+						.Append(Class.EngineName)
+						.Append("*)(O->GetNativeInterfaceAddress(U")
+						.Append(Class.EngineName)
+						.Append("::StaticClass())))\r\n");
+					Builder.Append("\t\t{\r\n");
+					Builder.Append("\t\t\t");
+					if (Function.bHasReturnProperty)
+					{
+						Builder.Append("Parms.ReturnValue = ");
+					}
+					Builder.Append("I->").Append(Function.SourceName).Append("_Implementation(");
+
+					bool bFirst = true;
+					foreach (UhtProperty Property in Function.ParameterProperties.Span)
+					{
+						if (!bFirst)
+						{
+							Builder.Append(",");
+						}
+						bFirst = false;
+						Builder.Append(Property.SourceName);
+					}
+					Builder.Append(");\r\n");
+					Builder.Append("\t\t}\r\n");
+				}
+
+				if (Function.bHasReturnProperty)
+				{
+					Builder.Append("\t\treturn Parms.ReturnValue;\r\n");
+				}
+				Builder.Append("\t}\r\n");
+			}
+			return Builder;
+		}
+
+		private StringBuilder AppendSerializer(StringBuilder Builder, UhtClass Class, UhtSerializerArchiveType SerializerType, string MacroText)
+		{
+			if (!Class.SerializerArchiveType.HasAnyFlags(SerializerType))
+			{
+				if (Class.EnclosingDefine.Length > 0)
+				{
+					Builder.Append("#if ").Append(Class.EnclosingDefine).Append("\r\n");
+				}
+				Builder.Append("\t").Append(MacroText).Append('(').Append(Class.SourceName).Append(")\r\n");
+				if (Class.EnclosingDefine.Length > 0)
+				{
+					Builder.Append("#endif\r\n");
+				}
+			}
+			return Builder;
+		}
+
+		private StringBuilder AppendNativeGeneratedInitCode(StringBuilder Builder, UhtClass Class)
+		{
+			const string MetaDataParamsName = "Class_MetaDataParams";
+			string SingletonName = GetSingletonName(Class, true);
+			string StaticsName = SingletonName + "_Statics";
+			string RegistrationName = $"Z_Registration_Info_UClass_{Class.SourceName}";
+			string[]? SparseDataTypes = Class.MetaData.GetStringArray(UhtNames.SparseClassDataTypes);
+
+			PropertyMemberContextImpl Context = new PropertyMemberContextImpl(this.CodeGenerator, Class, Class.SourceName, StaticsName);
+
+			bool bHasInterfaces = Class.Bases != null && Class.Bases.Any(x => x is UhtClass BaseClass ? BaseClass.ClassFlags.HasAnyFlags(EClassFlags.Interface) : false);
+
+			// Collect the functions to be exported
+			bool bAllEditorOnlyFunctions = true;
+			List<UhtFunction> SortedFunctions = new List<UhtFunction>();
+			foreach (UhtFunction Function in Class.Functions)
+			{
+				if (!Function.FunctionFlags.HasAnyFlags(EFunctionFlags.Delegate))
+				{
+					bAllEditorOnlyFunctions &= Function.FunctionFlags.HasAnyFlags(EFunctionFlags.EditorOnly);
+				}
+				SortedFunctions.Add(Function);
+			}
+			SortedFunctions.Sort((x, y) => StringComparerUE.OrdinalIgnoreCase.Compare(x.EngineName, y.EngineName));
+
+			// Output any function
+			foreach (UhtFunction Function in SortedFunctions)
+			{
+				if (!Function.FunctionFlags.HasAnyFlags(EFunctionFlags.Delegate))
+				{
+					AppendFunction(Builder, Function, Class.ClassFlags.HasAnyFlags(EClassFlags.NoExport));
+				}				
+			}
+
+			Builder.Append("\tIMPLEMENT_CLASS_NO_AUTO_REGISTRATION(").Append(Class.SourceName).Append(");\r\n");
+
+			// Everything from this point on will be part of the definition hash
+			int HashCodeBlockStart = Builder.Length;
+
+			// simple ::StaticClass wrapper to avoid header, link and DLL hell
+			{
+				Builder.Append("\tUClass* ").Append(GetSingletonName(Class, false)).Append("()\r\n");
+				Builder.Append("\t{\r\n");
+				Builder.Append("\t\treturn ").Append(Class.SourceName).Append("::StaticClass();\r\n");
+				Builder.Append("\t}\r\n");
+			}
+
+			// Declare the statics object
+			{
+				Builder.Append("\tstruct ").Append(StaticsName).Append("\r\n");
+				Builder.Append("\t{\r\n");
+				Builder.Append("\t\tstatic UObject* (*const DependentSingletons[])();\r\n");
+
+				if (SortedFunctions.Count != 0)
+				{
+					Builder.AppendBeginEditorOnlyGuard(bAllEditorOnlyFunctions);
+					Builder.Append("\t\tstatic const FClassFunctionLinkInfo FuncInfo[];\r\n");
+					Builder.AppendEndEditorOnlyGuard(bAllEditorOnlyFunctions);
+				}
+
+				Builder.AppendMetaDataDecl(Class, MetaDataParamsName, 2);
+
+				AppendPropertiesDecl(Builder, Class, Class.SourceName, StaticsName, 2);
+
+				if (bHasInterfaces)
+				{
+					Builder.Append("\t\tstatic const UECodeGen_Private::FImplementedInterfaceParams InterfaceParams[];\r\n");
+				}
+
+				Builder.Append("\t\tstatic const FCppClassTypeInfoStatic StaticCppClassTypeInfo;\r\n");
+
+				Builder.Append("\t\tstatic const UECodeGen_Private::FClassParams ClassParams;\r\n");
+
+				Builder.Append("\t};\r\n");
+			}
+
+			// Define the statics object
+			{
+				Builder.Append("\tUObject* (*const ").Append(StaticsName).Append("::DependentSingletons[])() = {\r\n");
+				if (Class.SuperClass != null && Class.SuperClass != Class)
+				{
+					Builder.Append("\t\t(UObject* (*)())").Append(GetSingletonName(Class.SuperClass, true)).Append(",\r\n");
+				}
+				Builder.Append("\t\t(UObject* (*)())").Append(GetSingletonName(this.Package, true)).Append(",\r\n");
+				Builder.Append("\t};\r\n");
+
+				if (SortedFunctions.Count > 0)
+				{
+					Builder.AppendBeginEditorOnlyGuard(bAllEditorOnlyFunctions);
+					Builder.Append("\tconst FClassFunctionLinkInfo ").Append(StaticsName).Append("::FuncInfo[] = {\r\n");
+
+					foreach (UhtFunction Function in SortedFunctions)
+					{
+						bool bIsEditorOnlyFunction = Function.FunctionFlags.HasAnyFlags(EFunctionFlags.EditorOnly);
+						Builder.AppendBeginEditorOnlyGuard(bIsEditorOnlyFunction);
+						Builder
+							.Append("\t\t{ &")
+							.Append(GetSingletonName(Function, true))
+							.Append(", ")
+							.AppendUTF8LiteralString(Function.EngineName)
+							.Append(" },")
+							.AppendObjectHash(Class, Context, Function)
+							.Append("\r\n");
+						Builder.AppendEndEditorOnlyGuard(bIsEditorOnlyFunction);
+					}
+
+					Builder.Append("\t};\r\n");
+					Builder.AppendEndEditorOnlyGuard(bAllEditorOnlyFunctions);
+				}
+
+				Builder.AppendMetaDataDef(Class, StaticsName, MetaDataParamsName, 1);
+
+				AppendPropertiesDefs(Builder, Class, Class.SourceName, StaticsName, 1);
+
+				if (Class.Bases != null && bHasInterfaces)
+				{
+					Builder.Append("\t\tconst UECodeGen_Private::FImplementedInterfaceParams ").Append(StaticsName).Append("::InterfaceParams[] = {\r\n");
+
+					foreach (UhtStruct Struct in Class.Bases)
+					{
+						if (Struct is UhtClass BaseClass)
+						{
+							if (BaseClass.ClassFlags.HasAnyFlags(EClassFlags.Interface))
+							{
+								Builder
+									.Append("\t\t\t{ ")
+									.Append(GetSingletonName(BaseClass.AlternateObject, false))
+									.Append(", (int32)VTABLE_OFFSET(")
+									.Append(Class.SourceName)
+									.Append(", ")
+									.Append(BaseClass.SourceName)
+									.Append("), false }, ")
+									.AppendObjectHash(Class, Context, BaseClass.AlternateObject)
+									.Append("\r\n");
+							}
+						}
+					}
+
+					Builder.Append("\t\t};\r\n");
+				}
+
+				Builder.Append("\tconst FCppClassTypeInfoStatic ").Append(StaticsName).Append("::StaticCppClassTypeInfo = {\r\n");
+				if (Class.ClassFlags.HasAnyFlags(EClassFlags.Interface))
+				{
+					Builder.Append("\t\tTCppClassTypeTraits<I").Append(Class.EngineName).Append(">::IsAbstract,\r\n");
+				}
+				else
+				{
+					Builder.Append("\t\tTCppClassTypeTraits<").Append(Class.SourceName).Append(">::IsAbstract,\r\n");
+				}
+				Builder.Append("\t};\r\n");
+
+				EClassFlags ClassFlags = Class.ClassFlags & EClassFlags.SaveInCompiledInClasses;
+
+				Builder.Append("\tconst UECodeGen_Private::FClassParams ").Append(StaticsName).Append("::ClassParams = {\r\n");
+				Builder.Append("\t\t&").Append(Class.SourceName).Append("::StaticClass,\r\n");
+				if (Class.Config.Length > 0)
+				{
+					Builder.Append("\t\t").AppendUTF8LiteralString(Class.Config).Append(",\r\n");
+				}
+				else
+				{
+					Builder.Append("\t\tnullptr,\r\n");
+				}
+				Builder.Append("\t\t&StaticCppClassTypeInfo,\r\n");
+				Builder.Append("\t\tDependentSingletons,\r\n");
+				if (SortedFunctions.Count == 0)
+				{
+					Builder.Append("\t\tnullptr,\r\n");
+				}
+				else if (bAllEditorOnlyFunctions)
+				{
+					Builder.Append("\t\tIF_WITH_EDITOR(FuncInfo, nullptr),\r\n");
+				}
+				else
+				{
+					Builder.Append("\t\tFuncInfo,\r\n");
+				}
+				Builder.AppendPropertiesParamsList(Class, StaticsName, 2, "\r\n");
+				Builder.Append("\t\t").Append(bHasInterfaces ? "InterfaceParams" : "nullptr").Append(",\r\n");
+				Builder.Append("\t\tUE_ARRAY_COUNT(DependentSingletons),\r\n");
+				if (SortedFunctions.Count == 0)
+				{
+					Builder.Append("\t\t0,\r\n");
+				}
+				else if (bAllEditorOnlyFunctions)
+				{
+					Builder.Append("\t\tIF_WITH_EDITOR(UE_ARRAY_COUNT(FuncInfo), 0),\r\n");
+				}
+				else
+				{
+					Builder.Append("\t\tUE_ARRAY_COUNT(FuncInfo),\r\n");
+				}
+				Builder.AppendPropertiesParamsCount(Class, StaticsName, 2, "\r\n");
+				Builder.Append("\t\t").Append(bHasInterfaces ? "UE_ARRAY_COUNT(InterfaceParams)" : "0").Append(",\r\n");
+				Builder.Append($"\t\t0x{(uint)ClassFlags:X8}u,\r\n");
+				Builder.Append("\t\t").AppendMetaDataParams(Class, StaticsName, MetaDataParamsName).Append("\r\n");
+				Builder.Append("\t};\r\n");
+			}
+
+			// Class registration
+			{
+				Builder.Append("\tUClass* ").Append(SingletonName).Append("()\r\n");
+				Builder.Append("\t{\r\n");
+				Builder.Append("\t\tif (!").Append(RegistrationName).Append(".OuterSingleton)\r\n");
+				Builder.Append("\t\t{\r\n");
+				Builder.Append("\t\t\tUECodeGen_Private::ConstructUClass(").Append(RegistrationName).Append(".OuterSingleton, ").Append(StaticsName).Append("::ClassParams);\r\n");
+				if (SparseDataTypes != null)
+				{
+					foreach (string SparseClass in SparseDataTypes)
+					{
+						Builder.Append("\t\t\t").Append(RegistrationName).Append(".OuterSingleton->SetSparseClassDataStruct(F").Append(SparseClass).Append("::StaticStruct());\r\n");
+					}
+				}
+				Builder.Append("\t\t}\r\n");
+				Builder.Append("\t\treturn ").Append(RegistrationName).Append(".OuterSingleton;\r\n");
+				Builder.Append("\t}\r\n");
+			}
+
+			// At this point, we can compute the hash... HOWEVER, in the old UHT extra data is appended to the hash block that isn't emitted to the actual output
+			using (BorrowStringBuilder HashBorrower = new BorrowStringBuilder(StringBuilderCache.Small))
+			{
+				StringBuilder HashBuilder = HashBorrower.StringBuilder;
+				HashBuilder.Append(Builder, HashCodeBlockStart, Builder.Length - HashCodeBlockStart);
+
+				int SaveLength = HashBuilder.Length;
+
+				// Append base class' hash at the end of the generated code, this will force update derived classes
+				// when base class changes during hot-reload.
+				uint BaseClassHash = 0;
+				if (Class.SuperClass != null && !Class.SuperClass.ClassFlags.HasAnyFlags(EClassFlags.Intrinsic))
+				{
+					BaseClassHash = this.ObjectInfos[Class.SuperClass.ObjectTypeIndex].Hash;
+				}
+				HashBuilder.Append($"\r\n// {BaseClassHash}\r\n");
+
+				// Append info for the sparse class data struct onto the text to be hashed
+				if (SparseDataTypes != null)
+				{
+					foreach (string SparseDataType in SparseDataTypes)
+					{
+						UhtType? Type = this.Session.FindType(Class, UhtFindOptions.ScriptStruct | UhtFindOptions.EngineName, SparseDataType);
+						if (Type != null)
+						{
+							HashBuilder.Append(Type.EngineName).Append("\r\n");
+							for (UhtScriptStruct? SparseStruct = Type as UhtScriptStruct; SparseStruct != null; SparseStruct = SparseStruct.SuperScriptStruct)
+							{
+								foreach (UhtProperty Property in SparseStruct.Properties)
+								{
+									HashBuilder.AppendPropertyText(Property, UhtPropertyTextType.SparseShort).Append(' ').Append(Property.SourceName).Append("\r\n");
+								}
+							}
+						}
+					}
+				}
+
+				if (Class.ClassFlags.HasAnyFlags(EClassFlags.NoExport))
+				{
+					Builder.Append("\t/* friend declarations for pasting into noexport class ").Append(Class.SourceName).Append("\r\n");
+					Builder.Append("\tfriend struct ").Append(StaticsName).Append(";\r\n");
+					Builder.Append("\t*/\r\n");
+				}
+
+				if (this.Session.bIncludeDebugOutput)
+				{
+					using (UhtBorrowBuffer BorrowBuffer = new UhtBorrowBuffer(HashBuilder, SaveLength, HashBuilder.Length - SaveLength))
+					{
+						Builder.Append("#if 0\r\n");
+						Builder.Append(BorrowBuffer.Buffer.Memory);
+						Builder.Append("#endif\r\n");
+					}
+				}
+
+				// Calculate generated class initialization code hash so that we know when it changes after hot-reload
+				using (UhtBorrowBuffer BorrowBuffer = new UhtBorrowBuffer(HashBuilder))
+				{
+					this.ObjectInfos[Class.ObjectTypeIndex].Hash = UhtHash.GenenerateTextHash(BorrowBuffer.Buffer.Memory.Span);
+				}
+			}
+
+			Builder.Append("\ttemplate<> ").Append(this.PackageApi).Append("UClass* StaticClass<").Append(Class.SourceName).Append(">()\r\n");
+			Builder.Append("\t{\r\n");
+			Builder.Append("\t\treturn ").Append(Class.SourceName).Append("::StaticClass();\r\n");
+			Builder.Append("\t}\r\n");
+
+			if (Class.ClassExportFlags.HasAnyFlags(UhtClassExportFlags.SelfHasReplicatedProperties))
+			{
+				Builder.Append("\r\n");
+				Builder.Append("\tvoid ").Append(Class.SourceName).Append("::ValidateGeneratedRepEnums(const TArray<struct FRepRecord>& ClassReps) const\r\n");
+				Builder.Append("\t{\r\n");
+
+				foreach (UhtProperty Property in Class.Properties)
+				{
+					if (Property.PropertyFlags.HasAnyFlags(EPropertyFlags.Net))
+					{
+						Builder.Append("\t\tstatic const FName Name_").Append(Property.SourceName).Append("(TEXT(\"").Append(Property.SourceName).Append("\"));\r\n");
+					}
+				}
+				Builder.Append("\r\n");
+				Builder.Append("\t\tconst bool bIsValid = true");
+				foreach (UhtProperty Property in Class.Properties)
+				{
+					if (Property.PropertyFlags.HasAnyFlags(EPropertyFlags.Net))
+					{
+						if (!Property.bIsStaticArray)
+						{
+							Builder.Append("\r\n\t\t\t&& Name_").Append(Property.SourceName).Append(" == ClassReps[(int32)ENetFields_Private::").Append(Property.SourceName).Append("].Property->GetFName()");
+						}
+						else
+						{
+							Builder.Append("\r\n\t\t\t&& Name_").Append(Property.SourceName).Append(" == ClassReps[(int32)ENetFields_Private::").Append(Property.SourceName).Append("_STATIC_ARRAY].Property->GetFName()");
+						}
+					}
+				}
+				Builder.Append(";\r\n");
+				Builder.Append("\r\n");
+				Builder.Append("\t\tcheckf(bIsValid, TEXT(\"UHT Generated Rep Indices do not match runtime populated Rep Indices for properties in ").Append(Class.SourceName).Append("\"));\r\n");
+				Builder.Append("\t}\r\n");
+			}
+			return Builder;
+		}
+
+		private StringBuilder AppendNatives(StringBuilder Builder, UhtClass Class)
+		{
+			Builder.Append("\tvoid ").Append(Class.SourceName).Append("::StaticRegisterNatives").Append(Class.SourceName).Append("()\r\n");
+			Builder.Append("\t{\r\n");
+
+			bool bAllEditorOnly = true;
+
+			List<UhtFunction> SortedFunctions = new List<UhtFunction>();
+			foreach (UhtFunction Function in Class.Functions)
+			{
+				if (Function.FunctionFlags.HasExactFlags(EFunctionFlags.Native | EFunctionFlags.NetRequest, EFunctionFlags.Native))
+				{
+					SortedFunctions.Add(Function);
+					bAllEditorOnly &= Function.FunctionFlags.HasAnyFlags(EFunctionFlags.EditorOnly);
+				}
+			}
+			SortedFunctions.Sort((x, y) => StringComparerUE.OrdinalIgnoreCase.Compare(x.EngineName, y.EngineName));
+
+			if (SortedFunctions.Count != 0)
+			{
+				using (UhtMacroBlockEmitter BlockEmitter = new UhtMacroBlockEmitter(Builder, "WITH_EDITOR", bAllEditorOnly))
+				{
+					Builder.Append("\t\tUClass* Class = ").Append(Class.SourceName).Append("::StaticClass();\r\n");
+					Builder.Append("\t\tstatic const FNameNativePtrPair Funcs[] = {\r\n");
+					
+					foreach (UhtFunction Function in SortedFunctions)
+					{
+						BlockEmitter.Set(Function.FunctionFlags.HasAnyFlags(EFunctionFlags.EditorOnly));
+						Builder
+							.Append("\t\t\t{ ")
+							.AppendUTF8LiteralString(Function.EngineName)
+							.Append(", &")
+							.AppendClassSourceNameOrInterfaceName(Class)
+							.Append("::exec")
+							.Append(Function.EngineName)
+							.Append(" },\r\n");
+					}
+
+					// This will close the block if we have one that isn't editor only
+					BlockEmitter.Set(bAllEditorOnly);
+
+					Builder.Append("\t\t};\r\n");
+					Builder.Append("\t\tFNativeFunctionRegistrar::RegisterFunctions(Class, Funcs, UE_ARRAY_COUNT(Funcs));\r\n");
+				}
+			}
+
+			Builder.Append("\t}\r\n");
+			return Builder;
+		}
+
+		private StringBuilder AppendCallbackFunctions(StringBuilder Builder, UhtClass Class, List<UhtFunction> CallbackFunctions)
+		{
+			if (CallbackFunctions.Count > 0)
+			{
+				bool bIsInterfaceClass = Class.ClassFlags.HasAnyFlags(EClassFlags.Interface);
+				using (UhtMacroBlockEmitter BlockEmitter = new UhtMacroBlockEmitter(Builder, "WITH_EDITOR"))
+				{
+					foreach (UhtFunction Function in CallbackFunctions)
+					{
+						// Net response functions don't go into the VM
+						if (Function.FunctionFlags.HasAnyFlags(EFunctionFlags.NetResponse))
+						{
+							continue;
+						}
+
+						BlockEmitter.Set(Function.FunctionFlags.HasAnyFlags(EFunctionFlags.EditorOnly));
+
+						if (!bIsInterfaceClass)
+						{
+							Builder.Append("\tstatic FName NAME_").Append(Class.SourceName).Append('_').Append(Function.EngineName).Append(" = FName(TEXT(\"").Append(Function.EngineName).Append("\"));\r\n");
+						}
+
+						AppendNativeFunctionHeader(Builder, Function, UhtPropertyTextType.EventFunction, false, null, null, UhtFunctionExportFlags.None, "\r\n");
+
+						if (bIsInterfaceClass)
+						{
+							Builder.Append("\t{\r\n");
+
+							// assert if this is ever called directly
+							Builder
+								.Append("\t\tcheck(0 && \"Do not directly call Event functions in Interfaces. Call Execute_")
+								.Append(Function.EngineName)
+								.Append(" instead.\");\r\n");
+
+							// satisfy compiler if it's expecting a return value
+							if (Function.ReturnProperty != null)
+							{
+								string EventParmStructName = GetEventStructParametersName(Class, Function.EngineName);
+								Builder.Append("\t\t").Append(EventParmStructName).Append(" Parms;\r\n");
+								Builder.Append("\t\treturn Parms.ReturnValue;\r\n");
+							}
+							Builder.Append("\t}\r\n");
+						}
+						else
+						{
+							AppendEventFunctionPrologue(Builder, Function, Function.EngineName, 1, "\r\n");
+
+							// Cast away const just in case, because ProcessEvent isn't const
+							Builder.Append("\t\t");
+							if (Function.FunctionFlags.HasAnyFlags(EFunctionFlags.Const))
+							{
+								Builder.Append("const_cast<").Append(Class.SourceName).Append("*>(this)->");
+							}
+							Builder
+								.Append("ProcessEvent(FindFunctionChecked(")
+								.Append("NAME_")
+								.Append(Class.SourceName)
+								.Append('_')
+								.Append(Function.EngineName)
+								.Append("),")
+								.Append(Function.Children.Count > 0 ? "&Parms" : "NULL")
+								.Append(");\r\n");
+
+							AppendEventFunctionEpilogue(Builder, Function, 1, "\r\n");
+						}
+					}
+				}
+			}
+			return Builder;
+		}
+
+		private StringBuilder AppendRpcFunctions(StringBuilder Builder, UhtClass Class, List<UhtFunction> ReversedFunctions, bool bEditorOnly)
+		{
+			foreach (UhtFunction Function in ReversedFunctions)
+			{
+				if (Function.FunctionFlags.HasAnyFlags(EFunctionFlags.EditorOnly) == bEditorOnly)
+				{
+					Builder.Append("\tDEFINE_FUNCTION(").AppendClassSourceNameOrInterfaceName(Class).Append("::").Append(Function.UnMarshalAndCallName).Append(")\r\n");
+					Builder.Append("\t{\r\n");
+					AppendFunctionThunk(Builder, Class, Function);
+					Builder.Append("\t}\r\n");
+				}
+			}
+			return Builder;
+		}
+
+		private StringBuilder AppendFunctionThunk(StringBuilder Builder, UhtClass Class, UhtFunction Function)
+		{
+			// Export the GET macro for the parameters
+			foreach (UhtProperty Parameter in Function.ParameterProperties.Span)
+			{
+				Builder.Append("\t\t").AppendFunctionThunkParameterGet(Parameter).Append(";\r\n");
+			}
+
+			Builder.Append("\t\tP_FINISH;\r\n");
+			Builder.Append("\t\tP_NATIVE_BEGIN;\r\n");
+
+			// Call the validate function if there is one
+			if (!Function.FunctionExportFlags.HasAnyFlags(UhtFunctionExportFlags.CppStatic) && Function.FunctionFlags.HasAnyFlags(EFunctionFlags.NetValidate))
+			{
+				Builder.Append("\t\tif (!P_THIS->").Append(Function.CppValidationImplName).Append("(").AppendFunctionThunkParameterNames(Function).Append("))\r\n");
+				Builder.Append("\t\t{\r\n");
+				Builder.Append("\t\t\tRPC_ValidateFailed(TEXT(\"").Append(Function.CppValidationImplName).Append("\"));\r\n");
+				Builder.Append("\t\t\treturn;\r\n");   // If we got here, the validation function check failed
+				Builder.Append("\t\t}\r\n");
+			}
+
+			// Write out the return value
+			Builder.Append("\t\t");
+			UhtProperty? ReturnProperty = Function.ReturnProperty;
+			if (ReturnProperty != null)
+			{
+				Builder.Append("*(").AppendFunctionThunkReturn(ReturnProperty).Append("*)Z_Param__Result=");
+			}
+
+			// Export the call to the C++ version
+			if (Function.FunctionExportFlags.HasAnyFlags(UhtFunctionExportFlags.CppStatic))
+			{
+				Builder.Append(Function.Outer?.SourceName).Append("::").Append(Function.CppImplName).Append("(").AppendFunctionThunkParameterNames(Function).Append(");\r\n");
+			}
+			else
+			{
+				Builder.Append("P_THIS->").Append(Function.CppImplName).Append("(").AppendFunctionThunkParameterNames(Function).Append(");\r\n");
+			}
+			Builder.Append("\t\tP_NATIVE_END;\r\n");
+			return Builder;
+		}
+
+		private static void FindNoExportStructsRecursive(List<UhtScriptStruct> Out, UhtStruct Struct)
+		{
+			for (UhtStruct? Current = Struct; Current != null; Current = Current.SuperStruct)
+			{
+				// Is isn't true for noexport structs
+				if (Current is UhtScriptStruct ScriptStruct)
+				{
+					if (ScriptStruct.ScriptStructFlags.HasAnyFlags(EStructFlags.Native))
+					{
+						break;
+					}
+
+					// these are a special cases that already exists and if wrong if exported naively
+					if (!ScriptStruct.bIsAlwaysAccessible)
+					{
+						Out.Remove(ScriptStruct);
+						Out.Add(ScriptStruct);
+					}
+				}
+
+				foreach (UhtType Type in Current.Children)
+				{
+					if (Type is UhtProperty Property)
+					{
+						foreach (UhtType ReferenceType in Property.EnumerateReferencedTypes())
+						{
+							if (ReferenceType is UhtScriptStruct PropertyScriptStruct)
+							{
+								FindNoExportStructsRecursive(Out, PropertyScriptStruct);
+							}
+						}
+					}
+				}
+			}
+		}
+
+		private static List<UhtScriptStruct> FindNoExportStructs(UhtStruct Struct)
+		{
+			List<UhtScriptStruct> Out = new List<UhtScriptStruct>();
+			FindNoExportStructsRecursive(Out, Struct);
+			Out.Reverse();
+			return Out;
+		}
+
+		private class PropertyMemberContextImpl : IUhtPropertyMemberContext
+		{
+			private readonly UhtCodeGenerator CodeGenerator;
+			private readonly UhtStruct OuterStructInternal;
+			private readonly string OuterStructSourceNameInternal;
+			private readonly string StaticsNameInternal;
+
+			public PropertyMemberContextImpl(UhtCodeGenerator CodeGenerator, UhtStruct OuterStruct, string OuterStructSourceName, string StaticsName)
+			{
+				this.CodeGenerator = CodeGenerator;
+				this.OuterStructInternal = OuterStruct;
+				this.StaticsNameInternal = StaticsName;
+				this.OuterStructSourceNameInternal = OuterStructSourceName.Length == 0 ? OuterStruct.SourceName : OuterStructSourceName;
+			}
+
+			public UhtStruct OuterStruct { get => this.OuterStructInternal; }
+			public string OuterStructSourceName { get => this.OuterStructSourceNameInternal; }
+			public string StaticsName { get => this.StaticsNameInternal; }
+			public string NamePrefix { get => "NewProp_"; }
+			public string MetaDataSuffix { get => "_MetaData"; }
+
+			public string GetSingletonName(UhtObject? Object, bool bRegistered)
+			{
+				return this.CodeGenerator.GetSingletonName(Object, bRegistered);
+			}
+
+			public uint GetTypeHash(UhtObject Object)
+			{
+				return this.CodeGenerator.ObjectInfos[Object.ObjectTypeIndex].Hash;
+			}
+		}
+	}
+
+	/// <summary>
+	/// Collection of string builder extensions used to generate the cpp files for individual headers.
+	/// </summary>
+	public static class UhtHeaderCodeGeneratorCppFileStringBuilderExtensinos
+	{
+
+		/// <summary>
+		/// Append the parameter names for a function
+		/// </summary>
+		/// <param name="Builder">Destination builder</param>
+		/// <param name="Function">Function in question</param>
+		/// <returns>Destination builder</returns>
+		public static StringBuilder AppendFunctionThunkParameterNames(this StringBuilder Builder, UhtFunction Function)
+		{
+			bool bFirst = true;
+			foreach (UhtProperty Parameter in Function.ParameterProperties.Span)
+			{
+				if (!bFirst)
+				{
+					Builder.Append(',');
+				}
+				Builder.AppendFunctionThunkParameterArg(Parameter);
+				bFirst = false;
+			}
+			return Builder;
+		}
+
+		/// <summary>
+		/// Append the parameter list and count as arguments to a structure constructor
+		/// </summary>
+		/// <param name="Builder">Destination builder</param>
+		/// <param name="Struct">Structure in question</param>
+		/// <param name="StaticsName">Name of the statics section for this structure</param>
+		/// <param name="Tabs">Number of tabs to precede the line</param>
+		/// <param name="Endl">String used to terminate a line</param>
+		/// <returns>Destination builder</returns>
+		public static StringBuilder AppendPropertiesParams(this StringBuilder Builder, UhtStruct Struct, string StaticsName, int Tabs, string Endl)
+		{
+			return Builder.AppendPropertiesParamsList(Struct, StaticsName, Tabs, Endl).AppendPropertiesParamsCount(Struct, StaticsName, Tabs, Endl);
+		}
+
+		/// <summary>
+		/// Append the parameter list as an argument to a structure constructor
+		/// </summary>
+		/// <param name="Builder">Destination builder</param>
+		/// <param name="Struct">Structure in question</param>
+		/// <param name="StaticsName">Name of the statics section for this structure</param>
+		/// <param name="Tabs">Number of tabs to precede the line</param>
+		/// <param name="Endl">String used to terminate a line</param>
+		/// <returns>Destination builder</returns>
+		public static StringBuilder AppendPropertiesParamsList(this StringBuilder Builder, UhtStruct Struct, string StaticsName, int Tabs, string Endl)
+		{
+			if (!Struct.Children.Any(x => x is UhtProperty))
+			{
+				Builder.AppendTabs(Tabs).Append("nullptr,").Append(Endl);
+			}
+			else if (Struct.Properties.Any(x => !x.bIsEditorOnlyProperty))
+			{
+				Builder.AppendTabs(Tabs).Append(StaticsName).Append("::PropPointers,").Append(Endl);
+			}
+			else
+			{
+				Builder.AppendTabs(Tabs).Append("IF_WITH_EDITORONLY_DATA(").Append(StaticsName).Append("::PropPointers, nullptr),").Append(Endl);
+			}
+			return Builder;
+		}
+
+		/// <summary>
+		/// Append the parameter count as an argument to a structure constructor
+		/// </summary>
+		/// <param name="Builder">Destination builder</param>
+		/// <param name="Struct">Structure in question</param>
+		/// <param name="StaticsName">Name of the statics section for this structure</param>
+		/// <param name="Tabs">Number of tabs to precede the line</param>
+		/// <param name="Endl">String used to terminate a line</param>
+		/// <returns>Destination builder</returns>
+		public static StringBuilder AppendPropertiesParamsCount(this StringBuilder Builder, UhtStruct Struct, string StaticsName, int Tabs, string Endl)
+		{
+			if (!Struct.Children.Any(x => x is UhtProperty))
+			{
+				Builder.AppendTabs(Tabs).Append("0,").Append(Endl);
+			}
+			else if (Struct.Properties.Any(x => !x.bIsEditorOnlyProperty))
+			{
+				Builder.AppendTabs(Tabs).Append("UE_ARRAY_COUNT(").Append(StaticsName).Append("::PropPointers),").Append(Endl);
+			}
+			else
+			{
+				Builder.AppendTabs(Tabs).Append("IF_WITH_EDITORONLY_DATA(UE_ARRAY_COUNT(").Append(StaticsName).Append("::PropPointers), 0),").Append(Endl);
+			}
+			return Builder;
+		}
+
+		/// <summary>
+		/// Append the given array list and count as arguments to a structure constructor
+		/// </summary>
+		/// <param name="Builder">Destination builder</param>
+		/// <param name="bIsEmpty">If true, the array list is empty</param>
+		/// <param name="bAllEditorOnlyData">If true, the array is all editor only data</param>
+		/// <param name="StaticsName">Name of the statics section</param>
+		/// <param name="ArrayName">The name of the arrray</param>
+		/// <returns>Destination builder</returns>
+		public static StringBuilder AppendArray(this StringBuilder Builder, bool bIsEmpty, bool bAllEditorOnlyData, string StaticsName, string ArrayName)
+		{
+			if (bIsEmpty)
+			{
+				Builder.Append("nullptr, 0");
+			}
+			else if (bAllEditorOnlyData)
+			{
+				Builder
+					.Append("IF_WITH_EDITORONLY_DATA(")
+					.Append(StaticsName).Append("::").Append(ArrayName)
+					.Append(", nullptr), IF_WITH_EDITORONLY_DATA(UE_ARRAY_COUNT(")
+					.Append(StaticsName).Append("::").Append(ArrayName)
+					.Append("), 0)");
+			}
+			else
+			{
+				Builder
+					.Append(StaticsName).Append("::").Append(ArrayName)
+					.Append(", UE_ARRAY_COUNT(")
+					.Append(StaticsName).Append("::").Append(ArrayName)
+					.Append(")");
+			}
+			return Builder;
+		}
+
+	}
+
+}
