UHT in UBT initial version.  Currently disabled.

diff --git a/Engine/Source/Programs/Shared/EpicGames.UHT/Exporters/CodeGen/UhtHeaderCodeGeneratorHFile.cs b/Engine/Source/Programs/Shared/EpicGames.UHT/Exporters/CodeGen/UhtHeaderCodeGeneratorHFile.cs
new file mode 100644
--- /dev/null
+++ b/Engine/Source/Programs/Shared/EpicGames.UHT/Exporters/CodeGen/UhtHeaderCodeGeneratorHFile.cs
@@ -0,0 +1,1397 @@
+// Copyright Epic Games, Inc. All Rights Reserved.
+
+using EpicGames.Core;
+using EpicGames.UHT.Types;
+using EpicGames.UHT.Utils;
+using System;
+using System.Collections.Generic;
+using System.IO;
+using System.Linq;
+using System.Text;
+
+namespace EpicGames.UHT.Exporters.CodeGen
+{
+	internal class UhtHeaderCodeGeneratorHFile
+		: UhtHeaderCodeGenerator
+	{
+		public static string RigVMParameterPrefix = "FRigVMExecuteContext& RigVMExecuteContext";
+
+		/// <summary>
+		/// Construct an instance of this generator object
+		/// </summary>
+		/// <param name="CodeGenerator">The base code generator</param>
+		/// <param name="Package">Package being generated</param>
+		/// <param name="HeaderFile">Header file being generated</param>
+		public UhtHeaderCodeGeneratorHFile(UhtCodeGenerator CodeGenerator, UhtPackage Package, UhtHeaderFile HeaderFile)
+			: base(CodeGenerator, Package, HeaderFile)
+		{
+		}
+
+		/// <summary>
+		/// For a given UE header file, generated the generated H file
+		/// </summary>
+		/// <param name="HeaderOutput">Output object</param>
+		public void Generate(IUhtExportOutput HeaderOutput)
+		{
+			ref UhtCodeGenerator.HeaderInfo HeaderInfo = ref this.HeaderInfos[this.HeaderFile.HeaderFileTypeIndex];
+			using (BorrowStringBuilder Borrower = new BorrowStringBuilder(StringBuilderCache.Big))
+			{
+				StringBuilder Builder = Borrower.StringBuilder;
+
+				Builder.Append(HeaderCopyright);
+				Builder.Append("#include \"UObject/ObjectMacros.h\"\r\n");
+				Builder.Append("#include \"UObject/ScriptMacros.h\"\r\n");
+				if (HeaderInfo.bNeedsPushModelHeaders)
+				{
+					Builder.Append("#include \"Net/Core/PushModel/PushModelMacros.h\"\r\n");
+				}
+				Builder.Append("\r\n");
+				Builder.Append(DisableDeprecationWarnings).Append("\r\n");
+
+				string StrippedName = Path.GetFileNameWithoutExtension(this.HeaderFile.FilePath);
+				string DefineName = $"{this.Package.ShortName.ToString().ToUpper()}_{StrippedName}_generated_h";
+
+				if (this.HeaderFile.References.ForwardDeclarations.Count > 0)
+				{
+					string[] Sorted = new string[this.HeaderFile.References.ForwardDeclarations.Count];
+					int Index = 0;
+					foreach (string ForwardDeclaration in this.HeaderFile.References.ForwardDeclarations)
+					{
+						Sorted[Index++] = ForwardDeclaration;
+					}
+					Array.Sort(Sorted, StringComparerUE.OrdinalIgnoreCase);
+					foreach (string ForwardDeclaration in Sorted)
+					{
+						Builder.Append(ForwardDeclaration).Append("\r\n");
+					}
+				}
+
+				Builder.Append("#ifdef ").Append(DefineName).Append("\r\n");
+				Builder.Append("#error \"").Append(StrippedName).Append(".generated.h already included, missing '#pragma once' in ").Append(StrippedName).Append(".h\"\r\n");
+				Builder.Append("#endif\r\n");
+				Builder.Append("#define ").Append(DefineName).Append("\r\n");
+				Builder.Append("\r\n");
+
+				foreach (UhtField Field in this.HeaderFile.References.ExportTypes)
+				{
+					if (Field is UhtEnum Enum)
+					{
+					}
+					else if (Field is UhtScriptStruct ScriptStruct)
+					{
+						if (ScriptStruct.ScriptStructFlags.HasAnyFlags(EStructFlags.Native))
+						{
+							AppendScriptStruct(Builder, ScriptStruct);
+						}
+					}
+					else if (Field is UhtFunction Function)
+					{
+						AppendDelegate(Builder, Function);
+					}
+					else if (Field is UhtClass Class)
+					{
+						if (!Class.ClassFlags.HasAnyFlags(EClassFlags.Intrinsic))
+						{
+							AppendClass(Builder, Class);
+						}
+					}
+				}
+
+				Builder.Append("#undef CURRENT_FILE_ID\r\n");
+				Builder.Append("#define CURRENT_FILE_ID ").Append(HeaderInfo.FileId).Append("\r\n\r\n\r\n");
+
+				foreach (UhtField Field in this.HeaderFile.References.ExportTypes)
+				{
+					if (Field is UhtEnum Enum)
+					{
+						AppendEnum(Builder, Enum);
+					}
+				}
+
+				Builder.Append(EnableDeprecationWarnings).Append("\r\n");
+				HeaderOutput.CommitOutput(Builder);
+			}
+		}
+
+		private StringBuilder AppendScriptStruct(StringBuilder Builder, UhtScriptStruct ScriptStruct)
+		{
+			if (ScriptStruct.RigVMStructInfo != null)
+			{
+				const string RigVMParameterPrefix = "FRigVMExecuteContext& RigVMExecuteContext";
+
+				Builder.Append("\r\n");
+				foreach (UhtRigVMMethodInfo MethodInfo in ScriptStruct.RigVMStructInfo.Methods)
+				{
+					Builder.Append("#define ").Append(ScriptStruct.SourceName).Append('_').Append(MethodInfo.Name).Append("() \\\r\n");
+					Builder.Append("\t").Append(MethodInfo.ReturnType).Append(' ').Append(ScriptStruct.SourceName).Append("::Static").Append(MethodInfo.Name).Append("( \\\r\n");
+					Builder.Append("\t\t").Append(RigVMParameterPrefix);
+					Builder.AppendParameterDecls(ScriptStruct.RigVMStructInfo.Members, true, ", \\\r\n\t\t", true, false);
+					Builder.AppendParameterDecls(MethodInfo.Parameters, true, ", \\\r\n\t\t", false, false);
+					Builder.Append(" \\\r\n");
+					Builder.Append("\t)\r\n");
+				}
+				Builder.Append("\r\n");
+			}
+
+			if (ScriptStruct.ScriptStructFlags.HasAnyFlags(EStructFlags.Native))
+			{
+				using (UhtMacroCreator Macro = new UhtMacroCreator(Builder, this, ScriptStruct, GeneratedBodyMacroSuffix))
+				{
+					Builder.Append("\tfriend struct Z_Construct_UScriptStruct_").Append(ScriptStruct.SourceName).Append("_Statics; \\\r\n");
+					Builder.Append("\t");
+					if (!ScriptStruct.ScriptStructFlags.HasAnyFlags(EStructFlags.RequiredAPI))
+					{
+						Builder.Append(this.PackageApi);
+					}
+					Builder.Append("static class UScriptStruct* StaticStruct(); \\\r\n");
+
+					// if we have RigVM methods on this struct we need to 
+					// declare the static method as well as the stub method
+					if (ScriptStruct.RigVMStructInfo != null)
+					{
+						foreach (UhtRigVMMethodInfo MethodInfo in ScriptStruct.RigVMStructInfo.Methods)
+						{
+							Builder.Append("\tstatic ").Append(MethodInfo.ReturnType).Append(" Static").Append(MethodInfo.Name).Append("( \\\r\n");
+							Builder.Append("\t\t").Append(RigVMParameterPrefix);
+							Builder.AppendParameterDecls(ScriptStruct.RigVMStructInfo.Members, true, ", \\\r\n\t\t", true, false);
+							Builder.AppendParameterDecls(MethodInfo.Parameters, true, ", \\\r\n\t\t", false, false);
+							Builder.Append(" \\\r\n");
+							Builder.Append("\t); \\\r\n");
+
+							Builder.Append("\tFORCEINLINE_DEBUGGABLE static ").Append(MethodInfo.ReturnType).Append(" RigVM").Append(MethodInfo.Name).Append("( \\\r\n");
+							Builder.Append("\t\t").Append(RigVMParameterPrefix).Append(", \\\r\n");
+							Builder.Append("\t\tFRigVMMemoryHandleArray RigVMMemoryHandles \\\r\n");
+							Builder.Append("\t) \\\r\n");
+							Builder.Append("\t{ \\\r\n");
+
+							// implement inline stub method body
+							if (MethodInfo.Parameters.Count > 0)
+							{
+								for (int ParameterIndex = 0; ParameterIndex < MethodInfo.Parameters.Count; ParameterIndex++)
+								{
+									UhtRigVMParameter Parameter = MethodInfo.Parameters[ParameterIndex];
+									Builder.Append("\t\t").Append(Parameter.Declaration()).Append(" = *(").Append(Parameter.TypeNoRef())
+										.Append("*)RigVMExecuteContext.OpaqueArguments[").Append(ParameterIndex).Append("]; \\\r\n");
+								}
+								Builder.Append("\t\t \\\r\n");
+							}
+
+							if (ScriptStruct.RigVMStructInfo.Members.Count > 0)
+							{
+								int OperandIndex = 0;
+								foreach (UhtRigVMParameter Parameter in ScriptStruct.RigVMStructInfo.Members)
+								{
+									string ParamTypeOriginal = Parameter.TypeOriginal(true);
+									string ParamNameOriginal = Parameter.NameOriginal(false);
+									string AdditionalParameters = string.Empty;
+									if (!Parameter.bInput && !Parameter.bOutput && !Parameter.bSingleton)
+									{
+										AdditionalParameters = ", RigVMExecuteContext.GetSlice().GetIndex()";
+									}
+
+									if (Parameter.bIsArray)
+									{
+										string ExtendedType = Parameter.ExtendedType();
+										Builder
+											.Append("\t\tTArray")
+											.Append(ExtendedType)
+											.Append("& ")
+											.Append(ParamNameOriginal)
+											.Append(" = *(TArray")
+											.Append(ExtendedType)
+											.Append("*)RigVMMemoryHandles[")
+											.Append(OperandIndex)
+											.Append("].GetData(false")
+											.Append(AdditionalParameters)
+											.Append("); \\\r\n");
+										OperandIndex++;
+									}
+									else
+									{
+										string VariableType = Parameter.TypeVariableRef(true);
+										string ExtractedType = Parameter.TypeOriginal();
+
+										Builder
+											.Append("\t\t")
+											.Append(VariableType)
+											.Append(' ')
+											.Append(ParamNameOriginal)
+											.Append(" = *(")
+											.Append(ExtractedType)
+											.Append("*)RigVMMemoryHandles[")
+											.Append(OperandIndex)
+											.Append("].GetData(false")
+											.Append(AdditionalParameters)
+											.Append("); \\\r\n");
+										OperandIndex++;
+									}
+								}
+								Builder.Append("\t\t \\\r\n");
+							}
+
+							Builder.Append("\t\t").Append(MethodInfo.ReturnPrefix()).Append("Static").Append(MethodInfo.Name).Append("( \\\r\n");
+							Builder.Append("\t\t\tRigVMExecuteContext");
+							Builder.AppendParameterNames(ScriptStruct.RigVMStructInfo.Members, true, ", \\\r\n\t\t\t", false);
+							Builder.AppendParameterNames(MethodInfo.Parameters, true, ", \\\r\n\t\t\t");
+							Builder.Append(" \\\r\n");
+							Builder.Append("\t\t); \\\r\n");
+							Builder.Append("\t} \\\r\n");
+						}
+					}
+
+					if (ScriptStruct.SuperScriptStruct != null)
+					{
+						Builder.Append("\ttypedef ").Append(ScriptStruct.SuperScriptStruct.SourceName).Append(" Super;\\\\r\n");
+					}
+				}
+
+				// Forward declare the StaticStruct specialization in the header
+				Builder.Append("template<> ").Append(this.PackageApi).Append("UScriptStruct* StaticStruct<struct ").Append(ScriptStruct.SourceName).Append(">();\r\n");
+				Builder.Append("\r\n");
+			}
+			return Builder;
+		}
+
+		private StringBuilder AppendEnum(StringBuilder Builder, UhtEnum Enum)
+		{
+			// Export FOREACH macro
+			Builder.Append("#define FOREACH_ENUM_").Append(Enum.EngineName.ToUpper()).Append("(op) ");
+			bool bHasExistingMax = EnumHasExistingMax(Enum);
+			long MaxEnumValue = bHasExistingMax ? GetMaxEnumValue(Enum) : 0;
+			foreach (UhtEnumValue EnumValue in Enum.EnumValues)
+			{
+				if (!bHasExistingMax || EnumValue.Value != MaxEnumValue)
+				{
+					Builder.Append("\\\r\n\top(").Append(EnumValue.Name).Append(") ");
+				}
+			}
+			Builder.Append("\r\n");
+
+			// Forward declare the StaticEnum<> specialization for enum classes
+			if (Enum.CppForm == UhtEnumCppForm.EnumClass)
+			{
+				Builder.Append("\r\n");
+				Builder.Append("enum class ").Append(Enum.CppType);
+				if (Enum.UnderlyingType != UhtEnumUnderlyingType.Unspecified)
+				{
+					Builder.Append(" : ").Append(Enum.UnderlyingType.ToString());
+				}
+				Builder.Append(";\r\n");
+				Builder.Append("template<> ").Append(this.PackageApi).Append("UEnum* StaticEnum<").Append(Enum.CppType).Append(">();\r\n");
+				Builder.Append("\r\n");
+			}
+
+			return Builder;
+		}
+
+		private StringBuilder AppendDelegate(StringBuilder Builder, UhtFunction Function)
+		{
+			using (UhtMacroCreator Macro = new UhtMacroCreator(Builder, this, Function, DelegateMacroSuffix))
+			{
+				int Tabs = 0;
+				string StrippedFunctionName = Function.StrippedFunctionName;
+				string ExportFunctionName = GetDelegateFunctionExportName(Function);
+				string ExtraParameter = GetDelegateFunctionExtraParameter(Function);
+
+				AppendEventParameter(Builder, Function, StrippedFunctionName, UhtPropertyTextType.EventParameterMember, true, Tabs, " \\\r\n");
+				Builder.Append("static ");
+				AppendNativeFunctionHeader(Builder, Function, UhtPropertyTextType.EventFunction, true, ExportFunctionName, ExtraParameter, UhtFunctionExportFlags.Inline, " \\\r\n");
+				AppendEventFunctionPrologue(Builder, Function, StrippedFunctionName, Tabs, " \\\r\n");
+				Builder
+					.Append("\t")
+					.Append(StrippedFunctionName)
+					.Append('.')
+					.Append(Function.FunctionFlags.HasAnyFlags(EFunctionFlags.MulticastDelegate) ? "ProcessMulticastDelegate" : "ProcessDelegate")
+					.Append("<UObject>(")
+					.Append(Function.Children.Count > 0 ? "&Parms" : "NULL")
+					.Append("); \\\r\n");
+				AppendEventFunctionEpilogue(Builder, Function, Tabs, " \\\r\n");
+			}
+			return Builder;
+		}
+
+		private StringBuilder AppendClass(StringBuilder Builder, UhtClass Class)
+		{
+			string Api = Class.ClassFlags.HasAnyFlags(EClassFlags.MinimalAPI) ? this.PackageApi : "NO_API ";
+			UhtClass? NativeInterface = this.ObjectInfos[Class.ObjectTypeIndex].NativeInterface;
+
+			// Write the spare declarations
+			AppendSparseDeclarations(Builder, Class);
+
+			// Collect the functions in reversed order
+			List<UhtFunction> ReversedFunctions = new List<UhtFunction>(Class.Functions);
+			ReversedFunctions.Reverse();
+
+			// Check to see if we have any RPC functions for the editor
+			bool bHasEditorRpc = ReversedFunctions.Any(x => IsRpcFunction(x, true));
+
+			// Output the RPC methods
+			AppendRpcFunctions(Builder, Class, ReversedFunctions, false);
+			if (bHasEditorRpc)
+			{
+				AppendRpcFunctions(Builder, Class, ReversedFunctions, true);
+			}
+
+			// Output property accessors
+			AppendPropertyAccessors(Builder, Class);
+
+			// Collect the callback function and sort by name to make the order stable
+			List<UhtFunction> CallbackFunctions = new List<UhtFunction>(Class.Functions.Where(x => x.FunctionFlags.HasAnyFlags(EFunctionFlags.Event) && x.SuperFunction == null));
+			CallbackFunctions.Sort((x, y) => StringComparerUE.OrdinalIgnoreCase.Compare(x.EngineName, y.EngineName));
+
+			// Generate the callback parameter structures
+			AppendCallbackParametersDecls(Builder, Class, CallbackFunctions);
+
+			// Generate the RPC wrappers for the callbacks
+			AppendCallbackRpcWrapperDecls(Builder, Class, CallbackFunctions);
+
+			// Only write out adapters if the user has provided one or the other of the Serialize overloads
+			if (Class.SerializerArchiveType != UhtSerializerArchiveType.None && Class.SerializerArchiveType != UhtSerializerArchiveType.All)
+			{
+				AppendSerializer(Builder, Class, Api, UhtSerializerArchiveType.Archive, "DECLARE_FARCHIVE_SERIALIZER");
+				AppendSerializer(Builder, Class, Api, UhtSerializerArchiveType.StructuredArchiveRecord, "DECLARE_FSTRUCTUREDARCHIVE_SERIALIZER");
+			}
+
+			if (Class.ClassFlags.HasAnyFlags(EClassFlags.Interface))
+			{
+				AppendStandardConstructors(Builder, Class, Api);
+				AppendEnhancedConstructors(Builder, Class, Api);
+
+				using (UhtMacroCreator Macro = new UhtMacroCreator(Builder, this, Class, GeneratedUInterfaceBodyMacroSuffix))
+				{
+					AppendCommonGeneratedBody(Builder, Class, Api);
+				}
+
+				using (UhtMacroCreator Macro = new UhtMacroCreator(Builder, this, Class, GeneratedBodyLegacyMacroSuffix))
+				{
+					Builder.Append('\t');
+					AppendGeneratedMacroDeprecationWarning(Builder, "GENERATED_UINTERFACE_BODY");
+					Builder.Append('\t').Append(DisableDeprecationWarnings).Append(" \\\r\n");
+					Builder.Append('\t').AppendMacroName(this, Class, GeneratedUInterfaceBodyMacroSuffix).Append(" \\\r\n");
+					Builder.Append('\t').AppendMacroName(this, Class, StandardConstructorsMacroSuffix).Append(" \\\r\n");
+					Builder.Append('\t').Append(EnableDeprecationWarnings).Append(" \\\r\n");
+				}
+
+				using (UhtMacroCreator Macro = new UhtMacroCreator(Builder, this, Class, GeneratedBodyMacroSuffix))
+				{
+					Builder.Append('\t').Append(DisableDeprecationWarnings).Append(" \\\r\n");
+					Builder.Append('\t').AppendMacroName(this, Class, GeneratedUInterfaceBodyMacroSuffix).Append(" \\\r\n");
+					Builder.Append('\t').AppendMacroName(this, Class, EnchancedConstructorsMacroSuffix).Append(" \\\r\n");
+					AppendAccessSpecifier(Builder, Class);
+					Builder.Append(" \\\r\n");
+					Builder.Append('\t').Append(EnableDeprecationWarnings).Append(" \\\r\n");
+				}
+
+				using (UhtMacroCreator Macro = new UhtMacroCreator(Builder, this, Class, InClassIInterfaceNoPureDeclsMacroSuffix))
+				{
+					AppendInClassIInterface(Builder, Class, CallbackFunctions, Api);
+				}
+
+				using (UhtMacroCreator Macro = new UhtMacroCreator(Builder, this, Class, InClassIInterfaceMacroSuffix))
+				{
+					AppendInClassIInterface(Builder, Class, CallbackFunctions, Api);
+				}
+			}
+			else
+			{
+				using (UhtMacroCreator Macro = new UhtMacroCreator(Builder, this, Class, InClassNoPureDeclsMacroSuffix))
+				{
+					AppendClassGeneratedBody(Builder, Class, Api);
+				}
+
+				using (UhtMacroCreator Macro = new UhtMacroCreator(Builder, this, Class, InClassMacroSuffix))
+				{
+					AppendClassGeneratedBody(Builder, Class, Api);
+				}
+
+				AppendStandardConstructors(Builder, Class, Api);
+				AppendEnhancedConstructors(Builder, Class, Api);
+			}
+
+			using (UhtMacroCreator Macro = new UhtMacroCreator(Builder, this, Class.PrologLineNumber, PrologMacroSuffix))
+			{
+				if (CallbackFunctions.Count > 0)
+				{
+					Builder.Append("\t").AppendMacroName(this, Class, EventParamsMacroSuffix).Append(" \\\r\n");
+				}
+			}
+
+			bool bHasCallbacks = CallbackFunctions.Count > 0;
+			if (Class.ClassFlags.HasAnyFlags(EClassFlags.Interface))
+			{
+				if (NativeInterface != null)
+				{
+					AppendGeneratedBodyMacroBlock(Builder, Class, NativeInterface, true, bHasEditorRpc, bHasCallbacks, null);
+					AppendGeneratedBodyMacroBlock(Builder, Class, NativeInterface, false, bHasEditorRpc, bHasCallbacks, null);
+				}
+			}
+			else
+			{
+				AppendGeneratedBodyMacroBlock(Builder, Class, Class, true, bHasEditorRpc, bHasCallbacks, "GENERATED_UCLASS_BODY");
+				AppendGeneratedBodyMacroBlock(Builder, Class, Class, false, bHasEditorRpc, bHasCallbacks, null);
+			}
+
+			// Forward declare the StaticClass specialization in the header
+			Builder.Append("template<> ").Append(this.PackageApi).Append("UClass* StaticClass<class ").Append(Class.SourceName).Append(">();\r\n");
+			Builder.Append("\r\n");
+			return Builder;
+		}
+
+		private StringBuilder AppendSparseDeclarations(StringBuilder Builder, UhtClass Class)
+		{
+			// Format the sparse data
+			using (UhtMacroCreator Macro = new UhtMacroCreator(Builder, this, Class, SparseDataMacroSuffix))
+			{
+				string[]? SparseDataTypes = Class.MetaData.GetStringArray(UhtNames.SparseClassDataTypes);
+				if (SparseDataTypes != null)
+				{
+					foreach (string SparseDataType in SparseDataTypes)
+					{
+						Builder.Append("F").Append(SparseDataType).Append("* Get").Append(SparseDataType).Append("() \\\r\n");
+						Builder.Append("{ \\\r\n");
+						Builder.Append("\treturn (F").Append(SparseDataType).Append("*)(GetClass()->GetOrCreateSparseClassData()); \\\r\n");
+						Builder.Append("} \\\r\n");
+
+						Builder.Append("F").Append(SparseDataType).Append("* Get").Append(SparseDataType).Append("() const \\\r\n");
+						Builder.Append("{ \\\r\n");
+						Builder.Append("\treturn (F").Append(SparseDataType).Append("*)(GetClass()->GetOrCreateSparseClassData()); \\\r\n");
+						Builder.Append("} \\\r\n");
+
+						Builder.Append("const F").Append(SparseDataType).Append("* Get").Append(SparseDataType).Append("(EGetSparseClassDataMethod GetMethod) const \\\r\n");
+						Builder.Append("{ \\\r\n");
+						Builder.Append("\treturn (const F").Append(SparseDataType).Append("*)(GetClass()->GetSparseClassData(GetMethod)); \\\r\n");
+						Builder.Append("} \\\r\n");
+					}
+
+					foreach (string SparseDataType in SparseDataTypes)
+					{
+						UhtScriptStruct? SparseScriptStruct = Class.FindType(UhtFindOptions.EngineName | UhtFindOptions.ScriptStruct | UhtFindOptions.NoSelf, SparseDataType) as UhtScriptStruct;
+						while (SparseScriptStruct != null)
+						{
+							foreach (UhtProperty SparseProperty in SparseScriptStruct.Properties)
+							{
+								if (!SparseProperty.MetaData.ContainsKey(UhtNames.NoGetter))
+								{
+									string PropertyName = SparseProperty.SourceName;
+									string CleanPropertyName = PropertyName;
+									if (SparseProperty is UhtBooleanProperty && PropertyName.StartsWith("b"))
+									{
+										CleanPropertyName = PropertyName.Substring(1);
+									}
+
+									bool bGetByRef = SparseProperty.MetaData.ContainsKey(UhtNames.GetByRef);
+
+									if (bGetByRef)
+									{
+										Builder.Append("const ").AppendSparse(SparseProperty).Append("& Get").Append(CleanPropertyName).Append("() \\\r\n");
+									}
+									else
+									{
+										Builder.AppendSparse(SparseProperty).Append(" Get").Append(CleanPropertyName).Append("() \\\r\n");
+									}
+									Builder.Append("{ \\\r\n");
+									Builder.Append("\treturn Get").Append(SparseDataType).Append("()->").Append(PropertyName).Append("; \\\r\n");
+									Builder.Append("} \\\r\n");
+
+									if (bGetByRef)
+									{
+										Builder.Append("const ").AppendSparse(SparseProperty).Append("& Get").Append(CleanPropertyName).Append("() const \\\r\n");
+									}
+									else
+									{
+										Builder.AppendSparse(SparseProperty).Append(" Get").Append(CleanPropertyName).Append("() const \\\r\n");
+									}
+									Builder.Append("{ \\\r\n");
+									Builder.Append("\treturn Get").Append(SparseDataType).Append("()->").Append(PropertyName).Append("; \\\r\n");
+									Builder.Append("} \\\r\n");
+								}
+							}
+
+							SparseScriptStruct = SparseScriptStruct.SuperScriptStruct;
+						}
+					}
+				}
+			}
+			return Builder;
+		}
+
+		private StringBuilder AppendRpcFunctions(StringBuilder Builder, UhtClass Class, List<UhtFunction> ReversedFunctions, bool bEditorOnly)
+		{
+			Builder.AppendBeginEditorOnlyGuard(bEditorOnly);
+
+			using (UhtMacroCreator Macro = new UhtMacroCreator(Builder, this, Class, bEditorOnly ? EditorOnlyRpcWrappersMacroSuffix : RpcWrappersMacroSuffix))
+			{
+				AppendAutogeneratedBlueprintFunctionDeclarations(Builder, Class, ReversedFunctions, bEditorOnly);
+				AppendRpcWrappers(Builder, Class, ReversedFunctions, bEditorOnly);
+			}
+
+			using (UhtMacroCreator Macro = new UhtMacroCreator(Builder, this, Class, bEditorOnly ? EditorOnlyRpcWrappersNoPureDeclsMacroSuffix : RpcWrappersNoPureDeclsMacroSuffix))
+			{
+				if (Class.GeneratedCodeVersion <= EGeneratedCodeVersion.V1)
+				{
+					AppendAutogeneratedBlueprintFunctionDeclarationsOnlyNotDeclared(Builder, Class, ReversedFunctions, bEditorOnly);
+				}
+				AppendRpcWrappers(Builder, Class, ReversedFunctions, bEditorOnly);
+			}
+
+			if (bEditorOnly)
+			{
+				Builder.Append("#else\r\n");
+				using (UhtMacroCreator Macro = new UhtMacroCreator(Builder, this, Class, bEditorOnly ? EditorOnlyRpcWrappersMacroSuffix : RpcWrappersMacroSuffix))
+				{
+				}
+
+				using (UhtMacroCreator Macro = new UhtMacroCreator(Builder, this, Class, bEditorOnly ? EditorOnlyRpcWrappersNoPureDeclsMacroSuffix : RpcWrappersNoPureDeclsMacroSuffix))
+				{
+				}
+				Builder.AppendEndEditorOnlyGuard();
+			}
+			return Builder;
+		}
+
+		private StringBuilder AppendPropertyAccessors(StringBuilder Builder, UhtClass Class)
+		{
+			using (UhtMacroCreator Macro = new UhtMacroCreator(Builder, this, Class, AccessorsMacroSuffix))
+			{
+				foreach (UhtType Type in Class.Children)
+				{
+					if (Type is UhtProperty Property)
+					{
+						if (Property.PropertyExportFlags.HasAnyFlags(UhtPropertyExportFlags.GetterFound))
+						{
+							Builder.Append("static void ").AppendPropertyGetterWrapperName(Property).Append("(const void* Object, void* OutValue); \\\r\n");
+						}
+						if (Property.PropertyExportFlags.HasAnyFlags(UhtPropertyExportFlags.SetterFound))
+						{
+							Builder.Append("static void ").AppendPropertySetterWrapperName(Property).Append("(void* Object, const void* InValue); \\\r\n");
+						}
+					}
+				}
+			}
+			return Builder;
+		}
+
+		private StringBuilder AppendAutogeneratedBlueprintFunctionDeclarations(StringBuilder Builder, UhtClass Class, List<UhtFunction> ReversedFunctions, bool bEditorOnly)
+		{
+			foreach (UhtFunction Function in ReversedFunctions)
+			{
+				if (!IsRpcFunction(Function, bEditorOnly) || Function.CppImplName == Function.SourceName)
+				{
+					continue;
+				}
+				AppendNetValidateDeclaration(Builder, Function);
+				AppendNativeFunctionHeader(Builder, Function, UhtPropertyTextType.ClassFunction, true, null, null, UhtFunctionExportFlags.None, "; \\\r\n");
+			}
+			return Builder;
+		}
+
+		private StringBuilder AppendAutogeneratedBlueprintFunctionDeclarationsOnlyNotDeclared(StringBuilder Builder, UhtClass Class, List<UhtFunction> ReversedFunctions, bool bEditorOnly)
+		{
+			foreach (UhtFunction Function in ReversedFunctions)
+			{
+				if (!IsRpcFunction(Function, bEditorOnly) || Function.CppImplName == Function.SourceName)
+				{
+					continue;
+				}
+				if (!Function.FunctionExportFlags.HasAnyFlags(UhtFunctionExportFlags.ImplFound))
+				{
+					AppendNetValidateDeclaration(Builder, Function);
+					AppendNativeFunctionHeader(Builder, Function, UhtPropertyTextType.ClassFunction, true, null, null, UhtFunctionExportFlags.None, "; \\\r\n");
+				}
+			}
+			return Builder;
+		}
+
+		private StringBuilder AppendNetValidateDeclaration(StringBuilder Builder, UhtFunction Function)
+		{
+			if (Function.FunctionFlags.HasAnyFlags(EFunctionFlags.NetValidate))
+			{
+				Builder.Append("\t");
+				if (!Function.FunctionFlags.HasAnyFlags(EFunctionFlags.Static) && !Function.FunctionExportFlags.HasAnyFlags(UhtFunctionExportFlags.Final))
+				{
+					Builder.Append("virtual");
+				}
+				Builder.Append(" bool ").Append(Function.CppValidationImplName);
+				AppendParameters(Builder, Function, UhtPropertyTextType.ClassFunction, null, true);
+				Builder.Append("; \\\r\n");
+			}
+			return Builder;
+		}
+
+		private StringBuilder AppendRpcWrappers(StringBuilder Builder, UhtClass Class, List<UhtFunction> ReversedFunctions, bool bEditorOnly)
+		{
+			bool bFirst = true;
+			foreach (UhtFunction Function in ReversedFunctions)
+			{
+				if (!IsRpcFunction(Function, bEditorOnly))
+				{
+					continue;
+				}
+				if (!ShouldExportFunction(Function))
+				{
+					continue;
+				}
+				//COMPATIBILITY-TODO - Remove once we transition to C# version
+				if (bFirst)
+				{
+					Builder.Append(" \\\r\n");
+					bFirst = false;
+				}
+				Builder.Append("\tDECLARE_FUNCTION(").Append(Function.UnMarshalAndCallName).Append("); \\\r\n");
+			}
+			return Builder;
+		}
+
+		private StringBuilder AppendCallbackParametersDecls(StringBuilder Builder,UhtClass Class, List<UhtFunction> CallbackFunctions)
+		{
+			if (CallbackFunctions.Count > 0)
+			{
+				using (UhtMacroCreator Macro = new UhtMacroCreator(Builder, this, Class, EventParamsMacroSuffix))
+				{
+					foreach (UhtFunction Function in CallbackFunctions)
+					{
+						AppendEventParameter(Builder, Function, Function.StrippedFunctionName, UhtPropertyTextType.EventParameterMember, true, 1, " \\\r\n");
+					}
+				}
+			}
+			return Builder;
+		}
+
+		private StringBuilder AppendCallbackRpcWrapperDecls(StringBuilder Builder, UhtClass Class, List<UhtFunction> CallbackFunctions)
+		{
+			if (CallbackFunctions.Count > 0)
+			{
+				using (UhtMacroCreator Macro = new UhtMacroCreator(Builder, this, Class, CallbackWrappersMacroSuffix))
+				{
+					foreach (UhtFunction Function in CallbackFunctions)
+					{
+						if (!Function.FunctionFlags.HasAnyFlags(EFunctionFlags.NetResponse) &&
+							Function.EngineName != Function.MarshalAndCallName)
+						{
+							AppendNativeFunctionHeader(Builder, Function, UhtPropertyTextType.EventFunction, true, null, null, UhtFunctionExportFlags.None, " \\\r\n");
+							Builder.Append(" \\\r\n");
+						}
+					}
+				}
+			}
+			return Builder;
+		}
+
+		private StringBuilder AppendSerializer(StringBuilder Builder, UhtClass Class, string Api, UhtSerializerArchiveType Type, string Declare)
+		{
+			if (!Class.SerializerArchiveType.HasAnyFlags(Type))
+			{
+				if (Class.EnclosingDefine.Length != 0)
+				{
+					Builder.Append("#if ").Append(Class.EnclosingDefine).Append("\r\n");
+				}
+				AppendSerializerFunction(Builder, Class, Api, Declare);
+				if (Class.EnclosingDefine.Length != 0)
+				{
+					Builder.Append("#else\r\n");
+					using (UhtMacroCreator Macro = new UhtMacroCreator(Builder, this, Class, ArchiveSerializerMacroSuffix))
+					{
+					}
+					Builder.Append("#endif\r\n");
+				}
+			}
+			return Builder;
+		}
+
+		private StringBuilder AppendSerializerFunction(StringBuilder Builder, UhtClass Class, string Api, string Declare)
+		{
+			using (UhtMacroCreator Macro = new UhtMacroCreator(Builder, this, Class, ArchiveSerializerMacroSuffix))
+			{
+				Builder.Append('\t').Append(Declare).Append('(').Append(Class.SourceName).Append(", ").Append(Api.Substring(0, Api.Length - 1)).Append(") \\\r\n");
+			}
+			return Builder;
+		}
+
+		/// <summary>
+		/// Generates standard constructor declarations
+		/// </summary>
+		/// <param name="Builder">Output builder</param>
+		/// <param name="Class">Class being exported</param>
+		/// <param name="Api">API text to be used</param>
+		/// <returns>Output builder</returns>
+		private StringBuilder AppendStandardConstructors(StringBuilder Builder, UhtClass Class, string Api)
+		{
+			using (UhtMacroCreator Macro = new UhtMacroCreator(Builder, this, Class, StandardConstructorsMacroSuffix))
+			{
+				if (!Class.ClassExportFlags.HasAnyFlags(UhtClassExportFlags.HasCustomConstructor))
+				{
+					Builder.Append("\t/** Standard constructor, called after all reflected properties have been initialized */ \\\r\n");
+					Builder.Append("\t").Append(Api).Append(Class.SourceName).Append("(const FObjectInitializer& ObjectInitializer");
+					if (!Class.ClassExportFlags.HasAnyFlags(UhtClassExportFlags.HasDefaultConstructor))
+					{
+						Builder.Append(" = FObjectInitializer::Get()");
+					}
+					Builder.Append("); \\\r\n");
+				}
+				if (Class.ClassFlags.HasAnyFlags(EClassFlags.Abstract))
+				{
+					Builder.Append("\tDEFINE_ABSTRACT_DEFAULT_OBJECT_INITIALIZER_CONSTRUCTOR_CALL(").Append(Class.SourceName).Append(") \\\r\n");
+				}
+				else
+				{
+					Builder.Append("\tDEFINE_DEFAULT_OBJECT_INITIALIZER_CONSTRUCTOR_CALL(").Append(Class.SourceName).Append(") \\\r\n");
+				}
+
+				AppendVTableHelperCtorAndCaller(Builder, Class, Api);
+				AppendCopyConstructorDefinition(Builder, Class, Api);
+			}
+			return Builder;
+		}
+
+		/// <summary>
+		/// Generates enhanced constructor declaration.
+		/// </summary>
+		/// <param name="Builder">Output builder</param>
+		/// <param name="Class">Class being exported</param>
+		/// <param name="Api">API text to be used</param>
+		/// <returns>Output builder</returns>
+		private StringBuilder AppendEnhancedConstructors(StringBuilder Builder, UhtClass Class, string Api)
+		{
+			using (UhtMacroCreator Macro = new UhtMacroCreator(Builder, this, Class, EnchancedConstructorsMacroSuffix))
+			{
+				AppendConstructorDefinition(Builder, Class, Api);
+				AppendVTableHelperCtorAndCaller(Builder, Class, Api);
+				AppendDefaultConstructorCallDefinition(Builder, Class, Api);
+			}
+			return Builder;
+		}
+
+		/// <summary>
+		/// Generates vtable helper caller and eventual constructor body.
+		/// </summary>
+		/// <param name="Builder">Output builder</param>
+		/// <param name="Class">Class being exported</param>
+		/// <param name="Api">API text to be used</param>
+		/// <returns>Output builder</returns>
+		private StringBuilder AppendVTableHelperCtorAndCaller(StringBuilder Builder, UhtClass Class, string Api)
+		{
+			if (!Class.ClassExportFlags.HasAnyFlags(UhtClassExportFlags.HasCustomVTableHelperConstructor))
+			{
+				Builder.Append("\tDECLARE_VTABLE_PTR_HELPER_CTOR(").Append(Api.Substring(0, Api.Length - 1)).Append(", ").Append(Class.SourceName).Append("); \\\r\n");
+			}
+			Builder.Append("\tDEFINE_VTABLE_PTR_HELPER_CTOR_CALLER(").Append(Class.SourceName).Append("); \\\r\n");
+			return Builder;
+		}
+
+		/// <summary>
+		/// Generates private copy-constructor declaration.
+		/// </summary>
+		/// <param name="Builder">Output builder</param>
+		/// <param name="Class">Class being exported</param>
+		/// <param name="Api">API text to be used</param>
+		/// <returns>Output builder</returns>
+		private StringBuilder AppendCopyConstructorDefinition(StringBuilder Builder, UhtClass Class, string Api)
+		{
+			Builder.Append("private: \\\r\n");
+			Builder.Append("\t/** Private move- and copy-constructors, should never be used */ \\\r\n");
+			Builder.Append("\t").Append(Api).Append(Class.SourceName).Append('(').Append(Class.SourceName).Append("&&); \\\r\n");
+			Builder.Append("\t").Append(Api).Append(Class.SourceName).Append("(const ").Append(Class.SourceName).Append("&); \\\r\n");
+			Builder.Append("public: \\\r\n");
+			return Builder;
+		}
+
+		/// <summary>
+		/// Generates private copy-constructor declaration.
+		/// </summary>
+		/// <param name="Builder">Output builder</param>
+		/// <param name="Class">Class being exported</param>
+		/// <param name="Api">API text to be used</param>
+		/// <returns>Output builder</returns>
+		private StringBuilder AppendConstructorDefinition(StringBuilder Builder, UhtClass Class, string Api)
+		{
+			if (!Class.ClassExportFlags.HasAnyFlags(UhtClassExportFlags.HasConstructor))
+			{
+				Builder.Append("\t/** Standard constructor, called after all reflected properties have been initialized */ \\\r\n");
+
+				// Assume super class has OI constructor, this may not always be true but we should always be able to check this.
+				// In any case, it will default to old behavior before we even checked this.
+				bool bSuperClassObjectInitializerConstructorDeclared = true;
+				UhtClass? SuperClass = Class.SuperClass;
+				if (SuperClass != null)
+				{
+					if (!SuperClass.HeaderFile.bIsNoExportTypes)
+					{
+						bSuperClassObjectInitializerConstructorDeclared = SuperClass.ClassExportFlags.HasAnyFlags(UhtClassExportFlags.HasObjectInitializerConstructor);
+					}
+				}
+
+				if (bSuperClassObjectInitializerConstructorDeclared)
+				{
+					Builder.Append("\t").Append(Api).Append(Class.SourceName).Append("(const FObjectInitializer& ObjectInitializer = FObjectInitializer::Get()) : Super(ObjectInitializer) { }; \\\r\n");
+					//ETSTODO - Modification of class
+					Class.ClassExportFlags |= UhtClassExportFlags.HasObjectInitializerConstructor;
+				}
+				else
+				{
+					Builder.Append("\t").Append(Api).Append(Class.SourceName).Append("() { }; \\\r\n");
+					//ETSTODO - Modification of class
+					Class.ClassExportFlags |= UhtClassExportFlags.HasDefaultConstructor;
+				}
+
+				// The original code would mark this true at this point.  However,
+				// we are no longer allowed to modify the data.
+				//ETSTODO - Modification of class
+				Class.ClassExportFlags |= UhtClassExportFlags.HasConstructor;
+			}
+			AppendCopyConstructorDefinition(Builder, Class, Api);
+			return Builder;
+		}
+
+		/// <summary>
+		/// Generates constructor call definition
+		/// </summary>
+		/// <param name="Builder">Output builder</param>
+		/// <param name="Class">Class being exported</param>
+		/// <param name="Api">API text to be used</param>
+		/// <returns>Output builder</returns>
+		private StringBuilder AppendDefaultConstructorCallDefinition(StringBuilder Builder, UhtClass Class, string Api)
+		{
+			if (Class.ClassExportFlags.HasAnyFlags(UhtClassExportFlags.HasObjectInitializerConstructor))
+			{
+				if (Class.ClassFlags.HasAnyFlags(EClassFlags.Abstract))
+				{
+					Builder.Append("\tDEFINE_ABSTRACT_DEFAULT_OBJECT_INITIALIZER_CONSTRUCTOR_CALL(").Append(Class.SourceName).Append(") \\\r\n");
+				}
+				else
+				{
+					Builder.Append("\tDEFINE_DEFAULT_OBJECT_INITIALIZER_CONSTRUCTOR_CALL(").Append(Class.SourceName).Append(") \\\r\n");
+				}
+
+			}
+			else if (Class.ClassExportFlags.HasAnyFlags(UhtClassExportFlags.HasDefaultConstructor))
+			{
+				if (Class.ClassFlags.HasAnyFlags(EClassFlags.Abstract))
+				{
+					Builder.Append("\tDEFINE_ABSTRACT_DEFAULT_CONSTRUCTOR_CALL(").Append(Class.SourceName).Append(") \\\r\n");
+				}
+				else
+				{
+					Builder.Append("\tDEFINE_DEFAULT_CONSTRUCTOR_CALL(").Append(Class.SourceName).Append(") \\\r\n");
+				}
+			}
+			else
+			{
+				Builder.Append("\tDEFINE_FORBIDDEN_DEFAULT_CONSTRUCTOR_CALL(").Append(Class.SourceName).Append(") \\\r\n");
+			}
+			return Builder;
+		}
+
+		/// <summary>
+		/// Generates generated body code for classes
+		/// </summary>
+		/// <param name="Builder">Output builder</param>
+		/// <param name="Class">Class being exported</param>
+		/// <param name="Api">API text to be used</param>
+		/// <returns>Output builder</returns>
+		private StringBuilder AppendClassGeneratedBody(StringBuilder Builder, UhtClass Class, string Api)
+		{
+			AppendCommonGeneratedBody(Builder, Class, Api);
+
+			// export the class's config name
+			UhtClass? SuperClass = Class.SuperClass;
+			if (SuperClass != null && Class.Config.Length > 0 && Class.Config != SuperClass.Config)
+			{
+				Builder.Append("\tstatic const TCHAR* StaticConfigName() {return TEXT(\"").Append(Class.Config).Append("\");} \\\r\n \\\r\n");
+			}
+
+			// export implementation of _getUObject for classes that implement interfaces
+			if (Class.Bases != null)
+			{
+				foreach (UhtStruct BaseStruct in Class.Bases)
+				{
+					if (BaseStruct is UhtClass BaseClass)
+					{
+						if (BaseClass.ClassFlags.HasAnyFlags(EClassFlags.Interface))
+						{
+							Builder.Append("\tvirtual UObject* _getUObject() const override { return const_cast<").Append(Class.SourceName).Append("*>(this); } \\\r\n");
+							break;
+						}
+					}
+				}
+			}
+
+			if (Class.ClassExportFlags.HasAnyFlags(UhtClassExportFlags.SelfHasReplicatedProperties))
+			{
+				AppendReplicatedMacroData(Builder, Class, Api);
+			}
+
+			return Builder;
+		}
+
+		/// <summary>
+		/// Generates standard generated body code for interfaces and non-interfaces
+		/// </summary>
+		/// <param name="Builder">Output builder</param>
+		/// <param name="Class">Class being exported</param>
+		/// <param name="Api">API text to be used</param>
+		/// <returns>Output builder</returns>
+		private StringBuilder AppendCommonGeneratedBody(StringBuilder Builder, UhtClass Class, string Api)
+		{
+			// Export the class's native function registration.
+			Builder.Append("private: \\\r\n");
+			Builder.Append("\tstatic void StaticRegisterNatives").Append(Class.SourceName).Append("(); \\\r\n");
+			if (!Class.ClassFlags.HasAnyFlags(EClassFlags.NoExport))
+			{
+				Builder.Append("\tfriend struct ").Append(this.ObjectInfos[Class.ObjectTypeIndex].RegisteredSingletonName).Append("_Statics; \\\r\n");
+			}
+			Builder.Append("public: \\\r\n");
+
+			UhtClass? SuperClass = Class.SuperClass;
+			bool bCastedClass = Class.ClassCastFlags.HasAnyFlags(EClassCastFlags.AllFlags) && SuperClass != null && Class.ClassCastFlags != SuperClass.ClassCastFlags;
+
+			Builder
+				.Append("\tDECLARE_CLASS(")
+				.Append(Class)
+				.Append(", ")
+				.Append(SuperClass != null ? SuperClass.SourceName : "None")
+				.Append(", COMPILED_IN_FLAGS(")
+				.Append(Class.ClassFlags.HasAnyFlags(EClassFlags.Abstract) ? "CLASS_Abstract" : "0");
+
+			AppendClassFlags(Builder, Class);
+			Builder.Append("), ");
+			if (bCastedClass)
+			{
+				Builder
+					.Append("CASTCLASS_")
+					.Append(Class.SourceName);
+			}
+			else
+			{
+				Builder.Append("CASTCLASS_None");
+			}
+			Builder.Append(", TEXT(\"").Append(this.Package.SourceName).Append("\"), ").Append(Api.Substring(0, Api.Length - 1)).Append(") \\\r\n");
+
+			Builder.Append("\tDECLARE_SERIALIZER(").Append(Class.SourceName).Append(") \\\r\n");
+
+			// Add the serialization function declaration if we generated one
+			if (Class.SerializerArchiveType != UhtSerializerArchiveType.None && Class.SerializerArchiveType != UhtSerializerArchiveType.All)
+			{
+				Builder.Append("\t").AppendMacroName(this, Class, ArchiveSerializerMacroSuffix).Append(" \\\r\n");
+			}
+
+			if (SuperClass != null && Class.ClassWithin != SuperClass.ClassWithin)
+			{
+				Builder.Append("\tDECLARE_WITHIN(").Append(Class.ClassWithin.SourceName).Append(") \\\r\n");
+			}
+			return Builder;
+		}
+
+		/// <summary>
+		/// Appends the class flags in the form of CLASS_Something|CLASS_Something which represents all class flags that are set 
+		/// for the specified class which need to be exported as part of the DECLARE_CLASS macro
+		/// </summary>
+		/// <param name="Builder">Output builder</param>
+		/// <param name="Class">Class in question</param>
+		/// <returns>Output builder</returns>
+		private StringBuilder AppendClassFlags(StringBuilder Builder, UhtClass Class)
+		{
+			if (Class.ClassFlags.HasAnyFlags(EClassFlags.Transient))
+			{
+				Builder.Append(" | CLASS_Transient");
+			}
+			if (Class.ClassFlags.HasAnyFlags(EClassFlags.Optional))
+			{
+				Builder.Append(" | CLASS_Optional");
+			}
+			if (Class.ClassFlags.HasAnyFlags(EClassFlags.DefaultConfig))
+			{
+				Builder.Append(" | CLASS_DefaultConfig");
+			}
+			if (Class.ClassFlags.HasAnyFlags(EClassFlags.GlobalUserConfig))
+			{
+				Builder.Append(" | CLASS_GlobalUserConfig");
+			}
+			if (Class.ClassFlags.HasAnyFlags(EClassFlags.ProjectUserConfig))
+			{
+				Builder.Append(" | CLASS_ProjectUserConfig");
+			}
+			if (Class.ClassFlags.HasAnyFlags(EClassFlags.Config))
+			{
+				Builder.Append(" | CLASS_Config");
+			}
+			if (Class.ClassFlags.HasAnyFlags(EClassFlags.Interface))
+			{
+				Builder.Append(" | CLASS_Interface");
+			}
+			if (Class.ClassFlags.HasAnyFlags(EClassFlags.Deprecated))
+			{
+				Builder.Append(" | CLASS_Deprecated");
+			}
+			return Builder;
+		}
+
+		/// <summary>
+		/// Appends preprocessor string to emit GENERATED_U*_BODY() macro is deprecated.
+		/// </summary>
+		/// <param name="Builder">Output builder</param>
+		/// <param name="MacroName">Name ofthe macro to deprecate</param>
+		/// <returns>Output builder</returns>
+		private StringBuilder AppendGeneratedMacroDeprecationWarning(StringBuilder Builder, string MacroName)
+		{
+			// Deprecation warning is disabled right now. After people get familiar with the new macro it should be re-enabled.
+			//Builder.Append("EMIT_DEPRECATED_WARNING_MESSAGE(\"").Append(MacroName).Append("() macro is deprecated. Please use GENERATED_BODY() macro instead.\") \\\r\n");
+			return Builder;
+		}
+
+		private StringBuilder AppendAccessSpecifier(StringBuilder Builder, UhtClass Class)
+		{
+			switch (Class.GeneratedBodyAccessSpecifier)
+			{
+				case UhtAccessSpecifier.Public:
+					Builder.Append("public:");
+					break;
+				case UhtAccessSpecifier.Private:
+					Builder.Append("private:");
+					break;
+				case UhtAccessSpecifier.Protected:
+					Builder.Append("protected:");
+					break;
+				default:
+					Builder.Append("static_assert(false, \"Unknown access specifier for GENERATED_BODY() macro in class ").Append(Class.EngineName).Append(".\");");
+					break;
+			}
+			return Builder;
+		}
+
+		private StringBuilder AppendInClassIInterface(StringBuilder Builder, UhtClass Class, List<UhtFunction> CallbackFunctions, string Api)
+		{
+			string InterfaceSourceName = "I" + Class.EngineName;
+			string SuperInterfaceSourceName = Class.SuperClass != null ? "I" + Class.SuperClass.EngineName : "";
+
+			Builder.Append("protected: \\\r\n");
+			Builder.Append("\tvirtual ~").Append(InterfaceSourceName).Append("() {} \\\r\n");
+			Builder.Append("public: \\\r\n");
+			Builder.Append("\ttypedef ").Append(Class.SourceName).Append(" UClassType; \\\r\n");
+			Builder.Append("\ttypedef ").Append(InterfaceSourceName).Append(" ThisClass; \\\r\n");
+
+			AppendInterfaceCallFunctions(Builder, Class, CallbackFunctions);
+
+			// we'll need a way to get to the UObject portion of a native interface, so that we can safely pass native interfaces
+			// to script VM functions
+			if (Class.SuperClass != null && Class.SuperClass.IsChildOf(this.Session.UInterface))
+			{
+				// Note: This used to be declared as a pure virtual function, but it was changed here in order to allow the Blueprint nativization process
+				// to detect C++ interface classes that explicitly declare pure virtual functions via type traits. This code will no longer trigger that check.
+				Builder.Append("\tvirtual UObject* _getUObject() const { return nullptr; } \\\r\n");
+			}
+
+			if (Class.ClassExportFlags.HasAnyFlags(UhtClassExportFlags.SelfHasReplicatedProperties))
+			{
+				AppendReplicatedMacroData(Builder, Class, Api);
+			}
+
+			return Builder;
+		}
+
+		private StringBuilder AppendReplicatedMacroData(StringBuilder Builder, UhtClass Class, string Api)
+		{
+			if (!Class.ClassExportFlags.HasAnyFlags(UhtClassExportFlags.HasGetLifetimeReplicatedProps))
+			{
+				// Default version autogenerates declarations.
+				if (Class.GeneratedCodeVersion == EGeneratedCodeVersion.V1)
+				{
+					Builder.Append("\tvoid GetLifetimeReplicatedProps(TArray<FLifetimeProperty>& OutLifetimeProps) const override; \\\r\n");
+				}
+			}
+
+			AppendNetData(Builder, Class, Api);
+
+			// If this class has replicated properties and it owns the first one, that means
+			// it's the base most replicated class. In that case, go ahead and add our interface macro.
+			if (Class.ClassExportFlags.HasExactFlags(UhtClassExportFlags.HasReplciatedProperties, UhtClassExportFlags.SelfHasReplicatedProperties))
+			{
+				Builder.Append("private: \\\r\n");
+				Builder.Append("\tREPLICATED_BASE_CLASS(").Append(Class.SourceName).Append(") \\\r\n");
+				Builder.Append("public: \\\r\n");
+			}
+			return Builder;
+		}
+
+		private StringBuilder AppendNetData(StringBuilder Builder, UhtClass Class, string Api)
+		{
+			bool bHasArray = false;
+			foreach (UhtProperty Property in Class.EnumerateReplicatedProperties(false))
+			{
+				if (Property.bIsStaticArray)
+				{
+					if (!bHasArray)
+					{
+						bHasArray = true;
+						Builder.Append("\tenum class EArrayDims_Private : uint16 \\\r\n");
+						Builder.Append("\t{ \\\r\n");
+					}
+					Builder.Append("\t\t").Append(Property.SourceName).Append('=').Append(Property.ArrayDimensions).Append(", \\\r\n");
+				}
+			}
+
+			if (bHasArray)
+			{
+				Builder.Append("\t}; \\\r\n");
+			}
+
+			Builder.Append("\tenum class ENetFields_Private : uint16 \\\r\n");
+			Builder.Append("\t{ \\\r\n");
+			Builder.Append("\t\tNETFIELD_REP_START=(uint16)((int32)Super::ENetFields_Private::NETFIELD_REP_END + (int32)1), \\\r\n");
+
+			bool bIsFirst = true;
+			UhtProperty? LastProperty = null;
+			foreach (UhtProperty Property in Class.EnumerateReplicatedProperties(false))
+			{
+				LastProperty = Property;
+				if (!Property.bIsStaticArray)
+				{
+					if (bIsFirst)
+					{
+						Builder.Append("\t\t").Append(Property.SourceName).Append("=NETFIELD_REP_START, \\\r\n");
+						bIsFirst = false;
+					}
+					else
+					{
+						Builder.Append("\t\t").Append(Property.SourceName).Append(", \\\r\n");
+					}
+				}
+				else
+				{
+					if (bIsFirst)
+					{
+						Builder.Append("\t\t").Append(Property.SourceName).Append("_STATIC_ARRAY=NETFIELD_REP_START, \\\r\n");
+						bIsFirst = false;
+					}
+					else
+					{
+						Builder.Append("\t\t").Append(Property.SourceName).Append("_STATIC_ARRAY, \\\r\n");
+					}
+
+					Builder
+						.Append("\t\t")
+						.Append(Property.SourceName)
+						.Append("_STATIC_ARRAY_END=((uint16)")
+						.Append(Property.SourceName)
+						.Append("_STATIC_ARRAY + (uint16)EArrayDims_Private::")
+						.Append(Property.SourceName)
+						.Append(" - (uint16)1), \\\r\n");
+				}
+			}
+
+			if (LastProperty != null)
+			{
+				Builder.Append("\t\tNETFIELD_REP_END=").Append(LastProperty.SourceName);
+				if (LastProperty.bIsStaticArray)
+				{
+					Builder.Append("_STATIC_ARRAY_END");
+				}
+			}
+			Builder.Append("\t}; \\\r\n");
+
+			Builder.Append("\t").Append(Api).Append("virtual void ValidateGeneratedRepEnums(const TArray<struct FRepRecord>& ClassReps) const override; \\\r\n");
+
+			return Builder;
+		}
+
+		private StringBuilder AppendInterfaceCallFunctions(StringBuilder Builder, UhtClass Class, List<UhtFunction> CallbackFunctions)
+		{
+			const string ExtraArg = "UObject* O";
+			const string ConstExtraArg = "const UObject* O";
+
+			foreach (UhtFunction Function in CallbackFunctions)
+			{
+				AppendNativeFunctionHeader(Builder, Function, UhtPropertyTextType.InterfaceFunction, true, null,
+					Function.FunctionFlags.HasAnyFlags(EFunctionFlags.Const) ? ConstExtraArg : ExtraArg, UhtFunctionExportFlags.None,
+					"; \\\r\n");
+			}
+			return Builder;
+		}
+
+		private StringBuilder AppendGeneratedBodyMacroBlock(StringBuilder Builder, UhtClass Class, UhtClass BodyClass, bool bIsLegacy, bool bHasEditorRpc, bool bHasCallbacks, string? DeprecatedMacroName)
+		{
+			bool bIsInterface = Class.ClassFlags.HasAnyFlags(EClassFlags.Interface);
+			using (UhtMacroCreator Macro = new UhtMacroCreator(Builder, this, BodyClass, bIsLegacy ? GeneratedBodyLegacyMacroSuffix : GeneratedBodyMacroSuffix))
+			{
+				if (DeprecatedMacroName != null)
+				{
+					AppendGeneratedMacroDeprecationWarning(Builder, DeprecatedMacroName);
+				}
+				Builder.Append(DisableDeprecationWarnings).Append(" \\\r\n");
+				Builder.Append("public: \\\r\n");
+				Builder.Append('\t').AppendMacroName(this, Class, SparseDataMacroSuffix).Append(" \\\r\n");
+				if (bIsLegacy)
+				{
+					Builder.Append('\t').AppendMacroName(this, Class, RpcWrappersMacroSuffix).Append(" \\\r\n");
+				}
+				else
+				{
+					Builder.Append('\t').AppendMacroName(this, Class, RpcWrappersNoPureDeclsMacroSuffix).Append(" \\\r\n");
+				}
+				if (bHasEditorRpc)
+				{
+					if (bIsLegacy)
+					{
+						Builder.Append('\t').AppendMacroName(this, Class, EditorOnlyRpcWrappersMacroSuffix).Append(" \\\r\n");
+					}
+					else
+					{
+						Builder.Append('\t').AppendMacroName(this, Class, EditorOnlyRpcWrappersNoPureDeclsMacroSuffix).Append(" \\\r\n");
+					}
+				}
+				Builder.Append('\t').AppendMacroName(this, Class, AccessorsMacroSuffix).Append(" \\\r\n");
+				if (bHasCallbacks)
+				{
+						Builder.Append('\t').AppendMacroName(this, Class, CallbackWrappersMacroSuffix).Append(" \\\r\n");
+					}
+				if (bIsInterface)
+				{
+					if (bIsLegacy)
+					{
+						Builder.Append('\t').AppendMacroName(this, Class, InClassIInterfaceMacroSuffix).Append(" \\\r\n");
+					}
+					else
+					{
+						Builder.Append('\t').AppendMacroName(this, Class, InClassIInterfaceNoPureDeclsMacroSuffix).Append(" \\\r\n");
+					}
+				}
+				else
+				{
+					if (bIsLegacy)
+					{
+						Builder.Append('\t').AppendMacroName(this, Class, InClassMacroSuffix).Append(" \\\r\n");
+					}
+					else
+					{
+						Builder.Append('\t').AppendMacroName(this, Class, InClassNoPureDeclsMacroSuffix).Append(" \\\r\n");
+					}
+				}
+				if (!bIsInterface)
+				{
+					if (bIsLegacy)
+					{
+						Builder.Append('\t').AppendMacroName(this, Class, StandardConstructorsMacroSuffix).Append(" \\\r\n");
+					}
+					else
+					{
+						Builder.Append('\t').AppendMacroName(this, Class, EnchancedConstructorsMacroSuffix).Append(" \\\r\n");
+					}
+				}
+				if (bIsLegacy)
+				{
+					Builder.Append("public: \\\r\n");
+				}
+				else
+				{
+					AppendAccessSpecifier(Builder, Class);
+					Builder.Append(" \\\r\n");
+				}
+				Builder.Append(EnableDeprecationWarnings).Append(" \\\r\n");
+			}
+
+
+			return Builder;
+		}
+
+		#region Enum helper methods
+		private static bool IsFullEnumName(string InEnumName)
+		{
+			return InEnumName.Contains("::");
+		}
+
+		private static StringView GenerateEnumPrefix(UhtEnum Enum)
+		{
+			StringView Prefix = new StringView();
+			if (Enum.EnumValues.Count > 0)
+			{
+				Prefix = Enum.EnumValues[0].Name;
+
+				// For each item in the enumeration, trim the prefix as much as necessary to keep it a prefix.
+				// This ensures that once all items have been processed, a common prefix will have been constructed.
+				// This will be the longest common prefix since as little as possible is trimmed at each step.
+				for (int NameIdx = 1; NameIdx < Enum.EnumValues.Count; ++NameIdx)
+				{
+					StringView EnumItemName = Enum.EnumValues[NameIdx].Name;
+
+					// Find the length of the longest common prefix of Prefix and EnumItemName.
+					int PrefixIdx = 0;
+					while (PrefixIdx < Prefix.Length && PrefixIdx < EnumItemName.Length && Prefix.Span[PrefixIdx] == EnumItemName.Span[PrefixIdx])
+					{
+						PrefixIdx++;
+					}
+
+					// Trim the prefix to the length of the common prefix.
+					Prefix = new StringView(Prefix, 0, PrefixIdx);
+				}
+
+				// Find the index of the rightmost underscore in the prefix.
+				int UnderscoreIdx = Prefix.Span.LastIndexOf('_');
+
+				// If an underscore was found, trim the prefix so only the part before the rightmost underscore is included.
+				if (UnderscoreIdx > 0)
+				{
+					Prefix = new StringView(Prefix, 0, UnderscoreIdx);
+				}
+				else
+				{
+					// no underscores in the common prefix - this probably indicates that the names
+					// for this enum are not using Epic's notation, so just empty the prefix so that
+					// the max item will use the full name of the enum
+					Prefix = new StringView();
+				}
+			}
+
+			// If no common prefix was found, or if the enum does not contain any entries,
+			// use the name of the enumeration instead.
+			if (Prefix.Length == 0)
+			{
+				Prefix = Enum.EngineName;
+			}
+			return Prefix;
+		}
+
+		private static string GenerateFullEnumName(UhtEnum Enum, string InEnumName)
+		{
+			if (Enum.CppForm == UhtEnumCppForm.Regular || IsFullEnumName(InEnumName))
+			{
+				return InEnumName;
+			}
+			return $"{Enum.EngineName}::{InEnumName}";
+		}
+
+		private static bool EnumHasExistingMax(UhtEnum Enum)
+		{
+			if (Enum.GetIndexByName(GenerateFullEnumName(Enum, "MAX")) != -1)
+			{
+				return true;
+			}
+
+			string MaxEnumItem = GenerateFullEnumName(Enum, GenerateEnumPrefix(Enum).ToString() + "_MAX");
+			if (Enum.GetIndexByName(MaxEnumItem) != -1)
+			{
+				return true;
+			}
+			return false;
+		}
+
+		private static long GetMaxEnumValue(UhtEnum Enum)
+		{
+			if (Enum.EnumValues.Count == 0)
+			{
+				return 0;
+			}
+
+			long MaxValue = Enum.EnumValues[0].Value;
+			for (int i = 1; i < Enum.EnumValues.Count; ++i)
+			{
+				long CurrentValue = Enum.EnumValues[i].Value;
+				if (CurrentValue > MaxValue)
+				{
+					MaxValue = CurrentValue;
+				}
+			}
+
+			return MaxValue;
+		}
+		#endregion
+	}
+}
diff --git a/Engine/Source/Programs/Shared/EpicGames.UHT/Exporters/CodeGen/UhtMacroBlockEmitter.cs b/Engine/Source/Programs/Shared/EpicGames.UHT/Exporters/CodeGen/UhtMacroBlockEmitter.cs
new file mode 100644
--- /dev/null
+++ b/Engine/Source/Programs/Shared/EpicGames.UHT/Exporters/CodeGen/UhtMacroBlockEmitter.cs
@@ -0,0 +1,44 @@
+// Copyright Epic Games, Inc. All Rights Reserved.
+
+using System;
+using System.Text;
+
+namespace EpicGames.UHT.Exporters.CodeGen
+{
+	internal struct UhtMacroBlockEmitter : IDisposable
+	{
+		private readonly StringBuilder Builder;
+		private readonly string MacroName;
+		private bool bEmitted;
+
+		public UhtMacroBlockEmitter(StringBuilder Builder, string MacroName, bool bInitialState = false)
+		{
+			this.Builder = Builder;
+			this.MacroName = MacroName;
+			this.bEmitted = false;
+			Set(bInitialState);
+		}
+
+		public void Set(bool bEmit)
+		{
+			if (this.bEmitted == bEmit)
+			{
+				return;
+			}
+			this.bEmitted = bEmit;
+			if (this.bEmitted)
+			{
+				Builder.Append("#if ").Append(this.MacroName).Append("\r\n");
+			}
+			else
+			{
+				Builder.Append("#endif // ").Append(this.MacroName).Append("\r\n");
+			}
+		}
+
+		public void Dispose()
+		{
+			Set(false);
+		}
+	}
+}
diff --git a/Engine/Source/Programs/Shared/EpicGames.UHT/Exporters/CodeGen/UhtMacroCreator.cs b/Engine/Source/Programs/Shared/EpicGames.UHT/Exporters/CodeGen/UhtMacroCreator.cs
new file mode 100644
--- /dev/null
+++ b/Engine/Source/Programs/Shared/EpicGames.UHT/Exporters/CodeGen/UhtMacroCreator.cs
@@ -0,0 +1,56 @@
+// Copyright Epic Games, Inc. All Rights Reserved.
+
+using EpicGames.UHT.Types;
+using System;
+using System.Text;
+
+namespace EpicGames.UHT.Exporters.CodeGen
+{
+	internal struct UhtMacroCreator : IDisposable
+	{
+		private readonly StringBuilder Builder;
+		private readonly int StartingLength;
+
+		public UhtMacroCreator(StringBuilder Builder, UhtHeaderCodeGenerator Generator, int LineNumber, string MacroSuffix)
+		{
+			Builder.Append("#define ").AppendMacroName(Generator, LineNumber, MacroSuffix).Append(" \\\r\n");
+			this.Builder = Builder;
+			this.StartingLength = Builder.Length;
+		}
+
+		public UhtMacroCreator(StringBuilder Builder, UhtHeaderCodeGenerator Generator, UhtClass Class, string MacroSuffix)
+		{
+			Builder.Append("#define ").AppendMacroName(Generator, Class, MacroSuffix).Append(" \\\r\n");
+			this.Builder = Builder;
+			this.StartingLength = Builder.Length;
+		}
+
+		public UhtMacroCreator(StringBuilder Builder, UhtHeaderCodeGenerator Generator, UhtScriptStruct ScriptStruct, string MacroSuffix)
+		{
+			Builder.Append("#define ").AppendMacroName(Generator, ScriptStruct, MacroSuffix).Append(" \\\r\n");
+			this.Builder = Builder;
+			this.StartingLength = Builder.Length;
+		}
+
+		public UhtMacroCreator(StringBuilder Builder, UhtHeaderCodeGenerator Generator, UhtFunction Function, string MacroSuffix)
+		{
+			Builder.Append("#define ").AppendMacroName(Generator, Function, MacroSuffix).Append(" \\\r\n");
+			this.Builder = Builder;
+			this.StartingLength = Builder.Length;
+		}
+
+		public void Dispose()
+		{
+			int FinalLength = Builder.Length;
+			Builder.Length = Builder.Length - 4;
+			if (FinalLength == StartingLength)
+			{
+				Builder.Append("\r\n");
+			}
+			else
+			{
+				Builder.Append("\r\n\r\n\r\n");
+			}
+		}
+	}
+}
diff --git a/Engine/Source/Programs/Shared/EpicGames.UHT/Exporters/CodeGen/UhtPackageCodeGenerator.cs b/Engine/Source/Programs/Shared/EpicGames.UHT/Exporters/CodeGen/UhtPackageCodeGenerator.cs
new file mode 100644
--- /dev/null
+++ b/Engine/Source/Programs/Shared/EpicGames.UHT/Exporters/CodeGen/UhtPackageCodeGenerator.cs
@@ -0,0 +1,165 @@
+// Copyright Epic Games, Inc. All Rights Reserved.
+
+using EpicGames.UHT.Types;
+using System.Text;
+
+namespace EpicGames.UHT.Exporters.CodeGen
+{
+	internal class UhtPackageCodeGenerator
+	{
+		public static string HeaderCopyright =
+			"// Copyright Epic Games, Inc. All Rights Reserved.\r\n" +
+			"/*===========================================================================\r\n" +
+			"\tGenerated code exported from UnrealHeaderTool.\r\n" +
+			"\tDO NOT modify this manually! Edit the corresponding .h files instead!\r\n" +
+			"===========================================================================*/\r\n" +
+			"\r\n";
+
+		public static string RequiredCPPIncludes = "#include \"UObject/GeneratedCppIncludes.h\"\r\n";
+
+		public static string EnableDeprecationWarnings = "PRAGMA_ENABLE_DEPRECATION_WARNINGS";
+		public static string DisableDeprecationWarnings = "PRAGMA_DISABLE_DEPRECATION_WARNINGS";
+
+		public static string BeginEditorOnlyGuard = "#if WITH_EDITOR\r\n";
+		public static string EndEditorOnlyGuard = "#endif //WITH_EDITOR\r\n"; //COMPATIBILITY-TODO - This does not match UhtMacroBlockEmitter
+
+		public readonly UhtCodeGenerator CodeGenerator;
+		public readonly UhtPackage Package;
+
+		public Utils.UhtSession Session => this.CodeGenerator.Session;
+		public UhtCodeGenerator.PackageInfo[] PackageInfos => this.CodeGenerator.PackageInfos;
+		public UhtCodeGenerator.HeaderInfo[] HeaderInfos => this.CodeGenerator.HeaderInfos;
+		public UhtCodeGenerator.ObjectInfo[] ObjectInfos => this.CodeGenerator.ObjectInfos;
+		public string PackageApi => this.PackageInfos[this.Package.PackageTypeIndex].Api;
+		public string PackageSingletonName => this.ObjectInfos[this.Package.ObjectTypeIndex].RegisteredSingletonName;
+
+		public UhtPackageCodeGenerator(UhtCodeGenerator CodeGenerator, UhtPackage Package)
+		{
+			this.CodeGenerator = CodeGenerator;
+			this.Package = Package;
+		}
+
+		#region Utility functions
+
+		/// <summary>
+		/// Return the singleton name for an object
+		/// </summary>
+		/// <param name="Object">The object in question.</param>
+		/// <param name="bRegistered">If true, return the registered singleton name.  Otherwise return the unregistered.</param>
+		/// <returns>Singleton name of "nullptr" if Object is null</returns>
+		public string GetSingletonName(UhtObject? Object, bool bRegistered)
+		{
+			return this.CodeGenerator.GetSingletonName(Object, bRegistered);
+		}
+
+		/// <summary>
+		/// Return the external declaration for an object
+		/// </summary>
+		/// <param name="Object">The object in question.</param>
+		/// <param name="bRegistered">If true, return the registered external declaration.  Otherwise return the unregistered.</param>
+		/// <returns>External declaration</returns>
+		public string GetExternalDecl(UhtObject Object, bool bRegistered)
+		{
+			return this.CodeGenerator.GetExternalDecl(Object, bRegistered);
+		}
+
+		/// <summary>
+		/// Return the external declaration for an object
+		/// </summary>
+		/// <param name="ObjectIndex">The object in question.</param>
+		/// <param name="bRegistered">If true, return the registered external declaration.  Otherwise return the unregistered.</param>
+		/// <returns>External declaration</returns>
+		public string GetExternalDecl(int ObjectIndex, bool bRegistered)
+		{
+			return this.CodeGenerator.GetExternalDecl(ObjectIndex, bRegistered);
+		}
+
+		/// <summary>
+		/// Return the cross reference for an object
+		/// </summary>
+		/// <param name="Object">The object in question.</param>
+		/// <param name="bRegistered">If true, return the registered cross reference.  Otherwise return the unregistered.</param>
+		/// <returns>Cross reference</returns>
+		public string GetCrossReference(UhtObject Object, bool bRegistered)
+		{
+			return this.CodeGenerator.GetCrossReference(Object, bRegistered);
+		}
+
+		/// <summary>
+		/// Return the cross reference for an object
+		/// </summary>
+		/// <param name="ObjectIndex">The object in question.</param>
+		/// <param name="bRegistered">If true, return the registered cross reference.  Otherwise return the unregistered.</param>
+		/// <returns>Cross reference</returns>
+		public string GetCrossReference(int ObjectIndex, bool bRegistered)
+		{
+			return this.CodeGenerator.GetCrossReference(ObjectIndex, bRegistered);
+		}
+
+		/// <summary>
+		/// Test to see if the given field is a delegate function
+		/// </summary>
+		/// <param name="Field">Field to be tested</param>
+		/// <returns>True if the field is a delegate function</returns>
+		public static bool IsDelegateFunction(UhtField Field)
+		{
+			if (Field is UhtFunction Function)
+			{
+				return Function.FunctionType.IsDelegate();
+			}
+			return false;
+		}
+
+		/// <summary>
+		/// Combines two hash values to get a third.
+		/// Note - this function is not commutative.
+		///
+		/// This function cannot change for backward compatibility reasons.
+		/// You may want to choose HashCombineFast for a better in-memory hash combining function.
+		/// 
+		/// NOTE: This is a copy of the method in TypeHash.h
+		/// </summary>
+		/// <param name="A">Hash to merge</param>
+		/// <param name="C">Previously combined hash</param>
+		/// <returns>Resulting hash value</returns>
+		public static uint HashCombine(uint A, uint C)
+		{
+			uint B = 0x9e3779b9;
+			A += B;
+
+			A -= B; A -= C; A ^= (C >> 13);
+			B -= C; B -= A; B ^= (A << 8);
+			C -= A; C -= B; C ^= (B >> 13);
+			A -= B; A -= C; A ^= (C >> 12);
+			B -= C; B -= A; B ^= (A << 16);
+			C -= A; C -= B; C ^= (B >> 5);
+			A -= B; A -= C; A ^= (C >> 3);
+			B -= C; B -= A; B ^= (A << 10);
+			C -= A; C -= B; C ^= (B >> 15);
+
+			return C;
+		}
+		#endregion
+	}
+
+	internal static class UhtPackageCodeGeneratorStringBuilderExtensions
+	{
+		public static StringBuilder AppendBeginEditorOnlyGuard(this StringBuilder Builder, bool bEnable = true)
+		{
+			if (bEnable)
+			{
+				Builder.Append(UhtPackageCodeGenerator.BeginEditorOnlyGuard);
+			}
+			return Builder;
+		}
+
+		public static StringBuilder AppendEndEditorOnlyGuard(this StringBuilder Builder, bool bEnable = true)
+		{
+			if (bEnable)
+			{
+				Builder.Append(UhtPackageCodeGenerator.EndEditorOnlyGuard);
+			}
+			return Builder;
+		}
+	}
+}
diff --git a/Engine/Source/Programs/Shared/EpicGames.UHT/Exporters/CodeGen/UhtPackageCodeGeneratorCppFile.cs b/Engine/Source/Programs/Shared/EpicGames.UHT/Exporters/CodeGen/UhtPackageCodeGeneratorCppFile.cs
new file mode 100644
--- /dev/null
+++ b/Engine/Source/Programs/Shared/EpicGames.UHT/Exporters/CodeGen/UhtPackageCodeGeneratorCppFile.cs
@@ -0,0 +1,147 @@
+// Copyright Epic Games, Inc. All Rights Reserved.
+
+using EpicGames.Core;
+using EpicGames.UHT.Types;
+using EpicGames.UHT.Utils;
+using System;
+using System.Collections.Generic;
+using System.Text;
+
+namespace EpicGames.UHT.Exporters.CodeGen
+{
+	internal class UhtPackageCodeGeneratorCppFile : UhtPackageCodeGenerator
+	{
+		/// <summary>
+		/// Construct an instance of this generator object
+		/// </summary>
+		/// <param name="CodeGenerator">The base code generator</param>
+		/// <param name="Package">Package being generated</param>
+		public UhtPackageCodeGeneratorCppFile(UhtCodeGenerator CodeGenerator, UhtPackage Package)
+			: base(CodeGenerator, Package)
+		{
+		}
+
+		/// <summary>
+		/// For a given UE header file, generated the generated cpp file
+		/// </summary>
+		/// <param name="CppOutput">Output object</param>
+		/// <param name="PackageSortedHeaders">Sorted list of headers by name, in the package.</param>
+		public void Generate(IUhtExportOutput CppOutput, List<UhtHeaderFile> PackageSortedHeaders)
+		{
+			using (BorrowStringBuilder Borrower = new BorrowStringBuilder(StringBuilderCache.Big))
+			{
+				StringBuilder Builder = Borrower.StringBuilder;
+				const string MetaDataParamsName = "Package_MetaDataParams";
+
+				// Collect information from all of the headers
+				List<UhtField> Singletons = new List<UhtField>();
+				StringBuilder Declarations = new StringBuilder();
+				uint BodiesHash = 0;
+				foreach (UhtHeaderFile HeaderFile in PackageSortedHeaders)
+				{
+					ref UhtCodeGenerator.HeaderInfo HeaderInfo = ref this.HeaderInfos[HeaderFile.HeaderFileTypeIndex];
+					ReadOnlyMemory<string> Sorted = HeaderFile.References.Declaration.GetSortedReferences(
+						(int ObjectIndex, bool bRegistered) => GetExternalDecl(ObjectIndex, bRegistered));
+					foreach (string Declaration in Sorted.Span)
+					{
+						Declarations.Append(Declaration);
+					}
+
+					Singletons.AddRange(HeaderFile.References.Singletons);
+
+					uint BodyHash = this.HeaderInfos[HeaderFile.HeaderFileTypeIndex].BodyHash;
+					if (BodiesHash == 0)
+					{
+						// Don't combine in the first case because it keeps GUID backwards compatibility
+						BodiesHash = BodyHash;
+					}
+					else
+					{
+						BodiesHash = HashCombine(BodyHash, BodiesHash);
+					}
+				}
+
+				// No need to generate output if we have no declarations
+				if (Declarations.Length == 0)
+				{
+					return;
+				}
+				uint DeclarationsHash = UhtHash.GenenerateTextHash(Declarations.ToString());
+
+				string StrippedName = this.PackageInfos[this.Package.PackageTypeIndex].StrippedName;
+				string SingletonName = this.PackageSingletonName;
+				Singletons.Sort((UhtField Lhs, UhtField Rhs) =>
+				{
+					bool bLhsIsDel = IsDelegateFunction(Lhs);
+					bool bRhsIsDel = IsDelegateFunction(Rhs);
+					if (bLhsIsDel != bRhsIsDel)
+					{
+						return !bLhsIsDel ? -1 : 1;
+					}
+					return StringComparerUE.OrdinalIgnoreCase.Compare(
+						this.ObjectInfos[Lhs.ObjectTypeIndex].RegisteredSingletonName,
+						this.ObjectInfos[Rhs.ObjectTypeIndex].RegisteredSingletonName);
+				});
+
+				Builder.Append(HeaderCopyright);
+				Builder.Append(RequiredCPPIncludes);
+				Builder.Append(DisableDeprecationWarnings).Append("\r\n");
+				Builder.Append("void EmptyLinkFunctionForGeneratedCode").Append(this.Package.ShortName).Append("_init() {}\r\n");
+
+				if (this.Session.bIncludeDebugOutput)
+				{
+					Builder.Append("#if 0\r\n");
+					Builder.Append(Declarations.ToString());
+					Builder.Append("#endif\r\n");
+				}
+
+				foreach (UhtObject Object in Singletons)
+				{
+					ref UhtCodeGenerator.ObjectInfo Info = ref this.ObjectInfos[Object.ObjectTypeIndex];
+					Builder.Append(Info.RegsiteredExternalDecl);
+				}
+
+				Builder.AppendMetaDataDef(this.Package, null, MetaDataParamsName, 3);
+
+				Builder.Append("\tstatic FPackageRegistrationInfo Z_Registration_Info_UPackage_").Append(StrippedName).Append(";\r\n");
+				Builder.Append("\tFORCENOINLINE UPackage* ").Append(SingletonName).Append("()\r\n");
+				Builder.Append("\t{\r\n");
+				Builder.Append("\t\tif (!Z_Registration_Info_UPackage_").Append(StrippedName).Append(".OuterSingleton)\r\n");
+				Builder.Append("\t\t{\r\n");
+
+				if (Singletons.Count != 0)
+				{
+					Builder.Append("\t\t\tstatic UObject* (*const SingletonFuncArray[])() = {\r\n");
+					foreach (UhtField Field in Singletons)
+					{
+						Builder.Append("\t\t\t\t(UObject* (*)())").Append(this.ObjectInfos[Field.ObjectTypeIndex].RegisteredSingletonName).Append(",\r\n");
+					}
+					Builder.Append("\t\t\t};\r\n");
+				}
+
+				EPackageFlags Flags = this.Package.PackageFlags & (EPackageFlags.ClientOptional | EPackageFlags.ServerSideOnly | EPackageFlags.EditorOnly | EPackageFlags.Developer | EPackageFlags.UncookedOnly);
+				Builder.Append("\t\t\tstatic const UECodeGen_Private::FPackageParams PackageParams = {\r\n");
+				Builder.Append("\t\t\t\t").AppendUTF8LiteralString(this.Package.SourceName).Append(",\r\n");
+				Builder.Append("\t\t\t\t").Append(Singletons.Count != 0 ? "SingletonFuncArray" : "nullptr").Append(",\r\n");
+				Builder.Append("\t\t\t\t").Append(Singletons.Count != 0 ? "UE_ARRAY_COUNT(SingletonFuncArray)" : "0").Append(",\r\n");
+				Builder.Append("\t\t\t\tPKG_CompiledIn | ").Append($"0x{(uint)Flags:X8}").Append(",\r\n");
+				Builder.Append("\t\t\t\t").Append($"0x{BodiesHash:X8}").Append(",\r\n");
+				Builder.Append("\t\t\t\t").Append($"0x{DeclarationsHash:X8}").Append(",\r\n");
+				Builder.Append("\t\t\t\t").AppendMetaDataParams(this.Package, null, MetaDataParamsName).Append("\r\n");
+				Builder.Append("\t\t\t};\r\n");
+				Builder.Append("\t\t\tUECodeGen_Private::ConstructUPackage(Z_Registration_Info_UPackage_").Append(StrippedName).Append(".OuterSingleton, PackageParams);\r\n");
+				Builder.Append("\t\t}\r\n");
+				Builder.Append("\t\treturn Z_Registration_Info_UPackage_").Append(StrippedName).Append(".OuterSingleton;\r\n");
+				Builder.Append("\t}\r\n");
+
+				// Do not change the Z_CompiledInDeferPackage_UPackage_ without changing LC_SymbolPatterns
+				Builder.Append("\tstatic FRegisterCompiledInInfo Z_CompiledInDeferPackage_UPackage_").Append(StrippedName).Append("(").Append(SingletonName)
+					.Append(", TEXT(\"").Append(this.Package.SourceName).Append("\"), Z_Registration_Info_UPackage_").Append(StrippedName).Append(", CONSTRUCT_RELOAD_VERSION_INFO(FPackageReloadVersionInfo, ")
+					.Append($"0x{BodiesHash:X8}, 0x{DeclarationsHash:X8}").Append("));\r\n");
+
+				Builder.Append(EnableDeprecationWarnings).Append("\r\n");
+				CppOutput.CommitOutput(Builder);
+			}
+		}
+	}
+}
diff --git a/Engine/Source/Programs/Shared/EpicGames.UHT/Exporters/CodeGen/UhtPackageCodeGeneratorHFile.cs b/Engine/Source/Programs/Shared/EpicGames.UHT/Exporters/CodeGen/UhtPackageCodeGeneratorHFile.cs
new file mode 100644
--- /dev/null
+++ b/Engine/Source/Programs/Shared/EpicGames.UHT/Exporters/CodeGen/UhtPackageCodeGeneratorHFile.cs
@@ -0,0 +1,67 @@
+// Copyright Epic Games, Inc. All Rights Reserved.
+
+using EpicGames.Core;
+using EpicGames.UHT.Types;
+using EpicGames.UHT.Utils;
+using System.Collections.Generic;
+using System.Linq;
+using System.Text;
+
+namespace EpicGames.UHT.Exporters.CodeGen
+{
+	internal class UhtPackageCodeGeneratorHFile : UhtPackageCodeGenerator
+	{
+		/// <summary>
+		/// Construct an instance of this generator object
+		/// </summary>
+		/// <param name="CodeGenerator">The base code generator</param>
+		/// <param name="Package">Package being generated</param>
+		public UhtPackageCodeGeneratorHFile(UhtCodeGenerator CodeGenerator, UhtPackage Package)
+			: base(CodeGenerator, Package)
+		{
+		}
+
+		/// <summary>
+		/// For a given UE header file, generated the generated H file
+		/// </summary>
+		/// <param name="HeaderOutput">Output object</param>
+		/// <param name="PackageSortedHeaders">Sorted list of headers by name of all headers in the package</param>
+		public void Generate(IUhtExportOutput HeaderOutput, List<UhtHeaderFile> PackageSortedHeaders)
+		{
+			using (BorrowStringBuilder Borrower = new BorrowStringBuilder(StringBuilderCache.Big))
+			{
+				StringBuilder Builder = Borrower.StringBuilder;
+
+				Builder.Append(HeaderCopyright);
+				Builder.Append("#pragma once\r\n");
+				Builder.Append("\r\n");
+				Builder.Append("\r\n");
+
+				List<UhtHeaderFile> HeaderFiles = new List<UhtHeaderFile>(this.Package.Children.Count * 2);
+				HeaderFiles.AddRange(PackageSortedHeaders);
+
+				foreach (UhtHeaderFile HeaderFile in this.Package.Children)
+				{
+					if (HeaderFile.HeaderFileType == UhtHeaderFileType.Classes)
+					{
+						HeaderFiles.Add(HeaderFile);
+					}
+				}
+
+				List<UhtHeaderFile> SortedHeaderFiles = new List<UhtHeaderFile>(HeaderFiles.Distinct());
+				SortedHeaderFiles.Sort((x, y) => StringComparerUE.OrdinalIgnoreCase.Compare(x.FilePath, y.FilePath));
+
+				foreach (UhtHeaderFile HeaderFile in SortedHeaderFiles)
+				{
+					if (HeaderFile.HeaderFileType == UhtHeaderFileType.Classes)
+					{
+						Builder.Append("#include \"").Append(this.HeaderInfos[HeaderFile.HeaderFileTypeIndex].IncludePath).Append("\"\r\n");
+					}
+				}
+
+				Builder.Append("\r\n");
+				HeaderOutput.CommitOutput(Builder);
+			}
+		}
+	}
+}
diff --git a/Engine/Source/Programs/Shared/EpicGames.UHT/Exporters/CodeGen/UhtRigVMStringBuilderExtensions.cs b/Engine/Source/Programs/Shared/EpicGames.UHT/Exporters/CodeGen/UhtRigVMStringBuilderExtensions.cs
new file mode 100644
--- /dev/null
+++ b/Engine/Source/Programs/Shared/EpicGames.UHT/Exporters/CodeGen/UhtRigVMStringBuilderExtensions.cs
@@ -0,0 +1,111 @@
+// Copyright Epic Games, Inc. All Rights Reserved.
+
+using EpicGames.UHT.Types;
+using System;
+using System.Collections.Generic;
+using System.Text;
+
+namespace EpicGames.UHT.Exporters.CodeGen
+{
+	internal static class UhtRigVMStringBuilderExtensions
+	{
+		public static StringBuilder AppendParameterNames(this StringBuilder Builder, IEnumerable<UhtRigVMParameter> Parameters, 
+			bool bLeadingSeparator = false, string Separator = ", ", bool bCastName = false, bool bIncludeEditorOnly = false)
+		{
+			bool bEmitSeparator = bLeadingSeparator;
+			foreach (UhtRigVMParameter Parameter in Parameters)
+			{
+				if (bIncludeEditorOnly && !Parameter.bEditorOnly)
+				{
+					continue;
+				}
+				if (bEmitSeparator)
+				{
+					Builder.Append(Separator);
+				}
+				bEmitSeparator = true;
+				Builder.Append(Parameter.NameOriginal(bCastName));
+			}
+			return Builder;
+		}
+
+		public static StringBuilder AppendTypeConstRef(this StringBuilder Builder, UhtRigVMParameter Parameter, bool bCastType = false)
+		{
+			Builder.Append("const ");
+			string String = Parameter.TypeOriginal(bCastType);
+			ReadOnlySpan<char> Span = String.AsSpan();
+			if (Span.EndsWith("&"))
+			{
+				Span = Span.Slice(0, Span.Length - 1);
+			}
+			Builder.Append(Span);
+			if (Span.Length > 0 && (Span[0] == 'T' || Span[0] == 'F'))
+			{
+				Builder.Append('&');
+			}
+			return Builder;
+		}
+
+		public static StringBuilder AppendTypeRef(this StringBuilder Builder, UhtRigVMParameter Parameter, bool bCastType = false)
+		{
+			string String = Parameter.TypeOriginal(bCastType);
+			ReadOnlySpan<char> Span = String.AsSpan();
+			if (Span.EndsWith("&"))
+			{
+				Builder.Append(Span);
+			}
+			else
+			{
+				Builder.Append(Span).Append('&');
+			}
+			return Builder;
+		}
+
+		public static StringBuilder AppendTypeNoRef(this StringBuilder Builder, UhtRigVMParameter Parameter, bool bCastType = false)
+		{
+			string String = Parameter.TypeOriginal(bCastType);
+			ReadOnlySpan<char> Span = String.AsSpan();
+			if (Span.EndsWith("&"))
+			{
+				Span = Span.Slice(0, Span.Length - 1);
+			}
+			return Builder.Append(Span);
+		}
+
+		public static StringBuilder AppendTypeVariableRef(this StringBuilder Builder, UhtRigVMParameter Parameter, bool bCastType = false)
+		{
+			if (Parameter.IsConst())
+			{
+				Builder.AppendTypeConstRef(Parameter, bCastType);
+			}
+			else
+			{
+				Builder.AppendTypeRef(Parameter, bCastType);
+			}
+			return Builder;
+		}
+
+		public static StringBuilder AppendParameterDecls(this StringBuilder Builder, IEnumerable<UhtRigVMParameter> Parameters, 
+			bool bLeadingSeparator = false, string Separator = ", ", bool bCastType = false, bool bCastName = false, bool bIncludeEditorOnly = false)
+		{
+			bool bEmitSeparator = bLeadingSeparator;
+			foreach (UhtRigVMParameter Parameter in Parameters)
+			{
+				if (bIncludeEditorOnly && !Parameter.bEditorOnly)
+				{
+					continue;
+				}
+				if (bEmitSeparator)
+				{
+					Builder.Append(Separator);
+				}
+				bEmitSeparator = true;
+				Builder
+					.AppendTypeVariableRef(Parameter, bCastType)
+					.Append(' ')
+					.Append(bCastName ? Parameter.CastName : Parameter.Name);
+			}
+			return Builder;
+		}
+	}
+}
diff --git a/Engine/Source/Programs/Shared/EpicGames.UHT/Exporters/Json/UhtJsonExporter.cs b/Engine/Source/Programs/Shared/EpicGames.UHT/Exporters/Json/UhtJsonExporter.cs
new file mode 100644
--- /dev/null
+++ b/Engine/Source/Programs/Shared/EpicGames.UHT/Exporters/Json/UhtJsonExporter.cs
@@ -0,0 +1,63 @@
+// Copyright Epic Games, Inc. All Rights Reserved.
+
+using EpicGames.Core;
+using EpicGames.UHT.Tables;
+using EpicGames.UHT.Types;
+using EpicGames.UHT.Utils;
+using System.Collections.Generic;
+using System.Text.Json;
+using System.Threading.Tasks;
+
+namespace EpicGames.UHT.Exporters.Json
+{
+	[UnrealHeaderTool]
+	internal class UhtJsonExporter
+	{
+
+		[UhtExporter(Name = "Json", Description = "Json description of packages", Options = UhtExporterOptions.None)]
+		public static void JsonExporter(IUhtExportFactory Factory)
+		{
+			new UhtJsonExporter(Factory).Export();
+		}
+
+		public readonly IUhtExportFactory Factory;
+		public UhtSession Session => Factory.Session;
+		public UhtExportOptions Options => Factory.Options;
+
+		private UhtJsonExporter(IUhtExportFactory Factory)
+		{
+			this.Factory = Factory;
+		}
+
+		private void Export()
+		{
+			// Generate the files for the packages
+			List<IUhtExportTask> GeneratedPackages = new List<IUhtExportTask>(this.Session.PackageTypeCount);
+			foreach (UhtPackage Package in this.Session.Packages)
+			{
+				UHTManifest.Module Module = Package.Module;
+				UhtExportOptions AdditionalOptions = UhtExportOptions.WriteOutput;
+
+				string JsonPath = this.Factory.MakePath(Package, ".json");
+
+				GeneratedPackages.Add(Factory.CreateTask(JsonPath, new List<IUhtExportTask>(), AdditionalOptions,
+					(IUhtExportOutput JsonOutput) =>
+					{
+						JsonSerializerOptions Options = new JsonSerializerOptions { WriteIndented = true, IgnoreNullValues = true };
+						JsonOutput.CommitOutput(JsonSerializer.Serialize(Package, Options));
+					}));
+			}
+
+			// Wait for all the packages to complete
+			List<Task> PackageTasks = new List<Task>(this.Session.PackageTypeCount);
+			foreach (IUhtExportTask Output in GeneratedPackages)
+			{
+				if (Output.ActionTask != null)
+				{
+					PackageTasks.Add(Output.ActionTask);
+				}
+			}
+			Task.WaitAll(PackageTasks.ToArray());
+		}
+	}
+}
diff --git a/Engine/Source/Programs/Shared/EpicGames.UHT/Parsers/UhtClassBaseParser.cs b/Engine/Source/Programs/Shared/EpicGames.UHT/Parsers/UhtClassBaseParser.cs
new file mode 100644
--- /dev/null
+++ b/Engine/Source/Programs/Shared/EpicGames.UHT/Parsers/UhtClassBaseParser.cs
@@ -0,0 +1,50 @@
+// Copyright Epic Games, Inc. All Rights Reserved.
+
+using EpicGames.Core;
+using EpicGames.UHT.Tokenizer;
+using EpicGames.UHT.Types;
+using System.Collections.Generic;
+
+namespace EpicGames.UHT.Parsers
+{
+	public class UhtClassBaseParser : UhtClass
+	{
+		public UhtToken SuperIdentifier;
+		public List<UhtToken[]>? BaseIdentifiers = null;
+
+		public UhtClassBaseParser(UhtType Outer, int LineNumber) : base(Outer, LineNumber)
+		{
+		}
+
+		protected override void ResolveSuper(UhtResolvePhase ResolvePhase)
+		{
+			base.ResolveSuper(ResolvePhase);
+			switch (ResolvePhase)
+			{
+				case UhtResolvePhase.Bases:
+					BindAndResolveSuper(ref this.SuperIdentifier, UhtFindOptions.Class);
+					BindAndResolveBases(this.BaseIdentifiers, UhtFindOptions.Class);
+
+					// Force the MatchedSerializers on for anything being exported
+					if (!this.ClassFlags.HasAnyFlags(EClassFlags.NoExport))
+					{
+						this.ClassFlags |= EClassFlags.MatchedSerializers;
+					}
+					break;
+			}
+		}
+
+		protected override bool ResolveSelf(UhtResolvePhase ResolvePhase)
+		{
+			bool bResult = base.ResolveSelf(ResolvePhase);
+
+			switch (ResolvePhase)
+			{
+				case UhtResolvePhase.Properties:
+					UhtPropertyParser.ResolveChildren(this, UhtPropertyParseOptions.AddModuleRelativePath);
+					break;
+			}
+			return bResult;
+		}
+	}
+}
diff --git a/Engine/Source/Programs/Shared/EpicGames.UHT/Parsers/UhtClassParser.cs b/Engine/Source/Programs/Shared/EpicGames.UHT/Parsers/UhtClassParser.cs
new file mode 100644
--- /dev/null
+++ b/Engine/Source/Programs/Shared/EpicGames.UHT/Parsers/UhtClassParser.cs
@@ -0,0 +1,420 @@
+// Copyright Epic Games, Inc. All Rights Reserved.
+
+using EpicGames.Core;
+using EpicGames.UHT.Tables;
+using EpicGames.UHT.Tokenizer;
+using EpicGames.UHT.Types;
+using EpicGames.UHT.Utils;
+using System;
+using System.Collections.Generic;
+using System.Text;
+
+namespace EpicGames.UHT.Parsers
+{
+	[UnrealHeaderTool]
+	public class UhtClassParser : UhtClassBaseParser
+	{
+		private static UhtKeywordTable KeywordTable = UhtKeywordTables.Instance.Get(UhtTableNames.Class);
+		private static UhtSpecifierTable SpecifierTable = UhtSpecifierTables.Instance.Get(UhtTableNames.Class);
+
+		public EClassFlags RemovedClassFlags = EClassFlags.None;
+		public string ClassWithinIdentifier = String.Empty;
+
+		public List<string> ShowCategories = new List<string>();
+		public List<string> HideCategories = new List<string>();
+		public List<string> AutoExpandCategories = new List<string>();
+		public List<string> AutoCollapseCategories = new List<string>();
+		public List<string> PrioritizeCategories = new List<string>();
+		public List<string> ShowFunctions = new List<string>();
+		public List<string> HideFunctions = new List<string>();
+		public List<string> SparseClassDataTypes = new List<string>();
+		public List<string> ClassGroupNames = new List<string>();
+
+		public void AddClassFlags(EClassFlags Flags)
+		{
+			this.ClassFlags |= Flags;
+			this.RemovedClassFlags &= ~Flags;
+		}
+
+		public void RemoveClassFlags(EClassFlags Flags)
+		{
+			this.RemovedClassFlags |= Flags;
+			this.ClassFlags &= ~Flags;
+		}
+
+		public UhtClassParser(UhtType Outer, int LineNumber) : base(Outer, LineNumber)
+		{
+		}
+
+		protected override void ResolveSuper(UhtResolvePhase ResolvePhase)
+		{
+			base.ResolveSuper(ResolvePhase);
+
+			// Make sure the class within is resolved.  We have to make sure we don't try to resolve ourselves.
+			if (this.ClassWithin != null && this.ClassWithin != this)
+			{
+				this.ClassWithin.Resolve(ResolvePhase);
+			}
+
+			switch (ResolvePhase)
+			{
+				case UhtResolvePhase.Bases:
+					UhtClass? SuperClass = this.SuperClass;
+
+					// Merge the super class flags
+					if (SuperClass != null)
+					{
+						this.ClassFlags |= SuperClass.ClassFlags & EClassFlags.ScriptInherit & ~this.RemovedClassFlags;
+					}
+
+					if (this.Bases != null)
+					{
+						foreach (UhtStruct Base in this.Bases)
+						{
+							if (Base is UhtClass BaseClass)
+							{
+								if (!BaseClass.ClassFlags.HasAnyFlags(EClassFlags.Interface))
+								{
+									this.LogError($"Class '{BaseClass.SourceName}' is not an interface; Can only inherit from non-UObjects or UInterface derived interfaces");
+								}
+								this.ClassFlags |= BaseClass.ClassFlags & EClassFlags.ScriptInherit & ~this.RemovedClassFlags;
+							}
+						}
+					}
+
+					// These following collections only inherit from he parent if they are empty in this class
+					if (this.SparseClassDataTypes.Count == 0)
+					{
+						AppendStringListMetaData(this.SuperClass, UhtNames.SparseClassDataTypes, this.SparseClassDataTypes);
+					}
+					if (this.PrioritizeCategories.Count == 0)
+					{
+						AppendStringListMetaData(this.SuperClass, UhtNames.PrioritizeCategories, this.PrioritizeCategories);
+					}
+
+					// Merge with categories inherited from the parent.
+					MergeCategories();
+
+					SetAndValidateWithinClass(ResolvePhase);
+					SetAndValidateConfigName(ResolvePhase);
+
+					// Copy all of the lists back to the meta data
+					this.MetaData.AddIfNotEmpty(UhtNames.ClassGroupNames, this.ClassGroupNames);
+					this.MetaData.AddIfNotEmpty(UhtNames.AutoCollapseCategories, this.AutoCollapseCategories);
+					this.MetaData.AddIfNotEmpty(UhtNames.AutoExpandCategories, this.AutoExpandCategories);
+					this.MetaData.AddIfNotEmpty(UhtNames.PrioritizeCategories, this.PrioritizeCategories);
+					this.MetaData.AddIfNotEmpty(UhtNames.HideCategories, this.HideCategories);
+					this.MetaData.AddIfNotEmpty(UhtNames.ShowCategories, this.ShowCategories);
+					this.MetaData.AddIfNotEmpty(UhtNames.SparseClassDataTypes, this.SparseClassDataTypes);
+					this.MetaData.AddIfNotEmpty(UhtNames.HideFunctions, this.HideFunctions);
+
+					this.MetaData.Add(UhtNames.IncludePath, this.HeaderFile.IncludeFilePath);
+					break;
+			}
+		}
+
+		private void MergeCategories()
+		{
+			MergeShowCategories();
+
+			// Merge ShowFunctions and HideFunctions
+			AppendStringListMetaData(this.SuperClass, UhtNames.HideFunctions, this.HideFunctions);
+			foreach (string Value in this.ShowFunctions)
+			{
+				this.HideFunctions.RemoveSwap(Value);
+			}
+			this.ShowFunctions.Clear();
+
+			// Merge AutoExpandCategories and AutoCollapseCategories (we still want to keep AutoExpandCategories though!)
+			List<string> ParentAutoExpandCategories = GetStringListMetaData(this.SuperClass, UhtNames.AutoExpandCategories);
+			List<string> ParentAutoCollapseCategories = GetStringListMetaData(this.SuperClass, UhtNames.AutoCollapseCategories);
+
+			foreach (string Value in this.AutoExpandCategories)
+			{
+				this.AutoCollapseCategories.RemoveSwap(Value);
+				ParentAutoCollapseCategories.RemoveSwap(Value);
+			}
+
+			// Do the same as above but the other way around
+			foreach (string Value in this.AutoCollapseCategories)
+			{
+				this.AutoExpandCategories.RemoveSwap(Value);
+				ParentAutoExpandCategories.RemoveSwap(Value);
+			}
+
+			// Once AutoExpandCategories and AutoCollapseCategories for THIS class have been parsed, add the parent inherited categories
+			this.AutoCollapseCategories.AddRange(ParentAutoCollapseCategories);
+			this.AutoExpandCategories.AddRange(ParentAutoExpandCategories);
+		}
+
+		private void MergeShowCategories()
+		{
+
+			// Add the super class hide categories and prime the output show categories with the parent
+			List<string> OutShowCategories = GetStringListMetaData(this.SuperClass, UhtNames.ShowCategories);
+			AppendStringListMetaData(this.SuperClass, UhtNames.HideCategories, this.HideCategories);
+
+			// If this class has new show categories, then merge them
+			if (this.ShowCategories.Count != 0)
+			{
+				StringBuilder SubCategoryPath = new StringBuilder();
+				foreach (string Value in this.ShowCategories)
+				{
+
+					// if we didn't find this specific category path in the HideCategories metadata
+					if (!this.HideCategories.RemoveSwap(Value))
+					{
+						string[] SubCategories = Value.ToString().Split('|', StringSplitOptions.RemoveEmptyEntries);
+
+						SubCategoryPath.Clear();
+						// look to see if any of the parent paths are excluded in the HideCategories list
+						for (int CategoryPathIndex = 0; CategoryPathIndex < SubCategories.Length - 1; ++CategoryPathIndex)
+						{
+							SubCategoryPath.Append(SubCategories[CategoryPathIndex]);
+							// if we're hiding a parent category, then we need to flag this sub category for show
+							if (HideCategories.Contains(SubCategoryPath.ToString()))
+							{
+								OutShowCategories.AddUnique(Value);
+								break;
+							}
+							SubCategoryPath.Append("|");
+						}
+					}
+				}
+			}
+
+			// Replace the show categories
+			this.ShowCategories = OutShowCategories;
+		}
+
+		private void SetAndValidateWithinClass(UhtResolvePhase ResolvePhase)
+		{
+			// The class within must be a child of any super class within
+			UhtClass ExpectedClassWithin = this.SuperClass != null ? this.SuperClass.ClassWithin : this.Session.UObject;
+
+			// Process all of the class specifiers
+			if (!string.IsNullOrEmpty(this.ClassWithinIdentifier))
+			{
+				UhtClass? SpecifiedClassWithin = (UhtClass?)this.Session.FindType(null, UhtFindOptions.EngineName | UhtFindOptions.Class, this.ClassWithinIdentifier);
+				if (SpecifiedClassWithin == null)
+				{
+					this.LogError($"Within class '{this.ClassWithinIdentifier}' not found");
+				}
+				else
+				{
+
+					// Make sure the class within has been resolved to this phase.  We don't need to worry about the super since we know 
+					// that it has already been resolved.
+					if (SpecifiedClassWithin != this)
+					{
+						SpecifiedClassWithin.Resolve(ResolvePhase);
+					}
+
+					if (SpecifiedClassWithin.IsChildOf(this.Session.UInterface))
+					{
+						this.LogError("Classes cannot be 'within' interfaces");
+					}
+					else if (!SpecifiedClassWithin.IsChildOf(ExpectedClassWithin))
+					{
+						this.LogError($"Cannot override within class with '{SpecifiedClassWithin.SourceName}' since it isn't a child of parent class's expected within '{ExpectedClassWithin.SourceName}'");
+					}
+					else
+					{
+						this.ClassWithin = SpecifiedClassWithin;
+					}
+				}
+			}
+
+			// If we don't have a class within set, then just use the expected within
+			else
+			{
+				this.ClassWithin = ExpectedClassWithin;
+			}
+		}
+
+		private void SetAndValidateConfigName(UhtResolvePhase ResolvePhase)
+		{
+			// Since this flag is computed in this method, we have to re-propagate the flag from the super
+			// just in case they were defined in this source file.
+			if (this.SuperClass != null)
+			{
+				this.ClassFlags |= this.SuperClass.ClassFlags & EClassFlags.Config;
+			}
+
+			// Set the class config flag if any properties have config
+			if (!this.ClassFlags.HasAnyFlags(EClassFlags.Config))
+			{
+				foreach (UhtProperty Property in this.Properties)
+				{
+					if (Property.PropertyFlags.HasAnyFlags(EPropertyFlags.Config))
+					{
+						this.ClassFlags |= EClassFlags.Config;
+						break;
+					}
+				}
+			}
+
+			if (this.Config.Length > 0)
+			{
+				// if the user specified "inherit", we're just going to use the parent class's config filename
+				// this is not actually necessary but it can be useful for explicitly communicating config-ness
+				if (this.Config.Equals("inherit", StringComparison.OrdinalIgnoreCase))
+				{
+					if (this.SuperClass == null)
+					{
+						this.LogError($"Cannot inherit config filename for class '{this.SourceName}' when it has no super class");
+					}
+					else if (this.SuperClass.Config.Length == 0)
+					{
+						this.LogError($"Cannot inherit config filename for class '{this.SourceName}' when parent class '{this.SuperClass.SourceName}' has no config filename");
+					}
+					else
+					{
+						this.Config = this.SuperClass.Config;
+					}
+				}
+			}
+			else if (this.ClassFlags.HasAnyFlags(EClassFlags.Config) && this.SuperClass != null)
+			{
+				this.Config = this.SuperClass.Config;
+			}
+
+			if (this.ClassFlags.HasAnyFlags(EClassFlags.Config) && this.Config.Length == 0)
+			{
+				this.LogError("Classes with config / globalconfig member variables need to specify config file.");
+				this.Config = "Engine";
+			}
+		}
+
+		private static List<string> GetStringListMetaData(UhtType? Type, string Key)
+		{
+			List<string> Out = new List<string>();
+			AppendStringListMetaData(Type, Key, Out);
+			return Out;
+		}
+
+		private static void AppendStringListMetaData(UhtType? Type, string Key, List<string> StringList)
+		{
+			if (Type != null)
+			{
+				string[]? Values = Type.MetaData.GetStringArray(Key);
+				if (Values != null)
+				{
+					foreach (string Value in Values)
+					{
+						//COMPATIBILITY-TODO - TEST - This preserves duplicates that are in old UHT
+						//StringList.AddUnique(Value);
+						StringList.Add(Value);
+					}
+				}
+			}
+		}
+
+		#region Keywords
+		[UhtKeyword(Extends = UhtTableNames.Global)]
+		private static UhtParseResult UCLASSKeyword(UhtParsingScope TopScope, UhtParsingScope ActionScope, ref UhtToken Token)
+		{
+			return ParseUClass(TopScope, ref Token);
+		}
+
+		[UhtKeyword(Extends = UhtTableNames.Class)]
+		[UhtKeyword(Extends = UhtTableNames.Class, Keyword = "GENERATED_BODY")]
+		private static UhtParseResult GENERATED_UCLASS_BODYKeyword(UhtParsingScope TopScope, UhtParsingScope ActionScope, ref UhtToken Token)
+		{
+			UhtClass Class = (UhtClass)TopScope.ScopeType;
+
+			if (Token.IsValue("GENERATED_BODY"))
+			{
+				Class.bHasGeneratedBody = true;
+				Class.GeneratedBodyAccessSpecifier = TopScope.AccessSpecifier;
+			}
+			else
+			{
+				TopScope.AccessSpecifier = UhtAccessSpecifier.Public;
+			}
+
+			UhtParserHelpers.ParseCompileVersionDeclaration(TopScope.TokenReader, Class);
+
+			Class.GeneratedBodyLineNumber = TopScope.TokenReader.InputLine;
+
+			TopScope.TokenReader.Optional(';');
+			return UhtParseResult.Handled;
+		}
+		#endregion
+
+		public static UhtParseResult ParseUClass(UhtParsingScope ParentScope, ref UhtToken Token)
+		{
+			UhtClassParser Class = new UhtClassParser(ParentScope.ScopeType, Token.InputLine);
+			using (var TopScope = new UhtParsingScope(ParentScope, Class, UhtClassParser.KeywordTable, UhtAccessSpecifier.Private))
+			{
+				const string ScopeName = "class";
+
+				using (var TokenContext = new UhtMessageContext(TopScope.TokenReader, ScopeName))
+				{
+					// Parse the specifiers
+					UhtSpecifierContext SpecifierContext = new UhtSpecifierContext(TopScope, TopScope.TokenReader, Class.MetaData);
+					UhtSpecifierParser Specifiers = TopScope.HeaderParser.GetSpecifierParser(SpecifierContext, ScopeName, UhtClassParser.SpecifierTable);
+					Specifiers.ParseSpecifiers();
+					Class.PrologLineNumber = TopScope.TokenReader.InputLine;
+					Class.ClassFlags |= EClassFlags.Native;
+
+					TopScope.AddFormattedCommentsAsTooltipMetaData();
+
+					// Consume the class specifier
+					TopScope.TokenReader.Require("class");
+
+					TopScope.TokenReader.SkipAlignasAndDeprecatedMacroIfNecessary();
+
+					// Read the class name and possible API macro name
+					UhtToken APIMacroToken;
+					TopScope.TokenReader.TryOptionalAPIMacro(out APIMacroToken);
+					Class.SourceName = TopScope.TokenReader.GetIdentifier().Value.ToString();
+
+					// Update the context for better error messages
+					TokenContext.Reset($"class '{Class.SourceName}'");
+
+					// Split the source name into the different parts
+					UhtEngineNameParts NameParts = UhtUtilities.GetEngineNameParts(Class.SourceName);
+					Class.EngineName = NameParts.EngineName.ToString();
+
+					// Check for an empty engine name
+					if (Class.EngineName.Length == 0)
+					{
+						TopScope.TokenReader.LogError($"When compiling class definition for '{Class.SourceName}', attempting to strip prefix results in an empty name. Did you leave off a prefix?");
+					}
+
+					// Skip optional final keyword
+					TopScope.TokenReader.Optional("final");
+
+					// Parse the inheritance
+					UhtParserHelpers.ParseInheritance(TopScope.TokenReader, out Class.SuperIdentifier, out Class.BaseIdentifiers);
+
+					if (APIMacroToken)
+					{
+						Class.ClassFlags |= EClassFlags.RequiredAPI;
+					}
+
+					Specifiers.ParseDeferred();
+
+					if (Class.Outer != null)
+					{
+						Class.Outer.AddChild(Class);
+					}
+
+					TopScope.AddModuleRelativePathToMetaData();
+
+					TopScope.HeaderParser.ParseStatements('{', '}', true);
+
+					TopScope.TokenReader.Require(';');
+
+					if (Class.GeneratedBodyLineNumber == -1)
+					{
+						TopScope.TokenReader.LogError("Expected a GENERATED_BODY() at the start of the class");
+					}
+				}
+			}
+
+			return UhtParseResult.Handled;
+		}
+	}
+}
