UHT in UBT initial version.  Currently disabled.

diff --git a/Engine/Source/Programs/Shared/EpicGames.UHT/Parsers/UhtEnumParser.cs b/Engine/Source/Programs/Shared/EpicGames.UHT/Parsers/UhtEnumParser.cs
new file mode 100644
--- /dev/null
+++ b/Engine/Source/Programs/Shared/EpicGames.UHT/Parsers/UhtEnumParser.cs
@@ -0,0 +1,229 @@
+// Copyright Epic Games, Inc. All Rights Reserved.
+
+using EpicGames.Core;
+using EpicGames.UHT.Tables;
+using EpicGames.UHT.Tokenizer;
+using EpicGames.UHT.Types;
+using EpicGames.UHT.Utils;
+
+namespace EpicGames.UHT.Parsers
+{
+	[UnrealHeaderTool]
+	public class UhtEnumParser : UhtEnum
+	{
+		private static UhtKeywordTable KeywordTable = UhtKeywordTables.Instance.Get(UhtTableNames.Enum);
+		private static UhtSpecifierTable SpecifierTable = UhtSpecifierTables.Instance.Get(UhtTableNames.Enum);
+
+		public UhtEnumParser(UhtType Outer, int LineNumber) : base(Outer, LineNumber)
+		{
+		}
+
+		#region Keywords
+		[UhtKeyword(Extends = UhtTableNames.Global)]
+		private static UhtParseResult UENUMKeyword(UhtParsingScope TopScope, UhtParsingScope ActionScope, ref UhtToken Token)
+		{
+			return ParseUEnum(TopScope, Token);
+		}
+		#endregion
+
+		private static UhtParseResult ParseUEnum(UhtParsingScope ParentScope, UhtToken KeywordToken)
+		{
+			UhtEnumParser Enum = new UhtEnumParser(ParentScope.ScopeType, KeywordToken.InputLine);
+			using (var TopScope = new UhtParsingScope(ParentScope, Enum, KeywordTable, UhtAccessSpecifier.Public))
+			{
+				const string ScopeName = "UENUM";
+
+				using (var TokenContext = new UhtMessageContext(TopScope.TokenReader, ScopeName))
+				{
+
+					// Parse the specifiers
+					UhtSpecifierContext SpecifierContext = new UhtSpecifierContext(TopScope, TopScope.TokenReader, Enum.MetaData);
+					UhtSpecifierParser Specifiers = TopScope.HeaderParser.GetSpecifierParser(SpecifierContext, ScopeName, SpecifierTable);
+					Specifiers.ParseSpecifiers();
+
+					// Read the name and the CPP type
+					switch (TopScope.TokenReader.TryOptional(new string[] { "namespace", "enum" }))
+					{
+						case 0: // namespace
+							Enum.CppForm = UhtEnumCppForm.Namespaced;
+							TopScope.TokenReader.SkipDeprecatedMacroIfNecessary();
+							break;
+						case 1: // enum
+							Enum.CppForm = TopScope.TokenReader.TryOptional(new string[] { "class", "struct" }) >= 0 ? UhtEnumCppForm.EnumClass : UhtEnumCppForm.Regular;
+							TopScope.TokenReader.SkipAlignasAndDeprecatedMacroIfNecessary();
+							break;
+						default:
+							throw new UhtTokenException(TopScope.TokenReader, TopScope.TokenReader.PeekToken(), null);
+					}
+
+					UhtToken EnumToken = TopScope.TokenReader.GetIdentifier("enumeration name");
+
+					Enum.SourceName = EnumToken.Value.ToString();
+
+					Specifiers.ParseFieldMetaData(ScopeName);
+					Specifiers.ParseDeferred();
+
+					if (Enum.Outer != null)
+					{
+						Enum.Outer.AddChild(Enum);
+					}
+
+					if ((TopScope.HeaderParser.GetCurrentCompositeCompilerDirective() & UhtCompilerDirective.WithEditorOnlyData) != 0)
+					{
+						Enum.bIsEditorOnly = true;
+					}
+
+					// Read base for enum class
+					if (Enum.CppForm == UhtEnumCppForm.EnumClass)
+					{
+						if (TopScope.TokenReader.TryOptional(':'))
+						{
+							UhtToken EnumType = TopScope.TokenReader.GetIdentifier("enumeration base");
+
+							UhtEnumUnderlyingType UnderlyingType;
+							if (!System.Enum.TryParse<UhtEnumUnderlyingType>(EnumType.Value.ToString(), out UnderlyingType) || UnderlyingType == UhtEnumUnderlyingType.Unspecified)
+							{
+								TopScope.TokenReader.LogError(EnumType.InputLine, $"Unsupported enum class base type '{EnumType.Value}'");
+							}
+							Enum.UnderlyingType = UnderlyingType;
+						}
+						else
+						{
+							Enum.UnderlyingType = UhtEnumUnderlyingType.Unspecified;
+						}
+					}
+					else
+					{
+						if ((Enum.EnumFlags & EEnumFlags.Flags) != 0)
+						{
+							TopScope.TokenReader.LogError("The 'Flags' specifier can only be used on enum classes");
+						}
+					}
+
+					if (Enum.UnderlyingType != UhtEnumUnderlyingType.uint8 && Enum.MetaData.ContainsKey("BlueprintType"))
+					{
+						TopScope.TokenReader.LogError("Invalid BlueprintType enum base - currently only uint8 supported");
+					}
+
+					//EnumDef.GetDefinitionRange().Start = &Input[InputPos];
+
+					// Get the opening brace
+					TopScope.TokenReader.Require('{');
+
+					switch (Enum.CppForm)
+					{
+						case UhtEnumCppForm.Namespaced:
+							// Now handle the inner true enum portion
+							TopScope.TokenReader.Require("enum");
+							TopScope.TokenReader.SkipAlignasAndDeprecatedMacroIfNecessary();
+
+							UhtToken InnerEnumToken = TopScope.TokenReader.GetIdentifier("enumeration type name");
+
+							TopScope.TokenReader.Require('{');
+							Enum.CppType = $"{Enum.SourceName}::{InnerEnumToken.Value}";
+							break;
+
+						case UhtEnumCppForm.EnumClass:
+						case UhtEnumCppForm.Regular:
+							Enum.CppType = Enum.SourceName;
+							break;
+					}
+
+					TokenContext.Reset($"UENUM {Enum.SourceName}");
+
+					TopScope.AddModuleRelativePathToMetaData();
+					TopScope.AddFormattedCommentsAsTooltipMetaData();
+
+					// Parse all enum tags
+					bool bHasUnparsedValue = false;
+					long CurrentEnumValue = 0;
+
+					TopScope.TokenReader.RequireList('}', ',', true, () =>
+					{
+						UhtToken TagToken = TopScope.TokenReader.GetIdentifier();
+
+						StringView FullEnumName;
+						switch (Enum.CppForm)
+						{
+							case UhtEnumCppForm.Namespaced:
+							case UhtEnumCppForm.EnumClass:
+								FullEnumName = new StringView($"{Enum.SourceName}::{TagToken.Value}");
+								break;
+
+							case UhtEnumCppForm.Regular:
+								FullEnumName = TagToken.Value;
+								break;
+
+							default:
+								throw new UhtIceException("Unexpected EEnumCppForm value");
+						}
+
+						// Save the new tag with a default value.  This will be replaced later
+						int EnumIndex = Enum.AddEnumValue(TagToken.Value.ToString(), 0);
+						TopScope.AddFormattedCommentsAsTooltipMetaData(EnumIndex);
+
+						// Try to read an optional explicit enum value specification
+						if (TopScope.TokenReader.TryOptional('='))
+						{
+							long ScatchValue;
+							bool bParsedValue = TopScope.TokenReader.TryOptionalConstLong(out ScatchValue);
+							if (bParsedValue)
+							{
+								CurrentEnumValue = ScatchValue;
+							}
+							else
+							{
+								// We didn't parse a literal, so set an invalid value
+								CurrentEnumValue = -1;
+								bHasUnparsedValue = true;
+							}
+
+							// Skip tokens until we encounter a comma, a closing brace or a UMETA declaration
+							// There are tokens after the initializer so it's not a standalone literal,
+							// so set it to an invalid value.
+							int SkippedCount = TopScope.TokenReader.ConsumeUntil(new string[] { ",", "}", "UMETA" });
+							if (SkippedCount == 0 && !bParsedValue)
+							{
+								throw new UhtTokenException(TopScope.TokenReader, TopScope.TokenReader.PeekToken(), "enumerator initializer");
+							}
+							if (SkippedCount > 0)
+							{
+								CurrentEnumValue = -1;
+								bHasUnparsedValue = true;
+							}
+						}
+
+						// Save the value and auto increment
+						UhtEnumValue Value = Enum.EnumValues[EnumIndex];
+						Value.Value = CurrentEnumValue;
+						Enum.EnumValues[EnumIndex] = Value;
+						if (CurrentEnumValue != -1)
+						{
+							++CurrentEnumValue;
+						}
+
+						Enum.MetaData.Add(UhtNames.Name, EnumIndex, Enum.EnumValues[EnumIndex].Name.ToString());
+
+						// check for metadata on this enum value
+						SpecifierContext.MetaNameIndex = EnumIndex;
+						Specifiers.ParseFieldMetaData("Enum value");
+					});
+
+					// Trailing brace and semicolon for the enum
+					TopScope.TokenReader.Require(';');
+					if (Enum.CppForm == UhtEnumCppForm.Namespaced)
+					{
+						TopScope.TokenReader.Require('}');
+					}
+
+					if (!bHasUnparsedValue && !Enum.IsValidEnumValue(0) && Enum.MetaData.ContainsKey(UhtNames.BlueprintType))
+					{
+						Enum.LogWarning($"'{Enum.SourceName}' does not have a 0 entry! (This is a problem when the enum is initialized by default)");
+					}
+				}
+
+				return UhtParseResult.Handled;
+			}
+		}
+	}
+}
diff --git a/Engine/Source/Programs/Shared/EpicGames.UHT/Parsers/UhtFunctionParser.cs b/Engine/Source/Programs/Shared/EpicGames.UHT/Parsers/UhtFunctionParser.cs
new file mode 100644
--- /dev/null
+++ b/Engine/Source/Programs/Shared/EpicGames.UHT/Parsers/UhtFunctionParser.cs
@@ -0,0 +1,766 @@
+// Copyright Epic Games, Inc. All Rights Reserved.
+
+using EpicGames.Core;
+using EpicGames.UHT.Tables;
+using EpicGames.UHT.Tokenizer;
+using EpicGames.UHT.Types;
+using EpicGames.UHT.Utils;
+using System;
+using System.Collections.Generic;
+using System.Linq;
+using System.Text;
+
+namespace EpicGames.UHT.Parsers
+{
+	/** 
+	 *	AdvancedDisplay can be used in two ways:
+	 *	1. 'AdvancedDisplay = "3"' - the number tells how many parameters (from beginning) should NOT BE marked
+	 *	2. 'AdvancedDisplay = "AttachPointName, Location, LocationType"' - list the parameters, that should BE marked
+	 */
+	struct UhtAdvancedDisplayParameterHandler
+	{
+		private UhtMetaData MetaData;
+		private string[]? ParameterNames;
+		private int NumberLeaveUnmarked;
+		private int AlreadyLeft;
+		private bool bUseNumber;
+
+		public UhtAdvancedDisplayParameterHandler(UhtMetaData MetaData)
+		{
+			this.MetaData = MetaData;
+			this.ParameterNames = null;
+			this.NumberLeaveUnmarked = -1;
+			this.AlreadyLeft = 0;
+			this.bUseNumber = false;
+
+			string? FoundString;
+			if (this.MetaData.TryGetValue(UhtNames.AdvancedDisplay, out FoundString))
+			{
+				this.ParameterNames = FoundString.ToString().Split(',', StringSplitOptions.RemoveEmptyEntries);
+				for (int Index = 0, EndIndex = this.ParameterNames.Length; Index < EndIndex; ++Index)
+				{
+					this.ParameterNames[Index] = this.ParameterNames[Index].Trim();
+				}
+				if (this.ParameterNames.Length == 1)
+				{
+					bUseNumber = int.TryParse(this.ParameterNames[0], out NumberLeaveUnmarked);
+				}
+			}
+		}
+
+		/** 
+		 * return if given parameter should be marked as Advance View, 
+		 * the function should be called only once for any parameter
+		 */
+		public bool ShouldMarkParameter(StringView ParameterName)
+		{
+			if (this.bUseNumber)
+			{
+				if (this.NumberLeaveUnmarked < 0)
+				{
+					return false;
+				}
+				if (this.AlreadyLeft < this.NumberLeaveUnmarked)
+				{
+					++this.AlreadyLeft;
+					return false;
+				}
+				return true;
+			}
+
+			if (this.ParameterNames == null)
+			{
+				return false;
+			}
+
+			foreach (string Element in this.ParameterNames)
+			{
+				if (ParameterName.Span.Equals(Element, StringComparison.OrdinalIgnoreCase))
+				{
+					return true;
+				}
+			}
+			return false;
+		}
+
+		/** return if more parameters can be marked */
+		public bool CanMarkMore()
+		{
+			return this.bUseNumber ? this.NumberLeaveUnmarked > 0 : (this.ParameterNames != null && this.ParameterNames.Length > 0);
+		}
+	}
+
+	[UnrealHeaderTool]
+	public class UhtFunctionParser : UhtFunction
+	{
+		private static UhtKeywordTable KeywordTable = UhtKeywordTables.Instance.Get(UhtTableNames.Function);
+		private static UhtSpecifierTable SpecifierTable = UhtSpecifierTables.Instance.Get(UhtTableNames.Function);
+
+		public UhtFunctionParser(UhtType Outer, int LineNumber) : base(Outer, LineNumber)
+		{
+		}
+
+		/// <summary>
+		/// True if the function specifier has a getter/setter specified
+		/// </summary>
+		public bool bSawPropertyAccessor = false;
+
+		protected override bool ResolveSelf(UhtResolvePhase ResolvePhase)
+		{
+			bool bResult = base.ResolveSelf(ResolvePhase);
+			switch (ResolvePhase)
+			{
+				case UhtResolvePhase.Properties:
+					UhtPropertyParser.ResolveChildren(this, GetPropertyParseOptions(false));
+					foreach (UhtProperty Property in this.Properties)
+					{
+						if (Property.DefaultValueTokens != null)
+						{
+							string Key = "CPP_Default_" + Property.EngineName;
+							if (!this.MetaData.ContainsKey(Key))
+							{
+								bool bParsed = false;
+								try
+								{
+									// All tokens MUST be consumed from the reader
+									StringBuilder Builder = new StringBuilder();
+									IUhtTokenReader DefaultValueReader = UhtTokenReplayReader.GetThreadInstance(Property, this.HeaderFile.Data.Memory, Property.DefaultValueTokens.ToArray(), UhtTokenType.EndOfDefault);
+									bParsed = Property.SanitizeDefaultValue(DefaultValueReader, Builder) && DefaultValueReader.bIsEOF;
+									if (bParsed)
+									{
+										this.MetaData.Add(Key, Builder.ToString());
+									}
+								}
+								catch (Exception)
+								{
+									// Ignore the exception for now
+								}
+
+								if (!bParsed)
+								{
+									StringView DefaultValueText = new StringView(this.HeaderFile.Data, Property.DefaultValueTokens.First().InputStartPos,
+										Property.DefaultValueTokens.Last().InputEndPos - Property.DefaultValueTokens.First().InputStartPos);
+									Property.LogError($"C++ Default parameter not parsed: {Property.SourceName} '{DefaultValueText}'");
+								}
+							}
+						}
+					}
+					break;
+			}
+			return bResult;
+		}
+
+		public UhtPropertyParseOptions GetPropertyParseOptions(bool bReturnValue)
+		{
+			switch (this.FunctionType)
+			{
+				case UhtFunctionType.Delegate:
+				case UhtFunctionType.SparseDelegate:
+					return (bReturnValue ? UhtPropertyParseOptions.None : UhtPropertyParseOptions.CommaSeparatedName) | UhtPropertyParseOptions.DontAddReturn;
+
+				case UhtFunctionType.Function:
+					UhtPropertyParseOptions Options = UhtPropertyParseOptions.DontAddReturn; // Fetch the function name
+					Options |= bReturnValue ? UhtPropertyParseOptions.FunctionNameIncluded : UhtPropertyParseOptions.NameIncluded;
+					if (this.FunctionFlags.HasAllFlags(EFunctionFlags.BlueprintEvent | EFunctionFlags.Native))
+					{
+						Options |= UhtPropertyParseOptions.NoAutoConst;
+					}
+					return Options;
+
+				default:
+					throw new UhtIceException("Unknown enumeration value");
+			}
+		}
+		#region Keywords
+		[UhtKeyword(Extends = UhtTableNames.Global)]
+		[UhtKeyword(Extends = UhtTableNames.Class)]
+		private static UhtParseResult UDELEGATEKeyword(UhtParsingScope TopScope, UhtParsingScope ActionScope, ref UhtToken Token)
+		{
+			return ParseUDelegate(TopScope, Token, true);
+		}
+
+		[UhtKeyword(Extends = UhtTableNames.Class)]
+		[UhtKeyword(Extends = UhtTableNames.Interface)]
+		[UhtKeyword(Extends = UhtTableNames.NativeInterface)]
+		private static UhtParseResult UFUNCTIONKeyword(UhtParsingScope TopScope, UhtParsingScope ActionScope, ref UhtToken Token)
+		{
+			return ParseUFunction(TopScope, Token);
+		}
+
+		[UhtKeywordCatchAll(Extends = UhtTableNames.Global)]
+		private static UhtParseResult ParseCatchAllKeyword(UhtParsingScope TopScope, ref UhtToken Token)
+		{
+			if (UhtFunctionParser.IsValidateDelegateDeclaration(Token))
+			{
+				return ParseUDelegate(TopScope, Token, false);
+			}
+			return UhtParseResult.Unhandled;
+		}
+		#endregion
+
+		private static UhtParseResult ParseUDelegate(UhtParsingScope ParentScope, UhtToken Token, bool bHasSpecifiers)
+		{
+			UhtFunctionParser Function = new UhtFunctionParser(ParentScope.ScopeType, Token.InputLine);
+
+			using (var TopScope = new UhtParsingScope(ParentScope, Function, KeywordTable, UhtAccessSpecifier.Public))
+			{
+				const string ScopeName = "delegate declaration";
+
+				using (var TokenContext = new UhtMessageContext(TopScope.TokenReader, ScopeName))
+				{
+					TopScope.AddModuleRelativePathToMetaData();
+
+					UhtSpecifierContext SpecifierContext = new UhtSpecifierContext(TopScope, TopScope.TokenReader, Function.MetaData);
+					UhtSpecifierParser Specifiers = TopScope.HeaderParser.GetSpecifierParser(SpecifierContext, ScopeName, SpecifierTable);
+
+					// If this is a UDELEGATE, parse the specifiers
+					StringView DelegateMacro = new StringView();
+					if (bHasSpecifiers)
+					{
+						Specifiers.ParseSpecifiers();
+						Specifiers.ParseDeferred();
+						FinalizeFunctionSpecifiers(Function);
+
+						UhtToken MacroToken = TopScope.TokenReader.GetToken();
+						if (!IsValidateDelegateDeclaration(MacroToken))
+						{
+							throw new UhtTokenException(TopScope.TokenReader, MacroToken, "delegate macro");
+						}
+						DelegateMacro = MacroToken.Value;
+					}
+					else
+					{
+						DelegateMacro = Token.Value;
+					}
+
+					// Break the delegate declaration macro down into parts
+					bool bHasReturnValue = DelegateMacro.Span.Contains("_RetVal".AsSpan(), StringComparison.Ordinal);
+					bool bDeclaredConst = DelegateMacro.Span.Contains("_Const".AsSpan(), StringComparison.Ordinal);
+					bool bIsMulticast = DelegateMacro.Span.Contains("_MULTICAST".AsSpan(), StringComparison.Ordinal);
+					bool bIsSparse = DelegateMacro.Span.Contains("_SPARSE".AsSpan(), StringComparison.Ordinal);
+
+					// Determine the parameter count
+					int FoundParamIndex = UhtConfig.Instance.FindDelegateParameterCount(DelegateMacro);
+
+					// Try reconstructing the string to make sure it matches our expectations
+					string ExpectedOriginalString = string.Format("DECLARE_DYNAMIC{0}{1}_DELEGATE{2}{3}{4}",
+						bIsMulticast ? "_MULTICAST" : "",
+						bIsSparse ? "_SPARSE" : "",
+						bHasReturnValue ? "_RetVal" : "",
+						UhtConfig.Instance.GetDelegateParameterCountString(FoundParamIndex),
+						bDeclaredConst ? "_Const" : "");
+					if (DelegateMacro != ExpectedOriginalString)
+					{
+						throw new UhtException(TopScope.TokenReader, $"Unable to parse delegate declaration; expected '{ExpectedOriginalString}' but found '{DelegateMacro}'.");
+					}
+
+					// Multi-cast delegate function signatures are not allowed to have a return value
+					if (bHasReturnValue && bIsMulticast)
+					{
+						throw new UhtException(TopScope.TokenReader, "Multi-cast delegates function signatures must not return a value");
+					}
+
+					// Delegate signature
+					Function.FunctionType = bIsSparse ? UhtFunctionType.SparseDelegate : UhtFunctionType.Delegate;
+					Function.FunctionFlags |= EFunctionFlags.Public | EFunctionFlags.Delegate;
+					if (bIsMulticast)
+					{
+						Function.FunctionFlags |= EFunctionFlags.MulticastDelegate;
+					}
+
+					// Now parse the macro body
+					TopScope.TokenReader.Require('(');
+
+					// Parse the return type
+					UhtProperty? ReturnValueProperty = null;
+					if (bHasReturnValue)
+					{
+						TopScope.HeaderParser.GetPropertyParser(TopScope).Parse(EPropertyFlags.None,
+							Function.GetPropertyParseOptions(true), UhtParsePropertyDeclarationStyle.None, UhtPropertyCategory.Return,
+							(UhtParsingScope TopScope, UhtProperty Property, ref UhtToken NameToken, UhtLayoutMacroType LayoutMacroType) =>
+							{
+								Property.PropertyFlags |= EPropertyFlags.Parm | EPropertyFlags.OutParm | EPropertyFlags.ReturnParm;
+								ReturnValueProperty = Property;
+							});
+						TopScope.TokenReader.Require(',');
+					}
+
+					// Skip white spaces to get InputPos exactly on beginning of function name.
+					TopScope.TokenReader.SkipWhitespaceAndComments();
+
+					// Get the delegate name
+					UhtToken FuncNameToken = TopScope.TokenReader.GetIdentifier("name");
+					Function.SourceName = FuncNameToken.Value.ToString();
+
+					// If this is a delegate function then go ahead and mangle the name so we don't collide with
+					// actual functions or properties
+					{
+						//@TODO: UCREMOVAL: Eventually this mangling shouldn't occur
+
+						// Remove the leading F
+						if (Function.SourceName[0] != 'F')
+						{
+							TopScope.TokenReader.LogError("Delegate type declarations must start with F");
+						}
+						Function.StrippedFunctionName = Function.SourceName.Substring(1);
+						Function.EngineName = $"{Function.StrippedFunctionName}{UhtFunction.GeneratedDelegateSignatureSuffix}";
+					}
+
+					SetFunctionNames(Function);
+					AddFunction(Function);
+
+					// determine whether this function should be 'const'
+					if (bDeclaredConst)
+					{
+						Function.FunctionFlags |= EFunctionFlags.Const;
+						Function.FunctionExportFlags |= UhtFunctionExportFlags.DeclaredConst;
+					}
+
+					if (bIsSparse)
+					{
+						TopScope.TokenReader.Require(',');
+						UhtToken Name = TopScope.TokenReader.GetIdentifier("OwningClass specifier");
+						UhtEngineNameParts Parts = UhtUtilities.GetEngineNameParts(Name.Value);
+						Function.SparseOwningClassName = Parts.EngineName.ToString();
+						TopScope.TokenReader.Require(',');
+						Function.SparseDelegateName = TopScope.TokenReader.GetIdentifier("delegate name").Value.ToString();
+					}
+
+					Function.FunctionLineNumber = TopScope.TokenReader.InputLine;
+
+					// Get parameter list.
+					if (FoundParamIndex >= 0)
+					{
+						TopScope.TokenReader.Require(',');
+
+						ParseParameterList(TopScope, Function.GetPropertyParseOptions(false));
+					}
+					else
+					{
+						// Require the closing paren even with no parameter list
+						TopScope.TokenReader.Require(')');
+					}
+
+					// Add back in the return value
+					if (ReturnValueProperty != null)
+					{
+						TopScope.ScopeType.AddChild(ReturnValueProperty);
+					}
+
+					// Verify the number of parameters (FoundParamIndex = -1 means zero parameters, 0 means one, ...)
+					int ExpectedProperties = FoundParamIndex + 1 + (bHasReturnValue ? 1 : 0);
+					int PropertiesCount = Function.Properties.Count();
+					if (PropertiesCount != ExpectedProperties)
+					{
+						throw new UhtException(TopScope.TokenReader, $"Expected {ExpectedProperties} parameters but found {PropertiesCount} parameters");
+					}
+
+					// The macro line must be set here
+					Function.MacroLineNumber = TopScope.TokenReader.InputLine;
+
+					// Try parsing metadata for the function
+					Specifiers.ParseFieldMetaData(Function.SourceName);
+
+					TopScope.AddFormattedCommentsAsTooltipMetaData();
+
+					// Optionally consume a semicolon, it's not required for the delegate macro since it contains one internally
+					TopScope.TokenReader.Optional(';');
+				}
+				return UhtParseResult.Handled;
+			}
+		}
+
+		private static UhtParseResult ParseUFunction(UhtParsingScope ParentScope, UhtToken Token)
+		{
+			UhtFunctionParser Function = new UhtFunctionParser(ParentScope.ScopeType, Token.InputLine);
+
+			using (var TopScope = new UhtParsingScope(ParentScope, Function, KeywordTable, UhtAccessSpecifier.Public))
+			{
+				UhtParsingScope OuterClassScope = TopScope.CurrentClassScope;
+				UhtClass OuterClass = (UhtClass)OuterClassScope.ScopeType;
+				string ScopeName = "function";
+
+				using (var TokenContext = new UhtMessageContext(TopScope.TokenReader, ScopeName))
+				{
+					TopScope.AddModuleRelativePathToMetaData();
+
+					UhtSpecifierContext SpecifierContext = new UhtSpecifierContext(TopScope, TopScope.TokenReader, Function.MetaData);
+					UhtSpecifierParser SpecifierParser = TopScope.HeaderParser.GetSpecifierParser(SpecifierContext, ScopeName, SpecifierTable);
+					SpecifierParser.ParseSpecifiers();
+
+					if (!OuterClass.ClassFlags.HasAnyFlags(EClassFlags.Native))
+					{
+						throw new UhtException(Function, "Should only be here for native classes!");
+					}
+
+					Function.MacroLineNumber = TopScope.TokenReader.InputLine;
+					Function.FunctionFlags |= EFunctionFlags.Native;
+
+					bool bAutomaticallyFinal = true;
+					switch (OuterClassScope.AccessSpecifier)
+					{
+						case UhtAccessSpecifier.Public:
+							Function.FunctionFlags |= EFunctionFlags.Public;
+							break;
+
+						case UhtAccessSpecifier.Protected:
+							Function.FunctionFlags |= EFunctionFlags.Protected;
+							break;
+
+						case UhtAccessSpecifier.Private:
+							Function.FunctionFlags |= EFunctionFlags.Private | EFunctionFlags.Final;
+
+							// This is automatically final as well, but in a different way and for a different reason
+							bAutomaticallyFinal = false;
+							break;
+					}
+
+					if (TopScope.TokenReader.TryOptional("static"))
+					{
+						Function.FunctionFlags |= EFunctionFlags.Static;
+						Function.FunctionExportFlags |= UhtFunctionExportFlags.CppStatic;
+					}
+
+					if (Function.MetaData.ContainsKey(UhtNames.CppFromBpEvent))
+					{
+						Function.FunctionFlags |= EFunctionFlags.Event;
+					}
+
+					if ((TopScope.HeaderParser.GetCurrentCompositeCompilerDirective() & UhtCompilerDirective.WithEditor) != 0)
+					{
+						Function.FunctionFlags |= EFunctionFlags.EditorOnly;
+					}
+
+					SpecifierParser.ParseDeferred();
+					FinalizeFunctionSpecifiers(Function);
+
+					if (Function.FunctionExportFlags.HasAnyFlags(UhtFunctionExportFlags.CustomThunk) && !Function.MetaData.ContainsKey(UhtNames.CustomThunk))
+					{
+						Function.MetaData.Add(UhtNames.CustomThunk, true);
+					}
+
+					if (Function.FunctionFlags.HasAnyFlags(EFunctionFlags.Net))
+					{
+						// Network replicated functions are always events, and are only final if sealed
+						ScopeName = "event";
+						TokenContext.Reset(ScopeName);
+						bAutomaticallyFinal = false;
+					}
+
+					if (Function.FunctionFlags.HasAnyFlags(EFunctionFlags.BlueprintEvent))
+					{
+						ScopeName = Function.FunctionFlags.HasAnyFlags(EFunctionFlags.Native) ? "BlueprintNativeEvent" : "BlueprintImplementableEvent";
+						TokenContext.Reset(ScopeName);
+						bAutomaticallyFinal = false;
+					}
+
+					// Record the tokens so we can detect this function as a declaration later (i.e. RPC)
+					using (UhtTokenRecorder TokenRecorder = new UhtTokenRecorder(ParentScope))
+					{
+
+						if (TopScope.TokenReader.TryOptional("virtual"))
+						{
+							Function.FunctionExportFlags |= UhtFunctionExportFlags.Virtual;
+						}
+
+						bool bInternalOnly = Function.MetaData.GetBoolean(UhtNames.BlueprintInternalUseOnly);
+
+						// Peek ahead to look for a CORE_API style DLL import/export token if present
+						UhtToken APIMacroToken;
+						if (TopScope.TokenReader.TryOptionalAPIMacro(out APIMacroToken))
+						{
+							//@TODO: Validate the module name for RequiredAPIMacroIfPresent
+							Function.FunctionFlags |= EFunctionFlags.RequiredAPI;
+							Function.FunctionExportFlags |= UhtFunctionExportFlags.RequiredAPI;
+						}
+
+						// Look for static again, in case there was an ENGINE_API token first
+						if (APIMacroToken && TopScope.TokenReader.TryOptional("static"))
+						{
+							TopScope.TokenReader.LogError($"Unexpected API macro '{APIMacroToken.Value}'. Did you mean to put '{APIMacroToken.Value}' after the static keyword?");
+						}
+
+						// Look for virtual again, in case there was an ENGINE_API token first
+						if (TopScope.TokenReader.TryOptional("virtual"))
+						{
+							Function.FunctionExportFlags |= UhtFunctionExportFlags.Virtual;
+						}
+
+						// If virtual, remove the implicit final, the user can still specifying an explicit final at the end of the declaration
+						if (Function.FunctionExportFlags.HasAnyFlags(UhtFunctionExportFlags.Virtual))
+						{
+							bAutomaticallyFinal = false;
+						}
+
+						// Handle the initial implicit/explicit final
+						// A user can still specify an explicit final after the parameter list as well.
+						if (bAutomaticallyFinal || Function.FunctionExportFlags.HasAnyFlags(UhtFunctionExportFlags.SealedEvent))
+						{
+							Function.FunctionFlags |= EFunctionFlags.Final;
+							Function.FunctionExportFlags |= UhtFunctionExportFlags.Final | UhtFunctionExportFlags.AutoFinal;
+						}
+
+						// Get return type.  C++ style functions always have a return value type, even if it's void
+						UhtToken FuncNameToken = new UhtToken();
+						UhtProperty? ReturnValueProperty = null;
+						TopScope.HeaderParser.GetPropertyParser(TopScope).Parse(EPropertyFlags.None,
+							Function.GetPropertyParseOptions(true), UhtParsePropertyDeclarationStyle.None, UhtPropertyCategory.Return,
+							(UhtParsingScope TopScope, UhtProperty Property, ref UhtToken NameToken, UhtLayoutMacroType LayoutMacroType) =>
+							{
+								Property.PropertyFlags |= EPropertyFlags.Parm | EPropertyFlags.OutParm | EPropertyFlags.ReturnParm;
+								FuncNameToken = NameToken;
+								if (!(Property is UhtVoidProperty))
+								{
+									ReturnValueProperty = Property;
+								}
+							});
+
+						if (FuncNameToken.Value.Length == 0)
+						{
+							throw new UhtException(TopScope.TokenReader, "expected return value and function name");
+						}
+
+						// Get function or operator name.
+						Function.SourceName = FuncNameToken.Value.ToString();
+
+						ScopeName = $"{ScopeName} '{Function.SourceName}'";
+						TokenContext.Reset(ScopeName);
+
+						TopScope.TokenReader.Require('(');
+
+						SetFunctionNames(Function);
+						AddFunction(Function);
+
+						// Get parameter list.
+						ParseParameterList(TopScope, Function.GetPropertyParseOptions(false));
+
+						// Add back in the return value
+						if (ReturnValueProperty != null)
+						{
+							TopScope.ScopeType.AddChild(ReturnValueProperty);
+						}
+
+						// determine whether this function should be 'const'
+						if (TopScope.TokenReader.TryOptional("const"))
+						{
+							if (Function.FunctionFlags.HasAnyFlags(EFunctionFlags.Native))
+							{
+								// @TODO: UCREMOVAL Reconsider?
+								//Throwf(TEXT("'const' may only be used for native functions"));
+							}
+
+							Function.FunctionFlags |= EFunctionFlags.Const;
+							Function.FunctionExportFlags |= UhtFunctionExportFlags.DeclaredConst;
+						}
+
+						// Try parsing metadata for the function
+						SpecifierParser.ParseFieldMetaData(Function.EngineName);
+
+						TopScope.AddFormattedCommentsAsTooltipMetaData();
+
+						// 'final' and 'override' can appear in any order before an optional '= 0' pure virtual specifier
+						bool bFoundFinal = TopScope.TokenReader.TryOptional("final");
+						bool bFoundOverride = TopScope.TokenReader.TryOptional("override");
+						if (!bFoundFinal && bFoundOverride)
+						{
+							bFoundFinal = TopScope.TokenReader.TryOptional("final");
+						}
+
+						// Handle C++ style functions being declared as abstract
+						if (TopScope.TokenReader.TryOptional('='))
+						{
+							int ZeroValue = 1;
+							bool bGotZero = TopScope.TokenReader.TryOptionalConstInt(out ZeroValue);
+							bGotZero = bGotZero && (ZeroValue == 0);
+							if (!bGotZero || ZeroValue != 0)
+							{
+								throw new UhtException(TopScope.TokenReader, "Expected 0 to indicate function is abstract");
+							}
+						}
+
+						// Look for the final keyword to indicate this function is sealed
+						if (bFoundFinal)
+						{
+							// This is a final (prebinding, non-overridable) function
+							Function.FunctionFlags |= EFunctionFlags.Final;
+							Function.FunctionExportFlags |= UhtFunctionExportFlags.Final;
+						}
+
+						// Optionally consume a semicolon
+						// This is optional to allow inline function definitions
+						if (TopScope.TokenReader.TryOptional(';'))
+						{
+							// Do nothing (consume it)
+						}
+						else if (TopScope.TokenReader.TryPeekOptional('{'))
+						{
+							// Skip inline function bodies
+							UhtToken TokenCopy = new UhtToken();
+							TopScope.TokenReader.SkipDeclaration(ref TokenCopy);
+						}
+					}
+				}
+				return UhtParseResult.Handled;
+			}
+		}
+
+		private static void FinalizeFunctionSpecifiers(UhtFunction Function)
+		{
+			if (Function.FunctionFlags.HasAnyFlags(EFunctionFlags.Net))
+			{
+				// Network replicated functions are always events
+				Function.FunctionFlags |= EFunctionFlags.Event;
+			}
+		}
+
+		internal static bool IsValidateDelegateDeclaration(UhtToken Token)
+		{
+			return (Token.IsIdentifier() && Token.Value.Span.StartsWith("DECLARE_DYNAMIC_"));
+		}
+
+		private static void ParseParameterList(UhtParsingScope TopScope, UhtPropertyParseOptions Options)
+		{
+			UhtFunction Function = (UhtFunction)TopScope.ScopeType;
+
+			bool bIsNetFunc = Function.FunctionFlags.HasAnyFlags(EFunctionFlags.Net);
+			UhtPropertyCategory PropertyCategory = bIsNetFunc ? UhtPropertyCategory.ReplicatedParameter : UhtPropertyCategory.RegularParameter;
+			EPropertyFlags DisallowFlags = ~(EPropertyFlags.ParmFlags | EPropertyFlags.AutoWeak | EPropertyFlags.RepSkip | EPropertyFlags.UObjectWrapper | EPropertyFlags.NativeAccessSpecifiers);
+
+			UhtAdvancedDisplayParameterHandler AdvancedDisplay = new UhtAdvancedDisplayParameterHandler(TopScope.ScopeType.MetaData);
+
+			TopScope.TokenReader.RequireList(')', ',', false, () =>
+			{
+				TopScope.HeaderParser.GetPropertyParser(TopScope).Parse(DisallowFlags, Options, UhtParsePropertyDeclarationStyle.None, PropertyCategory,
+					(UhtParsingScope TopScope, UhtProperty Property, ref UhtToken NameToken, UhtLayoutMacroType LayoutMacroType) =>
+					{
+						Property.PropertyFlags |= EPropertyFlags.Parm;
+						if (AdvancedDisplay.CanMarkMore() && AdvancedDisplay.ShouldMarkParameter(Property.EngineName))
+						{
+							Property.PropertyFlags |= EPropertyFlags.AdvancedDisplay;
+						}
+
+						// Default value.
+						if (TopScope.TokenReader.TryOptional('='))
+						{
+							List<UhtToken> DefaultValueTokens = new List<UhtToken>();
+							int ParenthesisNestCount = 0;
+							while (!TopScope.TokenReader.bIsEOF)
+							{
+								UhtToken Token = TopScope.TokenReader.PeekToken();
+								if (Token.IsSymbol(','))
+								{
+									if (ParenthesisNestCount == 0)
+									{
+										break;
+									}
+									DefaultValueTokens.Add(Token);
+									TopScope.TokenReader.ConsumeToken();
+								}
+								else if (Token.IsSymbol(')'))
+								{
+									if (ParenthesisNestCount == 0)
+									{
+										break;
+									}
+									DefaultValueTokens.Add(Token);
+									TopScope.TokenReader.ConsumeToken();
+									--ParenthesisNestCount;
+								}
+								else if (Token.IsSymbol('('))
+								{
+									++ParenthesisNestCount;
+									DefaultValueTokens.Add(Token);
+									TopScope.TokenReader.ConsumeToken();
+								}
+								else
+								{
+									DefaultValueTokens.Add(Token);
+									TopScope.TokenReader.ConsumeToken();
+								}
+							}
+
+							// allow exec functions to be added to the metaData, this is so we can have default params for them.
+							bool bStoreCppDefaultValueInMetaData = Function.FunctionFlags.HasAnyFlags(EFunctionFlags.BlueprintCallable | EFunctionFlags.Exec);
+							if (DefaultValueTokens.Count > 0 && bStoreCppDefaultValueInMetaData)
+							{
+								Property.DefaultValueTokens = DefaultValueTokens;
+							}
+						}
+					});
+			});
+		}
+
+		private static void AddFunction(UhtFunction Function)
+		{
+			if (Function.Outer != null)
+			{
+				Function.Outer.AddChild(Function);
+			}
+		}
+
+		private static void SetFunctionNames(UhtFunction Function)
+		{
+			// The source name won't have the suffix applied to delegate names, however, the engine name will
+			// We use the engine name because we need to detect the suffix for delegates
+			string FunctionName = Function.EngineName;
+			if (FunctionName.EndsWith(UhtFunction.GeneratedDelegateSignatureSuffix))
+			{
+				FunctionName = FunctionName.Substring(0, FunctionName.Length - UhtFunction.GeneratedDelegateSignatureSuffix.Length);
+			}
+
+			Function.UnMarshalAndCallName = "exec" + FunctionName;
+
+			if (Function.FunctionFlags.HasAnyFlags(EFunctionFlags.BlueprintEvent))
+			{
+				Function.MarshalAndCallName = FunctionName;
+				if (Function.FunctionFlags.HasAllFlags(EFunctionFlags.BlueprintEvent | EFunctionFlags.Native))
+				{
+					Function.CppImplName = Function.EngineName + "_Implementation";
+				}
+			}
+			else if (Function.FunctionFlags.HasAllFlags(EFunctionFlags.Native | EFunctionFlags.Net))
+			{
+				Function.MarshalAndCallName = FunctionName;
+				if (Function.FunctionFlags.HasAnyFlags(EFunctionFlags.NetResponse))
+				{
+					// Response function implemented by programmer and called directly from thunk
+					Function.CppImplName = Function.EngineName;
+				}
+				else
+				{
+					if (Function.CppImplName.Length == 0)
+					{
+						Function.CppImplName = Function.EngineName + "_Implementation";
+					}
+					else if (Function.CppImplName == FunctionName)
+					{
+						Function.LogError("Native implementation function must be different than original function name.");
+					}
+
+					if (Function.CppValidationImplName.Length == 0 && Function.FunctionFlags.HasAnyFlags(EFunctionFlags.NetValidate))
+					{
+						Function.CppValidationImplName = Function.EngineName + "_Validate";
+					}
+					else if (Function.CppValidationImplName == FunctionName)
+					{
+						Function.LogError("Validation function must be different than original function name.");
+					}
+				}
+			}
+			else if (Function.FunctionFlags.HasAnyFlags(EFunctionFlags.Delegate))
+			{
+				Function.MarshalAndCallName = "delegate" + FunctionName;
+			}
+
+			if (Function.CppImplName.Length == 0)
+			{
+				Function.CppImplName = FunctionName;
+			}
+
+			if (Function.MarshalAndCallName.Length == 0)
+			{
+				Function.MarshalAndCallName = "event" + FunctionName;
+			}
+		}
+	}
+}
diff --git a/Engine/Source/Programs/Shared/EpicGames.UHT/Parsers/UhtHeaderFileParser.cs b/Engine/Source/Programs/Shared/EpicGames.UHT/Parsers/UhtHeaderFileParser.cs
new file mode 100644
--- /dev/null
+++ b/Engine/Source/Programs/Shared/EpicGames.UHT/Parsers/UhtHeaderFileParser.cs
@@ -0,0 +1,1058 @@
+// Copyright Epic Games, Inc. All Rights Reserved.
+
+using EpicGames.Core;
+using EpicGames.UHT.Tables;
+using EpicGames.UHT.Tokenizer;
+using EpicGames.UHT.Types;
+using EpicGames.UHT.Utils;
+using System;
+using System.Collections.Generic;
+using System.IO;
+
+namespace EpicGames.UHT.Parsers
+{
+	public enum UhtParseResult
+	{
+		Handled,
+		Unhandled,
+		Invalid,
+	}
+
+	[Flags]
+	public enum UhtCompilerDirective
+	{
+		None = 0,
+		/// <summary>
+		/// This indicates we are in a "#if CPP" block
+		/// </summary>
+		CPPBlock = 1 << 0,
+		/// <summary>
+		/// This indicates we are in a "#if !CPP" block
+		/// </summary>
+		NotCPPBlock = 1 << 1,
+		/// <summary>
+		/// This indicates we are in a "#if 0" block
+		/// </summary>
+		ZeroBlock = 1 << 2,
+		/// <summary>
+		/// This indicates we are in a "#if 1" block
+		/// </summary>
+		OneBlock = 1 << 3,
+		/// <summary>
+		/// This indicates we are in a "#if WITH_EDITOR" block
+		/// </summary>
+		WithEditor = 1 << 4,
+		/// <summary>
+		/// This indicates we are in a "#if WITH_EDITORONLY_DATA" block
+		/// </summary>
+		WithEditorOnlyData = 1 << 5,
+		/// <summary>
+		/// This indicates we are in a "#if WITH_HOT_RELOAD" block
+		/// </summary>
+		WithHotReload = 1 << 6,
+		/// <summary>
+		/// This directive is unrecognized and does not change the code generation at all
+		/// </summary>
+		Unrecognized = 1 << 7,
+	}
+
+	/// <summary>
+	/// Helper methods for testing flags.  These methods perform better than the generic HasFlag which hits
+	/// the GC and stalls.
+	/// </summary>
+	public static class UhtCompilerDirectiveExtensions
+	{
+
+		/// <summary>
+		/// Test to see if any of the specified flags are set
+		/// </summary>
+		/// <param name="InFlags">Current flags</param>
+		/// <param name="TestFlags">Flags to test for</param>
+		/// <returns>True if any of the flags are set</returns>
+		public static bool HasAnyFlags(this UhtCompilerDirective InFlags, UhtCompilerDirective TestFlags)
+		{
+			return (InFlags & TestFlags) != 0;
+		}
+
+		/// <summary>
+		/// Test to see if all of the specified flags are set
+		/// </summary>
+		/// <param name="InFlags">Current flags</param>
+		/// <param name="TestFlags">Flags to test for</param>
+		/// <returns>True if all the flags are set</returns>
+		public static bool HasAllFlags(this UhtCompilerDirective InFlags, UhtCompilerDirective TestFlags)
+		{
+			return (InFlags & TestFlags) == TestFlags;
+		}
+
+		/// <summary>
+		/// Test to see if a specific set of flags have a specific value.
+		/// </summary>
+		/// <param name="InFlags">Current flags</param>
+		/// <param name="TestFlags">Flags to test for</param>
+		/// <param name="MatchFlags">Expected value of the tested flags</param>
+		/// <returns>True if the given flags have a specific value.</returns>
+		public static bool HasExactFlags(this UhtCompilerDirective InFlags, UhtCompilerDirective TestFlags, UhtCompilerDirective MatchFlags)
+		{
+			return (InFlags & TestFlags) == MatchFlags;
+		}
+	}
+
+	[UnrealHeaderTool]
+	public static class UhtAccessSpecifierKeywords
+	{
+		#region Keywords
+		[UhtKeyword(Extends = UhtTableNames.ClassBase, Keyword = "public")]
+		[UhtKeyword(Extends = UhtTableNames.ScriptStruct, Keyword = "public")]
+		private static UhtParseResult PublicKeyword(UhtParsingScope TopScope, UhtParsingScope ActionScope, ref UhtToken Token)
+		{
+			return SetAccessSpecifier(TopScope, UhtAccessSpecifier.Public);
+		}
+
+		[UhtKeyword(Extends = UhtTableNames.ClassBase, Keyword = "protected")]
+		[UhtKeyword(Extends = UhtTableNames.ScriptStruct, Keyword = "protected")]
+		private static UhtParseResult ProtectedKeyword(UhtParsingScope TopScope, UhtParsingScope ActionScope, ref UhtToken Token)
+		{
+			return SetAccessSpecifier(TopScope, UhtAccessSpecifier.Protected);
+		}
+
+		[UhtKeyword(Extends = UhtTableNames.ClassBase, Keyword = "private")]
+		[UhtKeyword(Extends = UhtTableNames.ScriptStruct, Keyword = "private")]
+		private static UhtParseResult PrivateKeyword(UhtParsingScope TopScope, UhtParsingScope ActionScope, ref UhtToken Token)
+		{
+			return SetAccessSpecifier(TopScope, UhtAccessSpecifier.Private);
+		}
+		#endregion
+
+		private static UhtParseResult SetAccessSpecifier(UhtParsingScope TopScope, UhtAccessSpecifier AccessSpecifier)
+		{
+			TopScope.AccessSpecifier = AccessSpecifier;
+			TopScope.TokenReader.Require(':');
+			return UhtParseResult.Handled;
+		}
+	}
+
+	public class UhtHeaderFileParser : IUhtTokenPreprocessor
+	{
+		private static UhtKeywordTable KeywordTable = UhtKeywordTables.Instance.Get(UhtTableNames.Global);
+
+		private struct CompilerDirective
+		{
+			public UhtCompilerDirective Element;
+			public UhtCompilerDirective Composite;
+		}
+
+		/// <summary>
+		/// Header file being parsed
+		/// </summary>
+		public readonly UhtHeaderFile HeaderFile;
+
+		/// <summary>
+		/// Token reader for the header
+		/// </summary>
+		public readonly IUhtTokenReader TokenReader;
+
+		/// <summary>
+		/// If true, this header file belongs to the engine
+		/// </summary>
+		public bool bIsPartOfEngine => this.HeaderFile.Package.bIsPartOfEngine;
+
+		/// <summary>
+		/// If true, the inclusion of the generated header file was seen
+		/// </summary>
+		public bool bSpottedAutogeneratedHeaderInclude = false;
+
+		/// <summary>
+		/// For a given header file, we share a common specifier parser to reduce the number of allocations.
+		/// Before the parser can be reused, the ParseDeferred method must be called to dispatch that list.
+		/// </summary>
+		private UhtSpecifierParser? SpecifierParser = null;
+
+		/// <summary>
+		/// For a given header file, we share a common property parser to reduce the number of allocations.
+		/// </summary>
+		public UhtPropertyParser? PropertyParser = null;
+
+		/// <summary>
+		/// Stack of current #if states
+		/// </summary>
+		private List<CompilerDirective> CompilerDirectives = new List<CompilerDirective>();
+
+		/// <summary>
+		/// Stack of current #if states saved as part of the preprocessor state
+		/// </summary>
+		private List<CompilerDirective> SavedCompilerDirectives = new List<CompilerDirective>();
+
+		/// <summary>
+		/// Current top of the parsing scopes.  Classes, structures and functions all allocate scopes.
+		/// </summary>
+		private UhtParsingScope? TopScope = null;
+
+		/// <summary>
+		/// The total number of statements parsed
+		/// </summary>
+		private int StatementsParsed = 0;
+
+		public static UhtHeaderFileParser Parse(UhtHeaderFile HeaderFile)
+		{
+			UhtHeaderFileParser HeaderParser = new UhtHeaderFileParser(HeaderFile);
+			using (UhtParsingScope TopScope = new UhtParsingScope(HeaderParser, HeaderParser.HeaderFile, UhtHeaderFileParser.KeywordTable))
+			{
+				HeaderParser.ParseStatements();
+
+				if (!HeaderParser.bSpottedAutogeneratedHeaderInclude && HeaderParser.HeaderFile.Data.Length > 0)
+				{
+					bool bNoExportClassesOnly = true;
+					foreach (UhtType Type in HeaderParser.HeaderFile.Children)
+					{
+						if (Type is UhtClass Class)
+						{
+							if (Class.ClassType != UhtClassType.NativeInterface && !Class.ClassFlags.HasAnyFlags(EClassFlags.NoExport))
+							{
+								bNoExportClassesOnly = false;
+								break;
+							}
+						}
+					}
+
+					if (!bNoExportClassesOnly)
+					{
+						HeaderParser.HeaderFile.LogError($"Expected an include at the top of the header the follows all other includes: '#include \"{HeaderParser.HeaderFile.GeneratedHeaderFileName}\"'");
+					}
+				}
+			}
+			return HeaderParser;
+		}
+
+		public UhtHeaderFileParser(UhtHeaderFile HeaderFile)
+		{
+			this.TokenReader = new UhtTokenBufferReader(HeaderFile, HeaderFile.Data.Memory);
+			this.HeaderFile = HeaderFile;
+			this.TokenReader.TokenPreprocessor = this;
+		}
+
+		public void PushScope(UhtParsingScope Scope)
+		{
+			if (Scope.ParentScope != this.TopScope)
+			{
+				throw new UhtIceException("Pushing a new scope whose parent isn't the current top scope.");
+			}
+			this.TopScope = Scope;
+		}
+
+		public void PopScope(UhtParsingScope Scope)
+		{
+			if (Scope != this.TopScope)
+			{
+				throw new UhtIceException("Attempt to pop a scope that isn't the top scope");
+			}
+			this.TopScope = Scope.ParentScope;
+		}
+
+		public UhtSpecifierParser GetSpecifierParser(UhtSpecifierContext SpecifierContext, StringView Context, UhtSpecifierTable Table)
+		{
+			if (this.SpecifierParser == null)
+			{
+				this.SpecifierParser = new UhtSpecifierParser(SpecifierContext, Context, Table);
+			}
+			else
+			{
+				this.SpecifierParser.Reset(SpecifierContext, Context, Table);
+			}
+			return this.SpecifierParser;
+		}
+
+		public UhtPropertyParser GetPropertyParser(UhtParsingScope TopScope)
+		{
+			if (this.PropertyParser == null)
+			{
+				this.PropertyParser = new UhtPropertyParser(TopScope, TopScope.TokenReader);
+			}
+			else
+			{
+				this.PropertyParser.Reset(TopScope, TopScope.TokenReader);
+			}
+			return this.PropertyParser;
+		}
+
+		/// <summary>
+		/// Return the current compiler directive
+		/// </summary>
+		/// <returns>Enumeration flags for all active compiler directives</returns>
+		public UhtCompilerDirective GetCurrentCompositeCompilerDirective()
+		{
+			return this.CompilerDirectives.Count > 0 ? this.CompilerDirectives[this.CompilerDirectives.Count - 1].Composite : UhtCompilerDirective.None;
+		}
+
+		public UhtCompilerDirective GetCurrentNonCompositeCompilerDirective()
+		{
+			return this.CompilerDirectives.Count > 0 ? this.CompilerDirectives[this.CompilerDirectives.Count - 1].Element : UhtCompilerDirective.None;
+		}
+
+		#region ITokenPreprocessor implementation
+		public bool ParsePreprocessorDirective(ref UhtToken Token, bool bIsBeingIncluded, out bool bClearComments, out bool bIllegalContentsCheck)
+		{
+			bClearComments = true;
+			bIllegalContentsCheck = true;
+			if (ParseDirectiveInternal(Token, bIsBeingIncluded))
+			{
+				bClearComments = ClearCommentsCompilerDirective();
+			}
+			bIllegalContentsCheck = !GetCurrentNonCompositeCompilerDirective().HasAnyFlags(UhtCompilerDirective.ZeroBlock | UhtCompilerDirective.WithEditorOnlyData);
+			return IncludeCurrentCompilerDirective();
+		}
+
+		public void SaveState()
+		{
+			this.SavedCompilerDirectives.Clear();
+			this.SavedCompilerDirectives.AddRange(this.CompilerDirectives);
+		}
+
+		public void RestoreState()
+		{
+			this.CompilerDirectives.Clear();
+			this.CompilerDirectives.AddRange(this.SavedCompilerDirectives);
+		}
+		#endregion
+
+		#region Statement parsing
+		public void ParseStatements()
+		{
+			ParseStatements((char)0, (char)0, true);
+		}
+
+		public void ParseStatements(char Initiator, char Terminator, bool bLogUnhandledKeywords)
+		{
+			if (this.TopScope == null)
+			{
+				return;
+			}
+
+			if (Initiator != 0)
+			{
+				this.TokenReader.Require(Initiator);
+			}
+			while (true)
+			{
+				UhtToken Token = this.TokenReader.GetToken();
+				if (Token.TokenType.IsEndType())
+				{
+					if (this.TopScope != null && this.TopScope.ParentScope == null)
+					{
+						CheckEof(ref Token);
+					}
+					return;
+				}
+				else if (Terminator != 0 && Token.IsSymbol(Terminator))
+				{
+					return;
+				}
+
+				if (this.TopScope != null)
+				{
+					ParseStatement(this.TopScope, ref Token, bLogUnhandledKeywords);
+					this.TokenReader.ClearComments();
+					++this.StatementsParsed;
+				}
+			}
+		}
+
+
+		public static bool ParseStatement(UhtParsingScope TopScope, ref UhtToken Token, bool bLogUnhandledKeywords)
+		{
+			UhtParseResult ParseResult = UhtParseResult.Unhandled;
+
+			switch (Token.TokenType)
+			{
+				case UhtTokenType.Identifier:
+					ParseResult = DispatchKeyword(TopScope, ref Token);
+					break;
+
+				case UhtTokenType.Symbol:
+					// Ignore any extra semicolons
+					if (Token.IsSymbol(';'))
+					{
+						return true;
+					}
+					ParseResult = DispatchKeyword(TopScope, ref Token);
+					break;
+			}
+
+			if (ParseResult == UhtParseResult.Unhandled)
+			{
+				ParseResult = DispatchCatchAll(TopScope, ref Token);
+			}
+
+			if (ParseResult == UhtParseResult.Unhandled)
+			{
+				ParseResult = ProbablyAnUnknownObjectLikeMacro(TopScope.TokenReader, ref Token);
+			}
+
+			if (ParseResult == UhtParseResult.Unhandled && bLogUnhandledKeywords)
+			{
+				UhtKeywordTables.Instance.LogUnhandledError(TopScope.TokenReader, Token);
+			}
+
+			if (ParseResult == UhtParseResult.Unhandled || ParseResult == UhtParseResult.Invalid)
+			{
+				if (TopScope.ScopeType is UhtClass Class)
+				{
+					using (UhtTokenRecorder Recorder = new UhtTokenRecorder(TopScope, ref Token))
+					{
+						ParseResult = TopScope.TokenReader.SkipDeclaration(ref Token) ? UhtParseResult.Handled : UhtParseResult.Invalid;
+						if (Recorder.Stop())
+						{
+							if (Class.Declarations != null)
+							{
+								UhtDeclaration Declaration = Class.Declarations[Class.Declarations.Count - 1];
+								if (TopScope.HeaderParser.CheckForConstructor(Class, Declaration))
+								{
+								}
+								else if (Class.ClassType == UhtClassType.Class)
+								{
+									if (TopScope.HeaderParser.CheckForSerialize(Class, Declaration))
+									{
+									}
+								}
+							}
+						}
+					}
+				}
+				else
+				{
+					ParseResult = TopScope.TokenReader.SkipDeclaration(ref Token) ? UhtParseResult.Handled : UhtParseResult.Invalid;
+				}
+			}
+			return true;
+		}
+
+		private static UhtParseResult DispatchKeyword(UhtParsingScope TopScope, ref UhtToken Token)
+		{
+			UhtParseResult ParseResult = UhtParseResult.Unhandled;
+			for (UhtParsingScope? CurrentScope = TopScope; CurrentScope != null && ParseResult == UhtParseResult.Unhandled; CurrentScope = CurrentScope.ParentScope)
+			{
+				UhtKeyword KeywordInfo;
+				if (CurrentScope.ScopeKeywordTable.TryGetValue(Token.Value, out KeywordInfo))
+				{
+					if (KeywordInfo.bAllScopes || TopScope == CurrentScope)
+					{
+						ParseResult = KeywordInfo.Delegate(TopScope, CurrentScope, ref Token);
+					}
+				}
+			}
+			return ParseResult;
+		}
+
+		private static UhtParseResult DispatchCatchAll(UhtParsingScope TopScope, ref UhtToken Token)
+		{
+			for (UhtParsingScope? CurrentScope = TopScope; CurrentScope != null; CurrentScope = CurrentScope.ParentScope)
+			{
+				foreach (var Delegate in CurrentScope.ScopeKeywordTable.CatchAlls)
+				{
+					UhtParseResult ParseResult = Delegate(TopScope, ref Token);
+					if (ParseResult != UhtParseResult.Unhandled)
+					{
+						return ParseResult;
+					}
+				}
+			}
+			return UhtParseResult.Unhandled;
+		}
+
+		/// <summary>
+		/// Tests if an identifier looks like a macro which doesn't have a following open parenthesis.
+		/// </summary>
+		/// <param name="Token">The current token that initiated the process</param>
+		/// <returns>Result if matching the token</returns>
+		private static UhtParseResult ProbablyAnUnknownObjectLikeMacro(IUhtTokenReader TokenReader, ref UhtToken Token)
+		{
+			// Non-identifiers are not macros
+			if (!Token.IsIdentifier())
+			{
+				return UhtParseResult.Unhandled;
+			}
+
+			// Macros must start with a capitalized alphanumeric character or underscore
+			char FirstChar = Token.Value.Span[0];
+			if (FirstChar != '_' && (FirstChar < 'A' || FirstChar > 'Z'))
+			{
+				return UhtParseResult.Unhandled;
+			}
+
+			// We'll guess about it being a macro based on it being fully-capitalized with at least one underscore.
+			int UnderscoreCount = 0;
+			foreach (char Ch in Token.Value.Span.Slice(1))
+			{
+				if (Ch == '_')
+				{
+					++UnderscoreCount;
+				}
+				else if ((Ch < 'A' || Ch > 'Z') && (Ch < '0' || Ch > '9'))
+				{
+					return UhtParseResult.Unhandled;
+				}
+			}
+
+			// We look for at least one underscore as a convenient way of whitelisting many known macros
+			// like FORCEINLINE and CONSTEXPR, and non-macros like FPOV and TCHAR.
+			if (UnderscoreCount == 0)
+			{
+				return UhtParseResult.Unhandled;
+			}
+
+			// Identifiers which end in _API are known
+			if (Token.Value.Span.Length > 4 && Token.Value.Span.EndsWith("_API"))
+			{
+				return UhtParseResult.Unhandled;
+			}
+
+			// Ignore certain known macros or identifiers that look like macros.
+			if (Token.IsValue("FORCEINLINE_DEBUGGABLE") ||
+				Token.IsValue("FORCEINLINE_STATS") ||
+				Token.IsValue("SIZE_T"))
+			{
+				return UhtParseResult.Unhandled;
+			}
+
+			// Check if there's an open parenthesis following the token.
+			return TokenReader.PeekToken().IsSymbol('(') ? UhtParseResult.Unhandled : UhtParseResult.Handled;
+		}
+
+		private void CheckEof(ref UhtToken Token)
+		{
+			if (this.CompilerDirectives.Count > 0)
+			{
+				throw new UhtException(this.TokenReader, Token.InputLine, "Missing #endif");
+			}
+		}
+		#endregion
+
+		#region Internals
+		/// <summary>
+		/// Parse a preprocessor directive.
+		/// </summary>
+		/// <param name="Token">Token that started the parse</param>
+		/// <param name="bIsBeingIncluded">If true, then this directive is in an active block</param>
+		/// <returns>True if we should check to see if tokenizer should clear comments</returns>
+		private bool ParseDirectiveInternal(UhtToken Token, bool bIsBeingIncluded)
+		{
+			bool bCheckClearComments = false;
+
+			// Collect all the lines of the preprocessor statement including any continuations.
+			// We assume that the vast majority of lines will not be continuations.  So avoid using the
+			// string builder as much as possible.
+			int StartingLine = this.TokenReader.InputLine;
+			StringViewBuilder SVB = new StringViewBuilder();
+			while (true)
+			{
+				UhtToken LineToken = this.TokenReader.GetLine();
+				if (LineToken.TokenType != UhtTokenType.Line)
+				{
+					break;
+				}
+				if (LineToken.Value.Span.Length > 0 && LineToken.Value.Span[LineToken.Value.Span.Length - 1] == '\\')
+				{
+					SVB.Append(new StringView(LineToken.Value, 0, LineToken.Value.Span.Length - 1));
+				}
+				else
+				{
+					SVB.Append(LineToken.Value);
+					break;
+				}
+			}
+			StringView Line = SVB.ToStringView();
+
+			// Create a token reader we will use to decode
+			UhtTokenBufferReader LineTokenReader = new UhtTokenBufferReader(this.HeaderFile, Line.Memory);
+			LineTokenReader.InputLine = StartingLine;
+
+			UhtToken Directive;
+			if (!LineTokenReader.TryOptionalIdentifier(out Directive))
+			{
+				if (bIsBeingIncluded)
+				{
+					throw new UhtException(this.TokenReader, Directive.InputLine, "Missing compiler directive after '#'");
+				}
+				return bCheckClearComments;
+			}
+
+			if (Directive.IsValue("error"))
+			{
+				if (bIsBeingIncluded)
+				{
+					throw new UhtException(this.TokenReader, Directive.InputLine, "#error directive encountered");
+				}
+			}
+			else if (Directive.IsValue("pragma"))
+			{
+				// Ignore all pragmas
+			}
+			else if (Directive.IsValue("linenumber"))
+			{
+				int NewInputLine;
+				if (!LineTokenReader.TryOptionalConstInt(out NewInputLine))
+				{
+					throw new UhtException(this.TokenReader, Directive.InputLine, "Missing line number in line number directive");
+				}
+				this.TokenReader.InputLine = NewInputLine;
+			}
+			else if (Directive.IsValue("include"))
+			{
+				if (bIsBeingIncluded)
+				{
+					UhtToken IncludeName = LineTokenReader.GetToken();
+					if (IncludeName.IsConstString())
+					{
+						StringView IncludeNameString = IncludeName.GetUnescapedString(this.HeaderFile);
+						if (IncludeNameString != "UObject/DefineUPropertyMacros.h" && IncludeNameString != "UObject/UndefineUPropertyMacros.h")
+						{
+							if (this.bSpottedAutogeneratedHeaderInclude)
+							{
+								this.HeaderFile.LogError("#include found after .generated.h file - the .generated.h file should always be the last #include in a header");
+							}
+							if (this.HeaderFile.GeneratedHeaderFileName.AsSpan().Equals(IncludeNameString.Span, StringComparison.OrdinalIgnoreCase))
+							{
+								this.bSpottedAutogeneratedHeaderInclude = true;
+							}
+							if (IncludeNameString.Span.Contains(".generated.h", StringComparison.Ordinal))
+							{
+								string RefName = Path.GetFileName(IncludeNameString.ToString()).Replace(".generated.h", ".h");
+								this.HeaderFile.AddReferencedHeader(RefName, true);
+							}
+							else
+							{
+								this.HeaderFile.AddReferencedHeader(IncludeNameString.ToString(), true);
+							}
+						}
+					}
+				}
+			}
+			else if (Directive.IsValue("if"))
+			{
+				bCheckClearComments = true;
+				PushCompilerDirective(ParserConditional(LineTokenReader));
+			}
+			else if (Directive.IsValue("ifdef") || Directive.IsValue("ifndef"))
+			{
+				bCheckClearComments = true;
+				PushCompilerDirective(UhtCompilerDirective.Unrecognized);
+			}
+			else if (Directive.IsValue("elif"))
+			{
+				bCheckClearComments = true;
+				UhtCompilerDirective OldCompilerDirective = PopCompilerDirective(Directive);
+				UhtCompilerDirective NewCompilerDirective = ParserConditional(LineTokenReader);
+				if (SupportsElif(OldCompilerDirective) != SupportsElif(NewCompilerDirective))
+				{
+					throw new UhtException(this.TokenReader, Directive.InputLine, 
+						$"Mixing {GetCompilerDirectiveText(OldCompilerDirective)} with {GetCompilerDirectiveText(NewCompilerDirective)} in an #elif preprocessor block is not supported");
+				}
+				PushCompilerDirective(NewCompilerDirective);
+			}
+			else if (Directive.IsValue("else"))
+			{
+				bCheckClearComments = true;
+				UhtCompilerDirective OldCompilerDirective = PopCompilerDirective(Directive);
+				switch (OldCompilerDirective)
+				{
+					case UhtCompilerDirective.ZeroBlock:
+						PushCompilerDirective(UhtCompilerDirective.OneBlock);
+						break;
+					case UhtCompilerDirective.OneBlock:
+						PushCompilerDirective(UhtCompilerDirective.ZeroBlock);
+						break;
+					case UhtCompilerDirective.NotCPPBlock:
+						PushCompilerDirective(UhtCompilerDirective.CPPBlock);
+						break;
+					case UhtCompilerDirective.CPPBlock:
+						PushCompilerDirective(UhtCompilerDirective.NotCPPBlock);
+						break;
+					case UhtCompilerDirective.WithHotReload:
+						throw new UhtException(this.TokenReader, Directive.InputLine, "Can not use WITH_HOT_RELOAD with an #else clause");
+					default:
+						PushCompilerDirective(OldCompilerDirective);
+						break;
+				}
+			}
+			else if (Directive.IsValue("endif"))
+			{
+				PopCompilerDirective(Directive);
+			}
+			else if (Directive.IsValue("define"))
+			{
+			}
+			else if (Directive.IsValue("undef"))
+			{
+			}
+			else
+			{
+				if (bIsBeingIncluded)
+				{
+					throw new UhtException(this.TokenReader, Directive.InputLine, $"Unrecognized compiler directive {Directive.Value}");
+				}
+			}
+			return bCheckClearComments;
+		}
+
+		private UhtCompilerDirective ParserConditional(UhtTokenBufferReader LineTokenReader)
+		{
+			// Get any possible ! and the identifier
+			UhtToken Define = LineTokenReader.GetToken();
+			bool bNotPresent = Define.IsSymbol('!');
+			if (bNotPresent)
+			{
+				Define = LineTokenReader.GetToken();
+			}
+
+			//COMPATIBILITY-TODO
+			// UModel.h contains a compound #if where the leading one is !CPP.
+			// Checking for this being the only token causes that to fail
+#if COMPATIBILITY_DISABLE
+			// Make sure there is nothing left
+			UhtToken End = LineTokenReader.GetToken();
+			if (!End.TokenType.IsEndType())
+			{
+				return UhtCompilerDirective.Unrecognized;
+			}
+#endif
+
+			switch (Define.TokenType)
+			{
+				case UhtTokenType.DecimalConst:
+					if (Define.IsValue("0"))
+					{
+						return UhtCompilerDirective.ZeroBlock;
+					}
+					else if (Define.IsValue("1"))
+					{
+						return UhtCompilerDirective.OneBlock;
+					}
+					break;
+
+				case UhtTokenType.Identifier:
+					if (Define.IsValue("WITH_EDITORONLY_DATA"))
+					{
+						return bNotPresent ? UhtCompilerDirective.Unrecognized : UhtCompilerDirective.WithEditorOnlyData;
+					}
+					else if (Define.IsValue("WITH_EDITOR"))
+					{
+						return bNotPresent ? UhtCompilerDirective.Unrecognized : UhtCompilerDirective.WithEditor;
+					}
+					else if (Define.IsValue("WITH_HOT_RELOAD"))
+					{
+						return UhtCompilerDirective.WithHotReload;
+					}
+					else if (Define.IsValue("CPP"))
+					{
+						return bNotPresent ? UhtCompilerDirective.NotCPPBlock : UhtCompilerDirective.CPPBlock;
+					}
+					break;
+
+				case UhtTokenType.EndOfFile:
+				case UhtTokenType.EndOfDefault:
+				case UhtTokenType.EndOfType:
+				case UhtTokenType.EndOfDeclaration:
+					throw new UhtException(this.TokenReader, Define.InputLine, "#if with no expression");
+			}
+
+			return UhtCompilerDirective.Unrecognized;
+		}
+
+		/// <summary>
+		/// Add a new compiler directive to the stack
+		/// </summary>
+		/// <param name="CompilerDirective">Directive to be added</param>
+		private void PushCompilerDirective(UhtCompilerDirective CompilerDirective)
+		{
+			CompilerDirective New = new CompilerDirective();
+			New.Element = CompilerDirective;
+			New.Composite = GetCurrentCompositeCompilerDirective() | CompilerDirective;
+			this.CompilerDirectives.Add(New);
+		}
+
+		/// <summary>
+		/// Remove the top level compiler directive from the stack
+		/// </summary>
+		private UhtCompilerDirective PopCompilerDirective(UhtToken Token)
+		{
+			if (this.CompilerDirectives.Count == 0)
+			{
+				throw new UhtException(this.TokenReader, Token.InputLine, $"Unmatched '#{Token.Value}'");
+			}
+			UhtCompilerDirective CompilerDirective = this.CompilerDirectives[this.CompilerDirectives.Count - 1].Element;
+			this.CompilerDirectives.RemoveAt(this.CompilerDirectives.Count - 1);
+			return CompilerDirective;
+		}
+
+		private bool IncludeCurrentCompilerDirective()
+		{
+			if (this.CompilerDirectives.Count == 0)
+			{
+				return true;
+			}
+			return (GetCurrentCompositeCompilerDirective() & (UhtCompilerDirective.NotCPPBlock | UhtCompilerDirective.OneBlock | UhtCompilerDirective.WithEditor | UhtCompilerDirective.WithEditorOnlyData | UhtCompilerDirective.WithHotReload)) != 0;
+		}
+
+		/// <summary>
+		/// The old UHT would preprocess the file and eliminate any #if blocks that were not required for 
+		/// any contextual information.  This results in comments before the #if block being considered 
+		/// for the next definition.  This routine classifies each #if block type into if comments should
+		/// be purged after the directive.
+		/// </summary>
+		/// <returns></returns>
+		/// <exception cref="UhtIceException"></exception>
+		private bool ClearCommentsCompilerDirective()
+		{
+			if (this.CompilerDirectives.Count == 0)
+			{
+				return true;
+			}
+			UhtCompilerDirective CompilerDirective = this.CompilerDirectives[this.CompilerDirectives.Count - 1].Element;
+			switch (CompilerDirective)
+			{
+				case UhtCompilerDirective.CPPBlock:
+				case UhtCompilerDirective.NotCPPBlock:
+				case UhtCompilerDirective.ZeroBlock:
+				case UhtCompilerDirective.OneBlock:
+				case UhtCompilerDirective.Unrecognized:
+					return false;
+
+				case UhtCompilerDirective.WithEditor:
+				case UhtCompilerDirective.WithEditorOnlyData:
+				case UhtCompilerDirective.WithHotReload:
+					return true;
+
+				default:
+					throw new UhtIceException("Unknown compiler directive flag");
+			}
+		}
+
+		private bool SupportsElif(UhtCompilerDirective CompilerDirective)
+		{
+			return
+				CompilerDirective == UhtCompilerDirective.WithEditor ||
+				CompilerDirective == UhtCompilerDirective.WithEditorOnlyData ||
+				CompilerDirective == UhtCompilerDirective.WithHotReload;
+		}
+
+		private string GetCompilerDirectiveText(UhtCompilerDirective CompilerDirective)
+		{
+			switch (CompilerDirective)
+			{
+				case UhtCompilerDirective.CPPBlock: return "CPP";
+				case UhtCompilerDirective.NotCPPBlock: return "!CPP";
+				case UhtCompilerDirective.ZeroBlock: return "0";
+				case UhtCompilerDirective.OneBlock: return "1";
+				case UhtCompilerDirective.WithHotReload: return "WITH_HOT_RELOAD";
+				case UhtCompilerDirective.WithEditor: return "WITH_EDITOR";
+				case UhtCompilerDirective.WithEditorOnlyData: return "WITH_EDITOR_DATA";
+				default: return "<unrecognized>";
+			}
+		}
+
+		private void SkipVirtualAndAPI(IUhtTokenReader ReplayReader)
+		{
+			while (true)
+			{
+				UhtToken PeekToken = ReplayReader.PeekToken();
+				if (!PeekToken.IsValue("virtual") && !PeekToken.Value.Span.EndsWith("_API"))
+				{
+					break;
+				}
+				ReplayReader.ConsumeToken();
+			}
+		}
+
+		private bool CheckForConstructor(UhtClass Class, UhtDeclaration Declaration)
+		{
+			IUhtTokenReader ReplayReader = UhtTokenReplayReader.GetThreadInstance(this.HeaderFile, this.HeaderFile.Data.Memory, new ReadOnlyMemory<UhtToken>(Declaration.Tokens), UhtTokenType.EndOfDeclaration);
+
+			// Allow explicit constructors
+			bool bFoundExplicit = ReplayReader.TryOptional("explicit");
+
+			bool bSkippedAPIToken = false;
+			if (bSkippedAPIToken = ReplayReader.PeekToken().Value.Span.EndsWith("_API"))
+			{
+				ReplayReader.ConsumeToken();
+				if (!bFoundExplicit)
+				{
+					bFoundExplicit = ReplayReader.TryOptional("explicit");
+				}
+			}
+
+			if (!ReplayReader.TryOptional(Class.SourceName) ||
+				!ReplayReader.TryOptional('('))
+			{
+				return false;
+			}
+
+			bool bOICtor = false;
+			bool bVTCtor = false;
+
+			if (!Class.ClassExportFlags.HasAnyFlags(UhtClassExportFlags.HasDefaultConstructor) && ReplayReader.TryOptional(')'))
+			{
+				Class.ClassExportFlags |= UhtClassExportFlags.HasDefaultConstructor;
+			}
+			else if (!Class.ClassExportFlags.HasAllFlags(UhtClassExportFlags.HasObjectInitializerConstructor | UhtClassExportFlags.HasCustomVTableHelperConstructor))
+			{
+				bool bIsConst = false;
+				bool bIsRef = false;
+				int ParenthesesNestingLevel = 1;
+
+				while (ParenthesesNestingLevel != 0)
+				{
+					UhtToken Token = ReplayReader.GetToken();
+					if (!Token)
+					{
+						break;
+					}
+
+					// Template instantiation or additional parameter excludes ObjectInitializer constructor.
+					if (Token.IsValue(',') || Token.IsValue('<'))
+					{
+						bOICtor = false;
+						bVTCtor = false;
+						break;
+					}
+
+					if (Token.IsValue('('))
+					{
+						ParenthesesNestingLevel++;
+						continue;
+					}
+
+					if (Token.IsValue(')'))
+					{
+						ParenthesesNestingLevel--;
+						continue;
+					}
+
+					if (Token.IsValue("const"))
+					{
+						bIsConst = true;
+						continue;
+					}
+
+					if (Token.IsValue('&'))
+					{
+						bIsRef = true;
+						continue;
+					}
+
+					// FPostConstructInitializeProperties is deprecated, but left here, so it won't break legacy code.
+					if (Token.IsValue("FObjectInitializer") || Token.IsValue("FPostConstructInitializeProperties"))
+					{
+						bOICtor = true;
+					}
+
+					if (Token.IsValue("FVTableHelper"))
+					{
+						bVTCtor = true;
+					}
+				}
+
+				// Parse until finish.
+				if (ParenthesesNestingLevel != 0)
+				{
+					ReplayReader.SkipBrackets('(', ')', ParenthesesNestingLevel);
+				}
+
+				if (bOICtor && bIsRef && bIsConst)
+				{
+					Class.ClassExportFlags |= UhtClassExportFlags.HasObjectInitializerConstructor;
+					Class.MetaData.Add(UhtNames.ObjectInitializerConstructorDeclared, "");
+				}
+				if (bVTCtor && bIsRef)
+				{
+					Class.ClassExportFlags |= UhtClassExportFlags.HasCustomVTableHelperConstructor;
+				}
+			}
+
+			if (!bVTCtor)
+			{
+				Class.ClassExportFlags |= UhtClassExportFlags.HasConstructor;
+			}
+
+			return false;
+		}
+
+		bool CheckForSerialize(UhtClass Class, UhtDeclaration Declaration)
+		{
+			IUhtTokenReader ReplayReader = UhtTokenReplayReader.GetThreadInstance(this.HeaderFile, this.HeaderFile.Data.Memory, new ReadOnlyMemory<UhtToken>(Declaration.Tokens), UhtTokenType.EndOfDeclaration);
+
+			SkipVirtualAndAPI(ReplayReader);
+
+			if (!ReplayReader.TryOptional("void") ||
+				!ReplayReader.TryOptional("Serialize") ||
+				!ReplayReader.TryOptional('('))
+			{
+				return false;
+			}
+
+			UhtToken Token = ReplayReader.GetToken();
+
+			UhtSerializerArchiveType ArchiveType = UhtSerializerArchiveType.None;
+			if (Token.IsValue("FArchive"))
+			{
+				if (ReplayReader.TryOptional('&'))
+				{
+					// Allow the declaration to not define a name for the archive parameter
+					if (!ReplayReader.PeekToken().IsValue(')'))
+					{
+						ReplayReader.SkipOne();
+					}
+					if (ReplayReader.TryOptional(')'))
+					{
+						ArchiveType = UhtSerializerArchiveType.Archive;
+					}
+				}
+			}
+			else if (Token.IsValue("FStructuredArchive"))
+			{
+				if (ReplayReader.TryOptional("::") &&
+					ReplayReader.TryOptional("FRecord"))
+				{
+					// Allow the declaration to not define a name for the archive parameter
+					if (!ReplayReader.PeekToken().IsValue(')'))
+					{
+						ReplayReader.SkipOne();
+					}
+					if (ReplayReader.TryOptional(')'))
+					{
+						ArchiveType = UhtSerializerArchiveType.StructuredArchiveRecord;
+					}
+				}
+			}
+			else if (Token.IsValue("FStructuredArchiveRecord"))
+			{
+				// Allow the declaration to not define a name for the archive parameter
+				if (!ReplayReader.PeekToken().IsValue(')'))
+				{
+					ReplayReader.SkipOne();
+				}
+				if (ReplayReader.TryOptional(')'))
+				{
+					ArchiveType = UhtSerializerArchiveType.StructuredArchiveRecord;
+				}
+			}
+
+			if (ArchiveType != UhtSerializerArchiveType.None)
+			{
+				// Found what we want!
+				if (Declaration.CompilerDirectives == UhtCompilerDirective.None || Declaration.CompilerDirectives == UhtCompilerDirective.WithEditorOnlyData)
+				{
+					Class.SerializerArchiveType |= ArchiveType;
+					Class.EnclosingDefine = Declaration.CompilerDirectives == UhtCompilerDirective.None ? "" : "WITH_EDITORONLY_DATA";
+				}
+				else
+				{
+					Class.LogError("Serialize functions must not be inside preprocessor blocks, except for WITH_EDITORONLY_DATA");
+				}
+				return true;
+			}
+
+			return false;
+		}
+		#endregion
+	}
+}
diff --git a/Engine/Source/Programs/Shared/EpicGames.UHT/Parsers/UhtInterfaceClassParser.cs b/Engine/Source/Programs/Shared/EpicGames.UHT/Parsers/UhtInterfaceClassParser.cs
new file mode 100644
--- /dev/null
+++ b/Engine/Source/Programs/Shared/EpicGames.UHT/Parsers/UhtInterfaceClassParser.cs
@@ -0,0 +1,175 @@
+// Copyright Epic Games, Inc. All Rights Reserved.
+
+using EpicGames.Core;
+using EpicGames.UHT.Tables;
+using EpicGames.UHT.Tokenizer;
+using EpicGames.UHT.Types;
+using EpicGames.UHT.Utils;
+
+namespace EpicGames.UHT.Parsers
+{
+	[UnrealHeaderTool]
+	public class UhtInterfaceClassParser : UhtClassBaseParser
+	{
+		private static UhtKeywordTable KeywordTable = UhtKeywordTables.Instance.Get(UhtTableNames.Interface);
+		private static UhtSpecifierTable SpecifierTable = UhtSpecifierTables.Instance.Get(UhtTableNames.Interface);
+
+		public UhtInterfaceClassParser(UhtType Outer, int LineNumber) : base(Outer, LineNumber)
+		{
+		}
+
+		protected override void ResolveSuper(UhtResolvePhase ResolvePhase)
+		{
+			base.ResolveSuper( ResolvePhase);
+			switch (ResolvePhase)
+			{
+				case UhtResolvePhase.Bases:
+					UhtClass? SuperClass = this.SuperClass;
+					if (SuperClass != null)
+					{
+						this.ClassFlags |= SuperClass.ClassFlags & EClassFlags.ScriptInherit;
+						if (!SuperClass.ClassFlags.HasAnyFlags(EClassFlags.Native))
+						{
+							throw new UhtException(this, $"Native classes cannot extend non-native classes");
+						}
+						this.ClassWithin = SuperClass.ClassWithin;
+					}
+					else
+					{
+						this.ClassWithin = this.Session.UObject;
+					}
+					break;
+			}
+		}
+
+		protected override bool ResolveSelf(UhtResolvePhase ResolvePhase)
+		{
+			bool bResult = base.ResolveSelf(ResolvePhase);
+
+			switch (ResolvePhase)
+			{
+				case UhtResolvePhase.InvalidCheck:
+					{
+						string NativeInterfaceName = "I" + this.EngineName;
+						UhtType? NativeInterface = this.Session.FindType(null, UhtFindOptions.SourceName | UhtFindOptions.Class, NativeInterfaceName);
+						if (NativeInterface == null)
+						{
+							this.LogError($"UInterface '{this.SourceName}' parsed without a corresponding '{NativeInterfaceName}'");
+						}
+						else
+						{
+							// Copy the children
+							this.AddChildren(NativeInterface.DetchChildren());
+						}
+					}
+					break;
+			}
+
+			return bResult;
+		}
+		#region Keywords
+		[UhtKeyword(Extends = UhtTableNames.Global)]
+		private static UhtParseResult UINTERFACEKeyword(UhtParsingScope TopScope, UhtParsingScope ActionScope, ref UhtToken Token)
+		{
+			return ParseUInterface(TopScope, ref Token);
+		}
+
+		[UhtKeyword(Extends = UhtTableNames.Interface)]
+		[UhtKeyword(Extends = UhtTableNames.Interface, Keyword = "GENERATED_BODY")]
+		private static UhtParseResult GENERATED_UINTERFACE_BODYKeyword(UhtParsingScope TopScope, UhtParsingScope ActionScope, ref UhtToken Token)
+		{
+			UhtClass Class = (UhtClass)TopScope.ScopeType;
+
+			UhtParserHelpers.ParseCompileVersionDeclaration(TopScope.TokenReader, Class);
+
+			Class.GeneratedBodyAccessSpecifier = TopScope.AccessSpecifier;
+			Class.GeneratedBodyLineNumber = TopScope.TokenReader.InputLine;
+
+			if (Token.IsValue("GENERATED_UINTERFACE_BODY"))
+			{
+				TopScope.AccessSpecifier = UhtAccessSpecifier.Public;
+			}
+			return UhtParseResult.Handled;
+		}
+		#endregion
+
+		public static UhtParseResult ParseUInterface(UhtParsingScope ParentScope, ref UhtToken Token)
+		{
+			UhtInterfaceClassParser Class = new UhtInterfaceClassParser(ParentScope.ScopeType, Token.InputLine);
+			Class.ClassType = UhtClassType.Interface;
+
+			using (var TopScope = new UhtParsingScope(ParentScope, Class, KeywordTable, UhtAccessSpecifier.Private))
+			{
+				const string ScopeName = "interface";
+
+				using (var TokenContext = new UhtMessageContext(TopScope.TokenReader, ScopeName))
+				{
+					// Parse the specifiers
+					UhtSpecifierContext SpecifierContext = new UhtSpecifierContext(TopScope, TopScope.TokenReader, Class.MetaData);
+					UhtSpecifierParser Specifiers = TopScope.HeaderParser.GetSpecifierParser(SpecifierContext, ScopeName, SpecifierTable);
+					Specifiers.ParseSpecifiers();
+					Class.PrologLineNumber = TopScope.TokenReader.InputLine;
+					Class.ClassFlags |= EClassFlags.Native | EClassFlags.Interface | EClassFlags.Abstract;
+					Class.InternalObjectFlags |= EInternalObjectFlags.Native;
+
+					// Consume the class specifier
+					TopScope.TokenReader.Require("class");
+
+					// Get the optional API macro
+					UhtToken APIMacroToken;
+					TopScope.TokenReader.TryOptionalAPIMacro(out APIMacroToken);
+
+					// Get the name of the interface
+					Class.SourceName = TopScope.TokenReader.GetIdentifier().Value.ToString();
+
+					// Parse the inheritance
+					TopScope.TokenReader.OptionalInheritance((ref UhtToken Identifier) =>
+					{
+						Class.SuperIdentifier = Identifier;
+					});
+
+					// Split the source name into the different parts
+					UhtEngineNameParts NameParts = UhtUtilities.GetEngineNameParts(Class.SourceName);
+					Class.EngineName = NameParts.EngineName.ToString();
+
+					// Interfaces must start with a valid prefix
+					if (NameParts.Prefix != "U")
+					{
+						TopScope.TokenReader.LogError($"Interface name '{Class.SourceName}' is invalid, the first class should be identified as 'U{NameParts.EngineName}'");
+					}
+
+					// Check for an empty engine name
+					if (Class.EngineName.Length == 0)
+					{
+						TopScope.TokenReader.LogError($"When compiling class definition for '{Class.SourceName}', attempting to strip prefix results in an empty name. Did you leave off a prefix?");
+					}
+
+					if (APIMacroToken)
+					{
+						Class.ClassFlags |= EClassFlags.RequiredAPI;
+					}
+
+					Specifiers.ParseDeferred();
+
+					if (Class.Outer != null)
+					{
+						Class.Outer.AddChild(Class);
+					}
+
+					TopScope.AddModuleRelativePathToMetaData();
+
+					TopScope.HeaderParser.ParseStatements('{', '}', true);
+
+					TopScope.TokenReader.Require(';');
+
+					if (Class.GeneratedBodyLineNumber == -1)
+					{
+						TopScope.TokenReader.LogError("Expected a GENERATED_BODY() at the start of the interface");
+					}
+				}
+			}
+
+			return UhtParseResult.Handled;
+		}
+	}
+}
diff --git a/Engine/Source/Programs/Shared/EpicGames.UHT/Parsers/UhtNativeInterfaceClassParser.cs b/Engine/Source/Programs/Shared/EpicGames.UHT/Parsers/UhtNativeInterfaceClassParser.cs
new file mode 100644
--- /dev/null
+++ b/Engine/Source/Programs/Shared/EpicGames.UHT/Parsers/UhtNativeInterfaceClassParser.cs
@@ -0,0 +1,179 @@
+// Copyright Epic Games, Inc. All Rights Reserved.
+
+using EpicGames.Core;
+using EpicGames.UHT.Tables;
+using EpicGames.UHT.Tokenizer;
+using EpicGames.UHT.Types;
+using EpicGames.UHT.Utils;
+
+namespace EpicGames.UHT.Parsers
+{
+	[UnrealHeaderTool]
+	public class UhtParserNativeInterfaceClass : UhtClassBaseParser
+	{
+		private static UhtKeywordTable KeywordTable = UhtKeywordTables.Instance.Get(UhtTableNames.NativeInterface);
+		private static UhtSpecifierTable SpecifierTable = UhtSpecifierTables.Instance.Get(UhtTableNames.NativeInterface);
+
+		public UhtParserNativeInterfaceClass(UhtType Outer, int LineNumber) : base(Outer, LineNumber)
+		{
+		}
+
+		protected override bool ResolveSelf(UhtResolvePhase ResolvePhase)
+		{
+			bool bResult = base.ResolveSelf(ResolvePhase);
+
+			switch (ResolvePhase)
+			{
+				case UhtResolvePhase.InvalidCheck:
+					string InterfaceName = "U" + this.EngineName;
+					UhtClass? Interface = (UhtClass?) this.Session.FindType(null, UhtFindOptions.SourceName | UhtFindOptions.Class, InterfaceName);
+					if (Interface == null)
+					{
+						if (this.bHasGeneratedBody || this.Children.Count != 0)
+						{
+							this.LogError($"Native interface '{this.SourceName}' parsed without a corresponding '{InterfaceName}'");
+						}
+						else
+						{
+							this.bVisibleType = false;
+							bResult = false;
+						}
+					}
+					else
+					{
+						this.AlternateObject = Interface;
+						Interface.NativeInterface = this;
+						//COMPATIBILITY-TODO - Use the native interface access specifier
+						Interface.GeneratedBodyAccessSpecifier = this.GeneratedBodyAccessSpecifier;
+
+						if (this.GeneratedBodyLineNumber == -1)
+						{
+							this.LogError("Expected a GENERATED_BODY() at the start of the native interface");
+						}
+					}
+					break;
+			}
+
+			return bResult;
+		}
+
+		#region Keywords
+		[UhtKeyword(Extends = UhtTableNames.Global, Keyword = "class", DisableUsageError = true)]
+		private static UhtParseResult ClassKeyword(UhtParsingScope TopScope, UhtParsingScope ActionScope, ref UhtToken Token)
+		{
+			using (var SaveState = new UhtTokenSaveState(TopScope.TokenReader))
+			{
+				UhtToken SourceName = new UhtToken();
+				UhtToken SuperName = new UhtToken();
+				if (TryParseIInterface(TopScope, ref Token, out SourceName, out SuperName))
+				{
+					SaveState.AbandonState();
+					ParseIInterface(TopScope, ref Token, SourceName, SuperName);
+					return UhtParseResult.Handled;
+				}
+			}
+			return TopScope.TokenReader.SkipDeclaration(ref Token) ? UhtParseResult.Handled : UhtParseResult.Invalid;
+		}
+
+		[UhtKeyword(Extends = UhtTableNames.NativeInterface)]
+		[UhtKeyword(Extends = UhtTableNames.NativeInterface, Keyword = "GENERATED_BODY")]
+		private static UhtParseResult GENERATED_IINTERFACE_BODYKeyword(UhtParsingScope TopScope, UhtParsingScope ActionScope, ref UhtToken Token)
+		{
+			UhtClass Class = (UhtClass)TopScope.ScopeType;
+
+			UhtParserHelpers.ParseCompileVersionDeclaration(TopScope.TokenReader, Class);
+
+			Class.GeneratedBodyAccessSpecifier = TopScope.AccessSpecifier;
+			Class.GeneratedBodyLineNumber = TopScope.TokenReader.InputLine;
+			Class.bHasGeneratedBody = true;
+
+			if (Token.IsValue("GENERATED_IINTERFACE_BODY"))
+			{
+				TopScope.AccessSpecifier = UhtAccessSpecifier.Public;
+			}
+			return UhtParseResult.Handled;
+		}
+		#endregion
+
+		private static bool TryParseIInterface(UhtParsingScope ParentScope, ref UhtToken Token, out UhtToken SourceName, out UhtToken SuperName)
+		{
+			IUhtTokenReader TokenReader = ParentScope.TokenReader;
+
+			// Get the optional API macro
+			UhtToken APIMacroToken;
+			TokenReader.TryOptionalAPIMacro(out APIMacroToken);
+
+			// Get the name of the interface
+			SourceName = TokenReader.GetIdentifier();
+
+			// Old UHT would still parse the inheritance 
+			SuperName = new UhtToken();
+			if (TokenReader.TryOptional(':'))
+			{
+				if (!TokenReader.TryOptional("public"))
+				{
+					return false;
+				}
+				SuperName = TokenReader.GetIdentifier();
+			}
+
+			// Only process classes starting with 'I'
+			if (SourceName.Value.Span[0] != 'I')
+			{
+				return false;
+			}
+
+			// If we end with a ';', then this is a forward declaration
+			if (TokenReader.TryOptional(';'))
+			{
+				return false;
+			}
+
+			// If we don't have a '{', then this is something else
+			if (!TokenReader.TryPeekOptional('{'))
+			{
+				return false;
+			}
+			return true;
+		}
+
+		private static void ParseIInterface(UhtParsingScope ParentScope, ref UhtToken Token, UhtToken SourceName, UhtToken SuperName)
+		{
+			IUhtTokenReader TokenReader = ParentScope.TokenReader;
+
+			UhtParserNativeInterfaceClass Class = new UhtParserNativeInterfaceClass(ParentScope.ScopeType, Token.InputLine);
+			Class.ClassType = UhtClassType.NativeInterface;
+			Class.SourceName = SourceName.Value.ToString();
+			Class.ClassFlags |= EClassFlags.Native | EClassFlags.Interface;
+
+			// Split the source name into the different parts
+			UhtEngineNameParts NameParts = UhtUtilities.GetEngineNameParts(Class.SourceName);
+			Class.EngineName = NameParts.EngineName.ToString();
+
+			// Check for an empty engine name
+			if (Class.EngineName.Length == 0)
+			{
+				TokenReader.LogError($"When compiling class definition for '{Class.SourceName}', attempting to strip prefix results in an empty name. Did you leave off a prefix?");
+			}
+
+			if (Class.Outer != null)
+			{
+				Class.Outer.AddChild(Class);
+			}
+
+			Class.SuperIdentifier = SuperName;
+
+			using (var TopScope = new UhtParsingScope(ParentScope, Class, KeywordTable, UhtAccessSpecifier.Private))
+			{
+				const string ScopeName = "native interface";
+
+				using (var TokenContext = new UhtMessageContext(TopScope.TokenReader, ScopeName))
+				{
+					TopScope.HeaderParser.ParseStatements('{', '}', false);
+					TokenReader.Require(';');
+				}
+			}
+			return;
+		}
+	}
+}
diff --git a/Engine/Source/Programs/Shared/EpicGames.UHT/Parsers/UhtParserHelpers.cs b/Engine/Source/Programs/Shared/EpicGames.UHT/Parsers/UhtParserHelpers.cs
new file mode 100644
--- /dev/null
+++ b/Engine/Source/Programs/Shared/EpicGames.UHT/Parsers/UhtParserHelpers.cs
@@ -0,0 +1,62 @@
+// Copyright Epic Games, Inc. All Rights Reserved.
+
+using EpicGames.Core;
+using EpicGames.UHT.Tokenizer;
+using EpicGames.UHT.Types;
+using EpicGames.UHT.Utils;
+using System;
+using System.Collections.Generic;
+
+namespace EpicGames.UHT.Parsers
+{
+	public static class UhtParserHelpers
+	{
+		public static void ParseInheritance(IUhtTokenReader TokenReader, out UhtToken SuperIdentifier, out List<UhtToken[]>? BaseIdentifiers)
+		{
+			UhtToken SuperIdentifierTemp = new UhtToken();
+			List<UhtToken[]>? BaseIdentifiersTemp = null;
+			TokenReader.OptionalInheritance(
+				(ref UhtToken Identifier) =>
+				{
+					UhtConfig.Instance.RedirectTypeIdentifier(ref Identifier);
+					SuperIdentifierTemp = Identifier;
+				},
+				(UhtTokenList Identifier) =>
+				{
+					if (BaseIdentifiersTemp == null)
+					{
+						BaseIdentifiersTemp = new List<UhtToken[]>();
+					}
+					BaseIdentifiersTemp.Add(Identifier.ToArray());
+				});
+			SuperIdentifier = SuperIdentifierTemp;
+			BaseIdentifiers = BaseIdentifiersTemp;
+		}
+
+		public static void ParseCompileVersionDeclaration(IUhtTokenReader TokenReader, UhtStruct Struct)
+		{
+
+			// Fetch the default generation code version. If supplied, then package code version overrides the default.
+			EGeneratedCodeVersion Version = Struct.Package.Module.GeneratedCodeVersion;
+			if (Version == EGeneratedCodeVersion.None)
+			{
+				Version = UhtConfig.Instance.DefaultGeneratedCodeVersion;
+			}
+
+			// Fetch the code version from header file
+			TokenReader
+				.Require('(')
+				.OptionalIdentifier((ref UhtToken Identifier) =>
+				{
+					if (!Enum.TryParse(Identifier.Value.ToString(), true, out Version))
+					{
+						Version = EGeneratedCodeVersion.None;
+					}
+				})
+				.Require(')');
+
+			// Save the results
+			Struct.GeneratedCodeVersion = Version;
+		}
+	}
+}
