UHT in UBT initial version.  Currently disabled.

diff --git a/Engine/Source/Programs/Shared/EpicGames.UHT/Parsers/UhtParsingScope.cs b/Engine/Source/Programs/Shared/EpicGames.UHT/Parsers/UhtParsingScope.cs
new file mode 100644
--- /dev/null
+++ b/Engine/Source/Programs/Shared/EpicGames.UHT/Parsers/UhtParsingScope.cs
@@ -0,0 +1,485 @@
+// Copyright Epic Games, Inc. All Rights Reserved.
+
+using EpicGames.Core;
+using EpicGames.UHT.Tables;
+using EpicGames.UHT.Tokenizer;
+using EpicGames.UHT.Types;
+using EpicGames.UHT.Utils;
+using System;
+using System.Text;
+
+namespace EpicGames.UHT.Parsers
+{
+	public class UhtParsingScope : IDisposable
+	{
+		public readonly UhtHeaderFileParser HeaderParser;
+		public readonly IUhtTokenReader TokenReader;
+		public readonly UhtParsingScope? ParentScope;
+		public readonly UhtType ScopeType;
+		public readonly UhtKeywordTable ScopeKeywordTable;
+		public UhtAccessSpecifier AccessSpecifier = UhtAccessSpecifier.Public;
+
+		/// <summary>
+		/// Return the current class scope being compiled
+		/// </summary>
+		public UhtParsingScope CurrentClassScope
+		{
+			get
+			{
+				UhtParsingScope? CurrentScope = this;
+				while (CurrentScope != null)
+				{
+					if (CurrentScope.ScopeType is UhtClass Class)
+					{
+						return CurrentScope;
+					}
+					CurrentScope = CurrentScope.ParentScope;
+				}
+				throw new UhtIceException("Attempt to fetch the current class when a class isn't currently being parsed");
+			}
+		}
+
+		/// <summary>
+		/// Return the current class being compiled
+		/// </summary>
+		public UhtClass CurrentClass => (UhtClass)this.CurrentClassScope.ScopeType;
+
+		public UhtParsingScope(UhtHeaderFileParser HeaderParser, UhtType ScopeType, UhtKeywordTable KeywordTable)
+		{
+			this.HeaderParser = HeaderParser;
+			this.TokenReader = HeaderParser.TokenReader;
+			this.ParentScope = null;
+			this.ScopeType = ScopeType;
+			this.ScopeKeywordTable = KeywordTable;
+			this.HeaderParser.PushScope(this);
+		}
+
+		public UhtParsingScope(UhtParsingScope ParentScope, UhtType ScopeType, UhtKeywordTable KeywordTable, UhtAccessSpecifier AccessSpecifier)
+		{
+			this.HeaderParser = ParentScope.HeaderParser;
+			this.TokenReader = ParentScope.TokenReader;
+			this.ParentScope = ParentScope;
+			this.ScopeType = ScopeType;
+			this.ScopeKeywordTable = KeywordTable;
+			this.AccessSpecifier = AccessSpecifier;
+			this.HeaderParser.PushScope(this);
+		}
+
+		public void Dispose()
+		{
+			this.HeaderParser.PopScope(this);
+		}
+
+		/// <summary>
+		/// Add the module's relative path to the type's meta data
+		/// </summary>
+		public void AddModuleRelativePathToMetaData()
+		{
+			AddModuleRelativePathToMetaData(this.ScopeType.MetaData, this.ScopeType.HeaderFile);
+		}
+
+		/// <summary>
+		/// Add the module's relative path to the meta data
+		/// </summary>
+		/// <param name="MetaData">The meta data to add the information to</param>
+		/// <param name="HeaderFile">The header file currently being parsed</param>
+		public static void AddModuleRelativePathToMetaData(UhtMetaData MetaData, UhtHeaderFile HeaderFile)
+		{
+			MetaData.Add(UhtNames.ModuleRelativePath, HeaderFile.ModuleRelativeFilePath);
+		}
+
+		/// <summary>
+		/// Format the current token reader comments and add it as meta data
+		/// </summary>
+		/// <param name="MetaNameIndex">Index for the meta data key.  This is used for enum values</param>
+		public void AddFormattedCommentsAsTooltipMetaData(int MetaNameIndex = UhtMetaData.INDEX_NONE)
+		{
+			AddFormattedCommentsAsTooltipMetaData(this.ScopeType, MetaNameIndex);
+		}
+
+		/// <summary>
+		/// Format the current token reader comments and add it as meta data
+		/// </summary>
+		/// <param name="Type">The type to add the meta data to</param>
+		/// <param name="MetaNameIndex">Index for the meta data key.  This is used for enum values</param>
+		public void AddFormattedCommentsAsTooltipMetaData(UhtType Type, int MetaNameIndex = UhtMetaData.INDEX_NONE)
+		{
+
+			// Don't add a tooltip if one already exists.
+			if (Type.MetaData.ContainsKey(UhtNames.ToolTip, MetaNameIndex))
+			{
+				return;
+			}
+
+			// Fetch the comments
+			ReadOnlySpan<StringView> Comments = this.TokenReader.Comments;
+
+			// If we don't have any comments, just return
+			string MergedString = String.Empty;
+			if (Comments.Length == 0)
+			{
+				return;
+			}
+
+			// Set the comment as just a simple concatenation of all the strings
+			if (Comments.Length == 1)
+			{
+				MergedString = Comments[0].ToString();
+				Type.MetaData.Add(UhtNames.Comment, MetaNameIndex, Comments[0].ToString());
+			}
+			else
+			{
+				using (BorrowStringBuilder Borrower = new BorrowStringBuilder(StringBuilderCache.Small))
+				{
+					StringBuilder Builder = Borrower.StringBuilder;
+					foreach (StringView Comment in Comments)
+					{
+						Builder.Append(Comment);
+					}
+					MergedString = Builder.ToString();
+					Type.MetaData.Add(UhtNames.Comment, MetaNameIndex, MergedString);
+				}
+			}
+
+			// Format the tooltip and set the metadata
+			StringView ToolTip = FormatCommentForToolTip(MergedString);
+			if (ToolTip.Span.Length > 0)
+			{
+				Type.MetaData.Add(UhtNames.ToolTip, MetaNameIndex, ToolTip.ToString());
+
+				//COMPATIBILITY-TODO - Old UHT would only clear the comments if there was some form of a tooltip
+				this.TokenReader.ClearComments();
+			}
+
+			//COMPATIBILITY-TODO
+			// Clear the comments since they have been consumed
+			//this.TokenReader.ClearComments();
+		}
+
+		// We consider any alpha/digit or code point > 0xFF as a valid comment char
+		/// <summary>
+		/// Given a list of comments, check to see if any have alpha, numeric, or unicode code points with a value larger than 0xFF.
+		/// </summary>
+		/// <param name="Comments">Comments to search</param>
+		/// <returns>True is a character in question was found</returns>
+		private static bool HasValidCommentChar(ReadOnlySpan<StringView> Comments)
+		{
+			foreach (StringView Comment in Comments)
+			{
+				if (HasValidCommentChar(Comment.Span))
+				{
+					return true;
+				}
+			}
+			return false;
+		}
+
+		// We consider any alpha/digit or code point > 0xFF as a valid comment char
+		/// <summary>
+		/// Given a list of comments, check to see if any have alpha, numeric, or unicode code points with a value larger than 0xFF.
+		/// </summary>
+		/// <param name="Comments">Comments to search</param>
+		/// <returns>True is a character in question was found</returns>
+		private static bool HasValidCommentChar(ReadOnlySpan<char> Comments)
+		{
+			foreach (char C in Comments)
+			{
+				if (UhtFCString.IsAlnum(C) || C > 0xFF)
+				{
+					return true;
+				}
+			}
+			return false;
+		}
+
+		/// <summary>
+		/// Convert the given list of comments to a tooltip.  Each string view is a comment where the // style comments also includes the trailing \r\n.
+		/// 
+		/// The following style comments are supported:
+		/// 
+		/// /* */ - C Style
+		/// /** */ - C Style JavaDocs
+		/// /*~ */ - C Style but ignore
+		/// //\r\n - C++ Style
+		/// ///\r\n - C++ Style JavaDocs
+		/// //~\r\n - C++ Style bug ignore
+		///
+		/// As per TokenReader, there will only be one C style comment ever present, and it will be the first one.  When a C style comment is parsed, any prior comments
+		/// are cleared.  However, if a C++ style comment follows a C style comment (regardless of any intermediate blank lines), then both blocks of comments will be present.
+		/// If any blank lines are encountered between blocks of C++ style comments, then any prior comments are cleared.
+		/// </summary>
+		/// <param name="Comments">Comments to be parsed</param>
+		/// <returns>The generated tooltip</returns>
+		private static StringView FormatCommentForToolTip(string Comments)
+		{
+			if (!HasValidCommentChar(Comments))
+			{
+				return new StringView();
+			}
+
+			string Result = Comments;
+
+			// Sweep out comments marked to be ignored.
+			{
+				int CommentStart, CommentEnd;
+
+				// Block comments go first
+				while ((CommentStart = Result.IndexOf("/*~")) != -1)
+				{
+					CommentEnd = Result.IndexOf("*/", CommentStart);
+					if (CommentEnd != -1)
+					{
+						Result = Result.Remove(CommentStart, CommentEnd - CommentStart + 2);
+					}
+					else
+					{
+						// This looks like an error - an unclosed block comment.
+						break;
+					}
+				}
+
+				// Leftover line comments go next
+				while ((CommentStart = Result.IndexOf("//~")) != -1)
+				{
+					CommentEnd = Result.IndexOf("\n", CommentStart);
+					if (CommentEnd != -1)
+					{
+						Result = Result.Remove(CommentStart, CommentEnd - CommentStart + 1);
+					}
+					else
+					{
+						Result = Result.Remove(CommentStart);
+						break;
+					}
+				}
+			}
+
+			// Check for known commenting styles.
+			bool bJavaDocStyle = Result.Contains("/**");
+			bool bCStyle = Result.Contains("/*");
+			bool bCPPStyle = Result.StartsWith("//");
+
+			if (bJavaDocStyle || bCStyle)
+			{
+				// Remove beginning and end markers.
+				if (bJavaDocStyle)
+				{
+					Result = Result.Replace("/**", "");
+				}
+				if (bCStyle)
+				{
+					Result = Result.Replace("/*", "");
+				}
+				Result = Result.Replace("*/", "");
+			}
+
+			if (bCPPStyle)
+			{
+				// Remove c++-style comment markers.  Also handle javadoc-style comments 
+				Result = Result.Replace("///", "");
+				Result = Result.Replace("//", "");
+
+				// Parser strips cpptext and replaces it with "// (cpptext)" -- prevent
+				// this from being treated as a comment on variables declared below the
+				// cpptext section
+				Result = Result.Replace("(cpptext)", "");
+			}
+
+			// Get rid of carriage return or tab characters, which mess up tooltips.
+			Result = Result.Replace("\r", "");
+
+			//wx widgets has a hard coded tab size of 8
+			{
+				int SpacesPerTab = 8;
+				Result = UhtFCString.TabsToSpaces(Result, SpacesPerTab, true);
+			}
+
+			// get rid of uniform leading whitespace and all trailing whitespace, on each line
+			string[] Lines = Result.Split('\n');
+
+			for (int Index = 0; Index < Lines.Length; ++Index)
+			{
+				// Remove trailing whitespace
+				string Line = Lines[Index].TrimEnd();
+
+				// Remove leading "*" and "* " in javadoc comments.
+				if (bJavaDocStyle)
+				{
+					// Find first non-whitespace character
+					int Pos = 0;
+					while (Pos < Line.Length && UhtFCString.IsWhitespace(Line[Pos]))
+					{
+						++Pos;
+					}
+
+					// Is it a *?
+					if (Pos < Line.Length && Line[Pos] == '*')
+					{
+						// Eat next space as well
+						if (Pos + 1 < Line.Length && UhtFCString.IsWhitespace(Line[Pos + 1]))
+						{
+							++Pos;
+						}
+
+						Line = Line.Substring(Pos + 1);
+					}
+				}
+				Lines[Index] = Line;
+			}
+
+			Func<string, int, char, bool> IsAllSameChar = (string Line, int StartIndex, char TestChar) =>
+			{
+				for (int Index = StartIndex, End = Line.Length; Index < End; ++Index)
+				{
+					if (Line[Index] != TestChar)
+					{
+						return false;
+					}
+				}
+				return true;
+			};
+
+			Func<string, bool> IsWhitespaceOrLineSeparator = (string Line) =>
+			{
+				// Skip any leading spaces
+				int Index = 0;
+				int EndPos = Line.Length;
+				for (; Index < EndPos && UhtFCString.IsWhitespace(Line[Index]); ++Index) { }
+				if (Index == EndPos)
+				{
+					return true;
+				}
+
+				// Check for the same character
+				return IsAllSameChar(Line, Index, '-') || IsAllSameChar(Line, Index, '=') || IsAllSameChar(Line, Index, '*');
+			};
+
+			// Find first meaningful line
+			int FirstIndex = 0;
+			for (; FirstIndex < Lines.Length && IsWhitespaceOrLineSeparator(Lines[FirstIndex]); ++FirstIndex)
+			{ }
+
+			int LastIndex = Lines.Length;
+			for (; LastIndex > FirstIndex && IsWhitespaceOrLineSeparator(Lines[LastIndex - 1]); --LastIndex)
+			{ }
+
+			Result = String.Empty;
+			if (FirstIndex != LastIndex)
+			{
+				string FirstLine = Lines[FirstIndex];
+
+				// Figure out how much whitespace is on the first line
+				int MaxNumWhitespaceToRemove = 0;
+				for (; MaxNumWhitespaceToRemove < FirstLine.Length; MaxNumWhitespaceToRemove++)
+				{
+					if (/*!UhtFCString.IsLinebreak(FirstLine[MaxNumWhitespaceToRemove]) && */!UhtFCString.IsWhitespace(FirstLine[MaxNumWhitespaceToRemove]))
+					{
+						break;
+					}
+				}
+
+				for (int Index = FirstIndex; Index != LastIndex; ++Index)
+				{
+					string Line = Lines[Index];
+
+					int TemporaryMaxWhitespace = MaxNumWhitespaceToRemove;
+
+					// Allow eating an extra tab on subsequent lines if it's present
+					if ((Index > FirstIndex) && (Line.Length > 0) && (Line[0] == '\t'))
+					{
+						TemporaryMaxWhitespace++;
+					}
+
+					// Advance past whitespace
+					int Pos = 0;
+					while (Pos < TemporaryMaxWhitespace && Pos < Line.Length && UhtFCString.IsWhitespace(Line[Pos]))
+					{
+						++Pos;
+					}
+
+					if (Pos > 0)
+					{
+						Line = Line.Substring(Pos);
+					}
+
+					if (Index > FirstIndex)
+					{
+						Result += "\n";
+					}
+
+					if (Line.Length > 0 && !IsAllSameChar(Line, 0, '='))
+					{
+						Result += Line;
+					}
+				}
+			}
+
+			//@TODO: UCREMOVAL: Really want to trim an arbitrary number of newlines above and below, but keep multiple newlines internally
+			// Make sure it doesn't start with a newline
+			if (Result.Length != 0 && UhtFCString.IsLinebreak(Result[0]))
+			{
+				Result = Result.Substring(1);
+			}
+
+			// Make sure it doesn't end with a dead newline
+			if (Result.Length != 0 && UhtFCString.IsLinebreak(Result[Result.Length - 1]))
+			{
+				Result = Result.Substring(0, Result.Length - 1);
+			}
+
+			// Done.
+			return Result;
+		}
+	}
+
+	public struct UhtTokenRecorder : IDisposable
+	{
+		private readonly UhtCompilerDirective CompilerDirective;
+		private readonly UhtParsingScope Scope;
+		private bool bFlushed;
+
+		public UhtTokenRecorder(UhtParsingScope Scope, ref UhtToken InitialToken)
+		{
+			this.Scope = Scope;
+			this.CompilerDirective = this.Scope.HeaderParser.GetCurrentCompositeCompilerDirective();
+			this.bFlushed = false;
+			if (this.Scope.ScopeType is UhtClass)
+			{
+
+				this.Scope.TokenReader.EnableRecording();
+				this.Scope.TokenReader.RecordToken(ref InitialToken);
+			}
+		}
+
+		public UhtTokenRecorder(UhtParsingScope Scope)
+		{
+			this.Scope = Scope;
+			this.CompilerDirective = this.Scope.HeaderParser.GetCurrentCompositeCompilerDirective();
+			this.bFlushed = false;
+			if (this.Scope.ScopeType is UhtClass)
+			{
+				this.Scope.TokenReader.EnableRecording();
+			}
+		}
+
+		public void Dispose()
+		{
+			Stop();
+		}
+
+		public bool Stop()
+		{
+			if (!bFlushed)
+			{
+				bFlushed = true;
+				if (this.Scope.ScopeType is UhtClass Class)
+				{
+					Class.AddDeclaration(this.CompilerDirective, this.Scope.TokenReader.RecordedTokens);
+					this.Scope.TokenReader.DisableRecording();
+					return true;
+				}
+			}
+			return false;
+		}
+	}
+}
diff --git a/Engine/Source/Programs/Shared/EpicGames.UHT/Parsers/UhtPropertyParser.cs b/Engine/Source/Programs/Shared/EpicGames.UHT/Parsers/UhtPropertyParser.cs
new file mode 100644
--- /dev/null
+++ b/Engine/Source/Programs/Shared/EpicGames.UHT/Parsers/UhtPropertyParser.cs
@@ -0,0 +1,1188 @@
+// Copyright Epic Games, Inc. All Rights Reserved.
+
+using EpicGames.Core;
+using EpicGames.UHT.Parsers;
+using EpicGames.UHT.Tables;
+using EpicGames.UHT.Tokenizer;
+using EpicGames.UHT.Types;
+using EpicGames.UHT.Utils;
+using System;
+using System.Collections.Generic;
+using System.Text;
+using System.Threading;
+
+namespace EpicGames.UHT.Parsers
+{
+	/// <summary>
+	/// Options that customize the parsing of properties.
+	/// </summary>
+	[Flags]
+	public enum UhtPropertyParseOptions
+	{
+		None = 0,
+
+		/// <summary>
+		/// Don't automatically mark properties as CPF_Const
+		/// </summary>
+		NoAutoConst = 1 << 0,
+
+		/// <summary>
+		/// Parse for the layout macro
+		/// </summary>
+		ParseLayoutMacro = 1 << 1,
+
+		/// <summary>
+		/// If set, then the name of the property will be parsed with the type
+		/// </summary>
+		FunctionNameIncluded = 1 << 2,
+
+		/// <summary>
+		/// If set, then the name of the property will be parsed with the type
+		/// </summary>
+		NameIncluded = 1 << 3,
+
+		/// <summary>
+		/// When parsing delegates, the name is separated by commas
+		/// </summary>
+		CommaSeparatedName = 1 << 4,
+
+		/// <summary>
+		/// Multiple properties can be defined separated by commas
+		/// </summary>
+		List = 1 << 5,
+
+		/// <summary>
+		/// Don't add a return type to the property list (return values go at the end)
+		/// </summary>
+		DontAddReturn = 1 << 6,
+
+		/// <summary>
+		/// If set, add the module relative path to the parameter's meta data
+		/// </summary>
+		AddModuleRelativePath = 1 << 7,
+	}
+
+	/// <summary>
+	/// Helper methods for testing flags.  These methods perform better than the generic HasFlag which hits
+	/// the GC and stalls.
+	/// </summary>
+	public static class UhtParsePropertyOptionsExtensions
+	{
+
+		/// <summary>
+		/// Test to see if any of the specified flags are set
+		/// </summary>
+		/// <param name="InFlags">Current flags</param>
+		/// <param name="TestFlags">Flags to test for</param>
+		/// <returns>True if any of the flags are set</returns>
+		public static bool HasAnyFlags(this UhtPropertyParseOptions InFlags, UhtPropertyParseOptions TestFlags)
+		{
+			return (InFlags & TestFlags) != 0;
+		}
+
+		/// <summary>
+		/// Test to see if all of the specified flags are set
+		/// </summary>
+		/// <param name="InFlags">Current flags</param>
+		/// <param name="TestFlags">Flags to test for</param>
+		/// <returns>True if all the flags are set</returns>
+		public static bool HasAllFlags(this UhtPropertyParseOptions InFlags, UhtPropertyParseOptions TestFlags)
+		{
+			return (InFlags & TestFlags) == TestFlags;
+		}
+
+		/// <summary>
+		/// Test to see if a specific set of flags have a specific value.
+		/// </summary>
+		/// <param name="InFlags">Current flags</param>
+		/// <param name="TestFlags">Flags to test for</param>
+		/// <param name="MatchFlags">Expected value of the tested flags</param>
+		/// <returns>True if the given flags have a specific value.</returns>
+		public static bool HasExactFlags(this UhtPropertyParseOptions InFlags, UhtPropertyParseOptions TestFlags, UhtPropertyParseOptions MatchFlags)
+		{
+			return (InFlags & TestFlags) == MatchFlags;
+		}
+	}
+
+	/// <summary>
+	/// The property style of a variable declaration being parsed
+	/// </summary>
+	public enum UhtParsePropertyDeclarationStyle
+	{
+
+		/// <summary>
+		/// Property is a function return, parameters or this is a template argument
+		/// </summary>
+		None,
+
+		/// <summary>
+		/// Class or script structure property
+		/// </summary>
+		UPROPERTY
+	};
+
+	public enum UhtLayoutMacroType
+	{
+		None,
+		Array,
+		ArrayEditorOnly,
+		Bitfield,
+		BitfieldEditorOnly,
+		Field,
+		FieldEditorOnly,
+		FieldInitialized,
+	}
+
+	public static class UhtLayoutMacroTypeExtensions
+	{
+		public static bool IsEditorOnly(this UhtLayoutMacroType LayoutMacroType)
+		{
+			switch (LayoutMacroType)
+			{
+				case UhtLayoutMacroType.ArrayEditorOnly:
+				case UhtLayoutMacroType.BitfieldEditorOnly:
+				case UhtLayoutMacroType.FieldEditorOnly:
+					return true;
+
+				default:
+					return false;
+			}
+		}
+
+		public static bool IsBitfield(this UhtLayoutMacroType LayoutMacroType)
+		{
+			switch (LayoutMacroType)
+			{
+				case UhtLayoutMacroType.Bitfield:
+				case UhtLayoutMacroType.BitfieldEditorOnly:
+					return true;
+
+				default:
+					return false;
+			}
+		}
+
+		public static bool IsArray(this UhtLayoutMacroType LayoutMacroType)
+		{
+			switch (LayoutMacroType)
+			{
+				case UhtLayoutMacroType.Array:
+				case UhtLayoutMacroType.ArrayEditorOnly:
+					return true;
+
+				default:
+					return false;
+			}
+		}
+
+		public static bool HasInitializer(this UhtLayoutMacroType LayoutMacroType)
+		{
+			switch (LayoutMacroType)
+			{
+				case UhtLayoutMacroType.FieldInitialized:
+					return true;
+
+				default:
+					return false;
+			}
+		}
+
+		public static StringView MacroName(this UhtLayoutMacroType LayoutMacroType)
+		{
+			switch (LayoutMacroType)
+			{
+				default:
+				case UhtLayoutMacroType.None:
+					throw new UhtIceException("Invalid macro name for ELayoutMacroType");
+
+				case UhtLayoutMacroType.Array:
+					return "LAYOUT_ARRAY";
+
+				case UhtLayoutMacroType.ArrayEditorOnly:
+					return "LAYOUT_ARRAY_EDITORONLY";
+
+				case UhtLayoutMacroType.Bitfield:
+					return "LAYOUT_BITFIELD";
+
+				case UhtLayoutMacroType.BitfieldEditorOnly:
+					return "LAYOUT_BITFIELD_EDITORONLY";
+
+				case UhtLayoutMacroType.Field:
+					return "LAYOUT_FIELD";
+
+				case UhtLayoutMacroType.FieldEditorOnly:
+					return "LAYOUT_FIELD_EDITORONLY";
+
+				case UhtLayoutMacroType.FieldInitialized:
+					return "LAYOUT_FIELD_INITIALIZED";
+			}
+		}
+
+		public static KeyValuePair<StringView, UhtLayoutMacroType> MacroNameAndValue(this UhtLayoutMacroType LayoutMacroType)
+		{
+			return new KeyValuePair<StringView, UhtLayoutMacroType>(LayoutMacroType.MacroName(), LayoutMacroType);
+		}
+	}
+
+	public delegate void UhtPropertyDelegate(UhtParsingScope TopScope, UhtProperty Property, ref UhtToken NameToken, UhtLayoutMacroType LayoutMacroType);
+
+	public class UhtPropertySpecifierContext : UhtSpecifierContext
+	{
+		public UhtPropertySettings PropertySettings;
+		public bool bSeenEditSpecifier = false;
+		public bool bSeenBlueprintWriteSpecifier = false;
+		public bool bSeenBlueprintReadOnlySpecifier = false;
+		public bool bSeenBlueprintGetterSpecifier = false;
+
+		public UhtPropertySpecifierContext(UhtParsingScope Scope, IUhtMessageSite MessageSite) : base(Scope, MessageSite, UhtMetaData.Empty)
+		{
+			this.PropertySettings = new UhtPropertySettings(Scope.ScopeType, 1, UhtPropertyCategory.Member, 0);
+		}
+	}
+
+	/// <summary>
+	/// A parsed property is a property that was parsed but couldn't yet be resolved.  It retains the list of tokens needed
+	/// to resolve the type of the property.  It will be replaced with the resolved property type during property resolution.
+	/// </summary>
+	public class UhtPreResolveProperty : UhtProperty
+	{
+		/// <inheritdoc/>
+		public override string EngineClassName { get => "UHTParsedProperty"; }
+
+		/// <inheritdoc/>
+		protected override string CppTypeText { get => "invalid"; }
+
+		/// <inheritdoc/>
+		protected override string PGetMacroText { get => "invalid"; }
+
+		public ReadOnlyMemory<UhtToken> TypeTokens;
+		public UhtPropertySettings PropertySettings;
+
+		public UhtPreResolveProperty(UhtPropertySettings PropertySettings, ReadOnlyMemory<UhtToken> TypeTokens) : base(PropertySettings)
+		{
+			this.TypeTokens = TypeTokens;
+			this.PropertySettings = PropertySettings;
+		}
+
+		/// <inheritdoc/>
+		public override bool IsSameType(UhtProperty Other)
+		{
+			throw new NotImplementedException();
+		}
+
+		/// <inheritdoc/>
+		public override StringBuilder AppendText(StringBuilder Builder, UhtPropertyTextType TextType, bool bIsTemplateArgument)
+		{
+			throw new NotImplementedException();
+		}
+
+		/// <inheritdoc/>
+		public override StringBuilder AppendMemberDecl(StringBuilder Builder, IUhtPropertyMemberContext Context, string Name, string NameSuffix, int Tabs)
+		{
+			throw new NotImplementedException();
+		}
+
+		/// <inheritdoc/>
+		public override StringBuilder AppendMemberDef(StringBuilder Builder, IUhtPropertyMemberContext Context, string Name, string NameSuffix, string? Offset, int Tabs)
+		{
+			throw new NotImplementedException();
+		}
+
+		/// <inheritdoc/>
+		public override StringBuilder AppendFunctionThunkParameterGet(StringBuilder Builder)
+		{
+			throw new NotImplementedException();
+		}
+
+		/// <inheritdoc/>
+		public override StringBuilder AppendNullConstructorArg(StringBuilder Builder, bool bIsInitializer)
+		{
+			throw new NotImplementedException();
+		}
+
+		/// <inheritdoc/>
+		public override bool SanitizeDefaultValue(IUhtTokenReader DefaultValueReader, StringBuilder InnerDefaultValue)
+		{
+			throw new NotImplementedException();
+		}
+
+		/// <inheritdoc/>
+		public override string? GetRigVMType(ref UhtRigVMParameterFlags ParameterFlags)
+		{
+			throw new NotImplementedException();
+		}
+	}
+
+	public class UhtPropertyParser : IUhtMessageExtraContext
+	{
+		private UhtPropertySpecifierContext SpecifierContext;
+		private UhtPropertyParseOptions Options;
+		private UhtParsePropertyDeclarationStyle DeclarationStyle;
+
+		private IUhtTokenReader TokenReader => this.SpecifierContext.Scope.TokenReader;
+		private UhtParsingScope TopScope => this.SpecifierContext.Scope;
+
+		// Scratch pad variables used by actions
+		private List<UhtToken> CurrentTypeTokens = new List<UhtToken>();
+		private UhtToken CurrentTerminatingToken = new UhtToken();
+		private int CurrentTemplateDepth = 0;
+
+		// Actions/Delegates
+		UhtTokensUntilDelegate GatherTypeTokensDelegate;
+
+		private static Dictionary<StringView, UhtLayoutMacroType> LayoutMacroTypes = new Dictionary<StringView, UhtLayoutMacroType>(new[]
+		{
+			UhtLayoutMacroType.Array.MacroNameAndValue(),
+			UhtLayoutMacroType.ArrayEditorOnly.MacroNameAndValue(),
+			UhtLayoutMacroType.Bitfield.MacroNameAndValue(),
+			UhtLayoutMacroType.BitfieldEditorOnly.MacroNameAndValue(),
+			UhtLayoutMacroType.Field.MacroNameAndValue(),
+			UhtLayoutMacroType.FieldEditorOnly.MacroNameAndValue(),
+			UhtLayoutMacroType.FieldInitialized.MacroNameAndValue(),
+		});
+
+		public UhtPropertyParser(UhtParsingScope Scope, IUhtMessageSite MessageSite)
+		{
+			this.SpecifierContext = new UhtPropertySpecifierContext(Scope, MessageSite);
+			this.GatherTypeTokensDelegate = GatherTypeTokens;
+		}
+
+		public void Reset(UhtParsingScope Scope, IUhtMessageSite MessageSite)
+		{
+			this.SpecifierContext.Scope = Scope;
+			this.SpecifierContext.MessageSite = MessageSite;
+		}
+
+		public UhtPropertyParser Parse(EPropertyFlags DisallowPropertyFlags, UhtPropertyParseOptions Options, UhtParsePropertyDeclarationStyle DeclarationStyle, UhtPropertyCategory Category, UhtPropertyDelegate Delegate)
+		{
+			// Reset the context and create the property
+			this.SpecifierContext.PropertySettings.Reset(this.SpecifierContext.Scope.ScopeType, 0, Category, DisallowPropertyFlags);
+			this.SpecifierContext.MetaData = this.SpecifierContext.PropertySettings.MetaData;
+			this.SpecifierContext.MetaNameIndex = UhtMetaData.INDEX_NONE;
+			this.SpecifierContext.bSeenEditSpecifier = false;
+			this.SpecifierContext.bSeenBlueprintWriteSpecifier = false;
+			this.SpecifierContext.bSeenBlueprintReadOnlySpecifier = false;
+			this.SpecifierContext.bSeenBlueprintGetterSpecifier = false;
+
+			// Initialize the property
+
+			this.Options = Options;
+			this.DeclarationStyle = DeclarationStyle;
+
+			this.CurrentTypeTokens = new List<UhtToken>();
+			this.CurrentTerminatingToken = new UhtToken();
+			this.CurrentTemplateDepth = 0;
+
+			using (var TokenContext = new UhtMessageContext(this.TokenReader, this))
+			{
+				ParseInternal(Delegate);
+			}
+			return this;
+		}
+
+		#region IMessageExtraContext implementation
+		public IEnumerable<object?>? MessageExtraContext
+		{
+			get
+			{
+				Stack<object?> ExtraContext = new Stack<object?>(1);
+				ExtraContext.Push(this.SpecifierContext.PropertySettings.PropertyCategory.GetHintText());
+				return ExtraContext;
+			}
+		}
+		#endregion
+
+		/// <summary>
+		/// Resolve the given property.  This method will resolve any immediate property during the parsing phase or 
+		/// resolve any previously parsed property to the final version.
+		/// </summary>
+		/// <param name="ResolvePhase">Used to detect if the property is being parsed or resolved</param>
+		/// <param name="Data">Character buffer for header file</param>
+		/// <param name="PropertySettings">The property settings.</param>
+		/// <param name="TypeTokens">The tokens that represent the type</param>
+		/// <returns></returns>
+		public static UhtProperty? ResolveProperty(UhtPropertyResolvePhase ResolvePhase, UhtPropertySettings PropertySettings, ReadOnlyMemory<char> Data, ReadOnlyMemory<UhtToken> TypeTokens)
+		{
+			UhtPropertyTypeTable PropertyTypeTable = UhtPropertyTypeTable.Instance;
+			IUhtTokenReader ReplayReader = UhtTokenReplayReader.GetThreadInstance(PropertySettings.Outer, Data, TypeTokens, UhtTokenType.EndOfType);
+
+			// Loop through the tokens until we find a known property type or the start of a template argument list
+			UhtPropertyType PropertyType = new UhtPropertyType();
+			for (int Index = 0; Index < TypeTokens.Length; ++Index)
+			{
+				if (TypeTokens.Span[Index].IsSymbol())
+				{
+					ReadOnlySpan<char> Span = TypeTokens.Span[Index].Value.Span;
+					if (Span.Length == 1 && (Span[0] == '<' || Span[0] == '>' || Span[0] == ','))
+					{
+						break;
+					}
+				}
+				else
+				{
+					if (PropertyTypeTable.TryGet(TypeTokens.Span[Index].Value, out PropertyType))
+					{
+						if (ResolvePhase == UhtPropertyResolvePhase.Resolving || PropertyType.Options.HasAnyFlags(UhtPropertyTypeOptions.Immediate))
+						{
+							return ResolvePropertyType(ResolvePhase, PropertySettings, ReplayReader, PropertyType, TypeTokens.Span[Index], false);
+						}
+					}
+				}
+			}
+
+			// Try the default processor.  We only use the default processor when trying to resolve something post parsing phase.
+			if (ResolvePhase == UhtPropertyResolvePhase.Resolving)
+			{
+				return ResolvePropertyType(ResolvePhase, PropertySettings, ReplayReader, PropertyTypeTable.Default, new UhtToken(), false);
+			}
+			return null;
+		}
+
+		private static ThreadLocal<UhtPropertySettings> TlsPropertySettings = new ThreadLocal<UhtPropertySettings>(() => { return new UhtPropertySettings(); });
+
+		/// <summary>
+		/// Given a type with children, resolve any children that couldn't be resolved during the parsing phase.
+		/// </summary>
+		/// <param name="Type">The type with children</param>
+		/// <param name="Options">Parsing options</param>
+		public static void ResolveChildren(UhtType Type, UhtPropertyParseOptions Options)
+		{
+			UhtPropertyOptions PropertyOptions = UhtPropertyOptions.None;
+			if (Options.HasAnyFlags(UhtPropertyParseOptions.NoAutoConst))
+			{
+				PropertyOptions |= UhtPropertyOptions.NoAutoConst;
+			}
+			bool bInSymbolTable = Type.EngineType.AddChildrenToSymbolTable();
+
+			UhtPropertySettings? PropertySettings = TlsPropertySettings.Value;
+			if (PropertySettings == null)
+			{
+				throw new UhtIceException("Unable to acquire threaded property settings");
+			}
+
+			for (int Index = 0; Index < Type.Children.Count; ++Index)
+			{
+				if (Type.Children[Index] is UhtPreResolveProperty Property)
+				{
+					PropertySettings.Reset(Property, PropertyOptions);
+					UhtProperty? Resolved = UhtPropertyParser.ResolveProperty(UhtPropertyResolvePhase.Resolving, PropertySettings, Property.HeaderFile.Data.Memory, Property.TypeTokens);
+					if (Resolved != null)
+					{
+						if (bInSymbolTable && Resolved != Property)
+						{
+							Type.Session.ReplaceTypeInSymbolTable(Property, Resolved);
+						}
+						Type.Children[Index] = Resolved;
+					}
+				}
+			}
+		}
+
+		public static UhtProperty? ParseTemplateParam(UhtPropertyResolvePhase ResolvePhase, UhtPropertySettings ParentPropertySettings, StringView ParamName, IUhtTokenReader TokenReader)
+		{
+			UhtPropertyTypeTable PropertyTypeTable = UhtPropertyTypeTable.Instance;
+
+			// Save the token reader state.  We need this to restore back to the start when invoking the resolve methods.
+			using (var SaveState = new UhtTokenSaveState(TokenReader))
+			{
+				UhtPropertySettings PropertySettings = new UhtPropertySettings(ParentPropertySettings, ParamName.ToString(), TokenReader);
+
+				UhtPropertyType PropertyType = new UhtPropertyType();
+				while (!TokenReader.bIsEOF)
+				{
+					UhtToken Token = TokenReader.GetToken();
+
+					if (Token.IsSymbol('<') || Token.IsSymbol('>') || Token.IsSymbol(','))
+					{
+						break;
+					}
+
+					if (PropertyTypeTable.TryGet(Token.Value, out PropertyType))
+					{
+						if (ResolvePhase == UhtPropertyResolvePhase.Resolving || PropertyType.Options.HasAnyFlags(UhtPropertyTypeOptions.Immediate))
+						{
+							SaveState.RestoreState();
+							return ResolvePropertyType(ResolvePhase, PropertySettings, TokenReader, PropertyType, Token, true);
+						}
+					}
+				}
+
+				// Try the default processor.  We only use the default processor when trying to resolve something post parsing phase.
+				if (ResolvePhase == UhtPropertyResolvePhase.Resolving)
+				{
+					SaveState.RestoreState();
+					return ResolvePropertyType(ResolvePhase, PropertySettings, TokenReader, PropertyTypeTable.Default, new UhtToken(), true);
+				}
+				return null;
+			}
+		}
+
+		private static UhtProperty? ResolvePropertyType(UhtPropertyResolvePhase ResolvePhase,
+			UhtPropertySettings PropertySettings, IUhtTokenReader TokenReader, UhtPropertyType PropertyType, UhtToken MatchedToken, bool bIsTemplate)
+		{
+			UhtProperty? Out = PropertyType.Delegate(ResolvePhase, PropertySettings, TokenReader, MatchedToken);
+			if (Out == null)
+			{
+				return null;
+			}
+
+			// If this is a simple type, skip the type
+			if (PropertyType.Options.HasAnyFlags(UhtPropertyTypeOptions.Simple))
+			{
+				if (!TokenReader.SkipExpectedType(MatchedToken.Value, PropertySettings.PropertyCategory == UhtPropertyCategory.Member))
+				{
+					return null;
+				}
+			}
+
+			// Handle any trailing const
+			if (Out.PropertyCategory == UhtPropertyCategory.Member)
+			{
+				//@TODO: UCREMOVAL: 'const' member variables that will get written post-construction by defaultproperties
+				UhtClass? OuterClass = Out.Outer as UhtClass;
+				if (OuterClass != null && OuterClass.ClassFlags.HasAnyFlags(EClassFlags.Const))
+				{
+					// Eat a 'not quite truthful' const after the type; autogenerated for member variables of const classes.
+					if (TokenReader.TryOptional("const"))
+					{
+						Out.MetaData.Add(UhtNames.NativeConst, "");
+					}
+				}
+			}
+			else
+			{
+				if (TokenReader.TryOptional("const"))
+				{
+					Out.MetaData.Add(UhtNames.NativeConst, "");
+					Out.PropertyFlags |= EPropertyFlags.ConstParm;
+				}
+			}
+
+			// Check for unexpected '*'
+			if (TokenReader.TryOptional('*'))
+			{
+				TokenReader.LogError($"Inappropriate '*' on variable of type '{Out.GetUserFacingDecl()}', cannot have an exposed pointer to this type.");
+			}
+
+			// Arrays are passed by reference but are only implicitly so; setting it explicitly could cause a problem with replicated functions
+			if (TokenReader.TryOptional('&'))
+			{
+				switch (Out.PropertyCategory)
+				{
+					case UhtPropertyCategory.RegularParameter:
+					case UhtPropertyCategory.Return:
+						Out.PropertyFlags |= EPropertyFlags.OutParm;
+
+						//@TODO: UCREMOVAL: How to determine if we have a ref param?
+						if (Out.PropertyFlags.HasAnyFlags(EPropertyFlags.ConstParm))
+						{
+							Out.PropertyFlags |= EPropertyFlags.ReferenceParm;
+						}
+						break;
+
+					case UhtPropertyCategory.ReplicatedParameter:
+						Out.PropertyFlags |= EPropertyFlags.ReferenceParm;
+						break;
+
+					default:
+						break;
+				}
+
+				if (Out.PropertyFlags.HasAnyFlags(EPropertyFlags.ConstParm))
+				{
+					Out.RefQualifier = UhtPropertyRefQualifier.ConstRef;
+				}
+				else
+				{
+					Out.RefQualifier = UhtPropertyRefQualifier.NonConstRef;
+				}
+			}
+
+			if (!bIsTemplate)
+			{
+				if (!TokenReader.bIsEOF)
+				{
+					throw new UhtTokenException(TokenReader, TokenReader.PeekToken(), "end of type declaration");
+				}
+			}
+			else
+			{
+				ref UhtToken Token = ref TokenReader.PeekToken();
+				if (!Token.IsSymbol(',') && !Token.IsSymbol('>'))
+				{
+					throw new UhtTokenException(TokenReader, TokenReader.PeekToken(), "end of type declaration");
+				}
+			}
+			return Out;
+		}
+
+		private void ParseInternal(UhtPropertyDelegate Delegate)
+		{
+			this.SpecifierContext.PropertySettings.LineNumber = this.TokenReader.InputLine;
+
+			bool bIsParamList = this.SpecifierContext.PropertySettings.PropertyCategory != UhtPropertyCategory.Member && this.TokenReader.TryOptional("UPARAM");
+
+			UhtSpecifierParser Specifiers = this.TopScope.HeaderParser.GetSpecifierParser(this.SpecifierContext, "Variable", 
+				bIsParamList ? UhtPropertyArgumentSpecifiers.SpecifierTable : UhtPropertyMemberSpecifiers.SpecifierTable);
+			if (DeclarationStyle == UhtParsePropertyDeclarationStyle.UPROPERTY || bIsParamList)
+			{
+				Specifiers.ParseSpecifiers();
+			}
+
+			if (this.SpecifierContext.PropertySettings.PropertyCategory != UhtPropertyCategory.Member)
+			{
+				// const before the variable type support (only for params)
+				if (this.TokenReader.TryOptional("const"))
+				{
+					this.SpecifierContext.PropertySettings.PropertyFlags |= EPropertyFlags.ConstParm;
+					this.SpecifierContext.PropertySettings.MetaData.Add(UhtNames.NativeConst, "");
+				}
+			}
+
+			UhtCompilerDirective CompilerDirective = this.TopScope.HeaderParser.GetCurrentCompositeCompilerDirective();
+			if (CompilerDirective.HasAnyFlags(UhtCompilerDirective.WithEditorOnlyData))
+			{
+				this.SpecifierContext.PropertySettings.PropertyFlags |= EPropertyFlags.EditorOnly;
+			}
+			else if (CompilerDirective.HasAnyFlags(UhtCompilerDirective.WithEditor))
+			{
+				// Checking for this error is a bit tricky given legacy code.  
+				// 1) If already wrapped in WITH_EDITORONLY_DATA (see above), then we ignore the error via the else 
+				// 2) Ignore any module that is an editor module
+				UhtPackage Package = this.TopScope.HeaderParser.HeaderFile.Package;
+				UHTManifest.Module Module = Package.Module;
+				bool bIsEditorModule =
+					Module.ModuleType == UHTModuleType.EngineEditor ||
+					Module.ModuleType == UHTModuleType.GameEditor ||
+					Module.ModuleType == UHTModuleType.EngineUncooked ||
+					Module.ModuleType == UHTModuleType.GameUncooked;
+				if (this.SpecifierContext.PropertySettings.PropertyCategory == UhtPropertyCategory.Member && !bIsEditorModule)
+				{
+					this.TokenReader.LogError("UProperties should not be wrapped by WITH_EDITOR, use WITH_EDITORONLY_DATA instead.");
+				}
+			}
+
+			// Store the start and end positions of the parsed type
+			int StartPos = this.TokenReader.InputPos;
+
+			// Process the specifiers
+			if (DeclarationStyle == UhtParsePropertyDeclarationStyle.UPROPERTY || bIsParamList)
+			{
+				Specifiers.ParseDeferred();
+			}
+
+			// If we saw a BlueprintGetter but did not see BlueprintSetter or 
+			// or BlueprintReadWrite then treat as BlueprintReadOnly
+			if (this.SpecifierContext.bSeenBlueprintGetterSpecifier && !this.SpecifierContext.bSeenBlueprintWriteSpecifier)
+			{
+				this.SpecifierContext.PropertySettings.PropertyFlags |= EPropertyFlags.BlueprintReadOnly;
+			}
+
+			if (this.SpecifierContext.PropertySettings.MetaData.ContainsKey(UhtNames.ExposeOnSpawn))
+			{
+				this.SpecifierContext.PropertySettings.PropertyFlags |= EPropertyFlags.ExposeOnSpawn;
+			}
+
+			UhtAccessSpecifier AccessSpecifier = TopScope.AccessSpecifier;
+			if (AccessSpecifier == UhtAccessSpecifier.Public || this.SpecifierContext.PropertySettings.PropertyCategory != UhtPropertyCategory.Member)
+			{
+				this.SpecifierContext.PropertySettings.PropertyFlags &= ~EPropertyFlags.Protected;
+				this.SpecifierContext.PropertySettings.PropertyExportFlags |= UhtPropertyExportFlags.Public;
+				this.SpecifierContext.PropertySettings.PropertyExportFlags &= ~(UhtPropertyExportFlags.Private | UhtPropertyExportFlags.Protected);
+
+				this.SpecifierContext.PropertySettings.PropertyFlags &= ~EPropertyFlags.NativeAccessSpecifiers;
+				this.SpecifierContext.PropertySettings.PropertyFlags |= EPropertyFlags.NativeAccessSpecifierPublic;
+			}
+			else if (AccessSpecifier == UhtAccessSpecifier.Protected)
+			{
+				this.SpecifierContext.PropertySettings.PropertyFlags |= EPropertyFlags.Protected;
+				this.SpecifierContext.PropertySettings.PropertyExportFlags |= UhtPropertyExportFlags.Protected;
+				this.SpecifierContext.PropertySettings.PropertyExportFlags &= ~(UhtPropertyExportFlags.Public | UhtPropertyExportFlags.Private);
+
+				this.SpecifierContext.PropertySettings.PropertyFlags &= ~EPropertyFlags.NativeAccessSpecifiers;
+				this.SpecifierContext.PropertySettings.PropertyFlags |= EPropertyFlags.NativeAccessSpecifierProtected;
+			}
+			else if (AccessSpecifier == UhtAccessSpecifier.Private)
+			{
+				this.SpecifierContext.PropertySettings.PropertyFlags &= ~EPropertyFlags.Protected;
+				this.SpecifierContext.PropertySettings.PropertyExportFlags |= UhtPropertyExportFlags.Private;
+				this.SpecifierContext.PropertySettings.PropertyExportFlags &= ~(UhtPropertyExportFlags.Public | UhtPropertyExportFlags.Protected);
+
+				this.SpecifierContext.PropertySettings.PropertyFlags &= ~EPropertyFlags.NativeAccessSpecifiers;
+				this.SpecifierContext.PropertySettings.PropertyFlags |= EPropertyFlags.NativeAccessSpecifierPrivate;
+			}
+			else
+			{
+				throw new UhtIceException("Unknown access level");
+			}
+
+			// Swallow inline keywords
+			if (this.SpecifierContext.PropertySettings.PropertyCategory == UhtPropertyCategory.Return)
+			{
+				this.TokenReader
+					.Optional("inline")
+					.Optional("FORCENOINLINE")
+					.OptionalStartsWith("FORCEINLINE");
+			}
+
+			// Handle MemoryLayout.h macros
+			bool bHasWrapperBrackets = false;
+			UhtLayoutMacroType LayoutMacroType = UhtLayoutMacroType.None;
+			if (this.Options.HasAnyFlags(UhtPropertyParseOptions.ParseLayoutMacro))
+			{
+				ref UhtToken LayoutToken = ref this.TokenReader.PeekToken();
+				if (LayoutToken.IsIdentifier())
+				{
+					if (LayoutMacroTypes.TryGetValue(LayoutToken.Value, out LayoutMacroType))
+					{
+						this.TokenReader.ConsumeToken();
+						this.TokenReader.Require('(');
+						bHasWrapperBrackets = this.TokenReader.TryOptional('(');
+						if (LayoutMacroType.IsEditorOnly())
+						{
+							this.SpecifierContext.PropertySettings.PropertyFlags |= EPropertyFlags.EditorOnly;
+						}
+					}
+				}
+
+				// This exists as a compatibility "shim" with UHT4/5.0.  If the fetched token wasn't an identifier,
+				// it wasn't returned to the tokenizer.  So, just consume the token here.  In theory, this should be
+				// removed once we have a good deprecated system.
+				//@TODO - deprecate
+				else // if (LayoutToken.IsSymbol(';'))
+				{
+					this.TokenReader.ConsumeToken();
+				}
+			}
+
+			//@TODO: Should flag as settable from a const context, but this is at least good enough to allow use for C++ land
+			this.TokenReader.Optional("mutable");
+
+			// Gather the type tokens and possibly the property name.
+			this.TokenReader.While(GatherTypeTokensDelegate);
+
+			// Verify we at least have one type
+			if (this.CurrentTypeTokens.Count < 1)
+			{
+				throw new UhtException(this.TokenReader, $"{this.SpecifierContext.PropertySettings.PropertyCategory.GetHintText()}: Missing variable type or name");
+			}
+
+			// Consume the wrapper brackets.  This is just an extra set
+			if (bHasWrapperBrackets)
+			{
+				this.TokenReader.Require(')');
+			}
+
+			// Check for any disallowed flags
+			if (this.SpecifierContext.PropertySettings.PropertyFlags.HasAnyFlags(this.SpecifierContext.PropertySettings.DisallowPropertyFlags))
+			{
+				this.TokenReader.LogError("Specified type modifiers not allowed here");
+			}
+
+			if (this.Options.HasAnyFlags(UhtPropertyParseOptions.AddModuleRelativePath))
+			{
+				UhtParsingScope.AddModuleRelativePathToMetaData(this.SpecifierContext.PropertySettings.MetaData, this.SpecifierContext.Scope.ScopeType.HeaderFile);
+			}
+
+			// Fetch the name of the property, bitfield and array size
+			if (LayoutMacroType != UhtLayoutMacroType.None)
+			{
+				this.TokenReader.Require(',');
+				UhtToken NameToken = this.TokenReader.GetIdentifier();
+				if (LayoutMacroType.IsArray())
+				{
+					this.TokenReader.Require(',');
+					RequireArray(this.SpecifierContext.PropertySettings, ref NameToken, ')');
+				}
+				else if (LayoutMacroType.IsBitfield())
+				{
+					this.TokenReader.Require(',');
+					RequireBitfield(this.SpecifierContext.PropertySettings, ref NameToken);
+				}
+				else if (LayoutMacroType.HasInitializer())
+				{
+					this.TokenReader.SkipBrackets('(', ')', 1);
+				}
+				this.TokenReader.Require(')');
+
+				Finalize(ref NameToken, new ReadOnlyMemory<UhtToken>(this.CurrentTypeTokens.ToArray()), LayoutMacroType, Delegate);
+			}
+			else if (this.Options.HasAnyFlags(UhtPropertyParseOptions.List))
+			{
+				// Extract the property name from the types
+				if (this.CurrentTypeTokens.Count < 2 || !this.CurrentTypeTokens[this.CurrentTypeTokens.Count - 1].IsIdentifier())
+				{
+					throw new UhtException(this.TokenReader, $"{this.SpecifierContext.PropertySettings.PropertyCategory.GetHintText()}: Expected name");
+				}
+				UhtToken NameToken = this.CurrentTypeTokens[this.CurrentTypeTokens.Count - 1];
+				this.CurrentTypeTokens.RemoveAt(this.CurrentTypeTokens.Count - 1);
+				CheckForOptionalParts(this.SpecifierContext.PropertySettings, ref NameToken);
+
+				ReadOnlyMemory<UhtToken> TypeTokens = new ReadOnlyMemory<UhtToken>(this.CurrentTypeTokens.ToArray());
+
+				while (true)
+				{
+					UhtProperty Property = Finalize(ref NameToken, TypeTokens, LayoutMacroType, Delegate);
+
+					// If we have reached the end
+					if (!this.TokenReader.TryOptional(','))
+					{
+						break;
+					}
+
+					// While we could continue parsing, the old UHT would flag this as an error.
+					throw new UhtException(TopScope.TokenReader, $"Comma delimited properties cannot be converted");
+
+					// we'll need any metadata tags we parsed later on when we call ConvertEOLCommentToTooltip() so the tags aren't clobbered
+					//this.SpecifierContext.PropertySettings.MetaData = Property.MetaData.Clone();
+					//this.SpecifierContext.PropertySettings.ResetTrailingSettings();
+
+					// Get the next property name
+					//NameToken = this.TokenReader.GetIdentifier();
+					//CheckForOptionalParts(this.SpecifierContext.PropertySettings, ref NameToken);
+				}
+			}
+			else if (this.Options.HasAnyFlags(UhtPropertyParseOptions.CommaSeparatedName))
+			{
+				this.TokenReader.Require(',');
+				UhtToken NameToken = this.TokenReader.GetIdentifier();
+				CheckForOptionalParts(this.SpecifierContext.PropertySettings, ref NameToken);
+				Finalize(ref NameToken, new ReadOnlyMemory<UhtToken>(this.CurrentTypeTokens.ToArray()), LayoutMacroType, Delegate);
+			}
+			else if (this.Options.HasAnyFlags(UhtPropertyParseOptions.FunctionNameIncluded))
+			{
+				UhtToken NameToken = this.CurrentTypeTokens[this.CurrentTypeTokens.Count - 1];
+				NameToken.Value = new StringView("Function");
+				if (CheckForOptionalParts(this.SpecifierContext.PropertySettings, ref NameToken))
+				{
+					NameToken = this.TokenReader.GetIdentifier("function name");
+				}
+				else
+				{
+					if (this.CurrentTypeTokens.Count < 2 || !this.CurrentTypeTokens[this.CurrentTypeTokens.Count - 1].IsIdentifier())
+					{
+						throw new UhtException(this.TokenReader, $"{this.SpecifierContext.PropertySettings.PropertyCategory.GetHintText()}: Expected name");
+					}
+					NameToken = this.CurrentTypeTokens[this.CurrentTypeTokens.Count - 1];
+					this.CurrentTypeTokens.RemoveAt(this.CurrentTypeTokens.Count - 1);
+				}
+				Finalize(ref NameToken, new ReadOnlyMemory<UhtToken>(this.CurrentTypeTokens.ToArray()), LayoutMacroType, Delegate);
+			}
+			else if (this.Options.HasAnyFlags(UhtPropertyParseOptions.NameIncluded))
+			{
+				if (this.CurrentTypeTokens.Count < 2 || !this.CurrentTypeTokens[this.CurrentTypeTokens.Count - 1].IsIdentifier())
+				{
+					throw new UhtException(this.TokenReader, $"{this.SpecifierContext.PropertySettings.PropertyCategory.GetHintText()}: Expected name");
+				}
+				UhtToken NameToken = this.CurrentTypeTokens[this.CurrentTypeTokens.Count - 1];
+				this.CurrentTypeTokens.RemoveAt(this.CurrentTypeTokens.Count - 1);
+				CheckForOptionalParts(this.SpecifierContext.PropertySettings, ref NameToken);
+				Finalize(ref NameToken, new ReadOnlyMemory<UhtToken>(this.CurrentTypeTokens.ToArray()), LayoutMacroType, Delegate);
+			}
+			else
+			{
+				UhtToken NameToken = new UhtToken();
+				CheckForOptionalParts(this.SpecifierContext.PropertySettings, ref NameToken);
+				Finalize(ref NameToken, new ReadOnlyMemory<UhtToken>(this.CurrentTypeTokens.ToArray()), LayoutMacroType, Delegate);
+			}
+		}
+
+		/// <summary>
+		/// Finish creating the property
+		/// </summary>
+		/// <param name="NameToken">The name of the property</param>
+		/// <param name="TypeTokens">Series of tokens that represent the type</param>
+		/// <param name="LayoutMacroType">Optional layout macro type being parsed</param>
+		/// <param name="Delegate">Delegate to invoke when processing has been completed</param>
+		/// <returns>The newly created property.  During the parsing phase, this will often be a temporary property if the type references engine types.</returns>
+		private UhtProperty Finalize(ref UhtToken NameToken, ReadOnlyMemory<UhtToken> TypeTokens, UhtLayoutMacroType LayoutMacroType, UhtPropertyDelegate Delegate)
+		{
+			this.SpecifierContext.PropertySettings.SourceName = this.SpecifierContext.PropertySettings.PropertyCategory == UhtPropertyCategory.Return ? "ReturnValue" : NameToken.Value.ToString();
+
+			// Try to resolve the property using any immediate mode property types
+			UhtProperty? NewProperty = ResolveProperty(UhtPropertyResolvePhase.Parsing, this.SpecifierContext.PropertySettings, this.SpecifierContext.PropertySettings.Outer.HeaderFile.Data.Memory, TypeTokens);
+			if (NewProperty == null)
+			{
+				NewProperty = new UhtPreResolveProperty(this.SpecifierContext.PropertySettings, TypeTokens);
+			}
+
+			// Force the category in non-engine projects
+			if (NewProperty.PropertyCategory == UhtPropertyCategory.Member)
+			{
+				if (!NewProperty.Package.bIsPartOfEngine && 
+					NewProperty.PropertyFlags.HasAnyFlags(EPropertyFlags.Edit | EPropertyFlags.BlueprintVisible) &&
+					!NewProperty.MetaData.ContainsKey(UhtNames.Category))
+				{
+					NewProperty.MetaData.Add(UhtNames.Category, NewProperty.Outer!.EngineName);
+				}
+			}
+
+			// Check to see if the variable is deprecated, and if so set the flag
+			{
+				int DeprecatedIndex = NewProperty.SourceName.IndexOf("_DEPRECATED");
+				int NativizedPropertyPostfixIndex = NewProperty.SourceName.IndexOf("__pf"); //@TODO: check OverrideNativeName in Meta Data, to be sure it's not a random occurrence of the "__pf" string.
+				bool bIgnoreDeprecatedWord = (NativizedPropertyPostfixIndex != -1) && (NativizedPropertyPostfixIndex > DeprecatedIndex);
+				if ((DeprecatedIndex != -1) && !bIgnoreDeprecatedWord)
+				{
+					if (DeprecatedIndex != NewProperty.SourceName.Length - 11)
+					{
+						this.TokenReader.LogError("Deprecated variables must end with _DEPRECATED");
+					}
+
+					// We allow deprecated properties in blueprints that have getters and setters assigned as they may be part of a backwards compatibility path
+					bool bBlueprintVisible = NewProperty.PropertyFlags.HasAnyFlags(EPropertyFlags.BlueprintVisible);
+					bool bWarnOnGetter = bBlueprintVisible && !NewProperty.MetaData.ContainsKey(UhtNames.BlueprintGetter);
+					bool bWarnOnSetter = bBlueprintVisible && !NewProperty.PropertyFlags.HasAnyFlags(EPropertyFlags.BlueprintReadOnly) && !NewProperty.MetaData.ContainsKey(UhtNames.BlueprintSetter);
+
+					if (bWarnOnGetter)
+					{
+						this.TokenReader.LogWarning($"{NewProperty.PropertyCategory.GetHintText()}: Deprecated property '{NewProperty.SourceName}' should not be marked as blueprint visible without having a BlueprintGetter");
+					}
+
+					if (bWarnOnSetter)
+					{
+						this.TokenReader.LogWarning($"{NewProperty.PropertyCategory.GetHintText()}: Deprecated property '{NewProperty.SourceName}' should not be marked as blueprint writable without having a BlueprintSetter");
+					}
+
+					// Warn if a deprecated property is visible
+					if (NewProperty.PropertyFlags.HasAnyFlags(EPropertyFlags.Edit | EPropertyFlags.EditConst) || // Property is marked as editable
+						(!bBlueprintVisible && NewProperty.PropertyFlags.HasAnyFlags(EPropertyFlags.BlueprintReadOnly) && 
+						!NewProperty.PropertyExportFlags.HasAnyFlags(UhtPropertyExportFlags.ImpliedBlueprintPure))) // Is BPRO, but not via Implied Flags and not caught by Getter/Setter path above
+					{
+						this.TokenReader.LogWarning($"{NewProperty.PropertyCategory.GetHintText()}: Deprecated property '{NewProperty.SourceName}' should not be marked as visible or editable");
+					}
+
+					NewProperty.PropertyFlags |= EPropertyFlags.Deprecated;
+					NewProperty.EngineName = NewProperty.SourceName.Substring(0, DeprecatedIndex);
+				}
+			}
+
+			// Try gathering metadata for member fields
+			if (NewProperty.PropertyCategory == UhtPropertyCategory.Member)
+			{
+				UhtSpecifierParser Specifiers = this.TopScope.HeaderParser.GetSpecifierParser(this.SpecifierContext, NewProperty.SourceName, UhtPropertyMemberSpecifiers.SpecifierTable);
+				Specifiers.ParseFieldMetaData(NewProperty.SourceName);
+				this.TopScope.AddFormattedCommentsAsTooltipMetaData(NewProperty);
+			}
+
+			Delegate(this.TopScope, NewProperty, ref NameToken, LayoutMacroType);
+
+			// Void properties don't get added when they are the return value
+			if (NewProperty.PropertyCategory != UhtPropertyCategory.Return || !this.Options.HasAnyFlags(UhtPropertyParseOptions.DontAddReturn))
+			{
+				this.TopScope.ScopeType.AddChild(NewProperty);
+			}
+			return NewProperty;
+		}
+
+		private bool CheckForOptionalParts(UhtPropertySettings PropertySettings, ref UhtToken NameToken)
+		{
+			bool bGotOptionalParts = false;
+
+			if (this.TokenReader.TryOptional('['))
+			{
+				RequireArray(PropertySettings, ref NameToken, ']');
+				this.TokenReader.Require(']');
+				bGotOptionalParts = true;
+			}
+
+			if (this.TokenReader.TryOptional(':'))
+			{
+				RequireBitfield(PropertySettings, ref NameToken);
+				bGotOptionalParts = true;
+			}
+			return bGotOptionalParts;
+		}
+
+		private void RequireBitfield(UhtPropertySettings PropertySettings, ref UhtToken NameToken)
+		{
+			int BitfieldSize;
+			if (!this.TokenReader.TryOptionalConstInt(out BitfieldSize) || BitfieldSize != 1)
+			{
+				throw new UhtException(this.TokenReader, $"Bad or missing bit field size for '{NameToken.Value}', must be 1.");
+			}
+			PropertySettings.bIsBitfield = true;
+		}
+
+		private void RequireArray(UhtPropertySettings PropertySettings, ref UhtToken NameToken, char Terminator)
+		{
+			// Ignore how the actual array dimensions are actually defined - we'll calculate those with the compiler anyway.
+			PropertySettings.ArrayDimensions = this.TokenReader.GetRawString(Terminator, UhtRawStringOptions.DontConsumeTerminator).ToString();
+			if (PropertySettings.ArrayDimensions.Length == 0)
+			{
+				throw new UhtException(this.TokenReader, $"{PropertySettings.PropertyCategory.GetHintText()} {NameToken.Value}: Missing array dimensions or terminating '{Terminator}'");
+			}
+		}
+
+		private bool GatherTypeTokens(ref UhtToken Token)
+		{
+			if (this.CurrentTemplateDepth == 0 && Token.IsSymbol() && (Token.IsValue(',') || Token.IsValue('(') || Token.IsValue(')') || Token.IsValue(';') || Token.IsValue('[') || Token.IsValue(':') || Token.IsValue('=') || Token.IsValue('{')))
+			{
+				this.CurrentTerminatingToken = Token;
+				return false;
+			}
+
+			this.CurrentTypeTokens.Add(Token);
+			if (Token.IsSymbol('<'))
+			{
+				++this.CurrentTemplateDepth;
+			}
+			else if (Token.IsSymbol('>'))
+			{
+				if (this.CurrentTemplateDepth == 0)
+				{
+					throw new UhtTokenException(TokenReader, Token, "',' or ')'");
+				}
+				--this.CurrentTemplateDepth;
+			}
+			return true;
+		}
+	}
+
+	[UnrealHeaderTool]
+	public static class PropertyKeywords
+	{
+#region Keywords
+		[UhtKeyword(Extends = UhtTableNames.Class)]
+		[UhtKeyword(Extends = UhtTableNames.ScriptStruct)]
+		private static UhtParseResult UPROPERTYKeyword(UhtParsingScope TopScope, UhtParsingScope ActionScope, ref UhtToken Token)
+		{
+			UhtPropertyParseOptions Options = UhtPropertyParseOptions.ParseLayoutMacro | UhtPropertyParseOptions.List | UhtPropertyParseOptions.AddModuleRelativePath;
+			TopScope.HeaderParser.GetPropertyParser(TopScope).Parse(EPropertyFlags.ParmFlags, Options, UhtParsePropertyDeclarationStyle.UPROPERTY, UhtPropertyCategory.Member, PropertyDelegate);
+			TopScope.TokenReader.Require(';');
+			return UhtParseResult.Handled;
+		}
+#endregion
+
+		private static UhtPropertyDelegate PropertyDelegate = PropertyParsed;
+
+		private static void PropertyParsed(UhtParsingScope TopScope, UhtProperty Property, ref UhtToken NameToken, UhtLayoutMacroType LayoutMacroType)
+		{
+			IUhtTokenReader TokenReader = TopScope.TokenReader;
+			UhtStruct Struct = (UhtStruct)TopScope.ScopeType;
+
+			// Skip any initialization
+			if (TokenReader.TryOptional('='))
+			{
+				TokenReader.SkipUntil(';');
+			}
+			else if (TokenReader.TryOptional('{'))
+			{
+				TokenReader.SkipBrackets('{', '}', 1);
+			}
+		}
+	}
+
+	[UnrealHeaderTool]
+	public static class UhtDefaultPropertyParser
+	{
+		[UhtPropertyType(Options = UhtPropertyTypeOptions.Default)]
+
+
+		public static UhtProperty? DefaultProperty(UhtPropertyResolvePhase ResolvePhase, UhtPropertySettings PropertySettings, IUhtTokenReader TokenReader, UhtToken MatchedToken)
+		{
+			int TypeStartPos = TokenReader.PeekToken().InputStartPos;
+
+			if (TokenReader.TryOptional("const"))
+			{
+				PropertySettings.MetaData.Add(UhtNames.NativeConst, "");
+			}
+
+			UhtFindOptions FindOptions = UhtFindOptions.DelegateFunction | UhtFindOptions.Enum | UhtFindOptions.Class | UhtFindOptions.ScriptStruct;
+			if (TokenReader.TryOptional("enum"))
+			{
+				if (PropertySettings.PropertyCategory == UhtPropertyCategory.Member)
+				{
+					TokenReader.LogError($"Cannot declare enum at variable declaration");
+				}
+				FindOptions = UhtFindOptions.Enum;
+			}
+			else if (TokenReader.TryOptional("class"))
+			{
+				FindOptions = UhtFindOptions.Class;
+			}
+			else if (TokenReader.TryOptional("struct"))
+			{
+				FindOptions = UhtFindOptions.ScriptStruct;
+			}
+
+			UhtTokenList Identifiers = TokenReader.GetCppIdentifier();
+			UhtType? Type = PropertySettings.Outer.FindType(UhtFindOptions.SourceName | FindOptions, Identifiers, TokenReader);
+			if (Type == null)
+			{
+				return null;
+			}
+			UhtTokenListCache.Return(Identifiers);
+
+			if (Type is UhtEnum Enum)
+			{
+				if (PropertySettings.PropertyCategory == UhtPropertyCategory.Member)
+				{
+					if (Enum.CppForm != UhtEnumCppForm.EnumClass)
+					{
+						TokenReader.LogError("You cannot use the raw enum name as a type for member variables, instead use TEnumAsByte or a C++11 enum class with an explicit underlying type.");
+					}
+				}
+
+				return new UhtEnumProperty(PropertySettings, Enum);
+			}
+			else if (Type is UhtScriptStruct ScriptStruct)
+			{
+				return new UhtStructProperty(PropertySettings, ScriptStruct);
+			}
+			else if (Type is UhtFunction Function)
+			{
+				if (!Function.FunctionFlags.HasAnyFlags(EFunctionFlags.MulticastDelegate))
+				{
+					return new UhtDelegateProperty(PropertySettings, Function);
+				}
+				else if (Function.FunctionType == UhtFunctionType.SparseDelegate)
+				{
+					return new UhtMulticastSparseDelegateProperty(PropertySettings, Function);
+				}
+				else
+				{
+					return new UhtMulticastInlineDelegateProperty(PropertySettings, Function);
+				}
+			}
+			else if (Type is UhtClass Class)
+			{
+				// Const after variable type but before pointer symbol
+				if (TokenReader.TryOptional("const"))
+				{
+					PropertySettings.MetaData.Add(UhtNames.NativeConst, "");
+				}
+				TokenReader.Require('*');
+
+				// Optionally emit messages about native pointer members and swallow trailing 'const' after pointer properties
+				UhtObjectPropertyBase.ConditionalLogPointerUsage(PropertySettings, UhtConfig.Instance.EngineNativePointerMemberBehavior,
+					UhtConfig.Instance.NonEngineNativePointerMemberBehavior, "Native pointer", TokenReader, TypeStartPos, "TObjectPtr");
+
+				if (PropertySettings.PropertyCategory == UhtPropertyCategory.Member)
+				{
+					TokenReader.TryOptional("const");
+				}
+
+				PropertySettings.PointerType = UhtPointerType.Native;
+
+				if (Class.ClassFlags.HasAnyFlags(EClassFlags.Interface))
+				{
+					return new UhtInterfaceProperty(PropertySettings, Class);
+				}
+				else if (Class.IsChildOf(Class.Session.UClass))
+				{
+					// UObject here specifies that there is no limiter
+					return new UhtClassProperty(PropertySettings, Class, Class.Session.UObject);
+				}
+				else
+				{
+					return new UhtObjectProperty(PropertySettings, Class);
+				}
+			}
+			else
+			{
+				throw new UhtIceException("Unexpected type found");
+			}
+		}
+	}
+}
diff --git a/Engine/Source/Programs/Shared/EpicGames.UHT/Parsers/UhtScriptStructParser.cs b/Engine/Source/Programs/Shared/EpicGames.UHT/Parsers/UhtScriptStructParser.cs
new file mode 100644
--- /dev/null
+++ b/Engine/Source/Programs/Shared/EpicGames.UHT/Parsers/UhtScriptStructParser.cs
@@ -0,0 +1,231 @@
+// Copyright Epic Games, Inc. All Rights Reserved.
+
+using EpicGames.Core;
+using EpicGames.UHT.Tables;
+using EpicGames.UHT.Tokenizer;
+using EpicGames.UHT.Types;
+using EpicGames.UHT.Utils;
+using System.Collections.Generic;
+
+namespace EpicGames.UHT.Parsers
+{
+	[UnrealHeaderTool]
+	public class UhtScriptStructParser : UhtScriptStruct
+	{
+		private static UhtKeywordTable KeywordTable = UhtKeywordTables.Instance.Get(UhtTableNames.ScriptStruct);
+		private static UhtSpecifierTable SpecifierTable = UhtSpecifierTables.Instance.Get(UhtTableNames.ScriptStruct);
+
+		public UhtToken SuperIdentifier;
+		public List<UhtToken[]>? BaseIdentifiers = null;
+
+		public UhtScriptStructParser(UhtType Outer, int LineNumber) : base(Outer, LineNumber)
+		{
+		}
+
+		protected override void ResolveSuper(UhtResolvePhase ResolvePhase)
+		{
+			base.ResolveSuper(ResolvePhase);
+
+			switch (ResolvePhase)
+			{
+				case UhtResolvePhase.Bases:
+					BindAndResolveSuper(ref this.SuperIdentifier, UhtFindOptions.ScriptStruct);
+
+					// if we have a base struct, propagate inherited struct flags now
+					UhtScriptStruct? SuperScriptStruct = this.SuperScriptStruct;
+					if (SuperScriptStruct != null)
+					{
+						this.ScriptStructFlags |= SuperScriptStruct.ScriptStructFlags & EStructFlags.Inherit;
+					}
+					break;
+			}
+		}
+
+		protected override bool ResolveSelf(UhtResolvePhase ResolvePhase)
+		{
+			bool bResult = base.ResolveSelf(ResolvePhase);
+
+			switch (ResolvePhase)
+			{
+				case UhtResolvePhase.Properties:
+					UhtPropertyParser.ResolveChildren(this, UhtPropertyParseOptions.AddModuleRelativePath);
+					break;
+			}
+			return bResult;
+		}
+
+		#region Keywords
+		[UhtKeyword(Extends = UhtTableNames.Global)]
+		private static UhtParseResult USTRUCTKeyword(UhtParsingScope TopScope, UhtParsingScope ActionScope, ref UhtToken Token)
+		{
+			return ParseUScriptStruct(TopScope, Token);
+		}
+
+		[UhtKeyword(Extends = UhtTableNames.ScriptStruct)]
+		[UhtKeyword(Extends = UhtTableNames.ScriptStruct, Keyword = "GENERATED_BODY")]
+		private static UhtParseResult GENERATED_USTRUCT_BODYKeyword(UhtParsingScope TopScope, UhtParsingScope ActionScope, ref UhtToken Token)
+		{
+			UhtScriptStruct ScriptStruct = (UhtScriptStruct)TopScope.ScopeType;
+
+			if (TopScope.AccessSpecifier != UhtAccessSpecifier.Public)
+			{
+				TopScope.TokenReader.LogError($"{Token.Value} must be in the public scope of '{ScriptStruct.SourceName}', not private or protected.");
+			}
+
+			if (ScriptStruct.MacroDeclaredLineNumber != -1)
+			{
+				TopScope.TokenReader.LogError($"Multiple {Token.Value} declarations found in '{ScriptStruct.SourceName}'");
+			}
+
+			ScriptStruct.MacroDeclaredLineNumber = TopScope.TokenReader.InputLine;
+
+			UhtParserHelpers.ParseCompileVersionDeclaration(TopScope.TokenReader, ScriptStruct);
+
+			TopScope.TokenReader.Optional(';');
+			return UhtParseResult.Handled;
+		}
+
+		[UhtKeyword(Extends = UhtTableNames.ScriptStruct)]
+		private static UhtParseResult RIGVM_METHODKeyword(UhtParsingScope TopScope, UhtParsingScope ActionScope, ref UhtToken Token)
+		{
+			ParseRigVM(TopScope);
+			return UhtParseResult.Handled;
+		}
+		#endregion
+
+		private static UhtParseResult ParseUScriptStruct(UhtParsingScope ParentScope, UhtToken KeywordToken)
+		{
+			UhtScriptStructParser ScriptStruct = new UhtScriptStructParser(ParentScope.ScopeType, KeywordToken.InputLine);
+			using (var TopScope = new UhtParsingScope(ParentScope, ScriptStruct, KeywordTable, UhtAccessSpecifier.Public))
+			{
+				const string ScopeName = "struct";
+
+				using (var TokenContext = new UhtMessageContext(TopScope.TokenReader, ScopeName))
+				{
+					// Parse the specifiers
+					UhtSpecifierContext SpecifierContext = new UhtSpecifierContext(TopScope, TopScope.TokenReader, ScriptStruct.MetaData);
+					UhtSpecifierParser Specifiers = TopScope.HeaderParser.GetSpecifierParser(SpecifierContext, ScopeName, SpecifierTable);
+					Specifiers.ParseSpecifiers();
+
+					// Consume the struct specifier
+					TopScope.TokenReader.Require("struct");
+
+					TopScope.TokenReader.SkipAlignasAndDeprecatedMacroIfNecessary();
+
+					// Read the struct name and possible API macro name
+					UhtToken APIMacroToken;
+					TopScope.TokenReader.TryOptionalAPIMacro(out APIMacroToken);
+					ScriptStruct.SourceName = TopScope.TokenReader.GetIdentifier().Value.ToString();
+
+					TopScope.AddModuleRelativePathToMetaData();
+
+					// Strip the name
+					if (ScriptStruct.SourceName[0] == 'T' || ScriptStruct.SourceName[0] == 'F')
+					{
+						ScriptStruct.EngineName = ScriptStruct.SourceName.Substring(1);
+					}
+					else
+					{
+						// This will be flagged later in the validation phase
+						ScriptStruct.EngineName = ScriptStruct.SourceName;
+					}
+
+					// Check for an empty engine name
+					if (ScriptStruct.EngineName.Length == 0)
+					{
+						TopScope.TokenReader.LogError($"When compiling struct definition for '{ScriptStruct.SourceName}', attempting to strip prefix results in an empty name. Did you leave off a prefix?");
+					}
+
+					// Parse the inheritance
+					UhtParserHelpers.ParseInheritance(TopScope.TokenReader, out ScriptStruct.SuperIdentifier, out ScriptStruct.BaseIdentifiers);
+
+					// Add the comments here for compatibility with old UHT
+					//COMPATIBILITY-TODO - Move this back to where the AddModuleRelativePathToMetaData is called.
+					TopScope.TokenReader.PeekToken();
+					TopScope.TokenReader.CommitPendingComments();
+					TopScope.AddFormattedCommentsAsTooltipMetaData();
+
+					// Initialize the structure flags
+					ScriptStruct.ScriptStructFlags |= EStructFlags.Native;
+
+					// Record that this struct is RequiredAPI if the CORE_API style macro was present
+					if (APIMacroToken)
+					{
+						ScriptStruct.ScriptStructFlags |= EStructFlags.RequiredAPI;
+					}
+
+					// Process the deferred specifiers
+					Specifiers.ParseDeferred();
+
+					if (ScriptStruct.Outer != null)
+					{
+						ScriptStruct.Outer.AddChild(ScriptStruct);
+					}
+
+					TopScope.HeaderParser.ParseStatements('{', '}', true);
+
+					TopScope.TokenReader.Require(';');
+
+					if (ScriptStruct.MacroDeclaredLineNumber == -1 && ScriptStruct.ScriptStructFlags.HasAnyFlags(EStructFlags.Native))
+					{
+						TopScope.TokenReader.LogError("Expected a GENERATED_BODY() at the start of the struct");
+					}
+				}
+				return UhtParseResult.Handled;
+			}
+		}
+
+		private static void ParseRigVM(UhtParsingScope TopScope)
+		{
+			UhtScriptStruct ScriptStruct = (UhtScriptStruct)TopScope.ScopeType;
+
+			using (var TokenContext = new UhtMessageContext(TopScope.TokenReader, "RIGVM_METHOD"))
+			{
+
+				// Create the RigVM information if it doesn't already exist
+				UhtRigVMStructInfo? StructInfo = ScriptStruct.RigVMStructInfo;
+				if (StructInfo == null)
+				{
+					StructInfo = new UhtRigVMStructInfo();
+					StructInfo.Name = ScriptStruct.EngineName;
+					ScriptStruct.RigVMStructInfo = StructInfo;
+				}
+
+				// Create a new method information and add it
+				UhtRigVMMethodInfo MethodInfo = new UhtRigVMMethodInfo();
+
+				// NOTE: The argument list reader doesn't handle templates with multiple arguments (i.e. the ',' issue)
+				TopScope.TokenReader
+					.RequireList('(', ')')
+					.Optional("virtual")
+					.RequireIdentifier((ref UhtToken Identifier) => MethodInfo.ReturnType = Identifier.Value.ToString())
+					.RequireIdentifier((ref UhtToken Identifier) => MethodInfo.Name = Identifier.Value.ToString())
+					.RequireList('(', ')', ',', false, (IEnumerable<UhtToken> Tokens) =>
+					{
+						StringViewBuilder Builder = new StringViewBuilder();
+						UhtToken LastToken = new UhtToken();
+						foreach (UhtToken Token in Tokens)
+						{
+							if (Token.IsSymbol('='))
+							{
+								break;
+							}
+							if (LastToken)
+							{
+								if (Builder.Length != 0)
+								{
+									Builder.Append(' ');
+								}
+								Builder.Append(LastToken.Value);
+							}
+							LastToken = Token;
+						}
+						MethodInfo.Parameters.Add(new UhtRigVMParameter(LastToken.Value.ToString(), Builder.ToString()));
+					})
+					.ConsumeUntil(';');
+
+				StructInfo.Methods.Add(MethodInfo);
+			}
+		}
+	}
+}
diff --git a/Engine/Source/Programs/Shared/EpicGames.UHT/Parsers/UhtSpecifierParser.cs b/Engine/Source/Programs/Shared/EpicGames.UHT/Parsers/UhtSpecifierParser.cs
new file mode 100644
--- /dev/null
+++ b/Engine/Source/Programs/Shared/EpicGames.UHT/Parsers/UhtSpecifierParser.cs
@@ -0,0 +1,417 @@
+// Copyright Epic Games, Inc. All Rights Reserved.
+
+using EpicGames.UHT.Utils;
+using System;
+using System.Collections.Generic;
+using EpicGames.UHT.Tables;
+using EpicGames.Core;
+using EpicGames.UHT.Tokenizer;
+
+namespace EpicGames.UHT.Parsers
+{
+
+	/// <summary>
+	/// Class responsible for parsing specifiers and the field meta data.  To reduce allocations, one specifier parser is shared between all objects in
+	/// a given header file.  This makes the Action pattern being used a bit more obtuse, but it does help performance by reducing the allocations fairly
+	/// significantly.
+	/// </summary>
+	public class UhtSpecifierParser : IUhtMessageExtraContext
+	{
+		struct DeferredSpecifier
+		{
+			public UhtSpecifier Specifier;
+			public object? Value;
+		}
+
+		private static List<KeyValuePair<StringView, StringView>> EmptyKVPValues = new List<KeyValuePair<StringView, StringView>>();
+
+		private UhtSpecifierContext SpecifierContext;
+		private IUhtTokenReader TokenReader;
+		private UhtSpecifierTable Table;
+		private StringView Context;
+		private StringView CurrentSpecifier = new StringView();
+		private List<DeferredSpecifier> DeferredSpecifiers = new List<DeferredSpecifier>();
+		private bool bIsParsingFieldMetaData = false;
+		private List<KeyValuePair<StringView, StringView>>? CurrentKVPValues = null;
+		private List<StringView>? CurrentStringValues = null;
+		private int UMetaElementsParsed = 0;
+
+		private readonly Action ParseAction;
+		private readonly Action ParseFieldMetaDataAction;
+		private readonly Action ParseKVPValueAction;
+		private readonly Action ParseStringViewListAction;
+
+		public UhtSpecifierParser(UhtSpecifierContext SpecifierContext, StringView Context, UhtSpecifierTable Table)
+		{
+			this.SpecifierContext = SpecifierContext;
+			this.TokenReader = SpecifierContext.Scope.TokenReader;
+			this.Context = Context;
+			this.Table = Table;
+
+			this.ParseAction = ParseInternal;
+			this.ParseFieldMetaDataAction = ParseFieldMetaDataInternal;
+			this.ParseKVPValueAction = () => 
+			{
+				if (this.CurrentKVPValues == null)
+				{
+					this.CurrentKVPValues = new List<KeyValuePair<StringView, StringView>>();
+				}
+				this.CurrentKVPValues.Add(ReadKVP()); 
+			};
+			this.ParseStringViewListAction = () => 
+			{ 
+				if (this.CurrentStringValues == null)
+				{
+					this.CurrentStringValues = new List<StringView>();
+				}
+				this.CurrentStringValues.Add(ReadValue()); 
+			};
+		}
+
+		public void Reset(UhtSpecifierContext SpecifierContext, StringView Context, UhtSpecifierTable Table)
+		{
+			this.SpecifierContext = SpecifierContext;
+			this.TokenReader = SpecifierContext.Scope.TokenReader;
+			this.Context = Context;
+			this.Table = Table;
+		}
+
+		public UhtSpecifierParser ParseSpecifiers()
+		{
+			this.bIsParsingFieldMetaData = false;
+			this.SpecifierContext.MetaData.LineNumber = TokenReader.InputLine;
+
+			using (var TokenContext = new UhtMessageContext(this.TokenReader, this))
+			{
+				this.TokenReader.RequireList('(', ')', ',', false, this.ParseAction);
+			}
+			return this;
+		}
+
+		public UhtSpecifierParser ParseFieldMetaData(StringView Context)
+		{
+			this.TokenReader = this.SpecifierContext.Scope.TokenReader;
+			this.bIsParsingFieldMetaData = true;
+
+			using (var TokenContext = new UhtMessageContext(this.TokenReader, this))
+			{
+				if (TokenReader.TryOptional("UMETA"))
+				{
+					this.UMetaElementsParsed = 0;
+					TokenReader.RequireList('(', ')', ',', false, this.ParseFieldMetaDataAction);
+					if (this.UMetaElementsParsed == 0)
+					{
+						this.TokenReader.LogError($"No metadata specified while parsing {UhtMessage.FormatContext(this)}");
+					}
+				}
+			}
+			return this;
+		}
+
+		public void ParseDeferred()
+		{
+			foreach (var Deferred in this.DeferredSpecifiers)
+			{
+				Dispatch(Deferred.Specifier, Deferred.Value);
+			}
+			this.DeferredSpecifiers.Clear();
+		}
+
+		#region IMessageExtraContext implementation
+		public IEnumerable<object?>? MessageExtraContext
+		{
+			get
+			{
+				Stack<object?> ExtraContext = new Stack<object?>(1);
+				string What = this.bIsParsingFieldMetaData ? "metadata" : "specifiers";
+				if (this.Context.Span.Length > 0)
+				{
+					ExtraContext.Push($"{this.Context} {What}");
+				}
+				else
+				{
+					ExtraContext.Push(What);
+				}
+				return ExtraContext;
+			}
+		}
+		#endregion
+
+		private void ParseInternal()
+		{
+			UhtToken Identifier = this.TokenReader.GetIdentifier();
+
+			this.CurrentSpecifier = Identifier.Value;
+			UhtSpecifier? Specifier;
+			if (this.Table.TryGetValue(CurrentSpecifier, out Specifier))
+			{
+				if (TryParseValue(Specifier.ValueType, out object? Value))
+				{
+					if (Specifier.When == UhtSpecifierWhen.Deferred)
+					{
+						if (this.DeferredSpecifiers == null)
+						{
+							this.DeferredSpecifiers = new List<DeferredSpecifier>();
+						}
+						this.DeferredSpecifiers.Add(new DeferredSpecifier { Specifier = Specifier, Value = Value });
+					}
+					else
+					{
+						Dispatch(Specifier, Value);
+					}
+				}
+			}
+			else
+			{
+				this.TokenReader.LogError($"Unknown specifier '{CurrentSpecifier}' found while parsing {UhtMessage.FormatContext(this)}");
+			}
+		}
+
+		private void ParseFieldMetaDataInternal()
+		{
+			UhtToken Key;
+			if (!TokenReader.TryOptionalIdentifier(out Key))
+			{
+				throw new UhtException(this.TokenReader, $"UMETA expects a key and optional value", this);
+			}
+
+			StringViewBuilder SVB = new StringViewBuilder();
+			if (TokenReader.TryOptional('='))
+			{
+				if (!ReadValue(TokenReader, SVB, true))
+				{
+					throw new UhtException(this.TokenReader, $"UMETA key '{Key.Value}' expects a value", this);
+				}
+			}
+			
+			++this.UMetaElementsParsed;
+			this.SpecifierContext.MetaData.Add(Key.Value.ToString(), this.SpecifierContext.MetaNameIndex, SVB.ToString());
+		}
+
+		private void Dispatch(UhtSpecifier Specifier, object? Value)
+		{
+			UhtSpecifierDispatchResults Results = Specifier.Dispatch(this.SpecifierContext, Value);
+			if (Results == UhtSpecifierDispatchResults.Unknown)
+			{
+				this.TokenReader.LogError($"Unknown specifier '{Specifier.Name}' found while parsing {UhtMessage.FormatContext(this)}");
+			}
+		}
+
+		private bool TryParseValue(UhtSpecifierValueType ValueType, out object? Value)
+		{
+			Value = null;
+
+			switch (ValueType)
+			{
+				case UhtSpecifierValueType.NotSet:
+					throw new UhtIceException("NotSet is an invalid value for value types");
+
+				case UhtSpecifierValueType.None:
+					if (this.TokenReader.TryOptional('='))
+					{
+						ReadValue(); // consume the value;
+						this.TokenReader.LogError($"The specifier '{this.CurrentSpecifier}' found a value when none was expected", this);
+						return false;
+					}
+					return true;
+
+				case UhtSpecifierValueType.String:
+					if (!this.TokenReader.TryOptional('='))
+					{
+						this.TokenReader.LogError($"The specifier '{this.CurrentSpecifier}' expects a value", this);
+						return false;
+					}
+					Value = ReadValue();
+					return true;
+
+				case UhtSpecifierValueType.OptionalString:
+					{
+						List<StringView>? StringList = ReadValueList();
+						if (StringList != null && StringList.Count > 0)
+						{
+							Value = StringList[0];
+						}
+						return true;
+					}
+
+				case UhtSpecifierValueType.SingleString:
+					{
+						List<StringView>? StringList = ReadValueList();
+						if (StringList == null || StringList.Count != 1)
+						{
+							this.TokenReader.LogError($"The specifier '{this.CurrentSpecifier}' expects a single value", this);
+							return false;
+						}
+						Value = StringList[0];
+						return true;
+					}
+
+				case UhtSpecifierValueType.StringList:
+					Value = ReadValueList();
+					return true;
+
+				case UhtSpecifierValueType.Legacy:
+					Value = ReadValueList();
+					return true;
+
+				case UhtSpecifierValueType.NonEmptyStringList:
+					{
+						List<StringView>? StringList = ReadValueList();
+						if (StringList == null || StringList.Count == 0)
+						{
+							this.TokenReader.LogError($"The specifier '{this.CurrentSpecifier}' expects at list one value", this);
+							return false;
+						}
+						Value = StringList;
+						return true;
+					}
+
+				case UhtSpecifierValueType.KeyValuePairList:
+					{
+						this.CurrentKVPValues = null;
+						this.TokenReader
+							.Require('=')
+							.RequireList('(', ')', ',', false, this.ParseKVPValueAction);
+						List<KeyValuePair<StringView, StringView>> Return = this.CurrentKVPValues ?? EmptyKVPValues;
+						this.CurrentKVPValues = null;
+						Value = Return;
+						return true;
+					}
+
+				case UhtSpecifierValueType.OptionalEqualsKeyValuePairList:
+					{
+						this.CurrentKVPValues = null;
+						// This parser isn't as strict as the other parsers...
+						if (this.TokenReader.TryOptional('='))
+						{
+							if (!this.TokenReader.TryOptionalList('(', ')', ',', false, this.ParseKVPValueAction))
+							{
+								this.ParseKVPValueAction();
+							}
+						}
+						else
+						{
+							this.TokenReader.TryOptionalList('(', ')', ',', false, this.ParseKVPValueAction);
+						}
+						List<KeyValuePair<StringView, StringView>> Return = this.CurrentKVPValues ?? EmptyKVPValues;
+						this.CurrentKVPValues = null;
+						Value = Return;
+						return true;
+					}
+
+				default:
+					throw new UhtIceException("Unknown value type");
+			}
+		}
+
+		private KeyValuePair<StringView, StringView> ReadKVP()
+		{
+			UhtToken Key;
+			if (!this.TokenReader.TryOptionalIdentifier(out Key))
+			{
+				throw new UhtException(this.TokenReader, $"The specifier '{this.CurrentSpecifier}' expects a key and optional value", this);
+			}
+
+			StringView Value = "";
+			if (this.TokenReader.TryOptional('='))
+			{
+				Value = ReadValue();
+			}
+			return new KeyValuePair<StringView, StringView>(Key.Value, Value);
+		}
+
+		private List<StringView>? ReadValueList()
+		{
+			this.CurrentStringValues = null;
+
+			// This parser isn't as strict as the other parsers...
+			if (this.TokenReader.TryOptional('='))
+			{
+				if (!this.TokenReader.TryOptionalList('(', ')', ',', false, this.ParseStringViewListAction))
+				{
+					this.ParseStringViewListAction();
+				}
+			}
+			else
+			{
+				this.TokenReader.TryOptionalList('(', ')', ',', false, this.ParseStringViewListAction);
+			}
+			List<StringView>? StringValues = this.CurrentStringValues;
+			this.CurrentStringValues = null;
+			return StringValues;
+		}
+
+		private StringView ReadValue()
+		{
+			StringViewBuilder SVB = new StringViewBuilder();
+			if (!ReadValue(this.TokenReader, SVB, false))
+			{
+				throw new UhtException(this.TokenReader, $"The specifier '{this.CurrentSpecifier}' expects a value", this);
+			}
+			return SVB.ToStringView();
+		}
+
+
+		/// <summary>
+		/// Parse the sequence of meta data
+		/// </summary>
+		/// <param name="TokenReader">Input token reader</param>
+		/// <param name="SVB">Output string builder</param>
+		/// <param name="bRespectQuotes">If true, do not convert \" to " in string constants.  This is required for UMETA data</param>
+		/// <returns>True if data was read</returns>
+		private static bool ReadValue(IUhtTokenReader TokenReader, StringViewBuilder SVB, bool bRespectQuotes)
+		{
+			UhtToken Token = TokenReader.GetToken();
+			switch (Token.TokenType)
+			{
+				case UhtTokenType.EndOfFile:
+				case UhtTokenType.EndOfDefault:
+				case UhtTokenType.EndOfType:
+				case UhtTokenType.EndOfDeclaration:
+					return false;
+
+				case UhtTokenType.Identifier:
+					// We handle true/false differently for compatibility with old UHT
+					if (Token.IsValue("true", true))
+					{
+						SVB.Append("TRUE");
+					}
+					else if (Token.IsValue("false", true))
+					{
+						SVB.Append("FALSE");
+					}
+					else
+					{
+						SVB.Append(Token.Value);
+					}
+					if (TokenReader.TryOptional('='))
+					{
+						SVB.Append('=');
+						if (!ReadValue(TokenReader, SVB, bRespectQuotes))
+						{
+							return false;
+						}
+					}
+					break;
+
+				case UhtTokenType.Symbol:
+					SVB.Append(Token.Value);
+					if (TokenReader.TryOptional('='))
+					{
+						SVB.Append('=');
+						if (!ReadValue(TokenReader, SVB, bRespectQuotes))
+						{
+							return false;
+						}
+					}
+					break;
+
+				default:
+					SVB.Append(Token.GetConstantValue(bRespectQuotes));
+					break;
+			}
+
+			return true;
+		}
+	}
+}
diff --git a/Engine/Source/Programs/Shared/EpicGames.UHT/Specifiers/UhtClassSpecifiers.cs b/Engine/Source/Programs/Shared/EpicGames.UHT/Specifiers/UhtClassSpecifiers.cs
new file mode 100644
--- /dev/null
+++ b/Engine/Source/Programs/Shared/EpicGames.UHT/Specifiers/UhtClassSpecifiers.cs
@@ -0,0 +1,341 @@
+// Copyright Epic Games, Inc. All Rights Reserved.
+
+using EpicGames.Core;
+using EpicGames.UHT.Tables;
+using EpicGames.UHT.Types;
+using EpicGames.UHT.Utils;
+using System.Collections.Generic;
+
+namespace EpicGames.UHT.Parsers
+{
+	[UnrealHeaderTool]
+	public static class UhtClassSpecifiers
+	{
+		[UhtSpecifier(Extends = UhtTableNames.Class, ValueType = UhtSpecifierValueType.Legacy)]
+		private static void NoExportSpecifier(UhtSpecifierContext SpecifierContext)
+		{
+			UhtClassParser Class = (UhtClassParser)SpecifierContext.Scope.ScopeType;
+			Class.AddClassFlags(EClassFlags.NoExport);
+		}
+
+		[UhtSpecifier(Extends = UhtTableNames.Class, ValueType = UhtSpecifierValueType.Legacy)]
+		private static void IntrinsicSpecifier(UhtSpecifierContext SpecifierContext)
+		{
+			UhtClassParser Class = (UhtClassParser)SpecifierContext.Scope.ScopeType;
+			Class.AddClassFlags(EClassFlags.Intrinsic);
+		}
+
+		[UhtSpecifier(Extends = UhtTableNames.Class, ValueType = UhtSpecifierValueType.Legacy)]
+		private static void ComponentWrapperClassSpecifier(UhtSpecifierContext SpecifierContext)
+		{
+			SpecifierContext.MetaData.Add(UhtNames.IgnoreCategoryKeywordsInSubclasses, true);
+		}
+
+		[UhtSpecifier(Extends = UhtTableNames.Class, ValueType = UhtSpecifierValueType.SingleString)]
+		private static void WithinSpecifier(UhtSpecifierContext SpecifierContext, StringView Value)
+		{
+			UhtClassParser Class = (UhtClassParser)SpecifierContext.Scope.ScopeType;
+			Class.ClassWithinIdentifier = Value.ToString();
+		}
+
+		[UhtSpecifier(Extends = UhtTableNames.Class, ValueType = UhtSpecifierValueType.Legacy)]
+		private static void EditInlineNewSpecifier(UhtSpecifierContext SpecifierContext)
+		{
+			UhtClassParser Class = (UhtClassParser)SpecifierContext.Scope.ScopeType;
+			Class.AddClassFlags(EClassFlags.EditInlineNew);
+		}
+
+		[UhtSpecifier(Extends = UhtTableNames.Class, ValueType = UhtSpecifierValueType.Legacy)]
+		private static void NotEditInlineNewSpecifier(UhtSpecifierContext SpecifierContext)
+		{
+			UhtClassParser Class = (UhtClassParser)SpecifierContext.Scope.ScopeType;
+			Class.RemoveClassFlags(EClassFlags.EditInlineNew);
+		}
+
+		[UhtSpecifier(Extends = UhtTableNames.Class, ValueType = UhtSpecifierValueType.Legacy)]
+		private static void PlaceableSpecifier(UhtSpecifierContext SpecifierContext)
+		{
+			UhtClassParser Class = (UhtClassParser)SpecifierContext.Scope.ScopeType;
+			Class.RemoveClassFlags(EClassFlags.NotPlaceable);
+		}
+
+		[UhtSpecifier(Extends = UhtTableNames.Class, ValueType = UhtSpecifierValueType.Legacy)]
+		private static void NotPlaceableSpecifier(UhtSpecifierContext SpecifierContext)
+		{
+			UhtClassParser Class = (UhtClassParser)SpecifierContext.Scope.ScopeType;
+			Class.AddClassFlags(EClassFlags.NotPlaceable);
+		}
+
+		[UhtSpecifier(Extends = UhtTableNames.Class, ValueType = UhtSpecifierValueType.Legacy)]
+		private static void DefaultToInstancedSpecifier(UhtSpecifierContext SpecifierContext)
+		{
+			UhtClassParser Class = (UhtClassParser)SpecifierContext.Scope.ScopeType;
+			Class.AddClassFlags(EClassFlags.DefaultToInstanced);
+		}
+
+		[UhtSpecifier(Extends = UhtTableNames.Class, ValueType = UhtSpecifierValueType.Legacy)]
+		private static void HideDropdownSpecifier(UhtSpecifierContext SpecifierContext)
+		{
+			UhtClassParser Class = (UhtClassParser)SpecifierContext.Scope.ScopeType;
+			Class.AddClassFlags(EClassFlags.HideDropDown);
+		}
+
+		[UhtSpecifier(Extends = UhtTableNames.Class, ValueType = UhtSpecifierValueType.Legacy)]
+		private static void HiddenSpecifier(UhtSpecifierContext SpecifierContext)
+		{
+			// Prevents class from appearing in the editor class browser and edit inline menus.
+			UhtClassParser Class = (UhtClassParser)SpecifierContext.Scope.ScopeType;
+			Class.AddClassFlags(EClassFlags.Hidden);
+		}
+
+		[UhtSpecifier(Extends = UhtTableNames.Class, ValueType = UhtSpecifierValueType.Legacy)]
+		private static void DependsOnSpecifier(UhtSpecifierContext SpecifierContext)
+		{
+			SpecifierContext.MessageSite.LogError("The dependsOn specifier is deprecated. Please use #include \"ClassHeaderFilename.h\" instead.");
+		}
+
+		[UhtSpecifier(Extends = UhtTableNames.Class, ValueType = UhtSpecifierValueType.Legacy)]
+		private static void MinimalAPISpecifier(UhtSpecifierContext SpecifierContext)
+		{
+			UhtClassParser Class = (UhtClassParser)SpecifierContext.Scope.ScopeType;
+			Class.AddClassFlags(EClassFlags.MinimalAPI);
+		}
+
+		[UhtSpecifier(Extends = UhtTableNames.Class, ValueType = UhtSpecifierValueType.Legacy)]
+		private static void ConstSpecifier(UhtSpecifierContext SpecifierContext)
+		{
+			UhtClassParser Class = (UhtClassParser)SpecifierContext.Scope.ScopeType;
+			Class.AddClassFlags(EClassFlags.Const);
+		}
+
+		[UhtSpecifier(Extends = UhtTableNames.Class, ValueType = UhtSpecifierValueType.Legacy)]
+		private static void PerObjectConfigSpecifier(UhtSpecifierContext SpecifierContext)
+		{
+			UhtClassParser Class = (UhtClassParser)SpecifierContext.Scope.ScopeType;
+			Class.AddClassFlags(EClassFlags.PerObjectConfig);
+		}
+
+		[UhtSpecifier(Extends = UhtTableNames.Class, ValueType = UhtSpecifierValueType.Legacy)]
+		private static void ConfigDoNotCheckDefaultsSpecifier(UhtSpecifierContext SpecifierContext)
+		{
+			UhtClassParser Class = (UhtClassParser)SpecifierContext.Scope.ScopeType;
+			Class.AddClassFlags(EClassFlags.ConfigDoNotCheckDefaults);
+		}
+
+		[UhtSpecifier(Extends = UhtTableNames.Class, ValueType = UhtSpecifierValueType.Legacy)]
+		private static void AbstractSpecifier(UhtSpecifierContext SpecifierContext)
+		{
+			UhtClassParser Class = (UhtClassParser)SpecifierContext.Scope.ScopeType;
+			Class.AddClassFlags(EClassFlags.Abstract);
+		}
+
+		[UhtSpecifier(Extends = UhtTableNames.Class, ValueType = UhtSpecifierValueType.Legacy)]
+		private static void DeprecatedSpecifier(UhtSpecifierContext SpecifierContext)
+		{
+			UhtClassParser Class = (UhtClassParser)SpecifierContext.Scope.ScopeType;
+			Class.AddClassFlags(EClassFlags.Deprecated | EClassFlags.NotPlaceable);
+		}
+
+		[UhtSpecifier(Extends = UhtTableNames.Class, ValueType = UhtSpecifierValueType.Legacy)]
+		private static void TransientSpecifier(UhtSpecifierContext SpecifierContext)
+		{
+			UhtClassParser Class = (UhtClassParser)SpecifierContext.Scope.ScopeType;
+			Class.AddClassFlags(EClassFlags.Transient);
+		}
+
+		[UhtSpecifier(Extends = UhtTableNames.Class, ValueType = UhtSpecifierValueType.Legacy)]
+		private static void NonTransientSpecifier(UhtSpecifierContext SpecifierContext)
+		{
+			UhtClassParser Class = (UhtClassParser)SpecifierContext.Scope.ScopeType;
+			Class.RemoveClassFlags(EClassFlags.Transient);
+		}
+
+		[UhtSpecifier(Extends = UhtTableNames.Class, ValueType = UhtSpecifierValueType.Legacy)]
+		private static void OptionalSpecifier(UhtSpecifierContext SpecifierContext)
+		{
+			// Optional class
+			UhtClassParser Class = (UhtClassParser)SpecifierContext.Scope.ScopeType;
+			Class.AddClassFlags(EClassFlags.Optional);
+		}
+
+		[UhtSpecifier(Extends = UhtTableNames.Class, ValueType = UhtSpecifierValueType.Legacy)]
+		private static void CustomConstructorSpecifier(UhtSpecifierContext SpecifierContext)
+		{
+			// we will not export a constructor for this class, assuming it is in the CPP block
+			UhtClass Class = (UhtClass)SpecifierContext.Scope.ScopeType;
+			Class.ClassExportFlags |= UhtClassExportFlags.HasCustomConstructor;
+		}
+
+		[UhtSpecifier(Extends = UhtTableNames.Class, ValueType = UhtSpecifierValueType.SingleString)]
+		private static void ConfigSpecifier(UhtSpecifierContext SpecifierContext, StringView Value)
+		{
+			UhtClass Class = (UhtClass)SpecifierContext.Scope.ScopeType;
+			Class.Config = Value.ToString();
+		}
+
+		[UhtSpecifier(Extends = UhtTableNames.Class, ValueType = UhtSpecifierValueType.Legacy)]
+		private static void DefaultConfigSpecifier(UhtSpecifierContext SpecifierContext)
+		{
+			// Save object config only to Default INIs, never to local INIs.
+			UhtClassParser Class = (UhtClassParser)SpecifierContext.Scope.ScopeType;
+			Class.AddClassFlags(EClassFlags.DefaultConfig);
+		}
+
+		[UhtSpecifier(Extends = UhtTableNames.Class, ValueType = UhtSpecifierValueType.Legacy)]
+		private static void GlobalUserConfigSpecifier(UhtSpecifierContext SpecifierContext)
+		{
+			// Save object config only to global user overrides, never to local INIs
+			UhtClassParser Class = (UhtClassParser)SpecifierContext.Scope.ScopeType;
+			Class.AddClassFlags(EClassFlags.GlobalUserConfig);
+		}
+
+		[UhtSpecifier(Extends = UhtTableNames.Class, ValueType = UhtSpecifierValueType.Legacy)]
+		private static void ProjectUserConfigSpecifier(UhtSpecifierContext SpecifierContext)
+		{
+			// Save object config only to project user overrides, never to INIs that are checked in
+			UhtClassParser Class = (UhtClassParser)SpecifierContext.Scope.ScopeType;
+			Class.AddClassFlags(EClassFlags.ProjectUserConfig);
+		}
+
+		[UhtSpecifier(Extends = UhtTableNames.Class, ValueType = UhtSpecifierValueType.SingleString)]
+		private static void EditorConfigSpecifier(UhtSpecifierContext SpecifierContext, StringView Value)
+		{
+			SpecifierContext.MetaData.Add(UhtNames.EditorConfig, Value.ToString());
+		}
+
+		[UhtSpecifier(Extends = UhtTableNames.Class, ValueType = UhtSpecifierValueType.NonEmptyStringList)]
+		private static void ShowCategoriesSpecifier(UhtSpecifierContext SpecifierContext, List<StringView> Value)
+		{
+			UhtClassParser Class = (UhtClassParser)SpecifierContext.Scope.ScopeType;
+			Class.ShowCategories.AddUniqueRange(Value);
+		}
+
+		[UhtSpecifier(Extends = UhtTableNames.Class, ValueType = UhtSpecifierValueType.NonEmptyStringList)]
+		private static void HideCategoriesSpecifier(UhtSpecifierContext SpecifierContext, List<StringView> Value)
+		{
+			UhtClassParser Class = (UhtClassParser)SpecifierContext.Scope.ScopeType;
+			Class.HideCategories.AddUniqueRange(Value);
+		}
+
+		[UhtSpecifier(Extends = UhtTableNames.Class, ValueType = UhtSpecifierValueType.NonEmptyStringList)]
+		private static void ShowFunctionsSpecifier(UhtSpecifierContext SpecifierContext, List<StringView> Value)
+		{
+			UhtClassParser Class = (UhtClassParser)SpecifierContext.Scope.ScopeType;
+			Class.ShowFunctions.AddUniqueRange(Value);
+		}
+
+		[UhtSpecifier(Extends = UhtTableNames.Class, ValueType = UhtSpecifierValueType.NonEmptyStringList)]
+		private static void HideFunctionsSpecifier(UhtSpecifierContext SpecifierContext, List<StringView> Value)
+		{
+			UhtClassParser Class = (UhtClassParser)SpecifierContext.Scope.ScopeType;
+			Class.HideFunctions.AddUniqueRange(Value);
+		}
+
+		[UhtSpecifier(Extends = UhtTableNames.Class, ValueType = UhtSpecifierValueType.SingleString)]
+		private static void SparseClassDataTypesSpecifier(UhtSpecifierContext SpecifierContext, StringView Value)
+		{
+			UhtClassParser Class = (UhtClassParser)SpecifierContext.Scope.ScopeType;
+			Class.SparseClassDataTypes.AddUnique(Value.ToString());
+		}
+
+		[UhtSpecifier(Extends = UhtTableNames.Class, ValueType = UhtSpecifierValueType.NonEmptyStringList)]
+		private static void ClassGroupSpecifier(UhtSpecifierContext SpecifierContext, List<StringView> Value)
+		{
+			UhtClassParser Class = (UhtClassParser)SpecifierContext.Scope.ScopeType;
+			foreach (StringView Eleemnt in Value)
+			{
+				Class.ClassGroupNames.Add(Eleemnt.ToString());
+			}
+		}
+
+		[UhtSpecifier(Extends = UhtTableNames.Class, ValueType = UhtSpecifierValueType.NonEmptyStringList)]
+		private static void AutoExpandCategoriesSpecifier(UhtSpecifierContext SpecifierContext, List<StringView> Value)
+		{
+			UhtClassParser Class = (UhtClassParser)SpecifierContext.Scope.ScopeType;
+			Class.AutoExpandCategories.AddUniqueRange(Value);
+			Class.AutoCollapseCategories.RemoveSwapRange(Value);
+		}
+
+		[UhtSpecifier(Extends = UhtTableNames.Class, ValueType = UhtSpecifierValueType.NonEmptyStringList)]
+		private static void AutoCollapseCategoriesSpecifier(UhtSpecifierContext SpecifierContext, List<StringView> Value)
+		{
+			UhtClassParser Class = (UhtClassParser)SpecifierContext.Scope.ScopeType;
+			Class.AutoCollapseCategories.AddUniqueRange(Value);
+			Class.AutoExpandCategories.RemoveSwapRange(Value);
+		}
+
+		[UhtSpecifier(Extends = UhtTableNames.Class, ValueType = UhtSpecifierValueType.NonEmptyStringList)]
+		private static void PrioritizeCategoriesSpecifier(UhtSpecifierContext SpecifierContext, List<StringView> Value)
+		{
+			UhtClassParser Class = (UhtClassParser)SpecifierContext.Scope.ScopeType;
+			Class.PrioritizeCategories.AddUniqueRange(Value);
+		}
+
+		[UhtSpecifier(Extends = UhtTableNames.Class, ValueType = UhtSpecifierValueType.NonEmptyStringList)]
+		private static void DontAutoCollapseCategoriesSpecifier(UhtSpecifierContext SpecifierContext, List<StringView> Value)
+		{
+			UhtClassParser Class = (UhtClassParser)SpecifierContext.Scope.ScopeType;
+			Class.AutoCollapseCategories.RemoveSwapRange(Value);
+		}
+
+		[UhtSpecifier(Extends = UhtTableNames.Class, ValueType = UhtSpecifierValueType.Legacy)]
+		private static void CollapseCategoriesSpecifier(UhtSpecifierContext SpecifierContext)
+		{
+			// Class' properties should not be shown categorized in the editor.
+			UhtClassParser Class = (UhtClassParser)SpecifierContext.Scope.ScopeType;
+			Class.AddClassFlags(EClassFlags.CollapseCategories);
+		}
+
+		[UhtSpecifier(Extends = UhtTableNames.Class, ValueType = UhtSpecifierValueType.Legacy)]
+		private static void DontCollapseCategoriesSpecifier(UhtSpecifierContext SpecifierContext)
+		{
+			// Class' properties should be shown categorized in the editor.
+			UhtClassParser Class = (UhtClassParser)SpecifierContext.Scope.ScopeType;
+			Class.RemoveClassFlags(EClassFlags.CollapseCategories);
+		}
+
+		[UhtSpecifier(Extends = UhtTableNames.Class, ValueType = UhtSpecifierValueType.Legacy)]
+		private static void AdvancedClassDisplaySpecifier(UhtSpecifierContext SpecifierContext)
+		{
+			// By default the class properties are shown in advanced sections in UI
+			UhtClass Class = (UhtClass)SpecifierContext.Scope.ScopeType;
+			Class.MetaData.Add(UhtNames.AdvancedClassDisplay, "true");
+		}
+
+		[UhtSpecifier(Extends = UhtTableNames.Class, ValueType = UhtSpecifierValueType.Legacy)]
+		private static void ConversionRootSpecifier(UhtSpecifierContext SpecifierContext)
+		{
+			UhtClass Class = (UhtClass)SpecifierContext.Scope.ScopeType;
+			Class.MetaData.Add(UhtNames.IsConversionRoot, "true");
+		}
+
+		[UhtSpecifier(Extends = UhtTableNames.Class, ValueType = UhtSpecifierValueType.Legacy)]
+		private static void NeedsDeferredDependencyLoadingSpecifier(UhtSpecifierContext SpecifierContext)
+		{
+			UhtClassParser Class = (UhtClassParser)SpecifierContext.Scope.ScopeType;
+			Class.AddClassFlags(EClassFlags.NeedsDeferredDependencyLoading);
+		}
+
+		[UhtSpecifier(Extends = UhtTableNames.Class, ValueType = UhtSpecifierValueType.Legacy)]
+		private static void MatchedSerializersSpecifier(UhtSpecifierContext SpecifierContext)
+		{
+			if (!SpecifierContext.Scope.HeaderParser.HeaderFile.bIsNoExportTypes)
+			{
+				SpecifierContext.MessageSite.LogError("The 'MatchedSerializers' class specifier is only valid in the NoExportTypes.h file");
+			}
+			UhtClassParser Class = (UhtClassParser)SpecifierContext.Scope.ScopeType;
+			Class.AddClassFlags(EClassFlags.MatchedSerializers);
+		}
+
+		[UhtSpecifier(Extends = UhtTableNames.Class, ValueType = UhtSpecifierValueType.Legacy)]
+		private static void InterfaceSpecifier(UhtSpecifierContext SpecifierContext)
+		{
+			if (!SpecifierContext.Scope.HeaderParser.HeaderFile.bIsNoExportTypes)
+			{
+				SpecifierContext.MessageSite.LogError("The 'Interface' class specifier is only valid in the NoExportTypes.h file");
+			}
+			UhtClassParser Class = (UhtClassParser)SpecifierContext.Scope.ScopeType;
+			Class.AddClassFlags(EClassFlags.Interface);
+		}
+	}
+}
