UHT in UBT initial version.  Currently disabled.

diff --git a/Engine/Source/Programs/Shared/EpicGames.UHT/Specifiers/UhtDefaultSpecifiers.cs b/Engine/Source/Programs/Shared/EpicGames.UHT/Specifiers/UhtDefaultSpecifiers.cs
new file mode 100644
--- /dev/null
+++ b/Engine/Source/Programs/Shared/EpicGames.UHT/Specifiers/UhtDefaultSpecifiers.cs
@@ -0,0 +1,234 @@
+// Copyright Epic Games, Inc. All Rights Reserved.
+
+using EpicGames.Core;
+using EpicGames.UHT.Tables;
+using EpicGames.UHT.Types;
+using EpicGames.UHT.Utils;
+using System;
+using System.Collections.Generic;
+
+namespace EpicGames.UHT.Parsers
+{
+	[UnrealHeaderTool]
+	public static class UhtDefaultSpecifiers
+	{
+		private static void SetMetaData(UhtSpecifierContext SpecifierContext, StringView Key, StringView Value)
+		{
+			if (Key.Length == 0)
+			{
+				SpecifierContext.MessageSite.LogError("Invalid metadata, name can not be blank");
+				return;
+			}
+
+			// Trim the leading and ending whitespace
+			ReadOnlySpan<char> Span = Value.Span;
+			int Start = 0;
+			int End = Value.Length;
+			for (; Start < End; ++Start)
+			{
+				if (!UhtFCString.IsWhitespace(Span[Start]))
+				{
+					break;
+				}
+			}
+			for (; Start < End; --End)
+			{
+				if (!UhtFCString.IsWhitespace(Span[End - 1]))
+				{
+					break;
+				}
+			}
+
+			// Trim any quotes 
+			//COMPATIBILITY-TODO - This doesn't handle strings that end in an escaped ".  This is a bug in old UHT.
+			// Only remove quotes if we have quotes at both ends and the last isn't escaped.
+			if (Start < End && Span[Start] == '"')
+			{
+				++Start;
+			}
+			if (Start < End && Span[End - 1] == '"')
+			{
+				--End;
+			}
+
+			// Get the trimmed string
+			Value = new StringView(Value, Start, End - Start);
+
+			// Make sure this isn't a duplicate assignment
+			string KeyAsString = Key.ToString();
+			if (SpecifierContext.MetaData.TryGetValue(KeyAsString, out string? ExistingValue))
+			{
+				if (StringViewComparer.OrdinalIgnoreCase.Compare(ExistingValue, Value) != 0)
+				{
+					SpecifierContext.MessageSite.LogError($"Metadata key '{Key}' first seen with value '{ExistingValue}' then '{Value}'");
+				}
+			}
+			else
+			{
+				SpecifierContext.MetaData.Add(KeyAsString, SpecifierContext.MetaNameIndex, Value.ToString());
+			}
+		}
+
+		private static void SetMetaData(UhtSpecifierContext SpecifierContext, StringView Key, bool Value)
+		{
+			SetMetaData(SpecifierContext, Key, Value ? "true" : "false");
+		}
+
+		#region Specifiers
+		[UhtSpecifier(Extends = UhtTableNames.Default, ValueType = UhtSpecifierValueType.KeyValuePairList, When = UhtSpecifierWhen.Immediate)]
+		private static void MetaSpecifier(UhtSpecifierContext SpecifierContext, List<KeyValuePair<StringView, StringView>> Value)
+		{
+			foreach (KeyValuePair<StringView, StringView> KVP in (List<KeyValuePair<StringView, StringView>>)Value)
+			{
+				SetMetaData(SpecifierContext, KVP.Key, KVP.Value);
+			}
+		}
+
+		[UhtSpecifier(Extends = UhtTableNames.Default, ValueType = UhtSpecifierValueType.String, When = UhtSpecifierWhen.Immediate)]
+		private static void DisplayNameSpecifier(UhtSpecifierContext SpecifierContext, StringView Value)
+		{
+			SetMetaData(SpecifierContext, "DisplayName", Value);
+		}
+
+		[UhtSpecifier(Extends = UhtTableNames.Default, ValueType = UhtSpecifierValueType.String, When = UhtSpecifierWhen.Immediate)]
+		private static void FriendlyNameSpecifier(UhtSpecifierContext SpecifierContext, StringView Value)
+		{
+			SetMetaData(SpecifierContext, "FriendlyName", Value);
+		}
+
+		[UhtSpecifier(Extends = UhtTableNames.Default, ValueType = UhtSpecifierValueType.None, When = UhtSpecifierWhen.Immediate)]
+		private static void BlueprintInternalUseOnlySpecifier(UhtSpecifierContext SpecifierContext)
+		{
+			SetMetaData(SpecifierContext, "BlueprintInternalUseOnly", true);
+			SetMetaData(SpecifierContext, "BlueprintType", true);
+		}
+
+		[UhtSpecifier(Extends = UhtTableNames.Default, ValueType = UhtSpecifierValueType.None, When = UhtSpecifierWhen.Immediate)]
+		private static void BlueprintInternalUseOnlyHierarchicalSpecifier(UhtSpecifierContext SpecifierContext)
+		{
+			SetMetaData(SpecifierContext, "BlueprintInternalUseOnlyHierarchical", true);
+			SetMetaData(SpecifierContext, "BlueprintInternalUseOnly", true);
+			SetMetaData(SpecifierContext, "BlueprintType", true);
+		}
+
+		[UhtSpecifier(Extends = UhtTableNames.Default, ValueType = UhtSpecifierValueType.None, When = UhtSpecifierWhen.Immediate)]
+		private static void BlueprintTypeSpecifier(UhtSpecifierContext SpecifierContext)
+		{
+			SetMetaData(SpecifierContext, "BlueprintType", true);
+		}
+
+		[UhtSpecifier(Extends = UhtTableNames.Default, ValueType = UhtSpecifierValueType.None, When = UhtSpecifierWhen.Immediate)]
+		private static void NotBlueprintTypeSpecifier(UhtSpecifierContext SpecifierContext)
+		{
+			SetMetaData(SpecifierContext, "NotBlueprintType", true);
+			SpecifierContext.MetaData.Remove("BlueprintType", SpecifierContext.MetaNameIndex);
+		}
+
+		[UhtSpecifier(Extends = UhtTableNames.Default, ValueType = UhtSpecifierValueType.None, When = UhtSpecifierWhen.Immediate)]
+		private static void BlueprintableSpecifier(UhtSpecifierContext SpecifierContext)
+		{
+			SetMetaData(SpecifierContext, "IsBlueprintBase", true);
+			SetMetaData(SpecifierContext, "BlueprintType", true);
+		}
+
+		[UhtSpecifier(Extends = UhtTableNames.Default, ValueType = UhtSpecifierValueType.None, When = UhtSpecifierWhen.Immediate)]
+		private static void CallInEditorSpecifier(UhtSpecifierContext SpecifierContext)
+		{
+			SetMetaData(SpecifierContext, "CallInEditor", true);
+		}
+
+		[UhtSpecifier(Extends = UhtTableNames.Default, ValueType = UhtSpecifierValueType.None, When = UhtSpecifierWhen.Immediate)]
+		private static void NotBlueprintableSpecifier(UhtSpecifierContext SpecifierContext)
+		{
+			SetMetaData(SpecifierContext, "IsBlueprintBase", false);
+			SpecifierContext.MetaData.Remove("BlueprintType", SpecifierContext.MetaNameIndex);
+		}
+
+		[UhtSpecifier(Extends = UhtTableNames.Default, ValueType = UhtSpecifierValueType.String, When = UhtSpecifierWhen.Immediate)]
+		private static void CategorySpecifier(UhtSpecifierContext SpecifierContext, StringView Value)
+		{
+			SetMetaData(SpecifierContext, "Category", Value);
+		}
+
+		[UhtSpecifier(Extends = UhtTableNames.Default, ValueType = UhtSpecifierValueType.None, When = UhtSpecifierWhen.Immediate)]
+		private static void ExperimentalSpecifier(UhtSpecifierContext SpecifierContext)
+		{
+			SetMetaData(SpecifierContext, "DevelopmentStatus", "Experimental");
+		}
+
+		[UhtSpecifier(Extends = UhtTableNames.Default, ValueType = UhtSpecifierValueType.None, When = UhtSpecifierWhen.Immediate)]
+		private static void EarlyAccessPreviewSpecifier(UhtSpecifierContext SpecifierContext)
+		{
+			SetMetaData(SpecifierContext, "DevelopmentStatus", "EarlyAccess");
+		}
+
+		[UhtSpecifier(Extends = UhtTableNames.Default, ValueType = UhtSpecifierValueType.None, When = UhtSpecifierWhen.Immediate)]
+		private static void DocumentationPolicySpecifier(UhtSpecifierContext SpecifierContext)
+		{
+			SetMetaData(SpecifierContext, UhtNames.DocumentationPolicy, "Strict");
+		}
+
+		[UhtSpecifier(Extends = UhtTableNames.Default, ValueType = UhtSpecifierValueType.String, When = UhtSpecifierWhen.Immediate)]
+		private static void SparseClassDataTypeSpecifier(UhtSpecifierContext SpecifierContext, StringView Value)
+		{
+			SetMetaData(SpecifierContext, "SparseClassDataType", Value);
+		}
+		#endregion
+
+		#region Validators
+		[UhtSpecifierValidator(Name = "UIMin", Extends = UhtTableNames.Default)]
+		[UhtSpecifierValidator(Name = "UIMax", Extends = UhtTableNames.Default)]
+		[UhtSpecifierValidator(Name = "ClampMin", Extends = UhtTableNames.Default)]
+		[UhtSpecifierValidator(Name = "ClampMax", Extends = UhtTableNames.Default)]
+		private static void ValidateNumeric(UhtType Type, UhtMetaData MetaData, UhtMetaDataKey Key, StringView Value)
+		{
+			if (!UhtFCString.IsNumeric(Value.Span))
+			{
+				Type.LogError($"Metadata value for '{Key}' is non-numeric : '{Value}'");
+			}
+		}
+
+		[UhtSpecifierValidator(Extends = UhtTableNames.Default)]
+		private static void DevelopmentStatusSpecifierValidator(UhtType Type, UhtMetaData MetaData, UhtMetaDataKey Key, StringView Value)
+		{
+			string[] AllowedValues = { "EarlyAccess", "Experimental" };
+			foreach (string AllowedValue in AllowedValues)
+			{
+				if (Value.Equals(AllowedValue, StringComparison.OrdinalIgnoreCase))
+				{
+					return;
+				}
+			}
+			Type.LogError($"'{Key.Name}' metadata was '{Value}' but it must be {UhtUtilities.MergeTypeNames(AllowedValues, "or", false)}");
+		}
+
+		[UhtSpecifierValidator(Extends = UhtTableNames.Default)]
+		private static void DocumentationPolicySpecifierValidator(UhtType Type, UhtMetaData MetaData, UhtMetaDataKey Key, StringView Value)
+		{
+			const string StrictValue = "Strict";
+			if (!Value.Span.Equals(StrictValue, StringComparison.OrdinalIgnoreCase))
+			{
+				Type.LogError(MetaData.LineNumber, $"'{Key}' metadata was '{Value}' but it must be '{StrictValue}'");
+			}
+		}
+
+		[UhtSpecifierValidator(Extends = UhtTableNames.Default)]
+		private static void UnitsSpecifierValidator(UhtType Type, UhtMetaData MetaData, UhtMetaDataKey Key, StringView Value)
+		{
+			// Check for numeric property
+			if (Type is UhtProperty)
+			{
+				if (!(Type is UhtNumericProperty) && !(Type is UhtStructProperty))
+				{
+					Type.LogError("'Units' meta data can only be applied to numeric and struct properties");
+				}
+			}
+
+			if (!UhtConfig.Instance.IsValidUnits(Value))
+			{
+				Type.LogError($"Unrecognized units '{Value}' specified for '{Type.FullName}'");
+			}
+		}
+		#endregion
+	}
+}
diff --git a/Engine/Source/Programs/Shared/EpicGames.UHT/Specifiers/UhtEnumSpecifiers.cs b/Engine/Source/Programs/Shared/EpicGames.UHT/Specifiers/UhtEnumSpecifiers.cs
new file mode 100644
--- /dev/null
+++ b/Engine/Source/Programs/Shared/EpicGames.UHT/Specifiers/UhtEnumSpecifiers.cs
@@ -0,0 +1,20 @@
+// Copyright Epic Games, Inc. All Rights Reserved.
+
+using EpicGames.Core;
+using EpicGames.UHT.Tables;
+using EpicGames.UHT.Types;
+using EpicGames.UHT.Utils;
+
+namespace EpicGames.UHT.Parsers
+{
+	[UnrealHeaderTool]
+	public static class UhtEnumSpecifiers
+	{
+		[UhtSpecifier(Extends = UhtTableNames.Enum, ValueType = UhtSpecifierValueType.Legacy)]
+		private static void FlagsSpecifier(UhtSpecifierContext SpecifierContext)
+		{
+			UhtEnum Enum = (UhtEnum)SpecifierContext.Scope.ScopeType;
+			Enum.EnumFlags |= EEnumFlags.Flags;
+		}
+	}
+}
diff --git a/Engine/Source/Programs/Shared/EpicGames.UHT/Specifiers/UhtFunctionSpecifiers.cs b/Engine/Source/Programs/Shared/EpicGames.UHT/Specifiers/UhtFunctionSpecifiers.cs
new file mode 100644
--- /dev/null
+++ b/Engine/Source/Programs/Shared/EpicGames.UHT/Specifiers/UhtFunctionSpecifiers.cs
@@ -0,0 +1,431 @@
+// Copyright Epic Games, Inc. All Rights Reserved.
+
+using EpicGames.Core;
+using EpicGames.UHT.Tables;
+using EpicGames.UHT.Tokenizer;
+using EpicGames.UHT.Types;
+using EpicGames.UHT.Utils;
+using System;
+using System.Collections.Generic;
+
+namespace EpicGames.UHT.Parsers
+{
+	[UnrealHeaderTool]
+	public static class UhtFunctionSpecifiers
+	{
+		[UhtSpecifier(Extends = UhtTableNames.Function, ValueType = UhtSpecifierValueType.Legacy)]
+		private static void BlueprintAuthorityOnlySpecifier(UhtSpecifierContext SpecifierContext)
+		{
+			UhtFunction Function = (UhtFunction)SpecifierContext.Scope.ScopeType;
+			Function.FunctionFlags |= EFunctionFlags.BlueprintAuthorityOnly;
+		}
+
+		[UhtSpecifier(Extends = UhtTableNames.Function, ValueType = UhtSpecifierValueType.Legacy)]
+		private static void BlueprintCallableSpecifier(UhtSpecifierContext SpecifierContext)
+		{
+			UhtFunction Function = (UhtFunction)SpecifierContext.Scope.ScopeType;
+			Function.FunctionFlags |= EFunctionFlags.BlueprintCallable;
+		}
+
+		[UhtSpecifier(Extends = UhtTableNames.Function, ValueType = UhtSpecifierValueType.Legacy)]
+		private static void BlueprintCosmeticSpecifier(UhtSpecifierContext SpecifierContext)
+		{
+			UhtFunction Function = (UhtFunction)SpecifierContext.Scope.ScopeType;
+			Function.FunctionFlags |= EFunctionFlags.BlueprintCosmetic;
+		}
+
+		[UhtSpecifier(Extends = UhtTableNames.Function, ValueType = UhtSpecifierValueType.Legacy)]
+		private static void BlueprintGetterSpecifier(UhtSpecifierContext SpecifierContext)
+		{
+			UhtFunctionParser Function = (UhtFunctionParser)SpecifierContext.Scope.ScopeType;
+
+			if (Function.FunctionFlags.HasAnyFlags(EFunctionFlags.Event))
+			{
+				SpecifierContext.MessageSite.LogError("Function cannot be a blueprint event and a blueprint getter.");
+			}
+
+			Function.bSawPropertyAccessor = true;
+			Function.FunctionFlags |= EFunctionFlags.BlueprintCallable;
+			Function.FunctionFlags |= EFunctionFlags.BlueprintPure;
+			Function.MetaData.Add(UhtNames.BlueprintGetter, "");
+		}
+
+		[UhtSpecifier(Extends = UhtTableNames.Function, ValueType = UhtSpecifierValueType.Legacy)]
+		private static void BlueprintImplementableEventSpecifier(UhtSpecifierContext SpecifierContext)
+		{
+			UhtFunctionParser Function = (UhtFunctionParser)SpecifierContext.Scope.ScopeType;
+
+			if (Function.FunctionFlags.HasAnyFlags(EFunctionFlags.Net))
+			{
+				SpecifierContext.MessageSite.LogError("BlueprintImplementableEvent functions cannot be replicated!");
+			}
+			else if (Function.FunctionFlags.HasAnyFlags(EFunctionFlags.BlueprintEvent) && Function.FunctionFlags.HasAnyFlags(EFunctionFlags.Native))
+			{
+				// already a BlueprintNativeEvent
+				SpecifierContext.MessageSite.LogError("A function cannot be both BlueprintNativeEvent and BlueprintImplementableEvent!");
+			}
+			else if (Function.bSawPropertyAccessor)
+			{
+				SpecifierContext.MessageSite.LogError("A function cannot be both BlueprintImplementableEvent and a Blueprint Property accessor!");
+			}
+			else if (Function.FunctionFlags.HasAnyFlags(EFunctionFlags.Private))
+			{
+				SpecifierContext.MessageSite.LogError("A Private function cannot be a BlueprintImplementableEvent!");
+			}
+
+			Function.FunctionFlags |= EFunctionFlags.Event;
+			Function.FunctionFlags |= EFunctionFlags.BlueprintEvent;
+			Function.FunctionFlags &= ~EFunctionFlags.Native;
+		}
+
+		[UhtSpecifier(Extends = UhtTableNames.Function, ValueType = UhtSpecifierValueType.Legacy)]
+		private static void BlueprintNativeEventSpecifier(UhtSpecifierContext SpecifierContext)
+		{
+			UhtFunctionParser Function = (UhtFunctionParser)SpecifierContext.Scope.ScopeType;
+
+			if (Function.FunctionFlags.HasAnyFlags(EFunctionFlags.Net))
+			{
+				SpecifierContext.MessageSite.LogError("BlueprintNativeEvent functions cannot be replicated!");
+			}
+			else if (Function.FunctionFlags.HasAnyFlags(EFunctionFlags.BlueprintEvent) && !Function.FunctionFlags.HasAnyFlags(EFunctionFlags.Native))
+			{
+				// already a BlueprintImplementableEvent
+				SpecifierContext.MessageSite.LogError("A function cannot be both BlueprintNativeEvent and BlueprintImplementableEvent!");
+			}
+			else if (Function.bSawPropertyAccessor)
+			{
+				SpecifierContext.MessageSite.LogError("A function cannot be both BlueprintNativeEvent and a Blueprint Property accessor!");
+			}
+			else if (Function.FunctionFlags.HasAnyFlags(EFunctionFlags.Private))
+			{
+				SpecifierContext.MessageSite.LogError("A Private function cannot be a BlueprintNativeEvent!");
+			}
+
+			Function.FunctionFlags |= EFunctionFlags.Event;
+			Function.FunctionFlags |= EFunctionFlags.BlueprintEvent;
+		}
+
+		[UhtSpecifier(Extends = UhtTableNames.Function, ValueType = UhtSpecifierValueType.OptionalString)]
+		private static void BlueprintPureSpecifier(UhtSpecifierContext SpecifierContext, StringView? Value)
+		{
+			UhtFunction Function = (UhtFunction)SpecifierContext.Scope.ScopeType;
+
+			// This function can be called, and is also pure.
+			Function.FunctionFlags |= EFunctionFlags.BlueprintCallable;
+
+			if (Value == null || UhtFCString.ToBool((StringView)Value))
+			{
+				Function.FunctionFlags |= EFunctionFlags.BlueprintPure;
+			}
+			else
+			{
+				Function.FunctionExportFlags |= UhtFunctionExportFlags.ForceBlueprintImpure;
+			}
+		}
+
+		[UhtSpecifier(Extends = UhtTableNames.Function, ValueType = UhtSpecifierValueType.Legacy)]
+		private static void BlueprintSetterSpecifier(UhtSpecifierContext SpecifierContext)
+		{
+			UhtFunctionParser Function = (UhtFunctionParser)SpecifierContext.Scope.ScopeType;
+
+			if (Function.FunctionFlags.HasAnyFlags(EFunctionFlags.Event))
+			{
+				SpecifierContext.MessageSite.LogError("Function cannot be a blueprint event and a blueprint setter.");
+			}
+
+			Function.bSawPropertyAccessor = true;
+			Function.FunctionFlags |= EFunctionFlags.BlueprintCallable;
+			Function.MetaData.Add(UhtNames.BlueprintSetter, "");
+		}
+
+		[UhtSpecifier(Extends = UhtTableNames.Function, ValueType = UhtSpecifierValueType.OptionalString)]
+		private static void ClientSpecifier(UhtSpecifierContext SpecifierContext, StringView? Value)
+		{
+			UhtFunction Function = (UhtFunction)SpecifierContext.Scope.ScopeType;
+
+			if (Function.FunctionFlags.HasAnyFlags(EFunctionFlags.BlueprintEvent))
+			{
+				SpecifierContext.MessageSite.LogError("BlueprintImplementableEvent or BlueprintNativeEvent functions cannot be declared as Client or Server");
+			}
+
+			if (Function.FunctionFlags.HasAnyFlags(EFunctionFlags.Exec))
+			{
+				SpecifierContext.MessageSite.LogError("Exec functions cannot be replicated!");
+			}
+
+			Function.FunctionFlags |= EFunctionFlags.Net;
+			Function.FunctionFlags |= EFunctionFlags.NetClient;
+
+			if (Value != null)
+			{
+				Function.CppImplName = ((StringView)Value).ToString();
+			}
+		}
+
+		[UhtSpecifier(Extends = UhtTableNames.Function, ValueType = UhtSpecifierValueType.Legacy)]
+		private static void CustomThunkSpecifier(UhtSpecifierContext SpecifierContext)
+		{
+			UhtFunction Function = (UhtFunction)SpecifierContext.Scope.ScopeType;
+			Function.FunctionExportFlags |= UhtFunctionExportFlags.CustomThunk;
+		}
+
+		[UhtSpecifier(Extends = UhtTableNames.Function, ValueType = UhtSpecifierValueType.Legacy)]
+		private static void ExecSpecifier(UhtSpecifierContext SpecifierContext)
+		{
+			UhtFunction Function = (UhtFunction)SpecifierContext.Scope.ScopeType;
+			Function.FunctionFlags |= EFunctionFlags.Exec;
+			if (Function.FunctionFlags.HasAnyFlags(EFunctionFlags.Net))
+			{
+				SpecifierContext.MessageSite.LogError("Exec functions cannot be replicated!");
+			}
+		}
+
+		[UhtSpecifier(Extends = UhtTableNames.Function, ValueType = UhtSpecifierValueType.Legacy)]
+		private static void NetMulticastSpecifier(UhtSpecifierContext SpecifierContext)
+		{
+			UhtFunction Function = (UhtFunction)SpecifierContext.Scope.ScopeType;
+			if (Function.FunctionFlags.HasAnyFlags(EFunctionFlags.BlueprintEvent))
+			{
+				SpecifierContext.MessageSite.LogError("BlueprintImplementableEvent or BlueprintNativeEvent functions cannot be declared as Multicast");
+			}
+
+			Function.FunctionFlags |= EFunctionFlags.Net;
+			Function.FunctionFlags |= EFunctionFlags.NetMulticast;
+		}
+
+		[UhtSpecifier(Extends = UhtTableNames.Function, ValueType = UhtSpecifierValueType.Legacy)]
+		private static void ReliableSpecifier(UhtSpecifierContext SpecifierContext)
+		{
+			UhtFunction Function = (UhtFunction)SpecifierContext.Scope.ScopeType;
+			Function.FunctionFlags |= EFunctionFlags.NetReliable;
+		}
+
+		[UhtSpecifier(Extends = UhtTableNames.Function, ValueType = UhtSpecifierValueType.Legacy)]
+		private static void SealedEventSpecifier(UhtSpecifierContext SpecifierContext)
+		{
+			UhtFunction Function = (UhtFunction)SpecifierContext.Scope.ScopeType;
+			Function.FunctionExportFlags |= UhtFunctionExportFlags.SealedEvent;
+		}
+
+		[UhtSpecifier(Extends = UhtTableNames.Function, ValueType = UhtSpecifierValueType.OptionalString)]
+		private static void ServerSpecifier(UhtSpecifierContext SpecifierContext, StringView? Value)
+		{
+			UhtFunction Function = (UhtFunction)SpecifierContext.Scope.ScopeType;
+
+			if (Function.FunctionFlags.HasAnyFlags(EFunctionFlags.BlueprintEvent))
+			{
+				SpecifierContext.MessageSite.LogError("BlueprintImplementableEvent or BlueprintNativeEvent functions cannot be declared as Client or Server");
+			}
+
+			if (Function.FunctionFlags.HasAnyFlags(EFunctionFlags.Exec))
+			{
+				SpecifierContext.MessageSite.LogError("Exec functions cannot be replicated!");
+			}
+
+			Function.FunctionFlags |= EFunctionFlags.Net;
+			Function.FunctionFlags |= EFunctionFlags.NetServer;
+
+			if (Value != null)
+			{
+				Function.CppImplName = ((StringView)Value).ToString();
+			}
+		}
+
+		[UhtSpecifier(Extends = UhtTableNames.Function, ValueType = UhtSpecifierValueType.OptionalEqualsKeyValuePairList)]
+		private static void ServiceRequestSpecifier(UhtSpecifierContext SpecifierContext, List<KeyValuePair<StringView, StringView>> Value)
+		{
+			UhtFunctionParser Function = (UhtFunctionParser)SpecifierContext.Scope.ScopeType;
+
+			if (Function.FunctionFlags.HasAnyFlags(EFunctionFlags.BlueprintEvent))
+			{
+				SpecifierContext.MessageSite.LogError("BlueprintImplementableEvent or BlueprintNativeEvent functions cannot be declared as a ServiceRequest");
+			}
+
+			Function.FunctionFlags |= EFunctionFlags.Net;
+			Function.FunctionFlags |= EFunctionFlags.NetReliable;
+			Function.FunctionFlags |= EFunctionFlags.NetRequest;
+			Function.FunctionExportFlags |= UhtFunctionExportFlags.CustomThunk;
+
+			ParseNetServiceIdentifiers(SpecifierContext, Value);
+
+			if (string.IsNullOrEmpty(Function.EndpointName))
+			{
+				SpecifierContext.MessageSite.LogError("ServiceRequest needs to specify an endpoint name");
+			}
+		}
+
+		[UhtSpecifier(Extends = UhtTableNames.Function, ValueType = UhtSpecifierValueType.OptionalEqualsKeyValuePairList)]
+		private static void ServiceResponseSpecifier(UhtSpecifierContext SpecifierContext, List<KeyValuePair<StringView, StringView>> Value)
+		{
+			UhtFunctionParser Function = (UhtFunctionParser)SpecifierContext.Scope.ScopeType;
+
+			if (Function.FunctionFlags.HasAnyFlags(EFunctionFlags.BlueprintEvent))
+			{
+				SpecifierContext.MessageSite.LogError("BlueprintImplementableEvent or BlueprintNativeEvent functions cannot be declared as a ServiceResponse");
+			}
+
+			Function.FunctionFlags |= EFunctionFlags.Net;
+			Function.FunctionFlags |= EFunctionFlags.NetReliable;
+			Function.FunctionFlags |= EFunctionFlags.NetResponse;
+			//Function.FunctionExportFlags |= EFunctionExportFlags.CustomThunk;
+
+			ParseNetServiceIdentifiers(SpecifierContext, Value);
+
+			if (string.IsNullOrEmpty(Function.EndpointName))
+			{
+				SpecifierContext.MessageSite.LogError("ServiceResponse needs to specify an endpoint name");
+			}
+		}
+
+		[UhtSpecifier(Extends = UhtTableNames.Function, ValueType = UhtSpecifierValueType.Legacy)]
+		private static void UnreliableSpecifier(UhtSpecifierContext SpecifierContext)
+		{
+			UhtFunction Function = (UhtFunction)SpecifierContext.Scope.ScopeType;
+			Function.FunctionExportFlags |= UhtFunctionExportFlags.Unreliable;
+		}
+
+		[UhtSpecifier(Extends = UhtTableNames.Function, ValueType = UhtSpecifierValueType.StringList)]
+		private static void WithValidationSpecifier(UhtSpecifierContext SpecifierContext, List<StringView>? Value)
+		{
+			UhtFunction Function = (UhtFunction)SpecifierContext.Scope.ScopeType;
+			Function.FunctionFlags |= EFunctionFlags.NetValidate;
+
+			if (Value != null && Value.Count > 0)
+			{
+				Function.CppValidationImplName = Value[0].ToString();
+			}
+		}
+
+		[UhtSpecifierValidator(Extends = UhtTableNames.Function)]
+		private static void BlueprintProtectedSpecifierValidator(UhtType Type, UhtMetaData MetaData, UhtMetaDataKey Key, StringView Value)
+		{
+			UhtFunction Function = (UhtFunction)Type;
+			if (Function.FunctionFlags.HasAnyFlags(EFunctionFlags.Static))
+			{
+				// Given the owning class, locate the class that the owning derives from up to the point of UObject
+				if (Function.Outer is UhtClass OuterClass)
+				{
+					UhtClass? ClassPriorToUObject = OuterClass;
+					for (; ClassPriorToUObject != null; ClassPriorToUObject = ClassPriorToUObject.SuperClass)
+					{
+						// If our super is UObject, then stop
+						if (ClassPriorToUObject.Super == Type.Session.UObject)
+						{
+							break;
+						}
+					}
+
+					if (ClassPriorToUObject != null && ClassPriorToUObject.SourceName == "UBlueprintFunctionLibrary")
+					{
+						Type.LogError(MetaData.LineNumber, $"'{Key.Name}' doesn't make sense on static method '{Type.SourceName}' in a blueprint function library");
+					}
+				}
+			}
+		}
+
+		[UhtSpecifierValidator(Extends = UhtTableNames.Function)]
+		private static void CommutativeAssociativeBinaryOperatorSpecifierValidator(UhtType Type, UhtMetaData MetaData, UhtMetaDataKey Key, StringView Value)
+		{
+			UhtFunction Function = (UhtFunction)Type;
+			UhtProperty? ReturnProperty = Function.ReturnProperty;
+			ReadOnlyMemory<UhtType> ParameterProperties = Function.ParameterProperties;
+
+			if (ReturnProperty == null || ParameterProperties.Length != 2 || !((UhtProperty)ParameterProperties.Span[0]).IsSameType((UhtProperty)ParameterProperties.Span[1]))
+			{
+				Function.LogError("Commutative associative binary operators must have exactly 2 parameters of the same type and a return value.");
+			}
+		}
+
+		[UhtSpecifierValidator(Name = "ExpandBoolAsExecs", Extends = UhtTableNames.Function)]
+		[UhtSpecifierValidator(Name = "ExpandEnumAsExecs", Extends = UhtTableNames.Function)]
+		private static void ExpandsSpecifierValidator(UhtType Type, UhtMetaData MetaData, UhtMetaDataKey Key, StringView Value)
+		{
+			// multiple entry parsing in the same format as eg SetParam.
+			UhtType? FirstInput = null;
+			foreach (string RawGroup in Value.ToString().Split(','))
+			{
+				foreach (string Entry in RawGroup.Split('|'))
+				{
+					string Trimmed = Entry.Trim();
+					if (string.IsNullOrEmpty(Trimmed))
+					{
+						continue;
+					}
+
+					UhtType? FoundField = Type.FindType(UhtFindOptions.SourceName | UhtFindOptions.SelfOnly | UhtFindOptions.Property, Trimmed, Type);
+					if (FoundField != null)
+					{
+						UhtProperty Property = (UhtProperty)FoundField;
+						if (!Property.PropertyFlags.HasAnyFlags(EPropertyFlags.ReturnParm) &&
+							(!Property.PropertyFlags.HasAnyFlags(EPropertyFlags.OutParm) ||
+							 Property.PropertyFlags.HasAnyFlags(EPropertyFlags.ReferenceParm)))
+						{
+							if (FirstInput == null)
+							{
+								FirstInput = FoundField;
+							}
+							else
+							{
+								Type.LogError($"Function already specified an ExpandEnumAsExec input '{FirstInput.SourceName}', but '{Trimmed}' is also an input parameter. Only one is permitted.");
+							}
+						}
+					}
+				}
+			}
+		}
+
+		private static void ParseNetServiceIdentifiers(UhtSpecifierContext SpecifierContext, List<KeyValuePair<StringView, StringView>> Identifiers)
+		{
+			IUhtTokenReader TokenReader = SpecifierContext.Scope.TokenReader;
+			UhtFunctionParser Function = (UhtFunctionParser)SpecifierContext.Scope.ScopeType;
+
+			foreach (var KVP in Identifiers)
+			{
+				if (KVP.Value.Span.Length > 0)
+				{
+					if (KVP.Key.Span.StartsWith("Id", StringComparison.OrdinalIgnoreCase))
+					{
+						int Id;
+						if (!int.TryParse(KVP.Value.Span, out Id) || Id <= 0 || Id > UInt16.MaxValue)
+						{
+							TokenReader.LogError($"Invalid network identifier {KVP.Key} for function");
+						}
+						else
+						{
+							Function.RPCId = (UInt16)Id;
+						}
+					}
+					else if (KVP.Key.Span.StartsWith("ResponseId", StringComparison.OrdinalIgnoreCase) || KVP.Key.Span.StartsWith("Priority", StringComparison.OrdinalIgnoreCase))
+					{
+						int Id;
+						if (!int.TryParse(KVP.Value.Span, out Id) || Id <= 0 || Id > UInt16.MaxValue)
+						{
+							TokenReader.LogError($"Invalid network identifier {KVP.Key} for function");
+						}
+						else
+						{
+							Function.RPCResponseId = (UInt16)Id;
+						}
+					}
+					else
+					{
+						// No error message???
+					}
+				}
+
+				// Assume it's an endpoint name
+				else
+				{
+					if (!string.IsNullOrEmpty(Function.EndpointName))
+					{
+						TokenReader.LogError($"Function should not specify multiple endpoints - '{KVP.Key}' found but already using '{Function.EndpointName}'");
+					}
+					else
+					{
+						Function.EndpointName = KVP.Key.ToString();
+					}
+				}
+			}
+		}
+	}
+}
diff --git a/Engine/Source/Programs/Shared/EpicGames.UHT/Specifiers/UhtInterfaceClassSpecifiers.cs b/Engine/Source/Programs/Shared/EpicGames.UHT/Specifiers/UhtInterfaceClassSpecifiers.cs
new file mode 100644
--- /dev/null
+++ b/Engine/Source/Programs/Shared/EpicGames.UHT/Specifiers/UhtInterfaceClassSpecifiers.cs
@@ -0,0 +1,34 @@
+// Copyright Epic Games, Inc. All Rights Reserved.
+
+using EpicGames.Core;
+using EpicGames.UHT.Tables;
+using EpicGames.UHT.Types;
+using EpicGames.UHT.Utils;
+
+namespace EpicGames.UHT.Parsers
+{
+	[UnrealHeaderTool]
+	class UhtInterfaceClassSpecifiers
+	{
+		[UhtSpecifier(Extends = UhtTableNames.Interface, ValueType = UhtSpecifierValueType.Legacy)]
+		private static void DependsOnSpecifier(UhtSpecifierContext SpecifierContext)
+		{
+			UhtClass Class = (UhtClass)SpecifierContext.Scope.ScopeType;
+			throw new UhtException(SpecifierContext.MessageSite, $"The dependsOn specifier is deprecated. Please use #include \"ClassHeaderFilename.h\" instead.");
+		}
+
+		[UhtSpecifier(Extends = UhtTableNames.Interface, ValueType = UhtSpecifierValueType.Legacy)]
+		private static void MinimalAPISpecifier(UhtSpecifierContext SpecifierContext)
+		{
+			UhtClass Class = (UhtClass)SpecifierContext.Scope.ScopeType;
+			Class.ClassFlags |= EClassFlags.MinimalAPI;
+		}
+
+		[UhtSpecifier(Extends = UhtTableNames.Interface, ValueType = UhtSpecifierValueType.Legacy)]
+		private static void ConversionRootSpecifier(UhtSpecifierContext SpecifierContext)
+		{
+			UhtClass Class = (UhtClass)SpecifierContext.Scope.ScopeType;
+			Class.MetaData.Add(UhtNames.IsConversionRoot, "true");
+		}
+	}
+}
diff --git a/Engine/Source/Programs/Shared/EpicGames.UHT/Specifiers/UhtPropertyArgumentSpecifiers.cs b/Engine/Source/Programs/Shared/EpicGames.UHT/Specifiers/UhtPropertyArgumentSpecifiers.cs
new file mode 100644
--- /dev/null
+++ b/Engine/Source/Programs/Shared/EpicGames.UHT/Specifiers/UhtPropertyArgumentSpecifiers.cs
@@ -0,0 +1,47 @@
+// Copyright Epic Games, Inc. All Rights Reserved.
+
+using EpicGames.Core;
+using EpicGames.UHT.Tables;
+using EpicGames.UHT.Types;
+using EpicGames.UHT.Utils;
+
+namespace EpicGames.UHT.Parsers
+{
+	[UnrealHeaderTool]
+	public static class UhtPropertyArgumentSpecifiers
+	{
+		public static readonly UhtSpecifierTable SpecifierTable = UhtSpecifierTables.Instance.Get(UhtTableNames.PropertyArgument);
+		public static readonly UhtSpecifierValidatorTable SpecifierValidatorTable = UhtSpecifierValidatorTables.Instance.Get(UhtTableNames.PropertyArgument);
+
+		#region Argument Property Specifiers
+		[UhtSpecifier(Extends = UhtTableNames.PropertyArgument, ValueType = UhtSpecifierValueType.Legacy)]
+		private static void ConstSpecifier(UhtSpecifierContext SpecifierContext)
+		{
+			UhtPropertySpecifierContext Context = (UhtPropertySpecifierContext)SpecifierContext;
+			Context.PropertySettings.PropertyFlags |= EPropertyFlags.ConstParm;
+		}
+
+		[UhtSpecifier(Extends = UhtTableNames.PropertyArgument, ValueType = UhtSpecifierValueType.Legacy)]
+		private static void RefSpecifier(UhtSpecifierContext SpecifierContext)
+		{
+			UhtPropertySpecifierContext Context = (UhtPropertySpecifierContext)SpecifierContext;
+			Context.PropertySettings.PropertyFlags |= EPropertyFlags.OutParm | EPropertyFlags.ReferenceParm;
+		}
+
+		[UhtSpecifier(Extends = UhtTableNames.PropertyArgument, ValueType = UhtSpecifierValueType.Legacy)]
+		private static void NotReplicatedSpecifier(UhtSpecifierContext SpecifierContext)
+		{
+			UhtPropertySpecifierContext Context = (UhtPropertySpecifierContext)SpecifierContext;
+			if (Context.PropertySettings.PropertyCategory == UhtPropertyCategory.ReplicatedParameter)
+			{
+				Context.PropertySettings.PropertyCategory = UhtPropertyCategory.RegularParameter;
+				Context.PropertySettings.PropertyFlags |= EPropertyFlags.RepSkip;
+			}
+			else
+			{
+				Context.MessageSite.LogError("Only parameters in service request functions can be marked NotReplicated");
+			}
+		}
+		#endregion
+	}
+}
diff --git a/Engine/Source/Programs/Shared/EpicGames.UHT/Specifiers/UhtPropertyMemberSpecifiers.cs b/Engine/Source/Programs/Shared/EpicGames.UHT/Specifiers/UhtPropertyMemberSpecifiers.cs
new file mode 100644
--- /dev/null
+++ b/Engine/Source/Programs/Shared/EpicGames.UHT/Specifiers/UhtPropertyMemberSpecifiers.cs
@@ -0,0 +1,436 @@
+// Copyright Epic Games, Inc. All Rights Reserved.
+
+using EpicGames.Core;
+using EpicGames.UHT.Tables;
+using EpicGames.UHT.Types;
+using EpicGames.UHT.Utils;
+using System;
+
+namespace EpicGames.UHT.Parsers
+{
+	[UnrealHeaderTool]
+	public static class UhtPropertyMemberSpecifiers
+	{
+		public static UhtSpecifierTable SpecifierTable = UhtSpecifierTables.Instance.Get(UhtTableNames.PropertyMember);
+		public static UhtSpecifierValidatorTable SpecifierValidatorTable = UhtSpecifierValidatorTables.Instance.Get(UhtTableNames.PropertyMember);
+
+		[UhtSpecifier(Extends = UhtTableNames.PropertyMember, ValueType = UhtSpecifierValueType.Legacy)]
+		private static void EditAnywhereSpecifier(UhtSpecifierContext SpecifierContext)
+		{
+			UhtPropertySpecifierContext Context = (UhtPropertySpecifierContext)SpecifierContext;
+			if (Context.bSeenEditSpecifier)
+			{
+				Context.MessageSite.LogError("Found more than one edit/visibility specifier (EditAnywhere), only one is allowed");
+			}
+			Context.PropertySettings.PropertyFlags |= EPropertyFlags.Edit;
+			Context.bSeenEditSpecifier = true;
+		}
+
+		[UhtSpecifier(Extends = UhtTableNames.PropertyMember, ValueType = UhtSpecifierValueType.Legacy)]
+		private static void EditInstanceOnlySpecifier(UhtSpecifierContext SpecifierContext)
+		{
+			UhtPropertySpecifierContext Context = (UhtPropertySpecifierContext)SpecifierContext;
+			if (Context.bSeenEditSpecifier)
+			{
+				Context.MessageSite.LogError("Found more than one edit/visibility specifier (EditInstanceOnly), only one is allowed");
+			}
+			Context.PropertySettings.PropertyFlags |= EPropertyFlags.Edit | EPropertyFlags.DisableEditOnTemplate;
+			Context.bSeenEditSpecifier = true;
+		}
+
+		[UhtSpecifier(Extends = UhtTableNames.PropertyMember, ValueType = UhtSpecifierValueType.Legacy)]
+		private static void EditDefaultsOnlySpecifier(UhtSpecifierContext SpecifierContext)
+		{
+			UhtPropertySpecifierContext Context = (UhtPropertySpecifierContext)SpecifierContext;
+			if (Context.bSeenEditSpecifier)
+			{
+				Context.MessageSite.LogError("Found more than one edit/visibility specifier (EditDefaultsOnly), only one is allowed");
+			}
+			Context.PropertySettings.PropertyFlags |= EPropertyFlags.Edit | EPropertyFlags.DisableEditOnInstance;
+			Context.bSeenEditSpecifier = true;
+		}
+
+		[UhtSpecifier(Extends = UhtTableNames.PropertyMember, ValueType = UhtSpecifierValueType.Legacy)]
+		private static void VisibleAnywhereSpecifier(UhtSpecifierContext SpecifierContext)
+		{
+			UhtPropertySpecifierContext Context = (UhtPropertySpecifierContext)SpecifierContext;
+			if (Context.bSeenEditSpecifier)
+			{
+				Context.MessageSite.LogError("Found more than one edit/visibility specifier (VisibleAnywhere), only one is allowed");
+			}
+			Context.PropertySettings.PropertyFlags |= EPropertyFlags.Edit | EPropertyFlags.EditConst;
+			Context.bSeenEditSpecifier = true;
+		}
+
+		[UhtSpecifier(Extends = UhtTableNames.PropertyMember, ValueType = UhtSpecifierValueType.Legacy)]
+		private static void VisibleInstanceOnlySpecifier(UhtSpecifierContext SpecifierContext)
+		{
+			UhtPropertySpecifierContext Context = (UhtPropertySpecifierContext)SpecifierContext;
+			if (Context.bSeenEditSpecifier)
+			{
+				Context.MessageSite.LogError("Found more than one edit/visibility specifier (VisibleInstanceOnly), only one is allowed");
+			}
+			Context.PropertySettings.PropertyFlags |= EPropertyFlags.Edit | EPropertyFlags.EditConst | EPropertyFlags.DisableEditOnTemplate;
+			Context.bSeenEditSpecifier = true;
+		}
+
+		[UhtSpecifier(Extends = UhtTableNames.PropertyMember, ValueType = UhtSpecifierValueType.Legacy)]
+		private static void VisibleDefaultsOnlySpecifier(UhtSpecifierContext SpecifierContext)
+		{
+			UhtPropertySpecifierContext Context = (UhtPropertySpecifierContext)SpecifierContext;
+			if (Context.bSeenEditSpecifier)
+			{
+				Context.MessageSite.LogError("Found more than one edit/visibility specifier (VisibleDefaultsOnly), only one is allowed");
+			}
+			Context.PropertySettings.PropertyFlags |= EPropertyFlags.Edit | EPropertyFlags.EditConst | EPropertyFlags.DisableEditOnInstance;
+			Context.bSeenEditSpecifier = true;
+		}
+
+		[UhtSpecifier(Extends = UhtTableNames.PropertyMember, ValueType = UhtSpecifierValueType.Legacy)]
+		private static void BlueprintReadWriteSpecifier(UhtSpecifierContext SpecifierContext)
+		{
+			UhtPropertySpecifierContext Context = (UhtPropertySpecifierContext)SpecifierContext;
+			if (Context.bSeenBlueprintReadOnlySpecifier)
+			{
+				Context.MessageSite.LogError("Cannot specify a property as being both BlueprintReadOnly and BlueprintReadWrite.");
+			}
+
+			string? PrivateAccessMD;
+			bool bAllowPrivateAccess = Context.MetaData.TryGetValue(UhtNames.AllowPrivateAccess, out PrivateAccessMD) && !PrivateAccessMD.Equals("false", StringComparison.OrdinalIgnoreCase);
+			if (SpecifierContext.Scope.AccessSpecifier == UhtAccessSpecifier.Private && !bAllowPrivateAccess)
+			{
+				Context.MessageSite.LogError("BlueprintReadWrite should not be used on private members");
+			}
+
+			if (Context.PropertySettings.PropertyFlags.HasAnyFlags(EPropertyFlags.EditorOnly) && Context.PropertySettings.Outer is UhtScriptStruct)
+			{
+				Context.MessageSite.LogError("Blueprint exposed struct members cannot be editor only");
+			}
+
+			Context.PropertySettings.PropertyFlags |= EPropertyFlags.BlueprintVisible;
+			Context.bSeenBlueprintWriteSpecifier = true;
+		}
+
+		[UhtSpecifier(Extends = UhtTableNames.PropertyMember, ValueType = UhtSpecifierValueType.Legacy)]
+		private static void BlueprintReadOnlySpecifier(UhtSpecifierContext SpecifierContext)
+		{
+			UhtPropertySpecifierContext Context = (UhtPropertySpecifierContext)SpecifierContext;
+			if (Context.bSeenBlueprintWriteSpecifier)
+			{
+				Context.MessageSite.LogError("Cannot specify both BlueprintReadOnly and BlueprintReadWrite or BlueprintSetter.");
+			}
+
+			string? PrivateAccessMD;
+			bool bAllowPrivateAccess = Context.MetaData.TryGetValue(UhtNames.AllowPrivateAccess, out PrivateAccessMD) && !PrivateAccessMD.Equals("false", StringComparison.OrdinalIgnoreCase);
+			if (SpecifierContext.Scope.AccessSpecifier == UhtAccessSpecifier.Private && !bAllowPrivateAccess)
+			{
+				Context.MessageSite.LogError("BlueprintReadOnly should not be used on private members");
+			}
+
+			if (Context.PropertySettings.PropertyFlags.HasAnyFlags(EPropertyFlags.EditorOnly) && Context.PropertySettings.Outer is UhtScriptStruct)
+			{
+				Context.MessageSite.LogError("Blueprint exposed struct members cannot be editor only");
+			}
+
+			Context.PropertySettings.PropertyFlags |= EPropertyFlags.BlueprintVisible | EPropertyFlags.BlueprintReadOnly;
+			Context.bSeenBlueprintReadOnlySpecifier = true;
+		}
+
+		[UhtSpecifier(Extends = UhtTableNames.PropertyMember, ValueType = UhtSpecifierValueType.SingleString)]
+		private static void BlueprintSetterSpecifier(UhtSpecifierContext SpecifierContext, StringView Value)
+		{
+			UhtPropertySpecifierContext Context = (UhtPropertySpecifierContext)SpecifierContext;
+			if (Context.bSeenBlueprintReadOnlySpecifier)
+			{
+				Context.MessageSite.LogError("Cannot specify a property as being both BlueprintReadOnly and having a BlueprintSetter.");
+			}
+
+			if (Context.PropertySettings.Outer is UhtScriptStruct)
+			{
+				Context.MessageSite.LogError("Cannot specify BlueprintSetter for a struct member.");
+			}
+
+			Context.MetaData.Add(UhtNames.BlueprintSetter, Value.ToString());
+
+			Context.PropertySettings.PropertyFlags |= EPropertyFlags.BlueprintVisible;
+			Context.bSeenBlueprintWriteSpecifier = true;
+		}
+
+		[UhtSpecifier(Extends = UhtTableNames.PropertyMember, ValueType = UhtSpecifierValueType.SingleString)]
+		private static void BlueprintGetterSpecifier(UhtSpecifierContext SpecifierContext, StringView Value)
+		{
+			UhtPropertySpecifierContext Context = (UhtPropertySpecifierContext)SpecifierContext;
+			if (Context.PropertySettings.Outer is UhtScriptStruct)
+			{
+				Context.MessageSite.LogError("Cannot specify BlueprintGetter for a struct member.");
+			}
+
+			Context.MetaData.Add(UhtNames.BlueprintGetter, Value.ToString());
+
+			Context.PropertySettings.PropertyFlags |= EPropertyFlags.BlueprintVisible;
+			Context.bSeenBlueprintGetterSpecifier = true;
+		}
+
+		[UhtSpecifier(Extends = UhtTableNames.PropertyMember, ValueType = UhtSpecifierValueType.Legacy)]
+		private static void ConfigSpecifier(UhtSpecifierContext SpecifierContext)
+		{
+			UhtPropertySpecifierContext Context = (UhtPropertySpecifierContext)SpecifierContext;
+			Context.PropertySettings.PropertyFlags |= EPropertyFlags.Config;
+		}
+
+		[UhtSpecifier(Extends = UhtTableNames.PropertyMember, ValueType = UhtSpecifierValueType.Legacy)]
+		private static void GlobalConfigSpecifier(UhtSpecifierContext SpecifierContext)
+		{
+			UhtPropertySpecifierContext Context = (UhtPropertySpecifierContext)SpecifierContext;
+			Context.PropertySettings.PropertyFlags |= EPropertyFlags.GlobalConfig | EPropertyFlags.Config;
+		}
+
+		[UhtSpecifier(Extends = UhtTableNames.PropertyMember, ValueType = UhtSpecifierValueType.Legacy)]
+		private static void LocalizedSpecifier(UhtSpecifierContext SpecifierContext)
+		{
+			UhtPropertySpecifierContext Context = (UhtPropertySpecifierContext)SpecifierContext;
+			Context.MessageSite.LogError("The Localized specifier is deprecated");
+		}
+
+		[UhtSpecifier(Extends = UhtTableNames.PropertyMember, ValueType = UhtSpecifierValueType.Legacy)]
+		private static void TransientSpecifier(UhtSpecifierContext SpecifierContext)
+		{
+			UhtPropertySpecifierContext Context = (UhtPropertySpecifierContext)SpecifierContext;
+			Context.PropertySettings.PropertyFlags |= EPropertyFlags.Transient;
+		}
+
+		[UhtSpecifier(Extends = UhtTableNames.PropertyMember, ValueType = UhtSpecifierValueType.Legacy)]
+		private static void DuplicateTransientSpecifier(UhtSpecifierContext SpecifierContext)
+		{
+			UhtPropertySpecifierContext Context = (UhtPropertySpecifierContext)SpecifierContext;
+			Context.PropertySettings.PropertyFlags |= EPropertyFlags.DuplicateTransient;
+		}
+
+		[UhtSpecifier(Extends = UhtTableNames.PropertyMember, ValueType = UhtSpecifierValueType.Legacy)]
+		private static void TextExportTransientSpecifier(UhtSpecifierContext SpecifierContext)
+		{
+			UhtPropertySpecifierContext Context = (UhtPropertySpecifierContext)SpecifierContext;
+			Context.PropertySettings.PropertyFlags |= EPropertyFlags.TextExportTransient;
+		}
+
+		[UhtSpecifier(Extends = UhtTableNames.PropertyMember, ValueType = UhtSpecifierValueType.Legacy)]
+		private static void NonPIETransientSpecifier(UhtSpecifierContext SpecifierContext)
+		{
+			UhtPropertySpecifierContext Context = (UhtPropertySpecifierContext)SpecifierContext;
+			Context.MessageSite.LogWarning("NonPIETransient is deprecated - NonPIEDuplicateTransient should be used instead");
+			Context.PropertySettings.PropertyFlags |= EPropertyFlags.NonPIEDuplicateTransient;
+		}
+
+		[UhtSpecifier(Extends = UhtTableNames.PropertyMember, ValueType = UhtSpecifierValueType.Legacy)]
+		private static void NonPIEDuplicateTransientSpecifier(UhtSpecifierContext SpecifierContext)
+		{
+			UhtPropertySpecifierContext Context = (UhtPropertySpecifierContext)SpecifierContext;
+			Context.PropertySettings.PropertyFlags |= EPropertyFlags.NonPIEDuplicateTransient;
+		}
+
+		[UhtSpecifier(Extends = UhtTableNames.PropertyMember, ValueType = UhtSpecifierValueType.Legacy)]
+		private static void ExportSpecifier(UhtSpecifierContext SpecifierContext)
+		{
+			UhtPropertySpecifierContext Context = (UhtPropertySpecifierContext)SpecifierContext;
+			Context.PropertySettings.PropertyFlags |= EPropertyFlags.ExportObject;
+		}
+
+		[UhtSpecifier(Extends = UhtTableNames.PropertyMember, ValueType = UhtSpecifierValueType.Legacy)]
+		private static void EditInlineSpecifier(UhtSpecifierContext SpecifierContext)
+		{
+			UhtPropertySpecifierContext Context = (UhtPropertySpecifierContext)SpecifierContext;
+			Context.MessageSite.LogError("EditInline is deprecated. Remove it, or use Instanced instead.");
+		}
+
+		[UhtSpecifier(Extends = UhtTableNames.PropertyMember, ValueType = UhtSpecifierValueType.Legacy)]
+		private static void NoClearSpecifier(UhtSpecifierContext SpecifierContext)
+		{
+			UhtPropertySpecifierContext Context = (UhtPropertySpecifierContext)SpecifierContext;
+			Context.PropertySettings.PropertyFlags |= EPropertyFlags.NoClear;
+		}
+
+		[UhtSpecifier(Extends = UhtTableNames.PropertyMember, ValueType = UhtSpecifierValueType.Legacy)]
+		private static void EditFixedSizeSpecifier(UhtSpecifierContext SpecifierContext)
+		{
+			UhtPropertySpecifierContext Context = (UhtPropertySpecifierContext)SpecifierContext;
+			Context.PropertySettings.PropertyFlags |= EPropertyFlags.EditFixedSize;
+		}
+
+		[UhtSpecifier(Extends = UhtTableNames.PropertyMember, ValueType = UhtSpecifierValueType.Legacy)]
+		private static void ReplicatedSpecifier(UhtSpecifierContext SpecifierContext)
+		{
+			UhtPropertySpecifierContext Context = (UhtPropertySpecifierContext)SpecifierContext;
+			if (Context.PropertySettings.Outer is UhtScriptStruct)
+			{
+				Context.MessageSite.LogError("Struct members cannot be replicated");
+			}
+			Context.PropertySettings.PropertyFlags |= EPropertyFlags.Net;
+		}
+
+		[UhtSpecifier(Extends = UhtTableNames.PropertyMember, ValueType = UhtSpecifierValueType.SingleString)]
+		private static void ReplicatedUsingSpecifier(UhtSpecifierContext SpecifierContext, StringView Value)
+		{
+			UhtPropertySpecifierContext Context = (UhtPropertySpecifierContext)SpecifierContext;
+			if (Context.PropertySettings.Outer is UhtScriptStruct)
+			{
+				Context.MessageSite.LogError("Struct members cannot be replicated");
+			}
+			Context.PropertySettings.PropertyFlags |= EPropertyFlags.Net;
+
+			if (Value.Span.Length == 0)
+			{
+				Context.MessageSite.LogError("Must specify a valid function name for replication notifications");
+			}
+			Context.PropertySettings.RepNotifyName = Value.ToString();
+			Context.PropertySettings.PropertyFlags |= EPropertyFlags.RepNotify;
+		}
+
+		[UhtSpecifier(Extends = UhtTableNames.PropertyMember, ValueType = UhtSpecifierValueType.Legacy)]
+		private static void NotReplicatedSpecifier(UhtSpecifierContext SpecifierContext)
+		{
+			UhtPropertySpecifierContext Context = (UhtPropertySpecifierContext)SpecifierContext;
+			if (!(Context.PropertySettings.Outer is UhtScriptStruct))
+			{
+				Context.MessageSite.LogError("Only Struct members can be marked NotReplicated");
+			}
+			Context.PropertySettings.PropertyFlags |= EPropertyFlags.RepSkip;
+		}
+
+		[UhtSpecifier(Extends = UhtTableNames.PropertyMember, ValueType = UhtSpecifierValueType.Legacy)]
+		private static void RepRetrySpecifier(UhtSpecifierContext SpecifierContext)
+		{
+			UhtPropertySpecifierContext Context = (UhtPropertySpecifierContext)SpecifierContext;
+			Context.MessageSite.LogError("'RepRetry' is deprecated.");
+		}
+
+		[UhtSpecifier(Extends = UhtTableNames.PropertyMember, ValueType = UhtSpecifierValueType.Legacy)]
+		private static void InterpSpecifier(UhtSpecifierContext SpecifierContext)
+		{
+			UhtPropertySpecifierContext Context = (UhtPropertySpecifierContext)SpecifierContext;
+			Context.PropertySettings.PropertyFlags |= EPropertyFlags.Edit | EPropertyFlags.BlueprintVisible | EPropertyFlags.Interp;
+		}
+
+		[UhtSpecifier(Extends = UhtTableNames.PropertyMember, ValueType = UhtSpecifierValueType.Legacy)]
+		private static void NonTransactionalSpecifier(UhtSpecifierContext SpecifierContext)
+		{
+			UhtPropertySpecifierContext Context = (UhtPropertySpecifierContext)SpecifierContext;
+			Context.PropertySettings.PropertyFlags |= EPropertyFlags.NonTransactional;
+		}
+
+		[UhtSpecifier(Extends = UhtTableNames.PropertyMember, ValueType = UhtSpecifierValueType.Legacy)]
+		private static void InstancedSpecifier(UhtSpecifierContext SpecifierContext)
+		{
+			UhtPropertySpecifierContext Context = (UhtPropertySpecifierContext)SpecifierContext;
+			Context.PropertySettings.PropertyFlags |= EPropertyFlags.PersistentInstance | EPropertyFlags.ExportObject | EPropertyFlags.InstancedReference;
+			Context.MetaData.Add(UhtNames.EditInline, true);
+		}
+
+		[UhtSpecifier(Extends = UhtTableNames.PropertyMember, ValueType = UhtSpecifierValueType.Legacy)]
+		private static void BlueprintAssignableSpecifier(UhtSpecifierContext SpecifierContext)
+		{
+			UhtPropertySpecifierContext Context = (UhtPropertySpecifierContext)SpecifierContext;
+			Context.PropertySettings.PropertyFlags |= EPropertyFlags.BlueprintAssignable;
+		}
+
+		[UhtSpecifier(Extends = UhtTableNames.PropertyMember, ValueType = UhtSpecifierValueType.Legacy)]
+		private static void BlueprintCallableSpecifier(UhtSpecifierContext SpecifierContext)
+		{
+			UhtPropertySpecifierContext Context = (UhtPropertySpecifierContext)SpecifierContext;
+			Context.PropertySettings.PropertyFlags |= EPropertyFlags.BlueprintCallable;
+		}
+
+		[UhtSpecifier(Extends = UhtTableNames.PropertyMember, ValueType = UhtSpecifierValueType.Legacy)]
+		private static void BlueprintAuthorityOnlySpecifier(UhtSpecifierContext SpecifierContext)
+		{
+			UhtPropertySpecifierContext Context = (UhtPropertySpecifierContext)SpecifierContext;
+			Context.PropertySettings.PropertyFlags |= EPropertyFlags.BlueprintAuthorityOnly;
+		}
+
+		[UhtSpecifier(Extends = UhtTableNames.PropertyMember, ValueType = UhtSpecifierValueType.Legacy)]
+		private static void AssetRegistrySearchableSpecifier(UhtSpecifierContext SpecifierContext)
+		{
+			UhtPropertySpecifierContext Context = (UhtPropertySpecifierContext)SpecifierContext;
+			Context.PropertySettings.PropertyFlags |= EPropertyFlags.AssetRegistrySearchable;
+		}
+
+		[UhtSpecifier(Extends = UhtTableNames.PropertyMember, ValueType = UhtSpecifierValueType.Legacy)]
+		private static void SimpleDisplaySpecifier(UhtSpecifierContext SpecifierContext)
+		{
+			UhtPropertySpecifierContext Context = (UhtPropertySpecifierContext)SpecifierContext;
+			Context.PropertySettings.PropertyFlags |= EPropertyFlags.SimpleDisplay;
+		}
+
+		[UhtSpecifier(Extends = UhtTableNames.PropertyMember, ValueType = UhtSpecifierValueType.Legacy)]
+		private static void AdvancedDisplaySpecifier(UhtSpecifierContext SpecifierContext)
+		{
+			UhtPropertySpecifierContext Context = (UhtPropertySpecifierContext)SpecifierContext;
+			Context.PropertySettings.PropertyFlags |= EPropertyFlags.AdvancedDisplay;
+		}
+
+		[UhtSpecifier(Extends = UhtTableNames.PropertyMember, ValueType = UhtSpecifierValueType.Legacy)]
+		private static void SaveGameSpecifier(UhtSpecifierContext SpecifierContext)
+		{
+			UhtPropertySpecifierContext Context = (UhtPropertySpecifierContext)SpecifierContext;
+			Context.PropertySettings.PropertyFlags |= EPropertyFlags.SaveGame;
+		}
+
+		[UhtSpecifier(Extends = UhtTableNames.PropertyMember, ValueType = UhtSpecifierValueType.Legacy)]
+		private static void SkipSerializationSpecifier(UhtSpecifierContext SpecifierContext)
+		{
+			UhtPropertySpecifierContext Context = (UhtPropertySpecifierContext)SpecifierContext;
+			Context.PropertySettings.PropertyFlags |= EPropertyFlags.SkipSerialization;
+		}
+
+		[UhtSpecifier(Extends = UhtTableNames.PropertyMember, ValueType = UhtSpecifierValueType.OptionalString)]
+		private static void GetterSpecifier(UhtSpecifierContext SpecifierContext, StringView? Value)
+		{
+			UhtPropertySpecifierContext Context = (UhtPropertySpecifierContext)SpecifierContext;
+			if (!(Context.PropertySettings.Outer is UhtClass))
+			{
+				Context.MessageSite.LogError("Only class members can have Setters");
+			}
+			Context.PropertySettings.PropertyExportFlags |= UhtPropertyExportFlags.GetterSpecified;
+			if (Value != null)
+			{
+				StringView Temp = (StringView)Value;
+				if (Temp.Length > 0)
+				{
+					if (Temp.Equals("None", StringComparison.OrdinalIgnoreCase))
+					{
+						Context.PropertySettings.PropertyExportFlags |= UhtPropertyExportFlags.GetterSpecifiedNone;
+					}
+					else
+					{
+						Context.PropertySettings.Getter = Value.ToString();
+					}
+				}
+			}
+		}
+
+		[UhtSpecifier(Extends = UhtTableNames.PropertyMember, ValueType = UhtSpecifierValueType.OptionalString)]
+		private static void SetterSpecifier(UhtSpecifierContext SpecifierContext, StringView? Value)
+		{
+			UhtPropertySpecifierContext Context = (UhtPropertySpecifierContext)SpecifierContext;
+			if (!(Context.PropertySettings.Outer is UhtClass))
+			{
+				Context.MessageSite.LogError("Only class members can have Getters");
+			}
+			Context.PropertySettings.PropertyExportFlags |= UhtPropertyExportFlags.SetterSpecified;
+			if (Value != null)
+			{
+				StringView Temp = (StringView)Value;
+				if (Temp.Length > 0)
+				{
+					if (Temp.Equals("None", StringComparison.OrdinalIgnoreCase))
+					{
+						Context.PropertySettings.PropertyExportFlags |= UhtPropertyExportFlags.SetterSpecifiedNone;
+					}
+					else
+					{
+						Context.PropertySettings.Setter = Value.ToString();
+					}
+				}
+			}
+		}
+	}
+}
diff --git a/Engine/Source/Programs/Shared/EpicGames.UHT/Specifiers/UhtScriptStructSpecifiers.cs b/Engine/Source/Programs/Shared/EpicGames.UHT/Specifiers/UhtScriptStructSpecifiers.cs
new file mode 100644
--- /dev/null
+++ b/Engine/Source/Programs/Shared/EpicGames.UHT/Specifiers/UhtScriptStructSpecifiers.cs
@@ -0,0 +1,79 @@
+﻿// Copyright Epic Games, Inc. All Rights Reserved.
+
+using EpicGames.Core;
+using EpicGames.UHT.Tables;
+using EpicGames.UHT.Types;
+using EpicGames.UHT.Utils;
+
+namespace EpicGames.UHT.Parsers
+{
+	[UnrealHeaderTool]
+	public static class UhtScriptStructSpecifiers
+	{
+		[UhtSpecifier(Extends = UhtTableNames.ScriptStruct, ValueType = UhtSpecifierValueType.Legacy)]
+		private static void NoExportSpecifier(UhtSpecifierContext SpecifierContext)
+		{
+			UhtScriptStruct ScriptStruct = (UhtScriptStruct)SpecifierContext.Scope.ScopeType;
+			ScriptStruct.ScriptStructFlags |= EStructFlags.NoExport;
+			ScriptStruct.ScriptStructFlags &= ~EStructFlags.Native;
+		}
+
+		[UhtSpecifier(Extends = UhtTableNames.ScriptStruct, ValueType = UhtSpecifierValueType.Legacy)]
+		private static void AtomicSpecifier(UhtSpecifierContext SpecifierContext)
+		{
+			UhtScriptStruct ScriptStruct = (UhtScriptStruct)SpecifierContext.Scope.ScopeType;
+			ScriptStruct.ScriptStructFlags |= EStructFlags.Atomic;
+		}
+
+		[UhtSpecifier(Extends = UhtTableNames.ScriptStruct, ValueType = UhtSpecifierValueType.Legacy)]
+		private static void ImmutableSpecifier(UhtSpecifierContext SpecifierContext)
+		{
+			UhtScriptStruct ScriptStruct = (UhtScriptStruct)SpecifierContext.Scope.ScopeType;
+			ScriptStruct.ScriptStructFlags |= EStructFlags.Atomic | EStructFlags.Immutable;
+		}
+
+		[UhtSpecifier(Extends = UhtTableNames.ScriptStruct, ValueType = UhtSpecifierValueType.Legacy)]
+		private static void HasDefaultsSpecifier(UhtSpecifierContext SpecifierContext)
+		{
+			if (!SpecifierContext.Scope.HeaderParser.HeaderFile.bIsNoExportTypes)
+			{
+				SpecifierContext.MessageSite.LogError("The 'HasDefaults' struct specifier is only valid in the NoExportTypes.h file");
+			}
+			UhtScriptStruct ScriptStruct = (UhtScriptStruct)SpecifierContext.Scope.ScopeType;
+			ScriptStruct.ScriptStructExportFlags |= UhtScriptStructExportFlags.HasDefaults;
+		}
+
+		[UhtSpecifier(Extends = UhtTableNames.ScriptStruct, ValueType = UhtSpecifierValueType.Legacy)]
+		private static void HasNoOpConstructorSpecifier(UhtSpecifierContext SpecifierContext)
+		{
+			if (!SpecifierContext.Scope.HeaderParser.HeaderFile.bIsNoExportTypes)
+			{
+				SpecifierContext.MessageSite.LogError("The 'HasNoOpConstructor' struct specifier is only valid in the NoExportTypes.h file");
+			}
+			UhtScriptStruct ScriptStruct = (UhtScriptStruct)SpecifierContext.Scope.ScopeType;
+			ScriptStruct.ScriptStructExportFlags |= UhtScriptStructExportFlags.HasNoOpConstructor;
+		}
+
+		[UhtSpecifier(Extends = UhtTableNames.ScriptStruct, ValueType = UhtSpecifierValueType.Legacy)]
+		private static void IsAlwaysAccessibleSpecifier(UhtSpecifierContext SpecifierContext)
+		{
+			if (!SpecifierContext.Scope.HeaderParser.HeaderFile.bIsNoExportTypes)
+			{
+				SpecifierContext.MessageSite.LogError("The 'IsAlwaysAccessible' struct specifier is only valid in the NoExportTypes.h file");
+			}
+			UhtScriptStruct ScriptStruct = (UhtScriptStruct)SpecifierContext.Scope.ScopeType;
+			ScriptStruct.ScriptStructExportFlags |= UhtScriptStructExportFlags.IsAlwaysAccessible;
+		}
+
+		[UhtSpecifier(Extends = UhtTableNames.ScriptStruct, ValueType = UhtSpecifierValueType.Legacy)]
+		private static void IsCoreTypeSpecifier(UhtSpecifierContext SpecifierContext)
+		{
+			if (!SpecifierContext.Scope.HeaderParser.HeaderFile.bIsNoExportTypes)
+			{
+				SpecifierContext.MessageSite.LogError("The 'IsCoreType' struct specifier is only valid in the NoExportTypes.h file");
+			}
+			UhtScriptStruct ScriptStruct = (UhtScriptStruct)SpecifierContext.Scope.ScopeType;
+			ScriptStruct.ScriptStructExportFlags |= UhtScriptStructExportFlags.IsCoreType;
+		}
+	}
+}
diff --git a/Engine/Source/Programs/Shared/EpicGames.UHT/Tables/UhtEngineClassTable.cs b/Engine/Source/Programs/Shared/EpicGames.UHT/Tables/UhtEngineClassTable.cs
new file mode 100644
--- /dev/null
+++ b/Engine/Source/Programs/Shared/EpicGames.UHT/Tables/UhtEngineClassTable.cs
@@ -0,0 +1,88 @@
+// Copyright Epic Games, Inc. All Rights Reserved.
+
+using EpicGames.Core;
+using EpicGames.UHT.Utils;
+using System;
+using System.Collections.Generic;
+
+namespace EpicGames.UHT.Tables
+{
+
+	/// <summary>
+	/// This attribute is placed on classes that represent Unreal Engine classes.
+	/// </summary>
+	[AttributeUsage(AttributeTargets.Class, AllowMultiple = true)]
+	public class UhtEngineClassAttribute : Attribute
+	{
+
+		/// <summary>
+		/// The name of the engine class excluding any prefix
+		/// </summary>
+		public string? Name = null;
+
+		/// <summary>
+		/// If true, this class is a property
+		/// </summary>
+		public bool IsProperty = false;
+	}
+
+	/// <summary>
+	/// Represents an engine class in the engine class table
+	/// </summary>
+	public struct UhtEngineClass
+	{
+		/// <summary>
+		/// The name of the engine class excluding any prefix
+		/// </summary>
+		public string Name;
+
+		/// <summary>
+		/// If true, this class is a property
+		/// </summary>
+		public bool bIsProperty;
+	}
+
+	/// <summary>
+	/// Table of all known engine class names.
+	/// </summary>
+	public class UhtEngineClassTable
+	{
+		/// <summary>
+		/// Global instance of the engine class table
+		/// </summary>
+		public static UhtEngineClassTable Instance = new UhtEngineClassTable();
+
+		/// <summary>
+		/// Internal mapping from engine class name to information
+		/// </summary>
+		private Dictionary<StringView, UhtEngineClass> EngineClasses = new Dictionary<StringView, UhtEngineClass>();
+
+		/// <summary>
+		/// Test to see if the given class name is a property
+		/// </summary>
+		/// <param name="Name">Name of the class without the prefix</param>
+		/// <returns>True if the class name is a property.  False if the class name isn't a property or isn't an engine class.</returns>
+		public bool IsValidPropertyTypeName(StringView Name)
+		{
+			if (this.EngineClasses.TryGetValue(Name, out UhtEngineClass EngineClass))
+			{
+				return EngineClass.bIsProperty;
+			}
+			return false;
+		}
+
+		/// <summary>
+		/// Add an entry to the table
+		/// </summary>
+		/// <param name="Type">Type associated with the given attribute</param>
+		/// <param name="EngineClassAttribute">The located attribute</param>
+		public void OnEngineClassAttribute(Type Type, UhtEngineClassAttribute EngineClassAttribute)
+		{
+			if (string.IsNullOrEmpty(EngineClassAttribute.Name))
+			{
+				throw new UhtIceException("EngineClassNames must have a name specified");
+			}
+			this.EngineClasses.Add(EngineClassAttribute.Name, new UhtEngineClass { Name = EngineClassAttribute.Name, bIsProperty = EngineClassAttribute.IsProperty });
+		}
+	}
+}
diff --git a/Engine/Source/Programs/Shared/EpicGames.UHT/Tables/UhtExporterTable.cs b/Engine/Source/Programs/Shared/EpicGames.UHT/Tables/UhtExporterTable.cs
new file mode 100644
--- /dev/null
+++ b/Engine/Source/Programs/Shared/EpicGames.UHT/Tables/UhtExporterTable.cs
@@ -0,0 +1,142 @@
+// Copyright Epic Games, Inc. All Rights Reserved.
+
+using EpicGames.UHT.Utils;
+using System;
+using System.Collections;
+using System.Collections.Generic;
+using System.Reflection;
+
+namespace EpicGames.UHT.Tables
+{
+	public delegate void UhtExporterDelegate(IUhtExportFactory Factory);
+
+	[Flags]
+	public enum UhtExporterOptions
+	{
+		None = 0,
+
+		/// <summary>
+		/// The exporter should be run by default
+		/// </summary>
+		Default = 1 << 0,
+	}
+
+	/// <summary>
+	/// Helper methods for testing flags.  These methods perform better than the generic HasFlag which hits
+	/// the GC and stalls.
+	/// </summary>
+	public static class UhtExporterOptionsExtensions
+	{
+
+		/// <summary>
+		/// Test to see if any of the specified flags are set
+		/// </summary>
+		/// <param name="InFlags">Current flags</param>
+		/// <param name="TestFlags">Flags to test for</param>
+		/// <returns>True if any of the flags are set</returns>
+		public static bool HasAnyFlags(this UhtExporterOptions InFlags, UhtExporterOptions TestFlags)
+		{
+			return (InFlags & TestFlags) != 0;
+		}
+
+		/// <summary>
+		/// Test to see if all of the specified flags are set
+		/// </summary>
+		/// <param name="InFlags">Current flags</param>
+		/// <param name="TestFlags">Flags to test for</param>
+		/// <returns>True if all the flags are set</returns>
+		public static bool HasAllFlags(this UhtExporterOptions InFlags, UhtExporterOptions TestFlags)
+		{
+			return (InFlags & TestFlags) == TestFlags;
+		}
+
+		/// <summary>
+		/// Test to see if a specific set of flags have a specific value.
+		/// </summary>
+		/// <param name="InFlags">Current flags</param>
+		/// <param name="TestFlags">Flags to test for</param>
+		/// <param name="MatchFlags">Expected value of the tested flags</param>
+		/// <returns>True if the given flags have a specific value.</returns>
+		public static bool HasExactFlags(this UhtExporterOptions InFlags, UhtExporterOptions TestFlags, UhtExporterOptions MatchFlags)
+		{
+			return (InFlags & TestFlags) == MatchFlags;
+		}
+	}
+
+	[AttributeUsage(AttributeTargets.Method, AllowMultiple = true)]
+	public class UhtExporterAttribute : Attribute
+	{
+		public string Name = string.Empty;
+		public string Description = string.Empty;
+		public UhtExporterOptions Options = UhtExporterOptions.None;
+		public string[]? CppFilters = null;
+		public string[]? HeaderFilters = null;
+		public string[]? OtherFilters = null;
+	}
+
+	public struct UhtExporter
+	{
+		public string Name;
+		public string Description;
+		public UhtExporterOptions Options;
+		public UhtExporterDelegate Delegate;
+		public string[]? CppFilters;
+		public string[]? HeaderFilters;
+		public string[]? OtherFilters;
+	}
+
+	public class UhtExporterTable : IEnumerable<UhtExporter>
+	{
+		public static UhtExporterTable Instance = new UhtExporterTable();
+
+		private Dictionary<string, UhtExporter> ExporterValues = new Dictionary<string, UhtExporter>(StringComparer.OrdinalIgnoreCase);
+
+		/// <summary>
+		/// Return the exporter associated with the given name
+		/// </summary>
+		/// <param name="Name"></param>
+		/// <param name="Value">Exporter associated with the name</param>
+		/// <returns></returns>
+		public bool TryGet(string Name, out UhtExporter Value)
+		{
+			return this.ExporterValues.TryGetValue(Name, out Value);
+		}
+
+		public void OnExporterAttribute(Type Type, MethodInfo MethodInfo, UhtExporterAttribute ExporterAttribute)
+		{
+			if (string.IsNullOrEmpty(ExporterAttribute.Name))
+			{
+				throw new UhtIceException("An exporter must have a name");
+			}
+
+			UhtExporter ExporterValue = new UhtExporter
+			{
+				Name = ExporterAttribute.Name,
+				Description = ExporterAttribute.Description,
+				Delegate = (UhtExporterDelegate)Delegate.CreateDelegate(typeof(UhtExporterDelegate), MethodInfo),
+				Options = ExporterAttribute.Options,
+				CppFilters = ExporterAttribute.CppFilters,
+				HeaderFilters = ExporterAttribute.HeaderFilters,
+				OtherFilters = ExporterAttribute.OtherFilters,
+			};
+
+			this.ExporterValues.Add(ExporterAttribute.Name, ExporterValue);
+		}
+
+		public IEnumerator<UhtExporter> GetEnumerator()
+		{
+			foreach (KeyValuePair<string, UhtExporter> KVP in this.ExporterValues)
+			{
+				yield return KVP.Value;
+			}
+		}
+
+		IEnumerator IEnumerable.GetEnumerator()
+		{
+			foreach (KeyValuePair<string, UhtExporter> KVP in this.ExporterValues)
+			{
+				yield return KVP.Value;
+			}
+		}
+	}
+}
diff --git a/Engine/Source/Programs/Shared/EpicGames.UHT/Tables/UhtKeywordTable.cs b/Engine/Source/Programs/Shared/EpicGames.UHT/Tables/UhtKeywordTable.cs
new file mode 100644
--- /dev/null
+++ b/Engine/Source/Programs/Shared/EpicGames.UHT/Tables/UhtKeywordTable.cs
@@ -0,0 +1,183 @@
+// Copyright Epic Games, Inc. All Rights Reserved.
+
+using System;
+using System.Collections.Generic;
+using System.Reflection;
+using EpicGames.Core;
+using EpicGames.UHT.Parsers;
+using EpicGames.UHT.Tokenizer;
+using EpicGames.UHT.Utils;
+
+namespace EpicGames.UHT.Tables
+{
+
+	/// <summary>
+	/// Invoke the given method when the keyword is parsed.
+	/// </summary>
+	[AttributeUsage(AttributeTargets.Method, AllowMultiple = true)]
+	public class UhtKeywordAttribute : Attribute
+	{
+		public string? Extends;
+		public string? Keyword = null;
+		public string? AllowText = null;
+		public bool AllScopes = false;
+		public bool DisableUsageError = false;
+	}
+
+	[AttributeUsage(AttributeTargets.Method, AllowMultiple = true)]
+	public class UhtKeywordCatchAllAttribute : Attribute
+	{
+		public string? Extends;
+	}
+
+	public delegate UhtParseResult UhtKeywordDelegate(UhtParsingScope TopScope, UhtParsingScope ActionScope, ref UhtToken Token);
+	public delegate UhtParseResult UhtKeywordCatchAllDelegate(UhtParsingScope TopScope, ref UhtToken Token);
+
+	public struct UhtKeyword
+	{
+		public readonly string Name;
+		public readonly UhtKeywordDelegate Delegate;
+		public readonly string? AllowText;
+		public readonly bool bAllScopes;
+		public readonly bool bDisableUsageError;
+
+		public UhtKeyword(string Name, UhtKeywordDelegate Delegate, UhtKeywordAttribute? Attribute)
+		{
+			this.Name = Name;
+			this.Delegate = Delegate;
+			if (Attribute != null)
+			{
+				this.AllowText = Attribute.AllowText;
+				this.bAllScopes = Attribute.AllScopes;
+				this.bDisableUsageError = Attribute.DisableUsageError;
+			}
+			else
+			{
+				this.AllowText = null;
+				this.bAllScopes = false;
+				this.bDisableUsageError = false;
+			}
+		}
+	}
+
+	public class UhtKeywordTable : UhtLookupTable<UhtKeyword>
+	{
+
+		/// <summary>
+		/// List of catch-alls associated with this table
+		/// </summary>
+		public List<UhtKeywordCatchAllDelegate> CatchAlls = new List<UhtKeywordCatchAllDelegate>();
+
+		/// <summary>
+		/// Construct a new keyword table
+		/// </summary>
+		public UhtKeywordTable() : base(StringViewComparer.Ordinal)
+		{
+		}
+
+		/// <summary>
+		/// Add the given value to the lookup table.  It will throw an exception if it is a duplicate.
+		/// </summary>
+		/// <param name="Value">Value to be added</param>
+		public UhtKeywordTable Add(UhtKeyword Value)
+		{
+			base.Add(Value.Name, Value);
+			return this;
+		}
+
+		/// <summary>
+		/// Add the given catch-all to the table.
+		/// </summary>
+		/// <param name="CatchAll">The catch-all to be added</param>
+		public UhtKeywordTable AddCatchAll(UhtKeywordCatchAllDelegate CatchAll)
+		{
+			this.CatchAlls.Add(CatchAll);
+			return this;
+		}
+
+		/// <summary>
+		/// Merge the given keyword table.  Duplicates in the BaseTypeTable will be ignored.
+		/// </summary>
+		/// <param name="BaseTypeTable">Base table being merged</param>
+		public override void Merge(UhtLookupTableBase BaseTable)
+		{
+			base.Merge(BaseTable);
+			this.CatchAlls.AddRange(((UhtKeywordTable)BaseTable).CatchAlls);
+		}
+	}
+
+	public class UhtKeywordTables : UhtLookupTables<UhtKeywordTable>
+	{
+		public static UhtKeywordTables Instance = new UhtKeywordTables();
+
+		public UhtKeywordTables() : base("keywords")
+		{
+		}
+
+		public void OnKeywordCatchAllAttribute(Type Type, MethodInfo MethodInfo, UhtKeywordCatchAllAttribute KeywordCatchAllAttribute)
+		{
+			if (string.IsNullOrEmpty(KeywordCatchAllAttribute.Extends))
+			{
+				throw new UhtIceException($"The 'KeywordCatchAlll' attribute on the {Type.Name}.{MethodInfo.Name} method doesn't have a table specified.");
+			}
+
+			UhtKeywordTable Table = Get(KeywordCatchAllAttribute.Extends);
+			Table.AddCatchAll((UhtKeywordCatchAllDelegate)Delegate.CreateDelegate(typeof(UhtKeywordCatchAllDelegate), MethodInfo));
+		}
+
+		public void OnKeywordAttribute(Type Type, MethodInfo MethodInfo, UhtKeywordAttribute KeywordAttribute)
+		{
+			string Name = UhtLookupTableBase.GetSuffixedName(Type, MethodInfo, KeywordAttribute.Keyword, "Keyword");
+
+			if (string.IsNullOrEmpty(KeywordAttribute.Extends))
+			{
+				throw new UhtIceException($"The 'Keyword' attribute on the {Type.Name}.{MethodInfo.Name} method doesn't have a table specified.");
+			}
+
+			UhtKeywordTable Table = Get(KeywordAttribute.Extends);
+			Table.Add(new UhtKeyword(Name, (UhtKeywordDelegate)Delegate.CreateDelegate(typeof(UhtKeywordDelegate), MethodInfo), KeywordAttribute));
+		}
+
+		public void LogUnhandledError(IUhtMessageSite MessageSite, UhtToken Token)
+		{
+			List<string>? Tables = null;
+			foreach (KeyValuePair<string, UhtKeywordTable> KVP in this.Tables)
+			{
+				UhtKeywordTable KeywordTable = KVP.Value;
+				if (KeywordTable.Internal)
+				{
+					continue;
+				}
+				UhtKeyword Info;
+				if (KeywordTable.TryGetValue(Token.Value, out Info))
+				{
+					if (Info.bDisableUsageError)
+					{
+						// Do not log anything for this keyword in this table
+					}
+					else
+					{
+						if (Tables == null)
+						{
+							Tables = new List<string>();
+						}
+						if (Info.AllowText != null)
+						{
+							Tables.Add($"{KeywordTable.UserName} {Info.AllowText}");
+						}
+						else
+						{
+							Tables.Add(KeywordTable.UserName);
+						}
+					}
+				}
+			}
+
+			if (Tables != null)
+			{
+				string Text = UhtUtilities.MergeTypeNames(Tables, "and");
+				MessageSite.LogError(Token.InputLine, $"Invalid use of keyword '{Token.Value}'.  It may only appear in {Text} scopes");
+			}
+		}
+	}
+}
diff --git a/Engine/Source/Programs/Shared/EpicGames.UHT/Tables/UhtLocTextDefaultValueTable.cs b/Engine/Source/Programs/Shared/EpicGames.UHT/Tables/UhtLocTextDefaultValueTable.cs
new file mode 100644
--- /dev/null
+++ b/Engine/Source/Programs/Shared/EpicGames.UHT/Tables/UhtLocTextDefaultValueTable.cs
@@ -0,0 +1,59 @@
+// Copyright Epic Games, Inc. All Rights Reserved.
+
+using EpicGames.Core;
+using EpicGames.UHT.Tokenizer;
+using System;
+using System.Collections.Generic;
+using System.Reflection;
+using EpicGames.UHT.Types;
+using EpicGames.UHT.Utils;
+using System.Text;
+
+namespace EpicGames.UHT.Tables
+{
+	public delegate bool UhtLocTextDefaultValueDelegate(UhtTextProperty Property, IUhtTokenReader DefaultValueReader, ref UhtToken MacroToken, StringBuilder InnerDefaultValue);
+
+	[AttributeUsage(AttributeTargets.Method, AllowMultiple = true)]
+	public class UhtLocTextDefaultValueAttribute : Attribute
+	{
+		public string? Name;
+	}
+
+	public struct UhtLocTextDefaultValue
+	{
+		public UhtLocTextDefaultValueDelegate Delegate;
+	}
+
+	public class UhtLocTextDefaultValueTable
+	{
+		public static UhtLocTextDefaultValueTable Instance = new UhtLocTextDefaultValueTable();
+
+		private Dictionary<StringView, UhtLocTextDefaultValue> LocTextDefaultValues = new Dictionary<StringView, UhtLocTextDefaultValue>();
+
+		/// <summary>
+		/// Return the loc text default value associated with the given name
+		/// </summary>
+		/// <param name="Name"></param>
+		/// <param name="LocTextDefaultValue">Loc text default value handler</param>
+		/// <returns></returns>
+		public bool TryGet(StringView Name, out UhtLocTextDefaultValue LocTextDefaultValue)
+		{
+			return this.LocTextDefaultValues.TryGetValue(Name, out LocTextDefaultValue);
+		}
+
+		public void OnLocTextDefaultValueAttribute(Type Type, MethodInfo MethodInfo, UhtLocTextDefaultValueAttribute LocTextDefaultValueAttribute)
+		{
+			if (string.IsNullOrEmpty(LocTextDefaultValueAttribute.Name))
+			{
+				throw new UhtIceException("A loc text default value attribute must have a name");
+			}
+
+			UhtLocTextDefaultValue LocTextDefaultValue = new UhtLocTextDefaultValue
+			{
+				Delegate = (UhtLocTextDefaultValueDelegate)Delegate.CreateDelegate(typeof(UhtLocTextDefaultValueDelegate), MethodInfo)
+			};
+
+			this.LocTextDefaultValues.Add(new StringView(LocTextDefaultValueAttribute.Name), LocTextDefaultValue);
+		}
+	}
+}
diff --git a/Engine/Source/Programs/Shared/EpicGames.UHT/Tables/UhtPropertyTypeTable.cs b/Engine/Source/Programs/Shared/EpicGames.UHT/Tables/UhtPropertyTypeTable.cs
new file mode 100644
--- /dev/null
+++ b/Engine/Source/Programs/Shared/EpicGames.UHT/Tables/UhtPropertyTypeTable.cs
@@ -0,0 +1,181 @@
+// Copyright Epic Games, Inc. All Rights Reserved.
+
+using EpicGames.Core;
+using EpicGames.UHT.Tokenizer;
+using EpicGames.UHT.Types;
+using EpicGames.UHT.Utils;
+using System;
+using System.Collections.Generic;
+using System.Reflection;
+
+namespace EpicGames.UHT.Tables
+{
+	[Flags]
+	public enum UhtPropertyTypeOptions
+	{
+		None = 0,
+
+		/// <summary>
+		/// Simple property type with just the property type.  (i.e. "int32 MyValue")
+		/// Simple types are not required to parse the supplied token list.
+		/// </summary>
+		Simple = 1 << 0,
+
+		/// <summary>
+		/// Use case insensitive string compares
+		/// </summary>
+		CaseInsensitive = 1 << 1,
+
+		/// <summary>
+		/// This property type is to be invoked when there are no keyword matches found
+		/// </summary>
+		Default = 1 << 2,
+
+		/// <summary>
+		/// This property type doesn't reference any engine types an can be resolved immediately
+		/// </summary>
+		Immediate = 1 << 3,
+	}
+
+	/// <summary>
+	/// Helper methods for testing flags.  These methods perform better than the generic HasFlag which hits
+	/// the GC and stalls.
+	/// </summary>
+	public static class UhtPropertyTypeOptionsExtensions
+	{
+
+		/// <summary>
+		/// Test to see if any of the specified flags are set
+		/// </summary>
+		/// <param name="InFlags">Current flags</param>
+		/// <param name="TestFlags">Flags to test for</param>
+		/// <returns>True if any of the flags are set</returns>
+		public static bool HasAnyFlags(this UhtPropertyTypeOptions InFlags, UhtPropertyTypeOptions TestFlags)
+		{
+			return (InFlags & TestFlags) != 0;
+		}
+
+		/// <summary>
+		/// Test to see if all of the specified flags are set
+		/// </summary>
+		/// <param name="InFlags">Current flags</param>
+		/// <param name="TestFlags">Flags to test for</param>
+		/// <returns>True if all the flags are set</returns>
+		public static bool HasAllFlags(this UhtPropertyTypeOptions InFlags, UhtPropertyTypeOptions TestFlags)
+		{
+			return (InFlags & TestFlags) == TestFlags;
+		}
+
+		/// <summary>
+		/// Test to see if a specific set of flags have a specific value.
+		/// </summary>
+		/// <param name="InFlags">Current flags</param>
+		/// <param name="TestFlags">Flags to test for</param>
+		/// <param name="MatchFlags">Expected value of the tested flags</param>
+		/// <returns>True if the given flags have a specific value.</returns>
+		public static bool HasExactFlags(this UhtPropertyTypeOptions InFlags, UhtPropertyTypeOptions TestFlags, UhtPropertyTypeOptions MatchFlags)
+		{
+			return (InFlags & TestFlags) == MatchFlags;
+		}
+	}
+
+	public enum UhtPropertyResolvePhase
+	{
+		Parsing,
+		Resolving,
+	}
+
+	/// <summary>
+	/// Delegate invoked to resolve a tokenized type into a UHTProperty type
+	/// </summary>
+	/// <param name="ResolvePhase">Specifies if this is being resolved during the parsing phase or the resolution phase.  Type lookups can not happen during the parsing phase</param>
+	/// <param name="PropertySettings">The configuration of the property</param>
+	/// <param name="TokenReader">The token reader containing the type</param>
+	/// <param name="MatchedToken">The token that matched the delegate unless the delegate is the default resolver.</param>
+	/// <returns></returns>
+	public delegate UhtProperty? UhtResolvePropertyDelegate(UhtPropertyResolvePhase ResolvePhase, UhtPropertySettings PropertySettings, IUhtTokenReader TokenReader, UhtToken MatchedToken);
+
+	[AttributeUsage(AttributeTargets.Method, AllowMultiple = true)]
+	public class UhtPropertyTypeAttribute : Attribute
+	{
+		public string? Keyword = null;
+		public UhtPropertyTypeOptions Options = UhtPropertyTypeOptions.None;
+	}
+
+	/// <summary>
+	/// Represents a property type as specified by the PropertyTypeAttribute
+	/// </summary>
+	public struct UhtPropertyType
+	{
+		public UhtResolvePropertyDelegate Delegate;
+		public UhtPropertyTypeOptions Options;
+	}
+
+	public class UhtPropertyTypeTable
+	{
+		public static UhtPropertyTypeTable Instance = new UhtPropertyTypeTable();
+
+		private Dictionary<StringView, UhtPropertyType> CaseSensitive = new Dictionary<StringView, UhtPropertyType>();
+		private Dictionary<StringView, UhtPropertyType> CaseInsensitive = new Dictionary<StringView, UhtPropertyType>(StringViewComparer.OrdinalIgnoreCase);
+		private UhtPropertyType? DefaultInternal = null;
+
+		public UhtPropertyType Default
+		{
+			get
+			{
+				if (this.DefaultInternal == null)
+				{
+					throw new UhtIceException("No property type has been marked as default");
+				}
+				return (UhtPropertyType)this.DefaultInternal;
+			}
+		}
+
+		/// <summary>
+		/// Return the property type associated with the given name
+		/// </summary>
+		/// <param name="Name"></param>
+		/// <param name="PropertyType">Property type if matched</param>
+		/// <returns></returns>
+		public bool TryGet(StringView Name, out UhtPropertyType PropertyType)
+		{
+			return
+				this.CaseSensitive.TryGetValue(Name, out PropertyType) ||
+				this.CaseInsensitive.TryGetValue(Name, out PropertyType);
+		}
+
+		public void OnPropertyTypeAttribute(Type Type, MethodInfo MethodInfo, UhtPropertyTypeAttribute PropertyTypeAttribute)
+		{
+			if (string.IsNullOrEmpty(PropertyTypeAttribute.Keyword) && !PropertyTypeAttribute.Options.HasAnyFlags(UhtPropertyTypeOptions.Default))
+			{
+				throw new UhtIceException("A property type must have a keyword or be marked as default");
+			}
+
+			UhtPropertyType PropertyType = new UhtPropertyType
+			{
+				Delegate = (UhtResolvePropertyDelegate)Delegate.CreateDelegate(typeof(UhtResolvePropertyDelegate), MethodInfo),
+				Options = PropertyTypeAttribute.Options,
+			};
+
+			if (PropertyTypeAttribute.Options.HasAnyFlags(UhtPropertyTypeOptions.Default))
+			{
+				if (this.DefaultInternal != null)
+				{
+					throw new UhtIceException("Only one property type dispatcher can be marked as default");
+				}
+				this.DefaultInternal = PropertyType;
+			}
+			else if (!string.IsNullOrEmpty(PropertyTypeAttribute.Keyword))
+			{
+				if (PropertyTypeAttribute.Options.HasAnyFlags(UhtPropertyTypeOptions.CaseInsensitive))
+				{
+					this.CaseInsensitive.Add(PropertyTypeAttribute.Keyword, PropertyType);
+				}
+				else
+				{
+					this.CaseSensitive.Add(PropertyTypeAttribute.Keyword, PropertyType);
+				}
+			}
+		}
+	}
+}
diff --git a/Engine/Source/Programs/Shared/EpicGames.UHT/Tables/UhtSpecifierTable.cs b/Engine/Source/Programs/Shared/EpicGames.UHT/Tables/UhtSpecifierTable.cs
new file mode 100644
--- /dev/null
+++ b/Engine/Source/Programs/Shared/EpicGames.UHT/Tables/UhtSpecifierTable.cs
@@ -0,0 +1,392 @@
+// Copyright Epic Games, Inc. All Rights Reserved.
+
+using EpicGames.Core;
+using EpicGames.UHT.Parsers; // It would be nice if we didn't need this here
+using EpicGames.UHT.Types;
+using EpicGames.UHT.Utils;
+using System;
+using System.Collections.Generic;
+using System.Reflection;
+
+namespace EpicGames.UHT.Tables
+{
+	public enum UhtSpecifierValueType
+	{
+		/// <summary>
+		/// Internal value used to detect if the attribute has a valid value
+		/// </summary>
+		NotSet,
+
+		/// <summary>
+		/// No values of any type are allowed
+		/// </summary>
+		None,
+
+		/// <summary>
+		/// A string value but can not be in the form of a list (i.e. =(bob))
+		/// </summary>
+		String,
+
+		/// <summary>
+		/// An optional string value but can not be in the form of a list
+		/// </summary>
+		OptionalString,
+
+		/// <summary>
+		/// A string value or a single element string list
+		/// </summary>
+		SingleString,
+
+		/// <summary>
+		/// A list of values in key=value pairs
+		/// </summary>
+		KeyValuePairList,
+
+		/// <summary>
+		/// A list of values in key=value pairs but the equals is optional
+		/// </summary>
+		OptionalEqualsKeyValuePairList,
+
+		/// <summary>
+		/// A list of values.
+		/// </summary>
+		StringList,
+
+		/// <summary>
+		/// A list of values and must contain at least one entry
+		/// </summary>
+		NonEmptyStringList,
+
+		/// <summary>
+		/// Accepts a string list but the value is ignored by the specifier and is automatically deferred.  This is for legacy UHT support.
+		/// </summary>
+		Legacy,
+	}
+
+	public enum UhtSpecifierDispatchResults
+	{
+		Known,
+		Unknown,
+	}
+
+	/// <summary>
+	/// The specifier context provides the default and simplest information about the specifiers being processed
+	/// </summary>
+	public class UhtSpecifierContext
+	{
+		public UhtParsingScope Scope;
+		public IUhtMessageSite MessageSite;
+		public UhtMetaData MetaData;
+		public int MetaNameIndex;
+
+		public UhtSpecifierContext(UhtParsingScope Scope, IUhtMessageSite MessageSite, UhtMetaData MetaData, int MetaNameIndex = UhtMetaData.INDEX_NONE)
+		{
+			this.Scope = Scope;
+			this.MessageSite = MessageSite;
+			this.MetaData = MetaData;
+			this.MetaNameIndex = MetaNameIndex;
+		}
+	}
+
+	/// <summary>
+	/// Specifiers are either processed immediately when the declaration is parse or deferred until later in the parsing of the object
+	/// </summary>
+	public enum UhtSpecifierWhen
+	{
+		/// <summary>
+		/// Specifier is parsed when the meta data section is parsed.
+		/// </summary>
+		Immediate,
+
+		/// <summary>
+		/// Specifier is executed after more of the object is parsed (but usually before members are parsed)
+		/// </summary>
+		Deferred,
+	}
+
+	public abstract class UhtSpecifier
+	{
+		public string Name = String.Empty;
+		public UhtSpecifierValueType ValueType;
+		public UhtSpecifierWhen When = UhtSpecifierWhen.Deferred;
+
+		public abstract UhtSpecifierDispatchResults Dispatch(UhtSpecifierContext SpecifierContext, object? Value);
+	}
+
+	public delegate void UhtSpecifierNoneDelegate(UhtSpecifierContext SpecifierContext);
+
+	public class UhtSpecifierNone : UhtSpecifier
+	{
+		private UhtSpecifierNoneDelegate Delegate;
+
+		public UhtSpecifierNone(string Name, UhtSpecifierWhen When, UhtSpecifierNoneDelegate Delegate)
+		{
+			this.Name = Name;
+			this.ValueType = UhtSpecifierValueType.None;
+			this.When = When;
+			this.Delegate = Delegate;
+		}
+
+		public override UhtSpecifierDispatchResults Dispatch(UhtSpecifierContext SpecifierContext, object? Value)
+		{
+			this.Delegate(SpecifierContext);
+			return UhtSpecifierDispatchResults.Known;
+		}
+	}
+
+	public delegate void UhtSpecifierStringDelegate(UhtSpecifierContext SpecifierContext, StringView Value);
+
+	public class UhtSpecifierString : UhtSpecifier
+	{
+		private UhtSpecifierStringDelegate Delegate;
+
+		public UhtSpecifierString(string Name, UhtSpecifierWhen When, UhtSpecifierStringDelegate Delegate)
+		{
+			this.Name = Name;
+			this.ValueType = UhtSpecifierValueType.String;
+			this.When = When;
+			this.Delegate = Delegate;
+		}
+
+		public override UhtSpecifierDispatchResults Dispatch(UhtSpecifierContext SpecifierContext, object? Value)
+		{
+			if (Value == null)
+			{
+				throw new UhtIceException("Required value is null");
+			}
+			this.Delegate(SpecifierContext, (StringView)Value);
+			return UhtSpecifierDispatchResults.Known;
+		}
+	}
+
+	public delegate void UhtSpecifierOptionalStringDelegate(UhtSpecifierContext SpecifierContext, StringView? Value);
+
+	public class UhtSpecifierOptionalString : UhtSpecifier
+	{
+		private UhtSpecifierOptionalStringDelegate Delegate;
+
+		public UhtSpecifierOptionalString(string Name, UhtSpecifierWhen When, UhtSpecifierOptionalStringDelegate Delegate)
+		{
+			this.Name = Name;
+			this.ValueType = UhtSpecifierValueType.OptionalString;
+			this.When = When;
+			this.Delegate = Delegate;
+		}
+
+		public override UhtSpecifierDispatchResults Dispatch(UhtSpecifierContext SpecifierContext, object? Value)
+		{
+			this.Delegate(SpecifierContext, (StringView?)Value);
+			return UhtSpecifierDispatchResults.Known;
+		}
+	}
+
+	public delegate void UhtSpecifierSingleStringDelegate(UhtSpecifierContext SpecifierContext, StringView Value);
+
+	public class UhtSpecifierSingleString : UhtSpecifier
+	{
+		private UhtSpecifierSingleStringDelegate Delegate;
+
+		public UhtSpecifierSingleString(string Name, UhtSpecifierWhen When, UhtSpecifierSingleStringDelegate Delegate)
+		{
+			this.Name = Name;
+			this.ValueType = UhtSpecifierValueType.SingleString;
+			this.When = When;
+			this.Delegate = Delegate;
+		}
+
+		public override UhtSpecifierDispatchResults Dispatch(UhtSpecifierContext SpecifierContext, object? Value)
+		{
+			if (Value == null)
+			{
+				throw new UhtIceException("Required value is null");
+			}
+			this.Delegate(SpecifierContext, (StringView)Value);
+			return UhtSpecifierDispatchResults.Known;
+		}
+	}
+
+	public delegate void UhtSpecifierKeyValuePairListDelegate(UhtSpecifierContext SpecifierContext, List<KeyValuePair<StringView, StringView>> Value);
+
+	public class UhtSpecifierKeyValuePairList : UhtSpecifier
+	{
+		private UhtSpecifierKeyValuePairListDelegate Delegate;
+
+		public UhtSpecifierKeyValuePairList(string Name, UhtSpecifierWhen When, bool bEqualsOptional, UhtSpecifierKeyValuePairListDelegate Delegate)
+		{
+			this.Name = Name;
+			this.ValueType = bEqualsOptional ? UhtSpecifierValueType.OptionalEqualsKeyValuePairList : UhtSpecifierValueType.KeyValuePairList;
+			this.When = When;
+			this.Delegate = Delegate;
+		}
+
+		public override UhtSpecifierDispatchResults Dispatch(UhtSpecifierContext SpecifierContext, object? Value)
+		{
+			if (Value == null)
+			{
+				throw new UhtIceException("Required value is null");
+			}
+			this.Delegate(SpecifierContext, (List<KeyValuePair<StringView, StringView>>)Value);
+			return UhtSpecifierDispatchResults.Known;
+		}
+	}
+
+	public delegate void UhtSpecifierLegacyDelegate(UhtSpecifierContext SpecifierContext);
+
+	public class UhtSpecifierLegacy : UhtSpecifier
+	{
+		private UhtSpecifierLegacyDelegate Delegate;
+
+		public UhtSpecifierLegacy(string Name, UhtSpecifierLegacyDelegate Delegate)
+		{
+			this.Name = Name;
+			this.ValueType = UhtSpecifierValueType.StringList;
+			this.When = UhtSpecifierWhen.Deferred;
+			this.Delegate = Delegate;
+		}
+
+		public override UhtSpecifierDispatchResults Dispatch(UhtSpecifierContext SpecifierContext, object? Value)
+		{
+			if (Value != null)
+			{
+				SpecifierContext.Scope.TokenReader.LogInfo($"Specifier '{this.Name}' has a value which is unused, future versions of UnrealHeaderTool will flag this as an error.");
+			}
+			this.Delegate(SpecifierContext);
+			return UhtSpecifierDispatchResults.Known;
+		}
+	}
+
+	public delegate void UhtSpecifierStringListDelegate(UhtSpecifierContext SpecifierContext, List<StringView>? Value);
+
+	public class UhtSpecifierStringList : UhtSpecifier
+	{
+		private UhtSpecifierStringListDelegate Delegate;
+
+		public UhtSpecifierStringList(string Name, UhtSpecifierWhen When, UhtSpecifierStringListDelegate Delegate)
+		{
+			this.Name = Name;
+			this.ValueType = UhtSpecifierValueType.StringList;
+			this.When = When;
+			this.Delegate = Delegate;
+		}
+
+		public override UhtSpecifierDispatchResults Dispatch(UhtSpecifierContext SpecifierContext, object? Value)
+		{
+			this.Delegate(SpecifierContext, (List<StringView>?)Value);
+			return UhtSpecifierDispatchResults.Known;
+		}
+	}
+
+	public delegate void UhtSpecifierNonEmptyStringListDelegate(UhtSpecifierContext SpecifierContext, List<StringView> Value);
+
+	public class UhtSpecifierNonEmptyStringList : UhtSpecifier
+	{
+		private UhtSpecifierNonEmptyStringListDelegate Delegate;
+
+		public UhtSpecifierNonEmptyStringList(string Name, UhtSpecifierWhen When, UhtSpecifierNonEmptyStringListDelegate Delegate)
+		{
+			this.Name = Name;
+			this.ValueType = UhtSpecifierValueType.NonEmptyStringList;
+			this.When = When;
+			this.Delegate = Delegate;
+		}
+
+		public override UhtSpecifierDispatchResults Dispatch(UhtSpecifierContext SpecifierContext, object? Value)
+		{
+			if (Value == null)
+			{
+				throw new UhtIceException("Required value is null");
+			}
+			this.Delegate(SpecifierContext, (List<StringView>)Value);
+			return UhtSpecifierDispatchResults.Known;
+		}
+	}
+
+	[AttributeUsage(AttributeTargets.Method, AllowMultiple = true)]
+	public class UhtSpecifierAttribute : Attribute
+	{
+		public string? Name;
+		public string? Extends;
+		public UhtSpecifierValueType ValueType = UhtSpecifierValueType.NotSet;
+		public UhtSpecifierWhen When = UhtSpecifierWhen.Deferred;
+	}
+
+	public class UhtSpecifierTable : UhtLookupTable<UhtSpecifier>
+	{
+
+		/// <summary>
+		/// Construct a new specifier table
+		/// </summary>
+		public UhtSpecifierTable() : base(StringViewComparer.OrdinalIgnoreCase)
+		{
+		}
+
+		/// <summary>
+		/// Add the given value to the lookup table.  It will throw an exception if it is a duplicate.
+		/// </summary>
+		/// <param name="Value">Value to be added</param>
+		public UhtSpecifierTable Add(UhtSpecifier Specifier)
+		{
+			base.Add(Specifier.Name, Specifier);
+			return this;
+		}
+	}
+
+	public class UhtSpecifierTables : UhtLookupTables<UhtSpecifierTable>
+	{
+		public static UhtSpecifierTables Instance = new UhtSpecifierTables();
+
+		public UhtSpecifierTables() : base("specifiers")
+		{
+		}
+
+		public void OnSpecifierAttribute(Type Type, MethodInfo MethodInfo, UhtSpecifierAttribute SpecifierAttribute)
+		{
+			string Name = UhtLookupTableBase.GetSuffixedName(Type, MethodInfo, SpecifierAttribute.Name, "Specifier");
+
+			if (string.IsNullOrEmpty(SpecifierAttribute.Extends))
+			{
+				throw new UhtIceException($"The 'Specifier' attribute on the {Type.Name}.{MethodInfo.Name} method doesn't have a table specified.");
+			}
+			else
+			{
+			}
+
+			if (SpecifierAttribute.ValueType == UhtSpecifierValueType.NotSet)
+			{
+				throw new UhtIceException($"The 'Specifier' attribute on the {Type.Name}.{MethodInfo.Name} method doesn't have a value type specified.");
+			}
+
+			UhtSpecifierTable Table = Get(SpecifierAttribute.Extends);
+			switch (SpecifierAttribute.ValueType)
+			{
+				case UhtSpecifierValueType.None:
+					Table.Add(new UhtSpecifierNone(Name, SpecifierAttribute.When, (UhtSpecifierNoneDelegate)Delegate.CreateDelegate(typeof(UhtSpecifierNoneDelegate), MethodInfo)));
+					break;
+				case UhtSpecifierValueType.String:
+					Table.Add(new UhtSpecifierString(Name, SpecifierAttribute.When, (UhtSpecifierStringDelegate)Delegate.CreateDelegate(typeof(UhtSpecifierStringDelegate), MethodInfo)));
+					break;
+				case UhtSpecifierValueType.OptionalString:
+					Table.Add(new UhtSpecifierOptionalString(Name, SpecifierAttribute.When, (UhtSpecifierOptionalStringDelegate)Delegate.CreateDelegate(typeof(UhtSpecifierOptionalStringDelegate), MethodInfo)));
+					break;
+				case UhtSpecifierValueType.SingleString:
+					Table.Add(new UhtSpecifierSingleString(Name, SpecifierAttribute.When, (UhtSpecifierSingleStringDelegate)Delegate.CreateDelegate(typeof(UhtSpecifierSingleStringDelegate), MethodInfo)));
+					break;
+				case UhtSpecifierValueType.KeyValuePairList:
+					Table.Add(new UhtSpecifierKeyValuePairList(Name, SpecifierAttribute.When, false, (UhtSpecifierKeyValuePairListDelegate)Delegate.CreateDelegate(typeof(UhtSpecifierKeyValuePairListDelegate), MethodInfo)));
+					break;
+				case UhtSpecifierValueType.OptionalEqualsKeyValuePairList:
+					Table.Add(new UhtSpecifierKeyValuePairList(Name, SpecifierAttribute.When, true, (UhtSpecifierKeyValuePairListDelegate)Delegate.CreateDelegate(typeof(UhtSpecifierKeyValuePairListDelegate), MethodInfo)));
+					break;
+				case UhtSpecifierValueType.StringList:
+					Table.Add(new UhtSpecifierStringList(Name, SpecifierAttribute.When, (UhtSpecifierStringListDelegate)Delegate.CreateDelegate(typeof(UhtSpecifierStringListDelegate), MethodInfo)));
+					break;
+				case UhtSpecifierValueType.NonEmptyStringList:
+					Table.Add(new UhtSpecifierNonEmptyStringList(Name, SpecifierAttribute.When, (UhtSpecifierNonEmptyStringListDelegate)Delegate.CreateDelegate(typeof(UhtSpecifierNonEmptyStringListDelegate), MethodInfo)));
+					break;
+				case UhtSpecifierValueType.Legacy:
+					Table.Add(new UhtSpecifierLegacy(Name, (UhtSpecifierLegacyDelegate)Delegate.CreateDelegate(typeof(UhtSpecifierLegacyDelegate), MethodInfo)));
+					break;
+			}
+		}
+	}
+}
diff --git a/Engine/Source/Programs/Shared/EpicGames.UHT/Tables/UhtSpecifierValidatorTable.cs b/Engine/Source/Programs/Shared/EpicGames.UHT/Tables/UhtSpecifierValidatorTable.cs
new file mode 100644
--- /dev/null
+++ b/Engine/Source/Programs/Shared/EpicGames.UHT/Tables/UhtSpecifierValidatorTable.cs
@@ -0,0 +1,73 @@
+// Copyright Epic Games, Inc. All Rights Reserved.
+
+using EpicGames.Core;
+using EpicGames.UHT.Types;
+using EpicGames.UHT.Utils;
+using System;
+using System.Reflection;
+
+namespace EpicGames.UHT.Tables
+{
+	public delegate void UhtSpecifierValidatorDelegate(UhtType Type, UhtMetaData MetaData, UhtMetaDataKey Key, StringView Value);
+
+	public class UhtSpecifierValidator
+	{
+		public string Name;
+		public UhtSpecifierValidatorDelegate Delegate;
+
+		public UhtSpecifierValidator(string Name, UhtSpecifierValidatorDelegate Delegate)
+		{
+			this.Name = Name;
+			this.Delegate = Delegate;
+		}
+	}
+
+	[AttributeUsage(AttributeTargets.Method, AllowMultiple = true)]
+	public class UhtSpecifierValidatorAttribute : Attribute
+	{
+		public string? Name;
+		public string? Extends;
+	}
+
+	public class UhtSpecifierValidatorTable : UhtLookupTable<UhtSpecifierValidator>
+	{
+
+		/// <summary>
+		/// Construct a new specifier table
+		/// </summary>
+		public UhtSpecifierValidatorTable() : base(StringViewComparer.OrdinalIgnoreCase)
+		{
+		}
+
+		/// <summary>
+		/// Add the given value to the lookup table.  It will throw an exception if it is a duplicate.
+		/// </summary>
+		/// <param name="Value">Value to be added</param>
+		public UhtSpecifierValidatorTable Add(UhtSpecifierValidator Specifier)
+		{
+			base.Add(Specifier.Name, Specifier);
+			return this;
+		}
+	}
+
+	public class UhtSpecifierValidatorTables : UhtLookupTables<UhtSpecifierValidatorTable>
+	{
+		public static UhtSpecifierValidatorTables Instance = new UhtSpecifierValidatorTables();
+
+		public UhtSpecifierValidatorTables() : base("specifier validators")
+		{
+		}
+
+		public void OnSpecifierValidatorAttribute(Type Type, MethodInfo MethodInfo, UhtSpecifierValidatorAttribute SpecifierValidatorAttribute)
+		{
+			string Name = UhtLookupTableBase.GetSuffixedName(Type, MethodInfo, SpecifierValidatorAttribute.Name, "SpecifierValidator");
+
+			if (string.IsNullOrEmpty(SpecifierValidatorAttribute.Extends))
+			{
+				throw new UhtIceException($"The 'SpecifierValidator' attribute on the {Type.Name}.{MethodInfo.Name} method doesn't have a table specified.");
+			}
+			UhtSpecifierValidatorTable Table = Get(SpecifierValidatorAttribute.Extends);
+			Table.Add(new UhtSpecifierValidator(Name, (UhtSpecifierValidatorDelegate)Delegate.CreateDelegate(typeof(UhtSpecifierValidatorDelegate), MethodInfo)));
+		}
+	}
+}
