UHT in UBT initial version.  Currently disabled.

diff --git a/Engine/Source/Programs/Shared/EpicGames.UHT/Tables/UhtStructDefaultValueTable.cs b/Engine/Source/Programs/Shared/EpicGames.UHT/Tables/UhtStructDefaultValueTable.cs
new file mode 100644
--- /dev/null
+++ b/Engine/Source/Programs/Shared/EpicGames.UHT/Tables/UhtStructDefaultValueTable.cs
@@ -0,0 +1,136 @@
+// Copyright Epic Games, Inc. All Rights Reserved.
+
+using EpicGames.Core;
+using EpicGames.UHT.Tokenizer;
+using System;
+using System.Collections.Generic;
+using System.Reflection;
+using EpicGames.UHT.Types;
+using EpicGames.UHT.Utils;
+using System.Text;
+
+namespace EpicGames.UHT.Tables
+{
+	public delegate bool UhtStructDefaultValueDelegate(UhtStructProperty Property, IUhtTokenReader DefaultValueReader, StringBuilder InnerDefaultValue);
+
+	[Flags]
+	public enum UhtStructDefaultValueOptions
+	{
+		None = 0,
+
+		/// <summary>
+		/// This method is to be invoked when there are no keyword matches found
+		/// </summary>
+		Default = 1 << 2,
+	}
+
+	/// <summary>
+	/// Helper methods for testing flags.  These methods perform better than the generic HasFlag which hits
+	/// the GC and stalls.
+	/// </summary>
+	public static class UhtStructDefaultValueOptionsExtensions
+	{
+
+		/// <summary>
+		/// Test to see if any of the specified flags are set
+		/// </summary>
+		/// <param name="InFlags">Current flags</param>
+		/// <param name="TestFlags">Flags to test for</param>
+		/// <returns>True if any of the flags are set</returns>
+		public static bool HasAnyFlags(this UhtStructDefaultValueOptions InFlags, UhtStructDefaultValueOptions TestFlags)
+		{
+			return (InFlags & TestFlags) != 0;
+		}
+
+		/// <summary>
+		/// Test to see if all of the specified flags are set
+		/// </summary>
+		/// <param name="InFlags">Current flags</param>
+		/// <param name="TestFlags">Flags to test for</param>
+		/// <returns>True if all the flags are set</returns>
+		public static bool HasAllFlags(this UhtStructDefaultValueOptions InFlags, UhtStructDefaultValueOptions TestFlags)
+		{
+			return (InFlags & TestFlags) == TestFlags;
+		}
+
+		/// <summary>
+		/// Test to see if a specific set of flags have a specific value.
+		/// </summary>
+		/// <param name="InFlags">Current flags</param>
+		/// <param name="TestFlags">Flags to test for</param>
+		/// <param name="MatchFlags">Expected value of the tested flags</param>
+		/// <returns>True if the given flags have a specific value.</returns>
+		public static bool HasExactFlags(this UhtStructDefaultValueOptions InFlags, UhtStructDefaultValueOptions TestFlags, UhtStructDefaultValueOptions MatchFlags)
+		{
+			return (InFlags & TestFlags) == MatchFlags;
+		}
+	}
+
+	[AttributeUsage(AttributeTargets.Method, AllowMultiple = true)]
+	public class UhtStructDefaultValueAttribute : Attribute
+	{
+		public string? Name;
+		public UhtStructDefaultValueOptions Options = UhtStructDefaultValueOptions.None;
+	}
+
+	public struct UhtStructDefaultValue
+	{
+		public UhtStructDefaultValueDelegate Delegate;
+	}
+
+	public class UhtStructDefaultValueTable
+	{
+		public static UhtStructDefaultValueTable Instance = new UhtStructDefaultValueTable();
+
+		private Dictionary<StringView, UhtStructDefaultValue> StructDefaultValues = new Dictionary<StringView, UhtStructDefaultValue>();
+		private UhtStructDefaultValue? DefaultInternal = null;
+		public UhtStructDefaultValue Default
+		{
+			get
+			{
+				if (this.DefaultInternal == null)
+				{
+					throw new UhtIceException("No struct default value has been marked as default");
+				}
+				return (UhtStructDefaultValue)this.DefaultInternal;
+			}
+		}
+
+		/// <summary>
+		/// Return the structure default value associated with the given name
+		/// </summary>
+		/// <param name="Name"></param>
+		/// <param name="StructDefaultValue">Structure default value handler</param>
+		/// <returns></returns>
+		public bool TryGet(StringView Name, out UhtStructDefaultValue StructDefaultValue)
+		{
+			return this.StructDefaultValues.TryGetValue(Name, out StructDefaultValue);
+		}
+
+		public void OnStructDefaultValueAttribute(Type Type, MethodInfo MethodInfo, UhtStructDefaultValueAttribute StructDefaultValueAttribute)
+		{
+			if (string.IsNullOrEmpty(StructDefaultValueAttribute.Name) && !StructDefaultValueAttribute.Options.HasAnyFlags(UhtStructDefaultValueOptions.Default))
+			{
+				throw new UhtIceException("A struct default value attribute must have a name or be marked as default");
+			}
+
+			UhtStructDefaultValue StructDefaultValue = new UhtStructDefaultValue
+			{
+				Delegate = (UhtStructDefaultValueDelegate)Delegate.CreateDelegate(typeof(UhtStructDefaultValueDelegate), MethodInfo)
+			};
+
+			if (StructDefaultValueAttribute.Options.HasAnyFlags(UhtStructDefaultValueOptions.Default))
+			{
+				if (this.DefaultInternal != null)
+				{
+					throw new UhtIceException("Only one struct default value attribute dispatcher can be marked as default");
+				}
+				this.DefaultInternal = StructDefaultValue;
+			}
+			else if (!string.IsNullOrEmpty(StructDefaultValueAttribute.Name))
+			{
+				this.StructDefaultValues.Add(new StringView(StructDefaultValueAttribute.Name), StructDefaultValue);
+			}
+		}
+	}
+}
diff --git a/Engine/Source/Programs/Shared/EpicGames.UHT/Tokenizer/UhtToken.cs b/Engine/Source/Programs/Shared/EpicGames.UHT/Tokenizer/UhtToken.cs
new file mode 100644
--- /dev/null
+++ b/Engine/Source/Programs/Shared/EpicGames.UHT/Tokenizer/UhtToken.cs
@@ -0,0 +1,498 @@
+// Copyright Epic Games, Inc. All Rights Reserved.
+
+using EpicGames.Core;
+using EpicGames.UHT.Utils;
+using System;
+using System.Collections.Generic;
+using System.Globalization;
+using System.Text;
+
+namespace EpicGames.UHT.Tokenizer
+{
+	public enum UhtTokenType
+	{
+		EndOfFile,          // End of file token.
+		EndOfDefault,       // End of default value.
+		EndOfType,          // End of type
+		EndOfDeclaration,	// End of declaration
+		Line,				// Line of text (when calling GetLine only)
+		Identifier,			// Alphanumeric identifier.
+		Symbol,				// Symbol.
+		TrueConst,			// True value via identifier
+		FalseConst,			// False value via identifier
+		FloatConst,			// Floating point constant
+		DecimalConst,		// Decimal Integer constant
+		HexConst,			// Hex integer constant
+		CharConst,			// Single character constant
+		StringConst,		// String constant
+	}
+
+	public static class UhtTokenTypeExtensions
+	{
+		public static bool IsEndType(this UhtTokenType TokenType)
+		{
+			switch (TokenType)
+			{
+				case UhtTokenType.EndOfFile:
+				case UhtTokenType.EndOfDefault:
+				case UhtTokenType.EndOfType:
+				case UhtTokenType.EndOfDeclaration:
+					return true;
+				default:
+					return false;
+			}
+		}
+	}
+
+	public struct UhtToken
+	{
+		public const int MaxNameLength = 1024;
+		public const int MaxStringLength = 1024;
+
+		public UhtTokenType TokenType { get; set; }
+		public int UngetPos { get; set; }
+		public int UngetLine { get; set; }
+		public int InputStartPos { get; set; }
+		public int InputEndPos { get => this.InputStartPos + this.Value.Span.Length; }
+		public int InputLine { get; set; }
+		public StringView Value { get; set; }
+
+		public UhtToken(UhtTokenType TokenType)
+		{
+			this.TokenType = TokenType;
+			this.UngetPos = 0;
+			this.UngetLine = 0;
+			this.InputStartPos = 0;
+			this.InputLine = 0;
+			this.Value = new StringView();
+		}
+
+		public UhtToken(UhtTokenType TokenType, int UngetPos, int UngetLine, int InputStartPos, int InputLine, StringView Value)
+		{
+			this.TokenType = TokenType;
+			this.UngetPos = UngetPos;
+			this.UngetLine = UngetLine;
+			this.InputStartPos = InputStartPos;
+			this.InputLine = InputLine;
+			this.Value = Value;
+		}
+
+		public static implicit operator bool(UhtToken Token)
+		{
+			return !Token.IsEndType();
+		}
+
+		public bool IsEndType()
+		{
+			return this.TokenType.IsEndType();
+		}
+
+		public bool IsValue(char Value)
+		{
+			return this.Value.Span.Length == 1 && this.Value.Span[0] == Value;
+		}
+
+		public bool IsValue(string Value, bool bIgnoreCase = false)
+		{
+			return this.Value.Span.Equals(Value, bIgnoreCase ? StringComparison.OrdinalIgnoreCase : StringComparison.Ordinal);
+		}
+
+		public bool IsValue(StringView Value, bool bIgnoreCase = false)
+		{
+			return this.Value.Span.Equals(Value.Span, bIgnoreCase ? StringComparison.OrdinalIgnoreCase : StringComparison.Ordinal);
+		}
+
+		public bool ValueStartsWith(string Value, bool bIgnoreCase = false)
+		{
+			return this.Value.Span.StartsWith(Value, bIgnoreCase ? StringComparison.OrdinalIgnoreCase : StringComparison.Ordinal);
+		}
+
+		public bool IsIdentifier()
+		{
+			return this.TokenType == UhtTokenType.Identifier;
+		}
+
+		public bool IsIdentifier(string Identifier, bool bIgnoreCase = false)
+		{
+			return IsIdentifier() && IsValue(Identifier, bIgnoreCase);
+		}
+
+		public bool IsIdentifier(StringView Identifier, bool bIgnoreCase = false)
+		{
+			return IsIdentifier() && IsValue(Identifier, bIgnoreCase);
+		}
+
+		public bool IsSymbol()
+		{
+			return this.TokenType == UhtTokenType.Symbol;
+		}
+
+		public bool IsSymbol(char Symbol)
+		{
+			return IsSymbol() && IsValue(Symbol);
+		}
+
+		public bool IsSymbol(string Symbol)
+		{
+			return IsSymbol() && IsValue(Symbol);
+		}
+
+		public bool IsSymbol(StringView Symbol)
+		{
+			return IsSymbol() && IsValue(Symbol);
+		}
+
+		public bool IsConstInt()
+		{
+			return this.TokenType == UhtTokenType.DecimalConst || this.TokenType == UhtTokenType.HexConst;
+		}
+
+		public bool IsConstFloat()
+		{
+			return this.TokenType == UhtTokenType.FloatConst;
+		}
+
+		// Get the integer value of the token.  Only supported for decimal, hexidecimal, and floating point values
+		public bool GetConstInt(out int Value)
+		{
+			switch (TokenType)
+			{
+				case UhtTokenType.DecimalConst:
+					Value = (int)GetDecimalValue();
+					return true;
+				case UhtTokenType.HexConst:
+					Value = (int)GetHexValue();
+					return true;
+				case UhtTokenType.FloatConst:
+					{
+						float Float = GetFloatValue();
+						Value = (int)Float;
+						return Float == Value;
+					}
+				default:
+					Value = 0;
+					return false;
+			}
+		}
+
+		// Get the integer value of the token.  Only supported for decimal, hexadecimal, and floating point values
+		public bool GetConstLong(out long Value)
+		{
+			switch (TokenType)
+			{
+				case UhtTokenType.DecimalConst:
+					Value = GetDecimalValue();
+					return true;
+				case UhtTokenType.HexConst:
+					Value = GetHexValue();
+					return true;
+				case UhtTokenType.FloatConst:
+					{
+						float Float = GetFloatValue();
+						Value = (long)Float;
+						return Float == Value;
+					}
+				default:
+					Value = 0;
+					return false;
+			}
+		}
+
+		// Get the float value of the token.  Only supported for decimal, hexadecimal, and floating point values
+		public bool GetConstFloat(out float Value)
+		{
+			switch (TokenType)
+			{
+				case UhtTokenType.DecimalConst:
+					Value = (float)GetDecimalValue();
+					return true;
+				case UhtTokenType.HexConst:
+					Value = (float)GetHexValue();
+					return true;
+				case UhtTokenType.FloatConst:
+					Value = GetFloatValue();
+					return true;
+				default:
+					Value = 0;
+					return false;
+			}
+		}
+
+		// Get the double value of the token.  Only supported for decimal, hexadecimal, and floating point values
+		public bool GetConstDouble(out double Value)
+		{
+			switch (TokenType)
+			{
+				case UhtTokenType.DecimalConst:
+					Value = (float)GetDecimalValue();
+					return true;
+				case UhtTokenType.HexConst:
+					Value = (float)GetHexValue();
+					return true;
+				case UhtTokenType.FloatConst:
+					Value = GetDoubleValue();
+					return true;
+				default:
+					Value = 0;
+					return false;
+			}
+		}
+
+		public bool IsConstString()
+		{
+			return this.TokenType == UhtTokenType.StringConst || this.TokenType == UhtTokenType.CharConst;
+		}
+
+		// Return the token value for string constants
+		public StringView GetUnescapedString(IUhtMessageSite MessageSite)
+		{
+			switch (this.TokenType)
+			{
+				case UhtTokenType.StringConst:
+					ReadOnlySpan<char> Span = this.Value.Span.Slice(1, this.Value.Span.Length - 2);
+					int Index = Span.IndexOf('\\');
+					if (Index == -1)
+					{
+						return new StringView(this.Value, 1, Span.Length);
+					}
+					else
+					{
+						StringBuilder Builder = new StringBuilder();
+						while (Index >= 0)
+						{
+							Builder.Append(Span.Slice(0, Index));
+							if (Span[Index + 1] == 'n')
+							{
+								Builder.Append('\n');
+							}
+							Span = Span.Slice(Index + 1);
+							Index = Span.IndexOf('\\');
+						}
+						Builder.Append(Span);
+						return new StringView(Builder.ToString());
+					}
+
+				case UhtTokenType.CharConst:
+					if (this.Value.Span[1] == '\\')
+					{
+						switch (this.Value.Span[2])
+						{
+							case 't':
+								return new StringView("\t");
+							case 'n':
+								return new StringView("\n");
+							case 'r':
+								return new StringView("\r");
+							default:
+								return new StringView(this.Value, 2, 1);
+						}
+					}
+					else
+					{
+						return new StringView(this.Value, 1, 1);
+					}
+			}
+
+			throw new UhtException(MessageSite, this.InputLine, "Call to GetUnescapedString on token that isn't a string or char constant");
+		}
+
+		public StringView GetConstantValue(bool bRespectQuotes = false)
+		{
+			switch (this.TokenType)
+			{
+				case UhtTokenType.DecimalConst:
+					return new StringView(GetDecimalValue().ToString(NumberFormatInfo.InvariantInfo));
+				case UhtTokenType.HexConst:
+					return new StringView(GetHexValue().ToString(NumberFormatInfo.InvariantInfo));
+				case UhtTokenType.FloatConst:
+					return new StringView(GetFloatValue().ToString("F6", NumberFormatInfo.InvariantInfo));
+				case UhtTokenType.CharConst:
+				case UhtTokenType.StringConst:
+					return GetTokenString(bRespectQuotes);
+				default:
+					return "NotConstant";
+			}
+		}
+
+		public StringView GetTokenString(bool bRespectQuotes = false)
+		{
+			StringViewBuilder SVB = new StringViewBuilder();
+			switch (this.TokenType)
+			{
+				case UhtTokenType.StringConst:
+					StringView SubValue = new StringView(this.Value, 1, this.Value.Span.Length - 2);
+					while (SubValue.Span.Length > 0)
+					{
+						int SlashIndex = SubValue.Span.IndexOf('\\');
+						if (SlashIndex == -1)
+						{
+							SVB.Append(SubValue);
+							break;
+						}
+						if (SlashIndex > 0)
+						{
+							SVB.Append(new StringView(SubValue, 0, SlashIndex));
+						}
+						if (SlashIndex + 1 == SubValue.Span.Length)
+						{
+							break;
+						}
+
+						if (SlashIndex + 1 < SubValue.Span.Length)
+						{
+							char c = SubValue.Span[SlashIndex + 1];
+							if (c == 'n')
+							{
+								c = '\n';
+							}
+							else if (bRespectQuotes && c == '"')
+							{
+								SVB.Append('\\');
+							}
+							SVB.Append(c);
+							SubValue = new StringView(SubValue, SlashIndex + 2);
+						}
+					}
+					break;
+
+				case UhtTokenType.CharConst:
+					char CharConst = this.Value.Span[1];
+					if (CharConst == '\\')
+					{
+						CharConst = this.Value.Span[2];
+						switch (CharConst)
+						{
+							case 't':
+								CharConst = '\t';
+								break;
+							case 'n':
+								CharConst = '\n';
+								break;
+							case 'r':
+								CharConst = '\r';
+								break;
+						}
+					}
+					SVB.Append(CharConst);
+					break;
+
+				default:
+					throw new UhtIceException("Call to GetTokenString on a token that isn't a string or char constant");
+			}
+			return SVB.ToStringView();
+		}
+
+		float GetFloatValue()
+		{
+			if (Value.Span.Length > 0)
+			{
+				if (UhtFCString.IsFloatMarker(Value.Span[Value.Span.Length - 1]))
+				{
+					return float.Parse(Value.Span.Slice(0, Value.Span.Length - 1));
+				}
+			}
+			return float.Parse(Value.Span);
+		}
+
+		double GetDoubleValue()
+		{
+			if (Value.Span.Length > 0)
+			{
+				if (UhtFCString.IsFloatMarker(Value.Span[Value.Span.Length - 1]))
+				{
+					return double.Parse(Value.Span.Slice(0, Value.Span.Length - 1));
+				}
+			}
+			return double.Parse(Value.Span);
+		}
+
+		long GetDecimalValue()
+		{
+			ReadOnlySpan<char> Span = this.Value.Span;
+			bool bIsUnsigned = false;
+			while (Span.Length > 0)
+			{
+				char C = Span[Span.Length - 1];
+				if (UhtFCString.IsLongMarker(C))
+				{
+					Span = Span.Slice(0, Span.Length - 1);
+				}
+				else if (UhtFCString.IsUnsignedMarker(C))
+				{
+					bIsUnsigned = true;
+					Span = Span.Slice(0, Span.Length - 1);
+				}
+				else
+				{
+					break;
+				}
+			}
+			return bIsUnsigned ? (long)Convert.ToUInt64(Span.ToString(), 10) : Convert.ToInt64(Span.ToString(), 10);
+		}
+
+		long GetHexValue()
+		{
+			return Convert.ToInt64(this.Value.ToString(), 16);
+		}
+
+		public static string Join(IEnumerable<UhtToken> Tokens)
+		{
+			StringBuilder Builder = new StringBuilder();
+			foreach (var Token in Tokens)
+			{
+				Builder.Append(Token.Value.ToString());
+			}
+			return Builder.ToString();
+		}
+
+		public static string Join(char Seperator, IEnumerable<UhtToken> Tokens)
+		{
+			StringBuilder Builder = new StringBuilder();
+			bool bIncludeSeperator = false;
+			foreach (var Token in Tokens)
+			{
+				if (!bIncludeSeperator)
+				{
+					bIncludeSeperator = true;
+				}
+				else
+				{
+					Builder.Append(Seperator);
+				}
+				Builder.Append(Token.Value.ToString());
+			}
+			return Builder.ToString();
+		}
+
+		public static string Join(string Seperator, IEnumerable<UhtToken> Tokens)
+		{
+			StringBuilder Builder = new StringBuilder();
+			bool bIncludeSeperator = false;
+			foreach (var Token in Tokens)
+			{
+				if (!bIncludeSeperator)
+				{
+					bIncludeSeperator = true;
+				}
+				else
+				{
+					Builder.Append(Seperator);
+				}
+				Builder.Append(Token.Value.ToString());
+			}
+			return Builder.ToString();
+		}
+
+		public override string ToString()
+		{
+			if (IsEndType())
+			{
+				return "<none>";
+			}
+			else
+			{
+				return Value.Span.ToString();
+			}
+		}
+	}
+}
diff --git a/Engine/Source/Programs/Shared/EpicGames.UHT/Tokenizer/UhtTokenBufferReader.cs b/Engine/Source/Programs/Shared/EpicGames.UHT/Tokenizer/UhtTokenBufferReader.cs
new file mode 100644
--- /dev/null
+++ b/Engine/Source/Programs/Shared/EpicGames.UHT/Tokenizer/UhtTokenBufferReader.cs
@@ -0,0 +1,1051 @@
+// Copyright Epic Games, Inc. All Rights Reserved.
+
+using EpicGames.Core;
+using EpicGames.UHT.Utils;
+using System;
+using System.Collections.Generic;
+using System.Diagnostics;
+
+namespace EpicGames.UHT.Tokenizer
+{
+	public class UhtTokenBufferReader : IUhtTokenReader, IUhtMessageLineNumber
+	{
+		private static StringView[] EmptyComments = new StringView[] { };
+
+		private readonly IUhtMessageSite MessageSiteInternal;
+		private List<StringView>? CommentsInternal = null;
+		private List<UhtToken> RecordedTokensInternal = new List<UhtToken>();
+		private IUhtTokenPreprocessor? TokenPreprocessorInternal = null;
+		private UhtToken CurrentToken = new UhtToken(); // PeekToken must have been invoked first
+		private ReadOnlyMemory<char> Data;
+		private int PrevPos = 0;
+		private int PrevLine = 1;
+		private bool bHasToken = false;
+		private int PreCurrentTokenInputPos = 0;
+		private int PreCurrentTokenInputLine = 1;
+		private int InputPosInternal = 0;
+		private int InputLineInternal = 1;
+		private int CommentsDisableCount = 0;
+		private int CommittedComments = 0;
+		private List<StringView>? SavedComments = null;
+		private int SavedInputPos = 0;
+		private int SavedInputLine = 0;
+		private bool bHasSavedState = false;
+		private bool bRecordTokens = false;
+		private int PreprocessorPendingCommentsCount = 0;
+
+		public UhtTokenBufferReader(IUhtMessageSite MessageSite, ReadOnlyMemory<char> Input)
+		{
+			this.MessageSiteInternal = MessageSite;
+			this.Data = Input;
+		}
+
+		#region IUHTMessageSite Implementation
+		IUhtMessageSession IUhtMessageSite.MessageSession => this.MessageSiteInternal.MessageSession;
+		IUhtMessageSource? IUhtMessageSite.MessageSource => this.MessageSiteInternal.MessageSource;
+		IUhtMessageLineNumber? IUhtMessageSite.MessageLineNumber => this;
+		#endregion
+
+		#region IUHTMessageLinenumber Implementation
+		int IUhtMessageLineNumber.MessageLineNumber => this.InputLine;
+		#endregion
+
+		#region ITokenReader Implementation
+		public bool bIsEOF
+		{
+			get
+			{
+				if (this.bHasToken)
+				{
+					return this.CurrentToken.TokenType.IsEndType();
+				}
+				else
+				{
+					return this.InputPos == this.Data.Length;
+				}
+			}
+		}
+
+		public int InputPos
+		{
+			get => this.bHasToken ? this.PreCurrentTokenInputPos : this.InputPosInternal;
+		}
+
+		public int InputLine
+		{
+			get => this.bHasToken ? this.PreCurrentTokenInputLine : this.InputLineInternal;
+			set { ClearToken(); this.InputLineInternal = value; }
+		}
+
+		public ReadOnlySpan<StringView> Comments 
+		{
+			get
+			{
+				if (this.CommentsInternal != null && this.CommittedComments != 0)
+				{
+					return new ReadOnlySpan<StringView>(this.CommentsInternal.ToArray(), 0, this.CommittedComments);
+				}
+				else
+				{
+					return new ReadOnlySpan<StringView>(UhtTokenBufferReader.EmptyComments);
+				}
+			}
+		}
+
+		public IUhtTokenPreprocessor? TokenPreprocessor { get => this.TokenPreprocessorInternal; set => this.TokenPreprocessorInternal = value; }
+
+		public ref UhtToken PeekToken()
+		{
+			if (!this.bHasToken)
+			{
+				this.CurrentToken = GetTokenInternal(true);
+				this.bHasToken = true;
+			}
+			return ref this.CurrentToken;
+		}
+
+		/// <summary>
+		/// Skip all leading whitespace and collect any comments
+		/// </summary>
+		public void SkipWhitespaceAndComments()
+		{
+			bool bGotInlineComment = false;
+			SkipWhitespaceAndCommentsInternal(ref bGotInlineComment, true);
+		}
+
+		/// <inheritdoc/>
+		public void ConsumeToken()
+		{
+			if (this.bRecordTokens && !this.CurrentToken.IsEndType())
+			{
+				this.RecordedTokensInternal.Add(this.CurrentToken);
+			}
+			this.bHasToken = false;
+
+			// When comments are disabled, we are still collecting comments, but aren't committing them
+			if (this.CommentsDisableCount == 0)
+			{
+				if (this.CommentsInternal != null)
+				{
+					this.CommittedComments = this.CommentsInternal.Count;
+				}
+			}
+			else
+			{
+				ClearPendingComments();
+			}
+		}
+
+		/// <inheritdoc/>
+		public UhtToken GetToken()
+		{
+			UhtToken Token = PeekToken();
+			ConsumeToken();
+			return Token;
+		}
+
+		/// <inheritdoc/>
+		public StringView GetRawString(char Terminator, UhtRawStringOptions Options)
+		{
+			ReadOnlySpan<char> Span = this.Data.Span;
+
+			ClearToken();
+			SkipWhitespaceAndComments();
+
+			int StartPos = this.InputPos;
+			bool bInQuotes = false;
+			while (true)
+			{
+				char C = InternalGetChar(Span);
+
+				// Check for end of file
+				if (C == 0)
+				{
+					break;
+				}
+
+				// Check for end of line
+				if (C == '\r' || C == '\n')
+				{
+					--this.InputPosInternal;
+					break;
+				}
+
+				// Check for terminator as long as we aren't in quotes
+				if (C == Terminator && !bInQuotes)
+				{
+					if (Options.HasAnyFlags(UhtRawStringOptions.DontConsumeTerminator))
+					{
+						--this.InputPosInternal;
+					}
+					break;
+				}
+
+				// Check for comment
+				if (!bInQuotes && C == '/')
+				{
+					char P = InternalPeekChar(Span);
+					if (P == '*' || P == '/')
+					{
+						--this.InputPosInternal;
+						break;
+					}
+				}
+
+				// Check for quotes
+				if (C == '"' && Options.HasAnyFlags(UhtRawStringOptions.RespectQuotes))
+				{
+					bInQuotes = !bInQuotes;
+				}
+			}
+
+			// If EOF, then error
+			if (bInQuotes)
+			{
+				throw new UhtException(this, "Unterminated quoted string");
+			}
+
+			// Remove trailing whitespace
+			int EndPos = this.InputPos;
+			for (; EndPos > StartPos; --EndPos)
+			{
+				char C = Span[EndPos - 1];
+				if (C != ' ' && C != '\t')
+				{
+					break;
+				}
+			}
+
+			// Check for too long
+			if (EndPos - StartPos >= UhtToken.MaxStringLength)
+			{
+				throw new UhtException(this, $"String exceeds maximum of {UhtToken.MaxStringLength} characters");
+			}
+			return new StringView(this.Data, StartPos, EndPos - StartPos);
+		}
+
+		/// <inheritdoc/>
+		public UhtToken GetLine()
+		{
+			ReadOnlySpan<char> Span = this.Data.Span;
+
+			ClearToken();
+			this.PrevPos = this.InputPosInternal;
+			this.PrevLine = this.InputLineInternal;
+
+			if (this.PrevPos == Span.Length)
+			{
+				return new UhtToken(UhtTokenType.EndOfFile, this.PrevPos, this.PrevLine, this.PrevPos, this.PrevLine, new StringView(this.Data.Slice(this.PrevPos, 0)));
+			}
+
+			int LastPos = this.InputPosInternal;
+			while (true)
+			{
+				char Char = InternalGetChar(Span);
+				if (Char == 0)
+				{
+					break;
+				}
+				else if (Char == '\r')
+				{
+				}
+				else if (Char == '\n')
+				{
+					++this.InputLineInternal;
+					break;
+				}
+				else
+				{
+					LastPos = this.InputPosInternal;
+				}
+			}
+
+			return new UhtToken(UhtTokenType.Line, this.PrevPos, this.PrevLine, this.PrevPos, this.PrevLine, new StringView(this.Data.Slice(this.PrevPos, LastPos - this.PrevPos)));
+		}
+
+		/// <inheritdoc/>
+		public StringView GetStringView(int StartPos, int Count)
+		{
+			return new StringView(this.Data, StartPos, Count);
+		}
+
+		/// <inheritdoc/>
+		public void ClearComments()
+		{
+			if (this.PreprocessorPendingCommentsCount != 0)
+			{
+				Debugger.Break();
+			}
+			if (this.CommentsInternal != null)
+			{
+
+				// Clearing comments does not remove any uncommitted comments
+				this.CommentsInternal.RemoveRange(0, this.CommittedComments);
+				this.CommittedComments = 0;
+			}
+		}
+
+		/// <inheritdoc/>
+		public void DisableComments()
+		{
+			++this.CommentsDisableCount;
+		}
+
+		/// <inheritdoc/>
+		public void EnableComments()
+		{
+			--this.CommentsDisableCount;
+		}
+
+		/// <inheritdoc/>
+		public void CommitPendingComments()
+		{
+			if (this.CommentsInternal != null)
+			{
+				this.CommittedComments = this.CommentsInternal.Count;
+			}
+		}
+
+		/// <inheritdoc/>
+		public bool IsFirstTokenInLine(ref UhtToken Token)
+		{
+			return IsFirstTokenInLine(this.Data.Span, Token.InputStartPos);
+		}
+
+		/// <inheritdoc/>
+		public void SaveState()
+		{
+			if (this.bHasSavedState)
+			{
+				throw new UhtIceException("Can not save more than one state");
+			}
+			this.bHasSavedState = true;
+			this.SavedInputLine = this.InputLine;
+			this.SavedInputPos = this.InputPos;
+			if (this.CommentsInternal != null)
+			{
+				if (this.SavedComments == null)
+				{
+					this.SavedComments = new List<StringView>();
+				}
+				this.SavedComments.Clear();
+				this.SavedComments.AddRange(this.CommentsInternal);
+			}
+			if (this.TokenPreprocessorInternal != null)
+			{
+				this.TokenPreprocessorInternal.SaveState();
+			}
+		}
+
+		/// <inheritdoc/>
+		public void RestoreState()
+		{
+			if (!this.bHasSavedState)
+			{
+				throw new UhtIceException("Can not restore state when none have been saved");
+			}
+			this.bHasSavedState = false;
+			ClearToken();
+			this.InputPosInternal = this.SavedInputPos;
+			this.InputLineInternal = this.SavedInputLine;
+			if (this.SavedComments != null && this.CommentsInternal != null)
+			{
+				this.CommentsInternal.Clear();
+				this.CommentsInternal.AddRange(this.SavedComments);
+			}
+			if (this.TokenPreprocessorInternal != null)
+			{
+				this.TokenPreprocessorInternal.RestoreState();
+			}
+		}
+
+		/// <inheritdoc/>
+		public void AbandonState()
+		{
+			if (!this.bHasSavedState)
+			{
+				throw new UhtIceException("Can not abandon state when none have been saved");
+			}
+			this.bHasSavedState = false;
+			ClearToken();
+		}
+
+		/// <inheritdoc/>
+		public void EnableRecording()
+		{
+			if (this.bRecordTokens)
+			{
+				throw new UhtIceException("Can not nest token recording");
+			}
+			this.bRecordTokens = true;
+		}
+
+		/// <inheritdoc/>
+		public void RecordToken(ref UhtToken Token)
+		{
+			if (!this.bRecordTokens)
+			{
+				throw new UhtIceException("Attempt to disable token recording when it isn't enabled");
+			}
+			this.RecordedTokensInternal.Add(Token);
+		}
+
+		/// <inheritdoc/>
+		public void DisableRecording()
+		{
+			if (!this.bRecordTokens)
+			{
+				throw new UhtIceException("Attempt to disable token recording when it isn't enabled");
+			}
+			this.RecordedTokensInternal.Clear();
+			this.bRecordTokens = false;
+		}
+
+		/// <inheritdoc/>
+		public List<UhtToken> RecordedTokens
+		{
+			get
+			{
+				if (!this.bRecordTokens)
+				{
+					throw new UhtIceException("Attempt to get recorded tokens when it isn't enabled");
+				}
+				return this.RecordedTokensInternal;
+			}
+		}
+		#endregion
+
+		#region Internals
+		private void ClearAllComments()
+		{
+			if (this.PreprocessorPendingCommentsCount != 0)
+			{
+				Debugger.Break();
+			}
+			if (this.CommentsInternal != null)
+			{
+				this.CommentsInternal.Clear();
+				this.CommittedComments = 0;
+			}
+		}
+
+		private void ClearPendingComments()
+		{
+			if (this.CommentsInternal != null)
+			{
+				int StartingComment = this.CommittedComments + this.PreprocessorPendingCommentsCount;
+				this.CommentsInternal.RemoveRange(StartingComment, this.CommentsInternal.Count - StartingComment);
+			}
+		}
+
+		private bool IsFirstTokenInLine(ReadOnlySpan<char> Span, int StartPos)
+		{
+			for (int Pos = StartPos; --Pos > 0;)
+			{
+				switch (Span[Pos])
+				{
+					case '\r':
+					case '\n':
+						return true;
+					case ' ':
+					case '\t':
+						break;
+					default:
+						return false;
+				}
+			}
+			return true;
+		}
+
+		/// <summary>
+		/// Get the next token
+		/// </summary>
+		/// <returns>Return the next token from the stream.  If the end of stream is reached, a token type of None will be returned.</returns>
+		private UhtToken GetTokenInternal(bool bEnablePreprocessor)
+		{
+			ReadOnlySpan<char> Span = this.Data.Span;
+			int SpanLength = Span.Length;
+			bool bGotInlineComment = false;
+		Restart:
+			this.PreCurrentTokenInputLine = this.InputLineInternal;
+			this.PreCurrentTokenInputPos = this.InputPosInternal;
+			SkipWhitespaceAndCommentsInternal(ref bGotInlineComment, bEnablePreprocessor);
+			this.PrevPos = this.InputPosInternal;
+			this.PrevLine = this.InputLineInternal;
+
+			UhtTokenType TokenType = UhtTokenType.EndOfFile;
+			int StartPos = this.InputPosInternal;
+			int StartLine = this.InputLineInternal;
+
+			char C = InternalGetChar(Span);
+			if (C == 0)
+			{
+			}
+
+			else if (UhtFCString.IsAlpha(C) || C == '_')
+			{
+
+				for (; this.InputPosInternal < Span.Length; ++this.InputPosInternal)
+				{
+					C = Span[this.InputPosInternal];
+					if (!(UhtFCString.IsAlpha(C) || UhtFCString.IsDigit(C) || C == '_'))
+					{
+						break;
+					}
+				}
+
+				if (this.InputPosInternal - StartPos >= UhtToken.MaxNameLength)
+				{
+					throw new UhtException(this, $"Identifier length exceeds maximum of {UhtToken.MaxNameLength}");
+				}
+
+				TokenType = UhtTokenType.Identifier;
+			}
+
+			// Check for any numerics 
+			else if (IsNumeric(Span, C))
+			{
+				// Integer or floating point constant.
+
+				bool bIsFloat = C == '.';
+				bool bIsHex = false;
+
+				// Ignore the starting sign for this part of the parsing
+				if (UhtFCString.IsSign(C))
+				{
+					C = InternalGetChar(Span); // We know this won't fail since the code above checked the Peek
+				}
+
+				// Check for a hex valid
+				if (C == '0')
+				{
+					bIsHex = UhtFCString.IsHexMarker(InternalPeekChar(Span));
+					if (bIsHex)
+					{
+						InternalGetChar(Span);
+					}
+				}
+
+				// If we have a hex constant
+				if (bIsHex)
+				{
+					for (; this.InputPosInternal < Span.Length && UhtFCString.IsHexDigit(Span[this.InputPosInternal]); ++this.InputPosInternal)
+					{
+					}
+				}
+
+				// We have decimal/octal value or possibly a floating point value
+				else
+				{
+
+					// Skip all digits
+					for (; this.InputPosInternal < Span.Length && UhtFCString.IsDigit(Span[this.InputPosInternal]); ++this.InputPosInternal)
+					{
+					}
+
+					// If we have a '.'
+					if (this.InputPosInternal < Span.Length && Span[this.InputPosInternal] == '.')
+					{
+						bIsFloat = true;
+						++this.InputPosInternal;
+
+						// Skip all digits
+						for (; this.InputPosInternal < Span.Length && UhtFCString.IsDigit(Span[this.InputPosInternal]); ++this.InputPosInternal)
+						{
+						}
+					}
+
+					// If we have a 'e'
+					if (this.InputPosInternal < Span.Length && UhtFCString.IsExponentMarker(Span[this.InputPosInternal]))
+					{
+						bIsFloat = true;
+						++this.InputPosInternal;
+
+						// Skip any signs
+						if (this.InputPosInternal < Span.Length && UhtFCString.IsSign(Span[this.InputPosInternal]))
+						{
+							++this.InputPosInternal;
+						}
+
+						// Skip all digits
+						for (; this.InputPosInternal < Span.Length && UhtFCString.IsDigit(Span[this.InputPosInternal]); ++this.InputPosInternal)
+						{
+						}
+					}
+
+					// If we have a 'f'
+					if (this.InputPosInternal < Span.Length && UhtFCString.IsFloatMarker(Span[this.InputPosInternal]))
+					{
+						bIsFloat = true;
+						++this.InputPosInternal;
+					}
+
+					// Check for u/l markers
+					while (this.InputPosInternal < Span.Length && 
+						(UhtFCString.IsUnsignedMarker(Span[this.InputPosInternal]) || UhtFCString.IsLongMarker(Span[this.InputPosInternal])))
+					{
+						++this.InputPosInternal;
+					}
+				}
+
+				if (this.InputPosInternal - StartPos >= UhtToken.MaxNameLength)
+				{
+					throw new UhtException(this, $"Number length exceeds maximum of {UhtToken.MaxNameLength}");
+				}
+
+				TokenType = bIsFloat ? UhtTokenType.FloatConst : (bIsHex ? UhtTokenType.HexConst : UhtTokenType.DecimalConst);
+			}
+
+			// Escaped character constant
+			else if (C == '\'')
+			{
+
+				// We try to skip the character constant value. But if it is backslash, we have to skip another character
+				char NextChar = InternalGetChar(Span);
+				if (NextChar == '\\')
+				{
+					NextChar = InternalGetChar(Span);
+				}
+
+				NextChar = InternalGetChar(Span);
+				if (NextChar != '\'')
+				{
+					throw new UhtException(this, "Unterminated character constant");
+				}
+
+				TokenType = UhtTokenType.CharConst;
+			}
+
+			// String constant
+			else if (C == '"')
+			{
+				for (; ; )
+				{
+					char NextChar = InternalGetChar(Span);
+					if (NextChar == '\r' || NextChar == '\n')
+					{
+						// throw
+					}
+
+					if (NextChar == '\\')
+					{
+						NextChar = InternalGetChar(Span);
+						if (NextChar == '\r' || NextChar == '\n')
+						{
+							// throw
+						}
+						NextChar = InternalGetChar(Span);
+					}
+
+					if (NextChar == '"')
+					{
+						break;
+					}
+				}
+
+				if (this.InputPosInternal - StartPos >= UhtToken.MaxStringLength)
+				{
+					throw new UhtException(this, $"String constant exceeds maximum of {UhtToken.MaxStringLength} characters");
+				}
+
+				TokenType = UhtTokenType.StringConst;
+			}
+
+			// Assume everything else is a symbol.
+			// Don't handle >> or >>>
+			else
+			{
+				{
+					if (this.InputPosInternal < Span.Length)
+					{
+						char D = Span[this.InputPosInternal];
+						if ((C == '<' && D == '<')
+							|| (C == '!' && D == '=')
+							|| (C == '<' && D == '=')
+							|| (C == '>' && D == '=')
+							|| (C == '+' && D == '+')
+							|| (C == '-' && D == '-')
+							|| (C == '+' && D == '=')
+							|| (C == '-' && D == '=')
+							|| (C == '*' && D == '=')
+							|| (C == '/' && D == '=')
+							|| (C == '&' && D == '&')
+							|| (C == '|' && D == '|')
+							|| (C == '^' && D == '^')
+							|| (C == '=' && D == '=')
+							|| (C == '*' && D == '*')
+							|| (C == '~' && D == '=')
+							|| (C == ':' && D == ':'))
+						{
+							++this.InputPosInternal;
+						}
+					}
+
+					// Comment processing while processing the preprocessor statements is complicated.  When a preprocessor statement
+					// parsers, the tokenizer is doing some form of a PeekToken.  And comments read ahead of the '#' will still be 
+					// pending.
+					//
+					// 1) If the block is being skipped, we want to preserve any comments ahead of the token and eliminate all other comments found in the #if block
+					// 2) If the block isn't being skipped or in the case of things such as #pragma or #include, they were considered statements in the UHT
+					//	  and would result in the comments being dropped.  In the new UHT, we just eliminate all pending tokens.
+					if (C == '#' && bEnablePreprocessor && this.TokenPreprocessorInternal != null && IsFirstTokenInLine(Span, StartPos))
+					{
+						this.PreprocessorPendingCommentsCount = this.CommentsInternal != null ? this.CommentsInternal.Count - this.CommittedComments : 0;
+						UhtToken Temp = new UhtToken(TokenType, StartPos, StartLine, StartPos, this.InputLineInternal, new StringView(this.Data.Slice(StartPos, this.InputPosInternal - StartPos)));
+						bool bClearComments = false;
+						bool bIllegalContentsCheck = false;
+						bool bInclude = this.TokenPreprocessorInternal.ParsePreprocessorDirective(ref Temp, true, out bClearComments, out bIllegalContentsCheck);
+						if (!bInclude)
+						{
+							++this.CommentsDisableCount;
+							int CheckStateMachine = 0;
+							while (true)
+							{
+								if (bIsEOF)
+								{
+									break;
+								}
+								UhtToken LocalToken = GetTokenInternal(false);
+								ReadOnlySpan<char> LocalValueSpan = LocalToken.Value.Span;
+								if (LocalToken.IsSymbol('#') && LocalValueSpan.Length == 1 && LocalValueSpan[0] == '#' && IsFirstTokenInLine(Span, LocalToken.InputStartPos))
+								{
+									if (this.TokenPreprocessorInternal.ParsePreprocessorDirective(ref Temp, false, out bool bScratch, out bool bScratchIllegalContentsCheck))
+									{
+										break;
+									}
+									bIllegalContentsCheck = bScratchIllegalContentsCheck;
+								}
+								else if (bIllegalContentsCheck)
+								{
+									switch (CheckStateMachine)
+									{
+										case 0:
+										ResetStateMachineCheck:
+											if (LocalValueSpan.CompareTo("UPROPERTY", StringComparison.Ordinal) == 0 ||
+												LocalValueSpan.CompareTo("UCLASS", StringComparison.Ordinal) == 0 ||
+												LocalValueSpan.CompareTo("USTRUCT", StringComparison.Ordinal) == 0 ||
+												LocalValueSpan.CompareTo("UENUM", StringComparison.Ordinal) == 0 ||
+												LocalValueSpan.CompareTo("UINTERFACE", StringComparison.Ordinal) == 0 ||
+												LocalValueSpan.CompareTo("UDELEGATE", StringComparison.Ordinal) == 0 ||
+												LocalValueSpan.CompareTo("UFUNCTION", StringComparison.Ordinal) == 0)
+											{
+												this.LogError($"'{LocalValueSpan.ToString()}' must not be inside preprocessor blocks, except for WITH_EDITORONLY_DATA");
+											}
+											else if (LocalValueSpan.CompareTo("void", StringComparison.Ordinal) == 0)
+											{
+												CheckStateMachine = 1;
+											}
+											break;
+
+										case 1:
+											if (LocalValueSpan.CompareTo("Serialize", StringComparison.Ordinal) == 0)
+											{
+												CheckStateMachine = 2;
+											}
+											else
+											{
+												CheckStateMachine = 0;
+												goto ResetStateMachineCheck;
+											}
+											break;
+
+										case 2:
+											if (LocalValueSpan.CompareTo("(", StringComparison.Ordinal) == 0)
+											{
+												CheckStateMachine = 3;
+											}
+											else
+											{
+												CheckStateMachine = 0;
+												goto ResetStateMachineCheck;
+											}
+											break;
+
+										case 3:
+											if (LocalValueSpan.CompareTo("FArchive", StringComparison.Ordinal) == 0 || 
+												LocalValueSpan.CompareTo("FStructuredArchiveRecord", StringComparison.Ordinal) == 0)
+											{
+												this.LogError($"Engine serialization functions must not be inside preprocessor blocks, except for WITH_EDITORONLY_DATA");
+												CheckStateMachine = 0;
+											}
+											else if (LocalValueSpan.CompareTo("FStructuredArchive", StringComparison.Ordinal) == 0)
+											{
+												CheckStateMachine = 4;
+											}
+											else
+											{
+												CheckStateMachine = 0;
+												goto ResetStateMachineCheck;
+											}
+											break;
+
+										case 4:
+											if (LocalValueSpan.CompareTo("::", StringComparison.Ordinal) == 0)
+											{
+												CheckStateMachine = 5;
+											}
+											else
+											{
+												CheckStateMachine = 0;
+												goto ResetStateMachineCheck;
+											}
+											break;
+
+										case 5:
+											if (LocalValueSpan.CompareTo("FRecord", StringComparison.Ordinal) == 0)
+											{
+												this.LogError($"Engine serialization functions must not be inside preprocessor blocks, except for WITH_EDITORONLY_DATA");
+												CheckStateMachine = 0;
+											}
+											else
+											{
+												CheckStateMachine = 0;
+												goto ResetStateMachineCheck;
+											}
+											break;
+									}
+								}
+							}
+							--this.CommentsDisableCount;
+
+							// Clear any extra pending comments at this time
+							ClearPendingComments();
+						}
+						this.PreprocessorPendingCommentsCount = 0;
+
+						// Depending on the type of directive and/or the #if expression, we clear comments.
+						if (bClearComments)
+						{
+							ClearPendingComments();
+							bGotInlineComment = false;
+						}
+						goto Restart;
+					}
+				}
+				TokenType = UhtTokenType.Symbol;
+			}
+
+			UhtToken Out = new UhtToken(TokenType, StartPos, StartLine, StartPos, this.InputLineInternal, new StringView(this.Data.Slice(StartPos, this.InputPosInternal - StartPos)));
+			return Out;
+		}
+
+		/// <summary>
+		/// Skip all leading whitespace and collect any comments
+		/// </summary>
+		private void SkipWhitespaceAndCommentsInternal(ref bool bPersistGotInlineComment, bool bEnablePreprocessor)
+		{
+			ReadOnlySpan<char> Span = this.Data.Span;
+
+			bool bGotNewlineBetweenComments = false;
+			bool bGotInlineComment = bPersistGotInlineComment;
+			for (; ; )
+			{
+				uint Char = InternalGetChar(Span);
+				if (Char == 0)
+				{
+					break;
+				}
+				else if (Char == '\n')
+				{
+					bGotNewlineBetweenComments |= bGotInlineComment;
+					++this.InputLineInternal;
+				}
+				else if (Char == '\r' || Char == '\t' || Char == ' ')
+				{
+				}
+				else if (Char == '/')
+				{
+					uint NextChar = InternalPeekChar(Span);
+					if (NextChar == '*')
+					{
+						if (bEnablePreprocessor)
+						{
+							ClearAllComments();
+						}
+						int CommentStart = this.InputPosInternal - 1;
+						++this.InputPosInternal;
+						for (; ; )
+						{
+							char CommentChar = InternalGetChar(Span);
+							if (CommentChar == 0)
+							{
+								if (bEnablePreprocessor)
+								{
+									ClearAllComments();
+								}
+								throw new UhtException(this, "End of header encountered inside comment");
+							}
+							else if (CommentChar == '\n')
+							{
+								++this.InputLineInternal;
+							}
+							else if (CommentChar == '*' && InternalPeekChar(Span) == '/')
+							{
+								++this.InputPosInternal;
+								break;
+							}
+						}
+						if (bEnablePreprocessor)
+						{
+							AddComment(new StringView(this.Data.Slice(CommentStart, this.InputPosInternal - CommentStart)));
+						}
+					}
+					else if (NextChar == '/')
+					{
+						if (bGotNewlineBetweenComments)
+						{
+							bGotNewlineBetweenComments = false;
+							if (bEnablePreprocessor)
+							{
+								ClearAllComments();
+							}
+						}
+						bGotInlineComment = true;
+						int CommentStart = this.InputPosInternal - 1;
+						++this.InputPosInternal;
+
+						// Scan to the end of the line
+						for (; ; )
+						{
+							char CommentChar = InternalGetChar(Span);
+							if (CommentChar == 0)
+							{
+								//--Pos;
+								break;
+							}
+							if (CommentChar == '\r')
+							{
+							}
+							else if (CommentChar == '\n')
+							{
+								++InputLineInternal;
+								break;
+							}
+						}
+						if (bEnablePreprocessor)
+						{
+							AddComment(new StringView(this.Data.Slice(CommentStart, this.InputPosInternal - CommentStart)));
+						}
+					}
+					else
+					{
+						--this.InputPosInternal;
+						break;
+					}
+				}
+				else
+				{
+					--this.InputPosInternal;
+					break;
+				}
+			}
+			bPersistGotInlineComment = bGotInlineComment;
+			return;
+		}
+
+		private bool IsNumeric(ReadOnlySpan<char> Span, char C)
+		{
+			// Check for [0..9]
+			if (UhtFCString.IsDigit(C))
+			{
+				return true;
+			}
+
+			// Check for [+-]...
+			if (UhtFCString.IsSign(C))
+			{
+				if (this.InputPosInternal == Span.Length)
+				{
+					return false;
+				}
+
+				// Check for [+-][0..9]...
+				if (UhtFCString.IsDigit(Span[this.InputPosInternal]))
+				{
+					return true;
+				}
+
+				// Check for [+-][.][0..9]...
+				if (Span[this.InputPosInternal] != '.')
+				{
+					return false;
+				}
+
+				if (this.InputPosInternal + 1 == Span.Length)
+				{
+					return false;
+				}
+
+				if (UhtFCString.IsDigit(Span[this.InputPosInternal + 1]))
+				{
+					return true;
+				}
+				return false;
+			}
+
+			if (C == '.')
+			{
+				if (this.InputPosInternal == Span.Length)
+				{
+					return false;
+				}
+
+				// Check for [.][0..9]...
+				if (UhtFCString.IsDigit(Span[this.InputPosInternal]))
+				{
+					return true;
+				}
+			}
+			return false;
+		}
+
+		/// <summary>
+		/// Fetch the next character in the input stream or zero if we have reached the end.
+		/// The current offset in the buffer is not advanced.  The method does not support UTF-8
+		/// </summary>
+		/// <param name="Span">The span containing the data</param>
+		/// <returns>Next character in the stream or zero</returns>
+		private char InternalPeekChar(ReadOnlySpan<char> Span)
+		{
+			return this.InputPosInternal < Span.Length ? Span[this.InputPosInternal] : '\0';
+		}
+
+		/// <summary>
+		/// Fetch the next character in the input stream or zero if we have reached the end.
+		/// The current offset in the buffer is advanced.  The method does not support UTF-8.
+		/// </summary>
+		/// <param name="Span">The span containing the data</param>
+		/// <returns>Next character in the stream or zero</returns>
+		private char InternalGetChar(ReadOnlySpan<char> Span)
+		{
+			return this.InputPosInternal < Span.Length ? Span[this.InputPosInternal++] : '\0';
+		}
+
+		/// <summary>
+		/// If we have a current token, then reset the pending comments and input position back to before the token.
+		/// </summary>
+		private void ClearToken()
+		{
+			if (this.bHasToken)
+			{
+				this.bHasToken = false;
+				if (this.CommentsInternal != null && this.CommentsInternal.Count > this.CommittedComments)
+				{
+					this.CommentsInternal.RemoveRange(this.CommittedComments, this.CommentsInternal.Count - this.CommittedComments);
+				}
+				this.InputPosInternal = this.CurrentToken.UngetPos;
+				this.InputLineInternal = this.CurrentToken.UngetLine;
+			}
+		}
+
+		private void AddComment(StringView Comment)
+		{
+			if (this.CommentsInternal == null)
+			{
+				this.CommentsInternal = new List<StringView>(4);
+			}
+			this.CommentsInternal.Add(Comment);
+		}
+		#endregion
+	}
+}
diff --git a/Engine/Source/Programs/Shared/EpicGames.UHT/Tokenizer/UhtTokenReader.cs b/Engine/Source/Programs/Shared/EpicGames.UHT/Tokenizer/UhtTokenReader.cs
new file mode 100644
--- /dev/null
+++ b/Engine/Source/Programs/Shared/EpicGames.UHT/Tokenizer/UhtTokenReader.cs
@@ -0,0 +1,386 @@
+// Copyright Epic Games, Inc. All Rights Reserved.
+
+using EpicGames.Core;
+using EpicGames.UHT.Utils;
+using System;
+using System.Collections.Generic;
+using System.Text;
+using System.Threading;
+
+namespace EpicGames.UHT.Tokenizer
+{
+	[Flags]
+	public enum UhtRawStringOptions
+	{
+		None,
+		RespectQuotes = 1 << 0,
+		DontConsumeTerminator = 1 << 1,
+	}
+
+	/// <summary>
+	/// Helper methods for testing flags.  These methods perform better than the generic HasFlag which hits
+	/// the GC and stalls.
+	/// </summary>
+	public static class UhtRawStringOptionsExtensions
+	{
+
+		/// <summary>
+		/// Test to see if any of the specified flags are set
+		/// </summary>
+		/// <param name="InFlags">Current flags</param>
+		/// <param name="TestFlags">Flags to test for</param>
+		/// <returns>True if any of the flags are set</returns>
+		public static bool HasAnyFlags(this UhtRawStringOptions InFlags, UhtRawStringOptions TestFlags)
+		{
+			return (InFlags & TestFlags) != 0;
+		}
+
+		/// <summary>
+		/// Test to see if all of the specified flags are set
+		/// </summary>
+		/// <param name="InFlags">Current flags</param>
+		/// <param name="TestFlags">Flags to test for</param>
+		/// <returns>True if all the flags are set</returns>
+		public static bool HasAllFlags(this UhtRawStringOptions InFlags, UhtRawStringOptions TestFlags)
+		{
+			return (InFlags & TestFlags) == TestFlags;
+		}
+
+		/// <summary>
+		/// Test to see if a specific set of flags have a specific value.
+		/// </summary>
+		/// <param name="InFlags">Current flags</param>
+		/// <param name="TestFlags">Flags to test for</param>
+		/// <param name="MatchFlags">Expected value of the tested flags</param>
+		/// <returns>True if the given flags have a specific value.</returns>
+		public static bool HasExactFlags(this UhtRawStringOptions InFlags, UhtRawStringOptions TestFlags, UhtRawStringOptions MatchFlags)
+		{
+			return (InFlags & TestFlags) == MatchFlags;
+		}
+	}
+
+	public interface IUhtTokenPreprocessor
+	{
+		public bool ParsePreprocessorDirective(ref UhtToken Token, bool bIsBeingIncluded, out bool bClearComments, out bool bIllegalContentsCheck);
+		public void SaveState();
+		public void RestoreState();
+	}
+
+	public interface IUhtTokenReader : IUhtMessageSite
+	{
+		/// <summary>
+		/// True if the reader is at the end of the stream
+		/// </summary>
+		public bool bIsEOF { get; }
+
+		/// <summary>
+		/// Current input position in the stream by characters
+		/// </summary>
+		public int InputPos { get; }
+
+		/// <summary>
+		/// Current input line in the stream
+		/// </summary>
+		public int InputLine { get; set; }
+
+		/// <summary>
+		/// Preprocessor attached to the token reader
+		/// </summary>
+		public IUhtTokenPreprocessor? TokenPreprocessor { get; set; }
+
+		/// <summary>
+		/// If the reader doesn't have a current token, then read the next token and return a reference to it.
+		/// Otherwise return a reference to the current token.
+		/// </summary>
+		/// <returns>The current token.  Will be invalidated by other calls to ITokenReader</returns>
+		public ref UhtToken PeekToken();
+
+		/// <summary>
+		/// Mark the current token as being consumed.  Any call to PeekToken or GetToken will read another token.
+		/// </summary>
+		public void ConsumeToken();
+
+		/// <summary>
+		/// Get the next token in the data.  If there is a current token, then that token is returned and marked as consumed.
+		/// </summary>
+		/// <returns></returns>
+		public UhtToken GetToken();
+
+		/// <summary>
+		/// Tests to see if the given token is the first token of a line
+		/// </summary>
+		/// <param name="Token">The token to test</param>
+		/// <returns>True if the token is the first token on the line</returns>
+		public bool IsFirstTokenInLine(ref UhtToken Token);
+
+		/// <summary>
+		/// Skip any whitespace and comments at the current buffer position
+		/// </summary>
+		public void SkipWhitespaceAndComments();
+
+		/// <summary>
+		/// Read the entire next line in the buffer
+		/// </summary>
+		/// <returns></returns>
+		public UhtToken GetLine();
+
+		/// <summary>
+		/// Get a view of the buffer being read
+		/// </summary>
+		/// <param name="StartPos">Starting character offset in the buffer.</param>
+		/// <param name="Count">Length of the span</param>
+		/// <returns>The string view into the buffer</returns>
+		public StringView GetStringView(int StartPos, int Count);
+
+		/// <summary>
+		/// Return a string terminated by the given character.
+		/// </summary>
+		/// <param name="Terminator">The character to stop at.</param>
+		/// <param name="Options">Options</param>
+		/// <returns>The parsed string</returns>
+		public StringView GetRawString(char Terminator, UhtRawStringOptions Options);
+
+		/// <summary>
+		/// The current collection of parsed comments.  This does not include any comments parsed as part of a 
+		/// call to PeekToken unless ConsumeToken has been invoked after a call to PeekToken.
+		/// </summary>
+		public ReadOnlySpan<StringView> Comments { get; }
+
+		/// <summary>
+		/// Clear the current collection of comments.  Any comments parsed by PeekToken prior to calling ConsomeToken will
+		/// not be cleared.
+		/// </summary>
+		public void ClearComments();
+
+		/// <summary>
+		/// Disable the processing of comments.  This is often done when skipping a bulk of the buffer.
+		/// </summary>
+		/// <returns></returns>
+		public void DisableComments();
+
+		/// <summary>
+		/// Enable comment collection.
+		/// </summary>
+		public void EnableComments();
+
+		/// <summary>
+		/// If there are any pending comments (due to a PeekToken), commit then so they will be return as current comments
+		/// </summary>
+		//COMPATIBILITY-TODO - Remove once the struct adding of tooltips is fixed
+		public void CommitPendingComments();
+
+		/// <summary>
+		/// Save the current parsing state.  There is a limited number of states that can be saved.
+		/// Invoke either RestoreState or AbandonState after calling SaveState.
+		/// </summary>
+		public void SaveState();
+
+		/// <summary>
+		/// Restore a previously saved state.
+		/// </summary>
+		public void RestoreState();
+
+		/// <summary>
+		/// Abandon a previously saved state
+		/// </summary>
+		public void AbandonState();
+
+		/// <summary>
+		/// Enable the recording of tokens
+		/// </summary>
+		public void EnableRecording();
+
+		/// <summary>
+		/// Disable the recording of tokens.  Any currently recorded tokens will be removed
+		/// </summary>
+		public void DisableRecording();
+
+		/// <summary>
+		/// Record the given token to the list of recorded tokens
+		/// </summary>
+		/// <param name="Token">Token to record</param>
+		public void RecordToken(ref UhtToken Token);
+
+		/// <summary>
+		/// Get the current collection of recorded tokens
+		/// </summary>
+		public List<UhtToken> RecordedTokens { get; }
+	}
+
+	/// <summary>
+	/// Represents a list of tokens. Follow the Next chain for each element in the list.
+	/// </summary>
+	public class UhtTokenList
+	{
+
+		/// <summary>
+		/// The token
+		/// </summary>
+		public UhtToken Token;
+
+		/// <summary>
+		/// The next token in the list
+		/// </summary>
+		public UhtTokenList? Next;
+
+		/// <summary>
+		/// Join the tokens in the list
+		/// </summary>
+		/// <param name="Builder">Destination builder</param>
+		/// <param name="Seperator">Separator between the tokens</param>
+		/// <returns></returns>
+		public void Join(StringBuilder Builder, string Seperator = "")
+		{
+			Builder.Append(this.Token.Value.ToString());
+			UhtTokenList List = this;
+			while (List.Next != null)
+			{
+				List = List.Next;
+				Builder.Append(Seperator);
+				Builder.Append(List.Token.Value.ToString());
+			}
+		}
+
+		/// <summary>
+		/// Join the tokens in the list
+		/// </summary>
+		/// <param name="Seperator">Separator between the tokens</param>
+		/// <returns></returns>
+		public string Join(string Seperator = "")
+		{
+			if (this.Next == null)
+			{
+				return this.Token.Value.ToString();
+			}
+			StringBuilder Builder = new StringBuilder();
+			Join(Builder, Seperator);
+			return Builder.ToString();
+		}
+
+		/// <summary>
+		/// Return the token list as an array
+		/// </summary>
+		/// <returns></returns>
+		public UhtToken[] ToArray()
+		{
+			int Count = 1;
+			for (UhtTokenList Temp = this; Temp.Next != null; Temp = Temp.Next)
+			{
+				++Count;
+			}
+			UhtToken[] Out = new UhtToken[Count];
+			Out[0] = this.Token;
+			Count = 1;
+			for (UhtTokenList Temp = this; Temp.Next != null; Temp = Temp.Next)
+			{
+				Out[Count] = Temp.Next.Token;
+				++Count;
+			}
+			return Out;
+		}
+	}
+
+	/// <summary>
+	/// Token list cache.  Token lists must be returned to the cache.
+	/// </summary>
+	public class UhtTokenListCache
+	{
+		static private ThreadLocal<UhtTokenList?> Tls = new ThreadLocal<UhtTokenList?>(() => null);
+
+		public static UhtTokenList Borrow(UhtToken Token)
+		{
+			UhtTokenList? Identifier = Tls.Value;
+			if (Identifier != null)
+			{
+				Tls.Value = Identifier.Next;
+			}
+			else
+			{
+				Identifier = new UhtTokenList();
+			}
+			Identifier.Token = Token;
+			Identifier.Next = null;
+			return Identifier;
+		}
+
+		public static void Return(UhtTokenList Identifier)
+		{
+			UhtTokenList? Tail = Tls.Value;
+			if (Tail != null)
+			{
+				Tail.Next = Identifier;
+			}
+
+			for (; Identifier.Next != null; Identifier = Identifier.Next)
+			{
+			}
+
+			Tls.Value = Identifier;
+		}
+	}
+
+	public delegate void UhtTokenDelegate(ref UhtToken Token);
+	public delegate bool UhtTokensUntilDelegate(ref UhtToken Token);
+	public delegate void UhtTokensDelegate(IEnumerable<UhtToken> Token);
+	public delegate void UhtTokenListDelegate(UhtTokenList TokenList);
+	public delegate void UhtTokenPositionDelegate(int Position);
+	public delegate void UhtTokenConstFloatDelegate(float Value);
+	public delegate void UhtTokenConstDoubleDelegate(double Value);
+
+	public struct UhtTokenDisableComments : IDisposable
+	{
+		private readonly IUhtTokenReader TokenReader;
+
+		public UhtTokenDisableComments(IUhtTokenReader TokenReader)
+		{
+			this.TokenReader = TokenReader;
+			this.TokenReader.DisableComments();
+		}
+
+		public void Dispose()
+		{
+			this.TokenReader.EnableComments();
+		}
+	}
+
+	public struct UhtTokenSaveState : IDisposable
+	{
+		private readonly IUhtTokenReader TokenReader;
+		private bool bHandled;
+
+		public UhtTokenSaveState(IUhtTokenReader TokenReader)
+		{
+			this.TokenReader = TokenReader;
+			this.bHandled = false;
+			this.TokenReader.SaveState();
+		}
+
+		public void Dispose()
+		{
+			if (!bHandled)
+			{
+				RestoreState();
+			}
+		}
+
+		public void RestoreState()
+		{
+			if (this.bHandled)
+			{
+				throw new UhtIceException("Can not call RestoreState/AbandonState more than once");
+			}
+			this.TokenReader.RestoreState();
+			this.bHandled = true;
+		}
+
+		public void AbandonState()
+		{
+			if (this.bHandled)
+			{
+				throw new UhtIceException("Can not call RestoreState/AbandonState more than once");
+			}
+			this.TokenReader.AbandonState();
+			this.bHandled = true;
+		}
+	}
+}
diff --git a/Engine/Source/Programs/Shared/EpicGames.UHT/Tokenizer/UhtTokenReaderFloatExtensions.cs b/Engine/Source/Programs/Shared/EpicGames.UHT/Tokenizer/UhtTokenReaderFloatExtensions.cs
new file mode 100644
--- /dev/null
+++ b/Engine/Source/Programs/Shared/EpicGames.UHT/Tokenizer/UhtTokenReaderFloatExtensions.cs
@@ -0,0 +1,184 @@
+// Copyright Epic Games, Inc. All Rights Reserved.
+
+using EpicGames.UHT.Utils;
+
+namespace EpicGames.UHT.Tokenizer
+{
+	public static class UhtTokenReaderFloatExtensions
+	{
+		/// <summary>
+		/// Get the next token as a float.  If the next token is not a float, no token is consumed.
+		/// </summary>
+		/// <param name="Value">The float value of the token</param>
+		/// <param name="ExceptionContext">If not null, an exception will be thrown with the given text as part of the message.</param>
+		/// <returns>True if the next token was an float, false if not.</returns>
+		public static bool TryOptionalConstFloat(this IUhtTokenReader TokenReader, out float Value)
+		{
+			ref UhtToken Token = ref TokenReader.PeekToken();
+			if (Token.IsConstFloat() && Token.GetConstFloat(out Value)) // NOTE: This is restricted to only float values
+			{
+				TokenReader.ConsumeToken();
+				return true;
+			}
+			Value = 0;
+			return false;
+		}
+
+		/// <summary>
+		/// Get the next token as a float.  If the next token is not a float, no token is consumed.
+		/// </summary>
+		/// <param name="Value">The float value of the token</param>
+		/// <param name="ExceptionContext">If not null, an exception will be thrown with the given text as part of the message.</param>
+		/// <returns>The token reader</returns>
+		public static IUhtTokenReader OptionalConstFloat(this IUhtTokenReader TokenReader)
+		{
+			float Value;
+			TokenReader.TryOptionalConstFloat(out Value);
+			return TokenReader;
+		}
+
+		/// <summary>
+		/// Get the next token as a float.  If the next token is not a float, no token is consumed.
+		/// </summary>
+		/// <param name="Value">The float value of the token</param>
+		/// <param name="ExceptionContext">If not null, an exception will be thrown with the given text as part of the message.</param>
+		/// <returns>The token reader</returns>
+		public static IUhtTokenReader OptionalConstFloat(this IUhtTokenReader TokenReader, UhtTokenConstFloatDelegate Delegate)
+		{
+			float Value;
+			if (TokenReader.TryOptionalConstFloat(out Value))
+			{
+				Delegate(Value);
+			}
+			return TokenReader;
+		}
+
+		/// <summary>
+		/// Get the next token as a float.  If the next token is not a float, an exception is thrown
+		/// </summary>
+		/// <param name="ExceptionContext">If not null, an exception will be thrown with the given text as part of the message.</param>
+		/// <returns>True if the next token was an float, false if not.</returns>
+		public static IUhtTokenReader RequireConstFloat(this IUhtTokenReader TokenReader, object? ExceptionContext = null)
+		{
+			float Value;
+			if (!TokenReader.TryOptionalConstFloat(out Value))
+			{
+				throw new UhtTokenException(TokenReader, TokenReader.PeekToken(), "constant float", ExceptionContext);
+			}
+			return TokenReader;
+		}
+
+		/// <summary>
+		/// Get the next token as a float.  If the next token is not a float, an exception is thrown
+		/// </summary>
+		/// <param name="ExceptionContext">If not null, an exception will be thrown with the given text as part of the message.</param>
+		/// <returns>The floating point value of the token</returns>
+		public static float GetConstFloat(this IUhtTokenReader TokenReader, object? ExceptionContext = null)
+		{
+			float Value;
+			if (!TokenReader.TryOptionalConstFloat(out Value))
+			{
+				throw new UhtTokenException(TokenReader, TokenReader.PeekToken(), "constant float", ExceptionContext);
+			}
+			return Value;
+		}
+
+		/// <summary>
+		/// Get the next float.  It also handles [+/-] token followed by an float.
+		/// </summary>
+		/// <param name="TokenReader">Source tokens</param>
+		/// <param name="Value">The float value of the token</param>
+		/// <returns>True if the next token was an float, false if not.</returns>
+		public static bool TryOptionalLeadingSignConstFloat(this IUhtTokenReader TokenReader, out float Value)
+		{
+			float LocalValue = 0;
+			bool Results = TokenReader.TryOptionalLeadingSignConstNumeric((ref UhtToken Token) =>
+			{
+				return (Token.IsConstInt() || Token.IsConstFloat()) && Token.GetConstFloat(out LocalValue);
+			});
+			Value = LocalValue;
+			return Results;
+		}
+
+		/// <summary>
+		/// Get the next float.  It also handles [+/-] token followed by an float.
+		/// </summary>
+		/// <param name="TokenReader">Source tokens</param>
+		/// <param name="Value">The float value of the token</param>
+		/// <returns>True if the next token was an float, false if not.</returns>
+		public static bool TryOptionalConstFloatExpression(this IUhtTokenReader TokenReader, out float Value)
+		{
+			float LocalValue = 0;
+			bool Results = TokenReader.TryOptionalLeadingSignConstNumeric((ref UhtToken Token) =>
+			{
+				return (Token.IsConstInt() || Token.IsConstFloat()) && Token.GetConstFloat(out LocalValue);
+			});
+			Value = LocalValue;
+			return Results;
+		}
+
+		/// <summary>
+		/// Get the next float.  It also handles [+/-] token followed by an float.
+		/// </summary>
+		/// <param name="TokenReader">Source tokens</param>
+		/// <returns>The double value</returns>
+		public static float GetConstFloatExpression(this IUhtTokenReader TokenReader)
+		{
+			float LocalValue = 0;
+			if (!TokenReader.TryOptionalConstFloatExpression(out LocalValue))
+			{
+				throw new UhtTokenException(TokenReader, TokenReader.PeekToken(), "constant float", null);
+			}
+			return LocalValue;
+		}
+
+		/// <summary>
+		/// Get the next double.  It also handles [+/-] token followed by an double.
+		/// </summary>
+		/// <param name="TokenReader">Source tokens</param>
+		/// <param name="Value">The double value of the token</param>
+		/// <returns>True if the next token was an double, false if not.</returns>
+		public static bool TryOptionalConstDoubleExpression(this IUhtTokenReader TokenReader, out double Value)
+		{
+			double LocalValue = 0;
+			bool Results = TokenReader.TryOptionalLeadingSignConstNumeric((ref UhtToken Token) =>
+			{
+				return (Token.IsConstInt() || Token.IsConstFloat()) && Token.GetConstDouble(out LocalValue);
+			});
+			Value = LocalValue;
+			return Results;
+		}
+
+		/// <summary>
+		/// Get the next double.  It also handles [+/-] token followed by an double.
+		/// </summary>
+		/// <param name="TokenReader">Source tokens</param>
+		/// <param name="Delegate">Delegate to invoke if the double is parsed</param>
+		/// <returns>The supplied token reader</returns>
+		public static IUhtTokenReader RequireConstDoubleExpression(this IUhtTokenReader TokenReader, UhtTokenConstDoubleDelegate Delegate)
+		{
+			double LocalValue = 0;
+			if (!TokenReader.TryOptionalConstDoubleExpression(out LocalValue))
+			{
+				throw new UhtTokenException(TokenReader, TokenReader.PeekToken(), "constant double", null);
+			}
+			Delegate(LocalValue);
+			return TokenReader;
+		}
+
+		/// <summary>
+		/// Get the next double.  It also handles [+/-] token followed by an double.
+		/// </summary>
+		/// <param name="TokenReader">Source tokens</param>
+		/// <returns>The double value</returns>
+		public static double GetConstDoubleExpression(this IUhtTokenReader TokenReader)
+		{
+			double LocalValue = 0;
+			if (!TokenReader.TryOptionalConstDoubleExpression(out LocalValue))
+			{
+				throw new UhtTokenException(TokenReader, TokenReader.PeekToken(), "constant double", null);
+			}
+			return LocalValue;
+		}
+	}
+}
diff --git a/Engine/Source/Programs/Shared/EpicGames.UHT/Tokenizer/UhtTokenReaderFlowControlExtensions.cs b/Engine/Source/Programs/Shared/EpicGames.UHT/Tokenizer/UhtTokenReaderFlowControlExtensions.cs
new file mode 100644
--- /dev/null
+++ b/Engine/Source/Programs/Shared/EpicGames.UHT/Tokenizer/UhtTokenReaderFlowControlExtensions.cs
@@ -0,0 +1,262 @@
+// Copyright Epic Games, Inc. All Rights Reserved.
+
+using EpicGames.UHT.Utils;
+using System;
+using System.Collections.Generic;
+
+namespace EpicGames.UHT.Tokenizer
+{
+	public static class UhtTokenReaderFlowControlExtensions
+	{
+		public static bool TryOptionalList(this IUhtTokenReader TokenReader, char Initiator, char Terminator, char Separator, bool bAllowTrailingSeparator, Action Action)
+		{
+			if (TokenReader.TryOptional(Initiator))
+			{
+				TokenReader.RequireList(Terminator, Separator, bAllowTrailingSeparator, Action);
+				return true;
+			}
+			return false;
+		}
+
+		public static IUhtTokenReader RequireList(this IUhtTokenReader TokenReader, char Initiator, char Terminator, char Separator, bool bAllowTrailingSeparator, Action Action)
+		{
+			return TokenReader.RequireList(Initiator, Terminator, Separator, bAllowTrailingSeparator, null, Action);
+		}
+
+		public static IUhtTokenReader RequireList(this IUhtTokenReader TokenReader, char Initiator, char Terminator, char Separator, bool bAllowTrailingSeparator, object? ExceptionContext, Action Action)
+		{
+			TokenReader.Require(Initiator);
+			return TokenReader.RequireList(Terminator, Separator, bAllowTrailingSeparator, ExceptionContext, Action);
+		}
+
+		public static IUhtTokenReader RequireList(this IUhtTokenReader TokenReader, char Terminator, char Separator, bool bAllowTrailingSeparator, Action Action)
+		{
+			return TokenReader.RequireList(Terminator, Separator, bAllowTrailingSeparator, null, Action);
+		}
+
+		public static IUhtTokenReader RequireList(this IUhtTokenReader TokenReader, char Terminator, char Separator, bool bAllowTrailingSeparator, object? ExceptionContext, Action Action)
+		{
+			// Check for an empty list
+			if (TokenReader.TryOptional(Terminator))
+			{
+				return TokenReader;
+			}
+
+			// Process the body
+			while (true)
+			{
+
+				// Read the element via the lambda
+				Action();
+
+				// Make sure we haven't reached the EOF
+				if (TokenReader.bIsEOF)
+				{
+					throw new UhtTokenException(TokenReader, TokenReader.GetToken(), $"'{Separator}' or '{Terminator}'", ExceptionContext);
+				}
+
+				// If we have a separator, then it might be a trailing separator 
+				if (TokenReader.TryOptional(Separator))
+				{
+					if (bAllowTrailingSeparator && TokenReader.TryOptional(Terminator))
+					{
+						return TokenReader;
+					}
+					continue;
+				}
+
+				// Otherwise, we must have an terminator
+				if (!TokenReader.TryOptional(Terminator))
+				{
+					throw new UhtTokenException(TokenReader, TokenReader.GetToken(), $"'{Separator}' or '{Terminator}'", ExceptionContext);
+				}
+				return TokenReader;
+			}
+		}
+
+		public static bool TryOptionalList(this IUhtTokenReader TokenReader, char Initiator, char Terminator, char Separator, bool bAllowTrailingSeparator, UhtTokensDelegate Delegate)
+		{
+			if (TokenReader.TryOptional(Initiator))
+			{
+				TokenReader.RequireList(Terminator, Separator, bAllowTrailingSeparator, Delegate);
+				return true;
+			}
+			return false;
+		}
+
+		public static IUhtTokenReader RequireList(this IUhtTokenReader TokenReader, char Initiator, char Terminator, char Separator, bool bAllowTrailingSeparator, UhtTokensDelegate Delegate)
+		{
+			return TokenReader.RequireList(Initiator, Terminator, Separator, bAllowTrailingSeparator, null, Delegate);
+		}
+
+		public static IUhtTokenReader RequireList(this IUhtTokenReader TokenReader, char Initiator, char Terminator, char Separator, bool bAllowTrailingSeparator, object? ExceptionContext, UhtTokensDelegate Delegate)
+		{
+			TokenReader.Require(Initiator);
+			return TokenReader.RequireList(Terminator, Separator, bAllowTrailingSeparator, ExceptionContext, Delegate);
+		}
+
+		public static IUhtTokenReader RequireList(this IUhtTokenReader TokenReader, char Terminator, char Separator, bool bAllowTrailingSeparator, UhtTokensDelegate Delegate)
+		{
+			return TokenReader.RequireList(Terminator, Separator, bAllowTrailingSeparator, null, Delegate);
+		}
+
+		public static IUhtTokenReader RequireList(this IUhtTokenReader TokenReader, char Terminator, char Separator, bool bAllowTrailingSeparator, object? ExceptionContext, UhtTokensDelegate Delegate)
+		{
+			// Check for an empty list
+			if (TokenReader.TryOptional(Terminator))
+			{
+				return TokenReader;
+			}
+
+			// Process the body
+			while (true)
+			{
+				List<UhtToken> Tokens = new List<UhtToken>();
+
+				// Read the tokens until we hit the end
+				while (true)
+				{
+					// Make sure we haven't reached the EOF
+					if (TokenReader.bIsEOF)
+					{
+						throw new UhtTokenException(TokenReader, TokenReader.GetToken(), $"'{Separator}' or '{Terminator}'", ExceptionContext);
+					}
+
+					// If we have a separator, then it might be a trailing separator 
+					if (TokenReader.TryOptional(Separator))
+					{
+						Delegate(Tokens);
+						if (TokenReader.TryOptional(Terminator))
+						{
+							if (bAllowTrailingSeparator)
+							{
+								return TokenReader;
+							}
+							throw new UhtException(TokenReader, $"A separator '{Separator}' followed immediately by the terminator '{Terminator}' is invalid");
+						}
+						break;
+					}
+
+					// If this is the terminator, then we are done
+					if (TokenReader.TryOptional(Terminator))
+					{
+						Delegate(Tokens);
+						return TokenReader;
+					}
+
+					Tokens.Add(TokenReader.GetToken());
+				}
+			}
+		}
+
+		/// <summary>
+		/// Consume a block of tokens bounded by the two given symbols.
+		/// </summary>
+		/// <param name="TokenReader">Token reader</param>
+		/// <param name="Initiator">The next token must be the given symbol.</param>
+		/// <param name="Terminator">The tokens are read until the given symbol is found.  The terminating symbol will be consumed.</param>
+		/// <param name="ExceptionContext">Extra context for any error messages</param>
+		/// <returns>The input token reader</returns>
+		public static IUhtTokenReader RequireList(this IUhtTokenReader TokenReader, char Initiator, char Terminator, object? ExceptionContext = null)
+		{
+			TokenReader.Require(Initiator);
+
+			// Process the body
+			while (true)
+			{
+
+				// Make sure we haven't reached the EOF
+				if (TokenReader.bIsEOF)
+				{
+					throw new UhtTokenException(TokenReader, TokenReader.GetToken(), $"'{Terminator}'", ExceptionContext);
+				}
+
+				// Look for the terminator
+				if (TokenReader.TryOptional(Terminator))
+				{
+					break;
+				}
+
+				// Consume the current token
+				TokenReader.ConsumeToken();
+			}
+			return TokenReader;
+		}
+
+		public static IUhtTokenReader While(this IUhtTokenReader TokenReader, string Text, Action Action)
+		{
+			while (TokenReader.TryOptional(Text))
+			{
+				Action();
+			}
+			return TokenReader;
+		}
+
+		public static IUhtTokenReader While(this IUhtTokenReader TokenReader, char Text, Action Action)
+		{
+			while (TokenReader.TryOptional(Text))
+			{
+				Action();
+			}
+			return TokenReader;
+		}
+
+		/// <summary>
+		/// Read tokens until the delegate return false.  The terminating token is not consumed.
+		/// </summary>
+		/// <param name="TokenReader">Token reader</param>
+		/// <param name="Delegate">Invoked with each read token.  Return true to continue tokenizing or false to terminate.</param>
+		/// <returns>Token reader</returns>
+		public static IUhtTokenReader While(this IUhtTokenReader TokenReader, UhtTokensUntilDelegate Delegate)
+		{
+			while (true)
+			{
+				if (TokenReader.bIsEOF)
+				{
+					throw new UhtTokenException(TokenReader, TokenReader.GetToken(), null, null);
+				}
+				ref UhtToken Token = ref TokenReader.PeekToken();
+				if (Delegate(ref Token))
+				{
+					TokenReader.ConsumeToken();
+				}
+				else
+				{
+					return TokenReader;
+				}
+			}
+		}
+
+		public static int ConsumeUntil(this IUhtTokenReader TokenReader, string[] Terminators)
+		{
+			int ConsumedTokens = 0;
+			while (!TokenReader.bIsEOF)
+			{
+				ref UhtToken Token = ref TokenReader.PeekToken();
+				foreach (string Candidate in Terminators)
+				{
+					if ((Token.IsIdentifier() || Token.IsSymbol()) && Token.IsValue(Candidate))
+					{
+						return ConsumedTokens;
+					}
+				}
+				TokenReader.ConsumeToken();
+				++ConsumedTokens;
+			}
+			return ConsumedTokens;
+		}
+
+		public static IUhtTokenReader ConsumeUntil(this IUhtTokenReader TokenReader, char Terminator)
+		{
+			while (!TokenReader.bIsEOF)
+			{
+				if (TokenReader.TryOptional(Terminator))
+				{
+					break;
+				}
+				TokenReader.ConsumeToken();
+			}
+			return TokenReader;
+		}
+	}
+}
diff --git a/Engine/Source/Programs/Shared/EpicGames.UHT/Tokenizer/UhtTokenReaderGeneralExtensions.cs b/Engine/Source/Programs/Shared/EpicGames.UHT/Tokenizer/UhtTokenReaderGeneralExtensions.cs
new file mode 100644
--- /dev/null
+++ b/Engine/Source/Programs/Shared/EpicGames.UHT/Tokenizer/UhtTokenReaderGeneralExtensions.cs
@@ -0,0 +1,238 @@
+// Copyright Epic Games, Inc. All Rights Reserved.
+
+using EpicGames.Core;
+using EpicGames.UHT.Utils;
+using System;
+
+namespace EpicGames.UHT.Tokenizer
+{
+	public static class UhtTokenReaderGeneralExtensions
+	{
+		public static bool TryOptional(this IUhtTokenReader TokenReader, string Text)
+		{
+			ref UhtToken Token = ref TokenReader.PeekToken();
+			if (Token.IsIdentifier(Text) || Token.IsSymbol(Text))
+			{
+				TokenReader.ConsumeToken();
+				return true;
+			}
+			return false;
+		}
+
+		public static bool TryOptional(this IUhtTokenReader TokenReader, StringView Text)
+		{
+			ref UhtToken Token = ref TokenReader.PeekToken();
+			if (Token.IsIdentifier(Text) || Token.IsSymbol(Text))
+			{
+				TokenReader.ConsumeToken();
+				return true;
+			}
+			return false;
+		}
+
+		public static bool TryPeekOptional(this IUhtTokenReader TokenReader, string Text)
+		{
+			ref UhtToken Token = ref TokenReader.PeekToken();
+			return Token.IsIdentifier(Text) || Token.IsSymbol(Text);
+		}
+
+		/// <summary>
+		/// Test to see if the next token is one of the given strings.
+		/// </summary>
+		/// <param name="TokenReader">Token reader</param>
+		/// <param name="Text">List of keywords to test</param>
+		/// <returns>Index of matched string or -1 if nothing matched</returns>
+		public static int TryOptional(this IUhtTokenReader TokenReader, string[] Text)
+		{
+			ref UhtToken Token = ref TokenReader.PeekToken();
+			for (int Index = 0, EndIndex = Text.Length; Index < EndIndex; ++Index)
+			{
+				if (Token.IsIdentifier(Text[Index]) || Token.IsSymbol(Text[Index]))
+				{
+					TokenReader.ConsumeToken();
+					return Index;
+				}
+			}
+			return -1;
+		}
+
+		public static bool TryOptional(this IUhtTokenReader TokenReader, char Text)
+		{
+			ref UhtToken Token = ref TokenReader.PeekToken();
+			if (Token.IsSymbol(Text))
+			{
+				TokenReader.ConsumeToken();
+				return true;
+			}
+			return false;
+		}
+
+		public static bool TryOptional(this IUhtTokenReader TokenReader, char Text, out UhtToken OutToken)
+		{
+			ref UhtToken Token = ref TokenReader.PeekToken();
+			if (Token.IsSymbol(Text))
+			{
+				OutToken = Token;
+				TokenReader.ConsumeToken();
+				return true;
+			}
+			OutToken = new UhtToken();
+			return false;
+		}
+
+		public static bool TryPeekOptional(this IUhtTokenReader TokenReader, char Text)
+		{
+			ref UhtToken Token = ref TokenReader.PeekToken();
+			return Token.IsSymbol(Text);
+		}
+
+		public static IUhtTokenReader Optional(this IUhtTokenReader TokenReader, string Text)
+		{
+			ref UhtToken Token = ref TokenReader.PeekToken();
+			if (Token.IsIdentifier(Text) || Token.IsSymbol(Text))
+			{
+				TokenReader.ConsumeToken();
+			}
+			return TokenReader;
+		}
+
+		public static IUhtTokenReader Optional(this IUhtTokenReader TokenReader, string Text, Action Action)
+		{
+			ref UhtToken Token = ref TokenReader.PeekToken();
+			if (Token.IsIdentifier(Text) || Token.IsSymbol(Text))
+			{
+				TokenReader.ConsumeToken();
+				Action();
+			}
+			return TokenReader;
+		}
+
+		public static IUhtTokenReader Optional(this IUhtTokenReader TokenReader, string Text, Action YesAction, Action NoAction)
+		{
+			ref UhtToken Token = ref TokenReader.PeekToken();
+			if (Token.IsIdentifier(Text) || Token.IsSymbol(Text))
+			{
+				TokenReader.ConsumeToken();
+				YesAction();
+			}
+			else
+			{
+				NoAction();
+			}
+			return TokenReader;
+		}
+
+		public static IUhtTokenReader Optional(this IUhtTokenReader TokenReader, char Text)
+		{
+			ref UhtToken Token = ref TokenReader.PeekToken();
+			if (Token.IsSymbol(Text))
+			{
+				TokenReader.ConsumeToken();
+			}
+			return TokenReader;
+		}
+
+		public static IUhtTokenReader Optional(this IUhtTokenReader TokenReader, char Text, Action Action)
+		{
+			ref UhtToken Token = ref TokenReader.PeekToken();
+			if (Token.IsSymbol(Text))
+			{
+				TokenReader.ConsumeToken();
+				Action();
+			}
+			return TokenReader;
+		}
+
+		public static IUhtTokenReader OptionalStartsWith(this IUhtTokenReader TokenReader, string Text)
+		{
+			ref UhtToken Token = ref TokenReader.PeekToken();
+			if (Token.IsIdentifier() && Token.ValueStartsWith(Text))
+			{
+				TokenReader.ConsumeToken();
+			}
+			return TokenReader;
+		}
+
+		public static IUhtTokenReader OptionalStartsWith(this IUhtTokenReader TokenReader, string Text, UhtTokenDelegate TokenDelegate)
+		{
+			ref UhtToken Token = ref TokenReader.PeekToken();
+			if (Token.IsIdentifier() && Token.ValueStartsWith(Text))
+			{
+				UhtToken CurrentToken = Token;
+				TokenReader.ConsumeToken();
+				TokenDelegate(ref CurrentToken);
+			}
+			return TokenReader;
+		}
+
+		public static IUhtTokenReader Require(this IUhtTokenReader TokenReader, string Text, object? ExceptionContext = null)
+		{
+			ref UhtToken Token = ref TokenReader.PeekToken();
+			if (Token.IsIdentifier(Text) || Token.IsSymbol(Text))
+			{
+				TokenReader.ConsumeToken();
+			}
+			else
+			{
+				throw new UhtTokenException(TokenReader, Token, Text, ExceptionContext);
+			}
+			return TokenReader;
+		}
+
+		public static IUhtTokenReader Require(this IUhtTokenReader TokenReader, string Text, UhtTokenDelegate TokenDelegate)
+		{
+			return TokenReader.Require(Text, null, TokenDelegate);
+		}
+
+		public static IUhtTokenReader Require(this IUhtTokenReader TokenReader, string Text, object? ExceptionContext, UhtTokenDelegate TokenDelegate)
+		{
+			ref UhtToken Token = ref TokenReader.PeekToken();
+			if (Token.IsIdentifier(Text) || Token.IsSymbol(Text))
+			{
+				UhtToken CurrentToken = Token;
+				TokenReader.ConsumeToken();
+				TokenDelegate(ref CurrentToken);
+			}
+			else
+			{
+				throw new UhtTokenException(TokenReader, Token, Text, ExceptionContext);
+			}
+			return TokenReader;
+		}
+
+		public static IUhtTokenReader Require(this IUhtTokenReader TokenReader, char Text, object? ExceptionContext = null)
+		{
+			ref UhtToken Token = ref TokenReader.PeekToken();
+			if (Token.IsSymbol(Text))
+			{
+				TokenReader.ConsumeToken();
+			}
+			else
+			{
+				throw new UhtTokenException(TokenReader, Token, Text, ExceptionContext);
+			}
+			return TokenReader;
+		}
+
+		public static IUhtTokenReader Require(this IUhtTokenReader TokenReader, char Text, UhtTokenDelegate TokenDelegate)
+		{
+			return TokenReader.Require(Text, null, TokenDelegate);
+		}
+
+		public static IUhtTokenReader Require(this IUhtTokenReader TokenReader, char Text, object? ExceptionContext, UhtTokenDelegate TokenDelegate)
+		{
+			ref UhtToken Token = ref TokenReader.PeekToken();
+			if (Token.IsSymbol(Text))
+			{
+				UhtToken CurrentToken = Token;
+				TokenReader.ConsumeToken();
+				TokenDelegate(ref CurrentToken);
+			}
+			else
+			{
+				throw new UhtTokenException(TokenReader, Token, Text, ExceptionContext);
+			}
+			return TokenReader;
+		}
+	}
+}
diff --git a/Engine/Source/Programs/Shared/EpicGames.UHT/Tokenizer/UhtTokenReaderIdentifierExtensions.cs b/Engine/Source/Programs/Shared/EpicGames.UHT/Tokenizer/UhtTokenReaderIdentifierExtensions.cs
new file mode 100644
--- /dev/null
+++ b/Engine/Source/Programs/Shared/EpicGames.UHT/Tokenizer/UhtTokenReaderIdentifierExtensions.cs
@@ -0,0 +1,334 @@
+// Copyright Epic Games, Inc. All Rights Reserved.
+
+using EpicGames.UHT.Utils;
+using System;
+using System.Collections.Generic;
+using System.Diagnostics.CodeAnalysis;
+
+namespace EpicGames.UHT.Tokenizer
+{
+	[Flags]
+	public enum UhtCppIdentifierOptions
+	{
+		None,
+		AllowTemplates = 1 << 0,
+		AllTokens = 1 << 1, // AllowTemplates implies AllTokens
+	}
+
+	/// <summary>
+	/// Helper methods for testing flags.  These methods perform better than the generic HasFlag which hits
+	/// the GC and stalls.
+	/// </summary>
+	public static class UhtGetCppIdentifierOptionsExtensions
+	{
+
+		/// <summary>
+		/// Test to see if any of the specified flags are set
+		/// </summary>
+		/// <param name="InFlags">Current flags</param>
+		/// <param name="TestFlags">Flags to test for</param>
+		/// <returns>True if any of the flags are set</returns>
+		public static bool HasAnyFlags(this UhtCppIdentifierOptions InFlags, UhtCppIdentifierOptions TestFlags)
+		{
+			return (InFlags & TestFlags) != 0;
+		}
+
+		/// <summary>
+		/// Test to see if all of the specified flags are set
+		/// </summary>
+		/// <param name="InFlags">Current flags</param>
+		/// <param name="TestFlags">Flags to test for</param>
+		/// <returns>True if all the flags are set</returns>
+		public static bool HasAllFlags(this UhtCppIdentifierOptions InFlags, UhtCppIdentifierOptions TestFlags)
+		{
+			return (InFlags & TestFlags) == TestFlags;
+		}
+
+		/// <summary>
+		/// Test to see if a specific set of flags have a specific value.
+		/// </summary>
+		/// <param name="InFlags">Current flags</param>
+		/// <param name="TestFlags">Flags to test for</param>
+		/// <param name="MatchFlags">Expected value of the tested flags</param>
+		/// <returns>True if the given flags have a specific value.</returns>
+		public static bool HasExactFlags(this UhtCppIdentifierOptions InFlags, UhtCppIdentifierOptions TestFlags, UhtCppIdentifierOptions MatchFlags)
+		{
+			return (InFlags & TestFlags) == MatchFlags;
+		}
+	}
+
+	public static class UhtTokenReaderIdentifierExtensions
+	{
+
+		/// <summary>
+		/// Get the next token and verify that it is an identifier
+		/// </summary>
+		/// <returns>True if it is an identifier, false if not.</returns>
+		public static bool TryOptionalIdentifier(this IUhtTokenReader TokenReader)
+		{
+			ref UhtToken Token = ref TokenReader.PeekToken();
+			if (Token.IsIdentifier())
+			{
+				TokenReader.ConsumeToken();
+				return true;
+			}
+			return false;
+		}
+
+		/// <summary>
+		/// Get the next token and verify that it is an identifier
+		/// </summary>
+		/// <param name="Identifier">The fetched value of the identifier</param>
+		/// <returns>True if it is an identifier, false if not.</returns>
+		public static bool TryOptionalIdentifier(this IUhtTokenReader TokenReader, out UhtToken Identifier)
+		{
+			ref UhtToken Token = ref TokenReader.PeekToken();
+			if (Token.IsIdentifier())
+			{
+				Identifier = Token;
+				TokenReader.ConsumeToken();
+				return true;
+			}
+			Identifier = new UhtToken();
+			return false;
+		}
+
+		public static IUhtTokenReader OptionalIdentifier(this IUhtTokenReader TokenReader, UhtTokenDelegate TokenDelegate)
+		{
+			ref UhtToken Token = ref TokenReader.PeekToken();
+			if (Token.IsIdentifier())
+			{
+				UhtToken TokenCopy = Token;
+				TokenReader.ConsumeToken();
+				TokenDelegate(ref TokenCopy);
+			}
+			return TokenReader;
+		}
+
+		public static IUhtTokenReader RequireIdentifier(this IUhtTokenReader TokenReader, object? ExceptionContext = null)
+		{
+			ref UhtToken Token = ref TokenReader.PeekToken();
+			if (Token.IsIdentifier())
+			{
+				TokenReader.ConsumeToken();
+				return TokenReader;
+			}
+			throw new UhtTokenException(TokenReader, Token, "an identifier", ExceptionContext);
+		}
+
+		public static IUhtTokenReader RequireIdentifier(this IUhtTokenReader TokenReader, UhtTokenDelegate TokenDelegate)
+		{
+			return TokenReader.RequireIdentifier(null, TokenDelegate);
+		}
+
+		public static IUhtTokenReader RequireIdentifier(this IUhtTokenReader TokenReader, object? ExceptionContext, UhtTokenDelegate TokenDelegate)
+		{
+			ref UhtToken Token = ref TokenReader.PeekToken();
+			if (Token.IsIdentifier())
+			{
+				UhtToken CurrentToken = Token;
+				TokenReader.ConsumeToken();
+				TokenDelegate(ref CurrentToken);
+				return TokenReader;
+			}
+			throw new UhtTokenException(TokenReader, Token, "an identifier", ExceptionContext);
+		}
+
+		public static UhtToken GetIdentifier(this IUhtTokenReader TokenReader, object? ExceptionContext = null)
+		{
+			ref UhtToken Token = ref TokenReader.PeekToken();
+			if (Token.IsIdentifier())
+			{
+				UhtToken CurrentToken = Token;
+				TokenReader.ConsumeToken();
+				return CurrentToken;
+			}
+			throw new UhtTokenException(TokenReader, Token, ExceptionContext != null ? ExceptionContext : "an identifier");
+		}
+
+		public static bool TryOptionalCppIdentifier(this IUhtTokenReader TokenReader, [NotNullWhen(true)] out IEnumerable<UhtToken>? Identifier, UhtCppIdentifierOptions Options = UhtCppIdentifierOptions.None)
+		{
+			Identifier = null;
+			List<UhtToken> LocalIdentifier = new List<UhtToken>();
+
+			using (var SavedState = new UhtTokenSaveState(TokenReader))
+			{
+				UhtToken Token;
+
+				if (TokenReader.TryOptionalIdentifier(out Token))
+				{
+					LocalIdentifier.Add(Token);
+				}
+				else
+				{
+					return false;
+				}
+
+				if (Options.HasAnyFlags(UhtCppIdentifierOptions.AllowTemplates | UhtCppIdentifierOptions.AllTokens))
+				{
+					bool bAllowTemplates = Options.HasAnyFlags(UhtCppIdentifierOptions.AllowTemplates);
+
+					while (true)
+					{
+						if (bAllowTemplates && TokenReader.TryPeekOptional('<'))
+						{
+							LocalIdentifier.Add(TokenReader.GetToken());
+							int NestedScopes = 1;
+							while (NestedScopes > 0)
+							{
+								Token = TokenReader.GetToken();
+								if (Token.TokenType.IsEndType())
+								{
+									return false;
+								}
+								LocalIdentifier.Add(Token);
+								if (Token.IsSymbol('<'))
+								{
+									++NestedScopes;
+								}
+								else if (Token.IsSymbol('>'))
+								{
+									--NestedScopes;
+								}
+							}
+						}
+
+						if (!TokenReader.TryPeekOptional("::"))
+						{
+							break;
+						}
+						LocalIdentifier.Add(TokenReader.GetToken());
+						if (!TokenReader.TryOptionalIdentifier(out Token))
+						{
+							LocalIdentifier.Add(Token);
+						}
+						else
+						{
+							return false;
+						}
+					}
+				}
+				else
+				{
+					while (TokenReader.TryOptional("::"))
+					{
+						if (TokenReader.TryOptionalIdentifier(out Token))
+						{
+							LocalIdentifier.Add(Token);
+						}
+						else
+						{
+							return false;
+						}
+					}
+				}
+				Identifier = LocalIdentifier;
+				SavedState.AbandonState();
+				return true;
+			}
+		}
+
+		public static IUhtTokenReader RequireCppIdentifier(this IUhtTokenReader TokenReader, UhtCppIdentifierOptions Options = UhtCppIdentifierOptions.None)
+		{
+			TokenReader.GetCppIdentifier(Options);
+			return TokenReader;
+		}
+
+		public static IUhtTokenReader RequireCppIdentifier(this IUhtTokenReader TokenReader, ref UhtToken InitialIdentifier, UhtCppIdentifierOptions Options = UhtCppIdentifierOptions.None)
+		{
+			TokenReader.GetCppIdentifier(ref InitialIdentifier, Options);
+			return TokenReader;
+		}
+
+		public static IUhtTokenReader RequireCppIdentifier(this IUhtTokenReader TokenReader, UhtTokenListDelegate TokenListDelegate)
+		{
+			TokenReader.RequireCppIdentifier(UhtCppIdentifierOptions.None, TokenListDelegate);
+			return TokenReader;
+		}
+
+		public static IUhtTokenReader RequireCppIdentifier(this IUhtTokenReader TokenReader, UhtCppIdentifierOptions Options, UhtTokenListDelegate TokenListDelegate)
+		{
+			UhtTokenList TokenList = TokenReader.GetCppIdentifier(Options);
+			TokenListDelegate(TokenList);
+			UhtTokenListCache.Return(TokenList);
+			return TokenReader;
+		}
+
+		public static IUhtTokenReader RequireCppIdentifier(this IUhtTokenReader TokenReader, ref UhtToken InitialIdentifier, UhtTokenListDelegate TokenListDelegate)
+		{
+			TokenReader.RequireCppIdentifier(ref InitialIdentifier, UhtCppIdentifierOptions.None, TokenListDelegate);
+			return TokenReader;
+		}
+
+		public static IUhtTokenReader RequireCppIdentifier(this IUhtTokenReader TokenReader, ref UhtToken InitialIdentifier, UhtCppIdentifierOptions Options, UhtTokenListDelegate TokenListDelegate)
+		{
+			UhtTokenList TokenList = TokenReader.GetCppIdentifier(ref InitialIdentifier, Options);
+			TokenListDelegate(TokenList);
+			UhtTokenListCache.Return(TokenList);
+			return TokenReader;
+		}
+
+		public static UhtTokenList GetCppIdentifier(this IUhtTokenReader TokenReader, UhtCppIdentifierOptions Options = UhtCppIdentifierOptions.None)
+		{
+			UhtToken Token = TokenReader.GetIdentifier();
+			return TokenReader.GetCppIdentifier(ref Token, Options);
+		}
+
+		public static UhtTokenList GetCppIdentifier(this IUhtTokenReader TokenReader, ref UhtToken InitialIdentifier, UhtCppIdentifierOptions Options = UhtCppIdentifierOptions.None)
+		{
+			UhtTokenList ListHead = UhtTokenListCache.Borrow(InitialIdentifier);
+			UhtTokenList ListTail = ListHead;
+
+			if (Options.HasAnyFlags(UhtCppIdentifierOptions.AllowTemplates | UhtCppIdentifierOptions.AllTokens))
+			{
+				bool bAllowTemplates = Options.HasAnyFlags(UhtCppIdentifierOptions.AllowTemplates);
+
+				while (true)
+				{
+					if (bAllowTemplates && TokenReader.TryPeekOptional('<'))
+					{
+						ListTail.Next = UhtTokenListCache.Borrow(TokenReader.GetToken());
+						ListTail = ListTail.Next;
+						int NestedScopes = 1;
+						while (NestedScopes > 0)
+						{
+							UhtToken Token = TokenReader.GetToken();
+							if (Token.TokenType.IsEndType())
+							{
+								throw new UhtTokenException(TokenReader, Token, new string[] { "<", ">" }, "template");
+							}
+							ListTail.Next = UhtTokenListCache.Borrow(Token);
+							ListTail = ListTail.Next;
+							if (Token.IsSymbol('<'))
+							{
+								++NestedScopes;
+							}
+							else if (Token.IsSymbol('>'))
+							{
+								--NestedScopes;
+							}
+						}
+					}
+
+					if (!TokenReader.TryPeekOptional("::"))
+					{
+						break;
+					}
+					ListTail.Next = UhtTokenListCache.Borrow(TokenReader.GetToken());
+					ListTail = ListTail.Next;
+					ListTail.Next = UhtTokenListCache.Borrow(TokenReader.GetIdentifier());
+					ListTail = ListTail.Next;
+				}
+			}
+			else if (TokenReader.PeekToken().IsSymbol("::"))
+			{
+				TokenReader.While("::", () =>
+				{
+					ListTail.Next = UhtTokenListCache.Borrow(TokenReader.GetIdentifier());
+					ListTail = ListTail.Next;
+				});
+			}
+			return ListHead;
+		}
+	}
+}
