UHT in UBT initial version.  Currently disabled.

diff --git a/Engine/Source/Programs/Shared/EpicGames.UHT/Tokenizer/UhtTokenReaderIntegerExtensions.cs b/Engine/Source/Programs/Shared/EpicGames.UHT/Tokenizer/UhtTokenReaderIntegerExtensions.cs
new file mode 100644
--- /dev/null
+++ b/Engine/Source/Programs/Shared/EpicGames.UHT/Tokenizer/UhtTokenReaderIntegerExtensions.cs
@@ -0,0 +1,233 @@
+// Copyright Epic Games, Inc. All Rights Reserved.
+
+using EpicGames.Core;
+using EpicGames.UHT.Utils;
+
+namespace EpicGames.UHT.Tokenizer
+{
+	public delegate bool UhtParseMergedSignToken(ref UhtToken Token);
+
+	public static class UhtTokenReaderIntegerExtensions
+	{
+		/// <summary>
+		/// Get the next token as an integer.  If the next token is not an integer, no token is consumed.
+		/// </summary>
+		/// <param name="Value">The integer value of the token</param>
+		/// <returns>True if the next token was an integer, false if not.</returns>
+		public static bool TryOptionalConstInt(this IUhtTokenReader TokenReader, out int Value)
+		{
+			ref UhtToken Token = ref TokenReader.PeekToken();
+			if (Token.GetConstInt(out Value))
+			{
+				TokenReader.ConsumeToken();
+				return true;
+			}
+			Value = 0;
+			return false;
+		}
+
+		/// <summary>
+		/// Get the next token as an integer.  If the next token is not an integer, an exception is thrown
+		/// </summary>
+		/// <param name="ExceptionContext">If not null, an exception will be thrown with the given text as part of the message.</param>
+		/// <returns>The value of the constant</returns>
+		public static IUhtTokenReader OptionalConstInt(this IUhtTokenReader TokenReader, object? ExceptionContext = null)
+		{
+			int Value;
+			TokenReader.TryOptionalConstInt(out Value);
+			return TokenReader;
+		}
+
+		/// <summary>
+		/// Get the next token as an integer.  If the next token is not an integer, an exception is thrown
+		/// </summary>
+		/// <param name="ExceptionContext">If not null, an exception will be thrown with the given text as part of the message.</param>
+		/// <returns>The value of the constant</returns>
+		public static IUhtTokenReader RequireConstInt(this IUhtTokenReader TokenReader, object? ExceptionContext = null)
+		{
+			int Value;
+			if (!TokenReader.TryOptionalConstInt(out Value))
+			{
+				throw new UhtTokenException(TokenReader, TokenReader.PeekToken(), "constant integer", ExceptionContext);
+			}
+			return TokenReader;
+		}
+
+		/// <summary>
+		/// Get the next token as an integer.  If the next token is not an integer, an exception is thrown
+		/// </summary>
+		/// <param name="ExceptionContext">If not null, an exception will be thrown with the given text as part of the message.</param>
+		/// <returns>The value of the constant</returns>
+		public static int GetConstInt(this IUhtTokenReader TokenReader, object? ExceptionContext = null)
+		{
+			int Value;
+			if (!TokenReader.TryOptionalConstInt(out Value))
+			{
+				throw new UhtTokenException(TokenReader, TokenReader.PeekToken(), "constant integer", ExceptionContext);
+			}
+			return Value;
+		}
+
+		/// <summary>
+		/// Get the next token as an integer.  If the next token is not an integer, no token is consumed.
+		/// </summary>
+		/// <param name="Value">The integer value of the token</param>
+		/// <returns>True if the next token was an integer, false if not.</returns>
+		public static bool TryOptionalConstLong(this IUhtTokenReader TokenReader, out long Value)
+		{
+			ref UhtToken Token = ref TokenReader.PeekToken();
+			if (Token.GetConstLong(out Value))
+			{
+				TokenReader.ConsumeToken();
+				return true;
+			}
+			Value = 0;
+			return false;
+		}
+
+		/// <summary>
+		/// Get the next token as an integer.  If the next token is not an integer, an exception is thrown
+		/// </summary>
+		/// <param name="ExceptionContext">If not null, an exception will be thrown with the given text as part of the message.</param>
+		/// <returns>The value of the constant</returns>
+		public static IUhtTokenReader OptionalConstLong(this IUhtTokenReader TokenReader, object? ExceptionContext = null)
+		{
+			ref UhtToken Token = ref TokenReader.PeekToken();
+
+			long Value;
+			if (Token.GetConstLong(out Value))
+			{
+				TokenReader.ConsumeToken();
+			}
+			return TokenReader;
+		}
+
+		/// <summary>
+		/// Get the next token as an integer.  If the next token is not an integer, an exception is thrown
+		/// </summary>
+		/// <param name="ExceptionContext">If not null, an exception will be thrown with the given text as part of the message.</param>
+		/// <returns>The value of the constant</returns>
+		public static IUhtTokenReader RequireConstLong(this IUhtTokenReader TokenReader, object? ExceptionContext = null)
+		{
+			long Value;
+			if (!TokenReader.TryOptionalConstLong(out Value))
+			{
+				throw new UhtTokenException(TokenReader, TokenReader.PeekToken(), "constant long integer", ExceptionContext);
+			}
+			return TokenReader;
+		}
+
+		/// <summary>
+		/// Get the next token as an integer.  If the next token is not an integer, an exception is thrown
+		/// </summary>
+		/// <param name="ExceptionContext">If not null, an exception will be thrown with the given text as part of the message.</param>
+		/// <returns>The value of the constant</returns>
+		public static long GetConstLong(this IUhtTokenReader TokenReader, object? ExceptionContext = null)
+		{
+			long Value;
+			if (!TokenReader.TryOptionalConstLong(out Value))
+			{
+				throw new UhtTokenException(TokenReader, TokenReader.PeekToken(), "constant long integer", ExceptionContext);
+			}
+			return Value;
+		}
+
+		/// <summary>
+		/// Helper method to combine any leading sign with the next numeric token
+		/// </summary>
+		/// <param name="TokenReader">Source tokens</param>
+		/// <param name="Delegate">Delegate to invoke with the merged value</param>
+		/// <returns>True if the next token was an parsed, false if not.</returns>
+		public static bool TryOptionalLeadingSignConstNumeric(this IUhtTokenReader TokenReader, UhtParseMergedSignToken Delegate)
+		{
+			using (var SavedState = new UhtTokenSaveState(TokenReader))
+			{
+				// Check for a leading sign token
+				char Sign = ' ';
+				UhtToken Token = TokenReader.PeekToken();
+				if (Token.IsSymbol() && Token.Value.Length == 1 && UhtFCString.IsSign(Token.Value.Span[0]))
+				{
+					Sign = Token.Value.Span[0];
+					TokenReader.ConsumeToken();
+					Token = TokenReader.PeekToken();
+					if (UhtFCString.IsSign(Token.Value.Span[0]))
+					{
+						return false;
+					}
+					Token.Value = new StringView($"{Sign}{Token.Value}");
+				}
+				if (Delegate(ref Token))
+				{
+					TokenReader.ConsumeToken();
+					SavedState.AbandonState();
+					return true;
+				}
+				return false;
+			}
+		}
+
+		/// <summary>
+		/// Get the next integer.  It also handled [+/-] token followed by an integer.
+		/// </summary>
+		/// <param name="TokenReader">Source tokens</param>
+		/// <param name="Value">The integer value of the token</param>
+		/// <returns>True if the next token was an integer, false if not.</returns>
+		public static bool TryOptionalConstIntExpression(this IUhtTokenReader TokenReader, out int Value)
+		{
+			int LocalValue = 0;
+			bool Results = TokenReader.TryOptionalLeadingSignConstNumeric((ref UhtToken Token) =>
+			{
+				return Token.IsConstInt() && Token.GetConstInt(out LocalValue);
+			});
+			Value = LocalValue;
+			return Results;
+		}
+
+		/// <summary>
+		/// Get the next integer.  It also handled [+/-] token followed by an integer.
+		/// </summary>
+		/// <param name="TokenReader">Source tokens</param>
+		/// <returns>The integer value</returns>
+		public static int GetConstIntExpression(this IUhtTokenReader TokenReader)
+		{
+			int LocalValue = 0;
+			bool Results = TokenReader.TryOptionalLeadingSignConstNumeric((ref UhtToken Token) =>
+			{
+				return Token.IsConstInt() && Token.GetConstInt(out LocalValue);
+			});
+			return LocalValue;
+		}
+
+		/// <summary>
+		/// Get the next integer.  It also handled [+/-] token followed by an integer.
+		/// </summary>
+		/// <param name="TokenReader">Source tokens</param>
+		/// <param name="Value">The integer value of the token</param>
+		/// <returns>True if the next token was an integer, false if not.</returns>
+		public static bool TryOptionalConstLongExpression(this IUhtTokenReader TokenReader, out long Value)
+		{
+			long LocalValue = 0;
+			bool Results = TokenReader.TryOptionalLeadingSignConstNumeric((ref UhtToken Token) =>
+			{
+				return Token.IsConstInt() && Token.GetConstLong(out LocalValue);
+			});
+			Value = LocalValue;
+			return Results;
+		}
+
+		/// <summary>
+		/// Get the next long.  It also handled [+/-] token followed by an long.
+		/// </summary>
+		/// <param name="TokenReader">Source tokens</param>
+		/// <returns>The long value</returns>
+		public static int GetConstLongExpression(this IUhtTokenReader TokenReader)
+		{
+			int LocalValue = 0;
+			bool Results = TokenReader.TryOptionalLeadingSignConstNumeric((ref UhtToken Token) =>
+			{
+				return Token.IsConstInt() && Token.GetConstInt(out LocalValue);
+			});
+			return LocalValue;
+		}
+	}
+}
diff --git a/Engine/Source/Programs/Shared/EpicGames.UHT/Tokenizer/UhtTokenReaderSkipExtensions.cs b/Engine/Source/Programs/Shared/EpicGames.UHT/Tokenizer/UhtTokenReaderSkipExtensions.cs
new file mode 100644
--- /dev/null
+++ b/Engine/Source/Programs/Shared/EpicGames.UHT/Tokenizer/UhtTokenReaderSkipExtensions.cs
@@ -0,0 +1,109 @@
+// Copyright Epic Games, Inc. All Rights Reserved.
+
+using EpicGames.UHT.Utils;
+
+namespace EpicGames.UHT.Tokenizer
+{
+	/// <summary>
+	/// Series of extensions to token reading that are far too specialized to be included in the reader.
+	/// </summary>
+	public static class UhtTokenReaderSkipExtensions
+	{
+		public static IUhtTokenReader SkipOne(this IUhtTokenReader TokenReader)
+		{
+			TokenReader.PeekToken();
+			TokenReader.ConsumeToken();
+			return TokenReader;
+		}
+
+		public static IUhtTokenReader SkipAlignasIfNecessary(this IUhtTokenReader TokenReader)
+		{
+			const string Identifier = "alignas";
+			if (TokenReader.TryOptional(Identifier))
+			{
+				TokenReader
+					.Require('(', Identifier)
+					.RequireConstInt(Identifier)
+					.Require(')', Identifier);
+			}
+			return TokenReader;
+		}
+
+		public static IUhtTokenReader SkipDeprecatedMacroIfNecessary(this IUhtTokenReader TokenReader)
+		{
+			if (TokenReader.TryOptional(new string[] { "DEPRECATED", "UE_DEPRECATED" }) >= 0)
+			{
+				TokenReader
+					.Require('(', "deprecation macro")
+					.RequireConstFloat("version in deprecation macro")
+					.Require(',', "deprecation macro")
+					.RequireConstString("message in deprecation macro")
+					.Require(')', "deprecation macro");
+			}
+			return TokenReader;
+		}
+
+		public static IUhtTokenReader SkipAlignasAndDeprecatedMacroIfNecessary(this IUhtTokenReader TokenReader)
+		{
+			// alignas() can come before or after the deprecation macro.
+			// We can't have both, but the compiler will catch that anyway.
+			return TokenReader
+				.SkipAlignasIfNecessary()
+				.SkipDeprecatedMacroIfNecessary()
+				.SkipAlignasIfNecessary();
+		}
+
+		public static IUhtTokenReader SkipBrackets(this IUhtTokenReader TokenReader, char Initiator, char Terminator, int InitialNesting, object? ExceptionContext = null)
+		{
+			int Nesting = InitialNesting;
+			if (Nesting == 0)
+			{
+				TokenReader.Require(Initiator, ExceptionContext);
+				++Nesting;
+			}
+
+			do
+			{
+				UhtToken SkipToken = TokenReader.GetToken();
+				if (SkipToken.TokenType.IsEndType())
+				{
+					throw new UhtTokenException(TokenReader, SkipToken, Terminator, ExceptionContext);
+				}
+				else if (SkipToken.IsSymbol(Initiator))
+				{
+					++Nesting;
+				}
+				else if (SkipToken.IsSymbol(Terminator))
+				{
+					--Nesting;
+				}
+			} while (Nesting != 0);
+			return TokenReader;
+		}
+
+		/// <summary>
+		/// Skip tokens until the given terminator is found.  The terminator will not be consumed.
+		/// </summary>
+		/// <param name="TokenReader">Token reader</param>
+		/// <param name="Terminator">Terminator to skip until</param>
+		/// <param name="ExceptionContext">Extra context for any exceptions</param>
+		/// <returns>Token reader</returns>
+		public static IUhtTokenReader SkipUntil(this IUhtTokenReader TokenReader, char Terminator, object? ExceptionContext = null)
+		{
+			while (true)
+			{
+				ref UhtToken SkipToken = ref TokenReader.PeekToken();
+				if (SkipToken.TokenType.IsEndType())
+				{
+					throw new UhtTokenException(TokenReader, SkipToken, Terminator, ExceptionContext);
+				}
+				else if (SkipToken.IsSymbol(Terminator))
+				{
+					break;
+				}
+				TokenReader.ConsumeToken();
+			}
+			return TokenReader;
+		}
+	}
+}
diff --git a/Engine/Source/Programs/Shared/EpicGames.UHT/Tokenizer/UhtTokenReaderStringExtensions.cs b/Engine/Source/Programs/Shared/EpicGames.UHT/Tokenizer/UhtTokenReaderStringExtensions.cs
new file mode 100644
--- /dev/null
+++ b/Engine/Source/Programs/Shared/EpicGames.UHT/Tokenizer/UhtTokenReaderStringExtensions.cs
@@ -0,0 +1,116 @@
+// Copyright Epic Games, Inc. All Rights Reserved.
+
+using EpicGames.Core;
+using EpicGames.UHT.Utils;
+
+namespace EpicGames.UHT.Tokenizer
+{
+	public static class UhtTokenReaderStringExtensions
+	{
+		/// <summary>
+		/// Get the next token as a string.  If the next token is not a string, no token is consumed.
+		/// </summary>
+		/// <param name="Value">The string value of the token</param>
+		/// <returns>True if the next token was an string, false if not.</returns>
+		public static bool TryOptionalConstString(this IUhtTokenReader TokenReader, out StringView Value)
+		{
+			ref UhtToken Token = ref TokenReader.PeekToken();
+			if (Token.IsConstString())
+			{
+				Value = Token.GetTokenString();
+				TokenReader.ConsumeToken();
+				return true;
+			}
+			Value = "";
+			return false;
+		}
+
+		/// <summary>
+		/// Get the next token as a string.  If the next token is not a string, no token is consumed.
+		/// </summary>
+		/// <returns>True if the next token was an string, false if not.</returns>
+		public static IUhtTokenReader OptionalConstString(this IUhtTokenReader TokenReader)
+		{
+			ref UhtToken Token = ref TokenReader.PeekToken();
+			if (Token.IsConstString())
+			{
+				TokenReader.ConsumeToken();
+			}
+			return TokenReader;
+		}
+
+		/// <summary>
+		/// Verify that the next token is a string.
+		/// </summary>
+		/// <param name="ExceptionContext">If not null, an exception will be thrown with the given text as part of the message.</param>
+		/// <returns>True if the next token was a string, false if not.</returns>
+		public static IUhtTokenReader RequireConstString(this IUhtTokenReader TokenReader, object? ExceptionContext)
+		{
+			ref UhtToken Token = ref TokenReader.PeekToken();
+			if (Token.IsConstString())
+			{
+				TokenReader.ConsumeToken();
+				return TokenReader;
+			}
+			throw new UhtTokenException(TokenReader, Token, "constant string", ExceptionContext);
+		}
+
+		/// <summary>
+		/// Get the next token as a string.  If the next token is not a string, an exception is thrown
+		/// </summary>
+		/// <param name="ExceptionContext">If not null, an exception will be thrown with the given text as part of the message.</param>
+		/// <returns>The value of the string.</returns>
+		public static StringView GetConstString(this IUhtTokenReader TokenReader, object? ExceptionContext = null)
+		{
+			ref UhtToken Token = ref TokenReader.PeekToken();
+			if (Token.IsConstString())
+			{
+				StringView Output = Token.GetTokenString();
+				TokenReader.ConsumeToken();
+				return Output;
+			}
+			throw new UhtTokenException(TokenReader, Token, "constant string", ExceptionContext);
+		}
+
+		/// <summary>
+		/// Get the next token as a quoted string.  If the next token is not a string, an exception is thrown.
+		/// Character constants are not considered strings by this routine.
+		/// </summary>
+		/// <param name="ExceptionContext">If not null, an exception will be thrown with the given text as part of the message.</param>
+		/// <returns>The value of the string.</returns>
+		public static StringView GetConstQuotedString(this IUhtTokenReader TokenReader, object? ExceptionContext = null)
+		{
+			ref UhtToken Token = ref TokenReader.PeekToken();
+			if (Token.TokenType == UhtTokenType.StringConst)
+			{
+				StringView Output = Token.Value;
+				TokenReader.ConsumeToken();
+				return Output;
+			}
+			throw new UhtTokenException(TokenReader, Token, "constant quoted string", ExceptionContext);
+		}
+
+		/// <summary>
+		/// Get a const string that can optionally be wrapped with a TEXT() macro
+		/// </summary>
+		/// <param name="TokenReader">Token reader</param>
+		/// <param name="ExceptionContext">If not null, an exception will be thrown with the given text as part of the message.</param>
+		/// <returns>The value of the string</returns>
+		public static StringView GetWrappedConstString(this IUhtTokenReader TokenReader, object? ExceptionContext = null)
+		{
+			ref UhtToken Token = ref TokenReader.PeekToken();
+			if (Token.IsIdentifier("TEXT"))
+			{
+				TokenReader.ConsumeToken();
+				TokenReader.Require('(');
+				StringView Out = TokenReader.GetConstString(ExceptionContext);
+				TokenReader.Require(')');
+				return Out;
+			}
+			else
+			{
+				return TokenReader.GetConstString(ExceptionContext);
+			}
+		}
+	}
+}
diff --git a/Engine/Source/Programs/Shared/EpicGames.UHT/Tokenizer/UhtTokenReaderUtilityExtensions.cs b/Engine/Source/Programs/Shared/EpicGames.UHT/Tokenizer/UhtTokenReaderUtilityExtensions.cs
new file mode 100644
--- /dev/null
+++ b/Engine/Source/Programs/Shared/EpicGames.UHT/Tokenizer/UhtTokenReaderUtilityExtensions.cs
@@ -0,0 +1,273 @@
+// Copyright Epic Games, Inc. All Rights Reserved.
+
+using EpicGames.Core;
+using EpicGames.UHT.Utils;
+using System;
+using System.Collections.Generic;
+
+namespace EpicGames.UHT.Tokenizer
+{
+	public static class UhtTokenReaderUtilityExtensions
+	{
+		private static HashSet<StringView> SkipDeclarationWarningStrings = new HashSet<StringView>
+			{
+				"GENERATED_BODY",
+				"GENERATED_IINTERFACE_BODY",
+				"GENERATED_UCLASS_BODY",
+				"GENERATED_UINTERFACE_BODY",
+				"GENERATED_USTRUCT_BODY",
+				// Leaving these disabled ATM since they can exist in the code without causing compile issues
+				//"RIGVM_METHOD",
+				//"UCLASS",
+				//"UDELEGATE",
+				//"UENUM",
+				//"UFUNCTION",
+				//"UINTERFACE",
+				//"UPROPERTY",
+				//"USTRUCT",
+			};
+
+		/// <summary>
+		/// When processing type, make sure that the next token is the expected token
+		/// </summary>
+		/// <param name="TokenReader">Token reader</param>
+		/// <param name="Token">Token that started the process</param>
+		/// <returns>true if there could be more header to process, false if the end was reached.</returns>
+		public static bool SkipExpectedType(this IUhtTokenReader TokenReader, StringView ExpectedIdentifier, bool bIsMember)
+		{
+			if (TokenReader.TryOptional(ExpectedIdentifier))
+			{
+				return true;
+			}
+			if (bIsMember && TokenReader.TryOptional("const"))
+			{
+				TokenReader.LogError("Const properties are not supported.");
+			}
+			else
+			{
+				TokenReader.LogError($"Inappropriate keyword '{TokenReader.PeekToken().Value}' on variable of type '{ExpectedIdentifier}'");
+			}
+			return false;
+		}
+
+		public static bool TryOptionalAPIMacro(this IUhtTokenReader TokenReader, out UhtToken APIMacroToken)
+		{
+			ref UhtToken Token = ref TokenReader.PeekToken();
+			if (Token.IsIdentifier() && Token.Value.Span.EndsWith("_API"))
+			{
+				APIMacroToken = Token;
+				TokenReader.ConsumeToken();
+				return true;
+			}
+			APIMacroToken = new UhtToken();
+			return false;
+		}
+
+		public static IUhtTokenReader OptionalInheritance(this IUhtTokenReader TokenReader, UhtTokenDelegate SuperClassDelegate)
+		{
+			TokenReader.Optional(':', () =>
+			{
+				TokenReader
+					.Require("public", "public access modifier")
+					.RequireIdentifier(SuperClassDelegate);
+			});
+			return TokenReader;
+		}
+
+		public static IUhtTokenReader OptionalInheritance(this IUhtTokenReader TokenReader, UhtTokenDelegate SuperClassDelegate, UhtTokenListDelegate BaseClassDelegate)
+		{
+			TokenReader.Optional(':', () =>
+			{
+				TokenReader
+					.Require("public", "public access modifier")
+					.RequireIdentifier(SuperClassDelegate)
+					.While(',', () =>
+					{
+						TokenReader
+							.Require("public", "public interface access specifier")
+							.RequireCppIdentifier(UhtCppIdentifierOptions.AllowTemplates, BaseClassDelegate);
+					});
+			});
+			return TokenReader;
+		}
+
+		/// <summary>
+		/// Given a declaration/statement that starts with the given token, skip the declaration in the header.
+		/// </summary>
+		/// <param name="TokenReader">Token reader</param>
+		/// <param name="Token">Token that started the process</param>
+		/// <returns>true if there could be more header to process, false if the end was reached.</returns>
+		public static bool SkipDeclaration(this IUhtTokenReader TokenReader, ref UhtToken Token)
+		{
+			UhtToken StarToken = Token;
+
+			// Consume all tokens until the end of declaration/definition has been found.
+			int NestedScopes = 0;
+			bool bEndOfDeclarationFound = false;
+
+			// Store the current value of PrevComment so it can be restored after we parsed everything.
+			using (UhtTokenDisableComments DisableComments = new UhtTokenDisableComments(TokenReader))
+			{
+
+				// Check if this is a class/struct declaration in which case it can be followed by member variable declaration.	
+				bool bPossiblyClassDeclaration = Token.IsIdentifier() && (Token.IsValue("class") || Token.IsValue("struct"));
+
+				// (known) macros can end without ; or } so use () to find the end of the declaration.
+				// However, we don't want to use it with DECLARE_FUNCTION, because we need it to be treated like a function.
+				bool bMacroDeclaration = IsProbablyAMacro(Token.Value) && !Token.IsIdentifier("DECLARE_FUNCTION");
+
+				bool bDefinitionFound = false;
+				char OpeningBracket = bMacroDeclaration ? '(' : '{';
+				char ClosingBracket = bMacroDeclaration ? ')' : '}';
+
+				bool bRetestCurrentToken = false;
+				while (true)
+				{
+					if (!bRetestCurrentToken)
+					{
+						Token = TokenReader.GetToken();
+						if (Token.TokenType.IsEndType())
+						{
+							break;
+						}
+					}
+					else
+					{
+						bRetestCurrentToken = false;
+					}
+					ReadOnlySpan<char> Span = Token.Value.Span;
+
+					// If this is a macro, consume it
+					// If we find parentheses at top-level and we think it's a class declaration then it's more likely
+					// to be something like: class UThing* GetThing();
+					if (bPossiblyClassDeclaration && NestedScopes == 0 && Token.IsSymbol() && Span.Length == 1 && Span[0] == '(')
+					{
+						bPossiblyClassDeclaration = false;
+					}
+
+					if (Token.IsSymbol() && Span.Length == 1 && Span[0] == ';' && NestedScopes == 0)
+					{
+						bEndOfDeclarationFound = true;
+						break;
+					}
+
+					if (Token.IsIdentifier())
+					{
+						// Use a trivial pre-filter to avoid doing the search on things that aren't UE keywords we care about
+						if (Span[0] == 'G' || Span[0] == 'R' || Span[0] == 'U')
+						{
+							if (SkipDeclarationWarningStrings.Contains(Token.Value))
+							{
+								TokenReader.LogWarning($"The identifier \'{Token.Value}\' was detected in a block being skipped. Was this intentional?");
+							}
+						}
+					}
+
+					if (!bMacroDeclaration && Token.IsIdentifier() && Span.Equals("PURE_VIRTUAL", StringComparison.Ordinal) && NestedScopes == 0)
+					{
+						OpeningBracket = '(';
+						ClosingBracket = ')';
+					}
+
+					if (Token.IsSymbol() && Span.Length == 1 && Span[0] == OpeningBracket)
+					{
+						// This is a function definition or class declaration.
+						bDefinitionFound = true;
+						NestedScopes++;
+					}
+					else if (Token.IsSymbol() && Span.Length == 1 && Span[0] == ClosingBracket)
+					{
+						NestedScopes--;
+						if (NestedScopes == 0)
+						{
+							// Could be a class declaration in all capitals, and not a macro
+							bool bReallyEndDeclaration = true;
+							if (bMacroDeclaration)
+							{
+								bReallyEndDeclaration = !TokenReader.TryPeekOptional('{');
+							}
+
+							if (bReallyEndDeclaration)
+							{
+								bEndOfDeclarationFound = true;
+								break;
+							}
+						}
+
+						if (NestedScopes < 0)
+						{
+							throw new UhtException(TokenReader, Token.InputLine, $"Unexpected '{ClosingBracket}'. Did you miss a semi-colon?");
+						}
+					}
+					else if (bMacroDeclaration && NestedScopes == 0)
+					{
+						bMacroDeclaration = false;
+						OpeningBracket = '{';
+						ClosingBracket = '}';
+						bRetestCurrentToken = true;
+					}
+				}
+				if (bEndOfDeclarationFound)
+				{
+					// Member variable declaration after class declaration (see bPossiblyClassDeclaration).
+					if (bPossiblyClassDeclaration && bDefinitionFound)
+					{
+						// Should syntax errors be also handled when someone declares a variable after function definition?
+						// Consume the variable name.
+						if (TokenReader.bIsEOF)
+						{
+							return true;
+						}
+						if (TokenReader.TryOptionalIdentifier())
+						{
+							TokenReader.Require(';');
+						}
+					}
+
+					// C++ allows any number of ';' after member declaration/definition.
+					while (TokenReader.TryOptional(';'))
+					{ 
+					}
+				}
+			}
+
+			// Successfully consumed C++ declaration unless mismatched pair of brackets has been found.
+			return NestedScopes == 0 && bEndOfDeclarationFound;
+		}
+
+		private static bool IsProbablyAMacro(StringView Identifier)
+		{
+			ReadOnlySpan<char> Span = Identifier.Span;
+			if (Span.Length == 0)
+			{
+				return false;
+			}
+
+			// Macros must start with a capitalized alphanumeric character or underscore
+			char FirstChar = Span[0];
+			if (FirstChar != '_' && (FirstChar < 'A' || FirstChar > 'Z'))
+			{
+				return false;
+			}
+
+			// Test for known delegate and event macros.
+			if (Span.StartsWith("DECLARE_MULTICAST_DELEGATE") ||
+				Span.StartsWith("DECLARE_DELEGATE") ||
+				Span.StartsWith("DECLARE_EVENT"))
+			{
+				return true;
+			}
+
+			// Failing that, we'll guess about it being a macro based on it being a fully-capitalized identifier.
+			foreach (char Ch in Span.Slice(1))
+			{
+				if (Ch != '_' && (Ch < 'A' || Ch > 'Z') && (Ch < '0' || Ch > '9'))
+				{
+					return false;
+				}
+			}
+
+			return true;
+		}
+	}
+}
diff --git a/Engine/Source/Programs/Shared/EpicGames.UHT/Tokenizer/UhtTokenReplayReader.cs b/Engine/Source/Programs/Shared/EpicGames.UHT/Tokenizer/UhtTokenReplayReader.cs
new file mode 100644
--- /dev/null
+++ b/Engine/Source/Programs/Shared/EpicGames.UHT/Tokenizer/UhtTokenReplayReader.cs
@@ -0,0 +1,303 @@
+// Copyright Epic Games, Inc. All Rights Reserved.
+
+using EpicGames.Core;
+using EpicGames.UHT.Utils;
+using System;
+using System.Collections.Generic;
+using System.Threading;
+
+namespace EpicGames.UHT.Tokenizer
+{
+	public class UhtTokenReplayReader : IUhtTokenReader, IUhtMessageLineNumber
+	{
+		static private ThreadLocal<UhtTokenReplayReader> Tls = new ThreadLocal<UhtTokenReplayReader>(() => new UhtTokenReplayReader());
+
+		const int MaxSavedStates = 2;
+		private struct SavedState
+		{
+			public int TokenIndex;
+			public bool bHasToken;
+		}
+
+		private IUhtMessageSite MessageSiteInternal;
+		private ReadOnlyMemory<UhtToken> Tokens;
+		private ReadOnlyMemory<char> Data;
+		private int CurrentTokenIndex = -1;
+		private bool bHasToken = false;
+		private UhtToken CurrentToken = new UhtToken();
+		private SavedState[] SavedStates = new SavedState[MaxSavedStates];
+		private int SavedStateCount = 0;
+		private UhtTokenType EndTokenType = UhtTokenType.EndOfFile;
+
+		public UhtTokenReplayReader(IUhtMessageSite MessageSite, ReadOnlyMemory<char> Data, ReadOnlyMemory<UhtToken> Tokens, UhtTokenType EndTokenType)
+		{
+			this.MessageSiteInternal = MessageSite;
+			this.Tokens = new UhtToken[0].AsMemory();
+			this.Data = new char[0].AsMemory();
+			this.EndTokenType = EndTokenType;
+			this.CurrentToken = new UhtToken(EndTokenType);
+		}
+
+		public UhtTokenReplayReader()
+		{
+			this.MessageSiteInternal = new UhtEmptyMessageSite();
+			this.Tokens = new UhtToken[0].AsMemory();
+			this.Data = new char[0].AsMemory();
+		}
+
+		public UhtTokenReplayReader Reset(IUhtMessageSite MessageSite, ReadOnlyMemory<char> Data, ReadOnlyMemory<UhtToken> Tokens, UhtTokenType EndTokenType)
+		{
+			this.MessageSiteInternal = MessageSite;
+			this.Tokens = Tokens;
+			this.Data = Data;
+			this.CurrentTokenIndex = -1;
+			this.bHasToken = false;
+			this.CurrentToken = new UhtToken(EndTokenType);
+			this.SavedStateCount = 0;
+			this.EndTokenType = EndTokenType;
+			return this;
+		}
+
+		/// <summary>
+		/// Return the replay reader associated with the current thread.  Only one replay reader is cached per thread.
+		/// </summary>
+		/// <param name="MessageSite">The message site used to log errors</param>
+		/// <param name="Data">Source data where tokens were originally parsed</param>
+		/// <param name="Tokens">Collection of tokens to replay</param>
+		/// <param name="EndTokenType">Type of end token marker to return when the end of the token list is reached.  This is used to produce errors in the context of the replay</param>
+		/// <returns>The threaded instance of the replay reader</returns>
+		/// <exception cref="UhtIceException">Thrown if the TLS value can not be retrieved.</exception>
+		public static UhtTokenReplayReader GetThreadInstance(IUhtMessageSite MessageSite, ReadOnlyMemory<char> Data, ReadOnlyMemory<UhtToken> Tokens, UhtTokenType EndTokenType)
+		{
+			UhtTokenReplayReader? Reader = Tls.Value;
+			if (Reader == null)
+			{
+				throw new UhtIceException("Unable to acquire a UhtTokenReplayReader");
+			}
+			Reader.Reset(MessageSite, Data, Tokens, EndTokenType);
+			return Reader;
+		}
+
+		#region IUHTMessageSite Implementation
+		IUhtMessageSession IUhtMessageSite.MessageSession => this.MessageSiteInternal.MessageSession;
+		IUhtMessageSource? IUhtMessageSite.MessageSource => this.MessageSiteInternal.MessageSource;
+		IUhtMessageLineNumber? IUhtMessageSite.MessageLineNumber => this;
+		#endregion
+
+		#region ITokenReader Implementation
+		public bool bIsEOF
+		{
+			get
+			{
+				if (!this.bHasToken)
+				{
+					GetTokenInternal();
+				}
+				return this.CurrentTokenIndex == this.Tokens.Span.Length;
+			}
+		}
+
+		public int InputPos 
+		{
+			get
+			{
+				if (!this.bHasToken)
+				{
+					GetTokenInternal();
+				}
+				return CurrentToken.InputStartPos;
+			}
+		}
+
+		public int InputLine 
+		{
+			get
+			{
+				if (!this.bHasToken)
+				{
+					GetTokenInternal();
+				}
+				return CurrentToken.InputLine;
+			}
+			set => throw new NotImplementedException(); 
+		}
+
+		public IUhtTokenPreprocessor? TokenPreprocessor { get => throw new NotImplementedException(); set => throw new NotImplementedException(); }
+
+		public int LookAheadEnableCount { get => throw new NotImplementedException(); set => throw new NotImplementedException(); }
+
+		public ReadOnlySpan<StringView> Comments => throw new NotImplementedException();
+
+		public void ClearComments()
+		{
+			throw new NotImplementedException();
+		}
+
+		public void ConsumeToken()
+		{
+			this.bHasToken = false;
+		}
+
+		public void DisableComments()
+		{
+			throw new NotImplementedException();
+		}
+
+		public void EnableComments()
+		{
+			throw new NotImplementedException();
+		}
+
+		public void CommitPendingComments()
+		{
+			throw new NotImplementedException();
+		}
+
+		public UhtToken GetLine()
+		{
+			throw new NotImplementedException();
+		}
+
+		public StringView GetRawString(char Terminator, UhtRawStringOptions Options)
+		{
+			throw new NotImplementedException();
+		}
+
+		public StringView GetStringView(int StartPos, int Count)
+		{
+			return new StringView(this.Data, StartPos, Count);
+		}
+
+		public UhtToken GetToken()
+		{
+			UhtToken Token = PeekToken();
+			ConsumeToken();
+			return Token;
+		}
+
+		public bool IsFirstTokenInLine(ref UhtToken Token)
+		{
+			throw new NotImplementedException();
+		}
+
+		public ref UhtToken PeekToken()
+		{
+			if (!this.bHasToken)
+			{
+				GetTokenInternal();
+			}
+			return ref this.CurrentToken;
+		}
+
+		public void SkipWhitespaceAndComments()
+		{
+			throw new NotImplementedException();
+		}
+
+		public void SaveState()
+		{
+			if (this.SavedStateCount == MaxSavedStates)
+			{
+				throw new UhtIceException("Token reader saved states full");
+			}
+			this.SavedStates[this.SavedStateCount] = new SavedState { TokenIndex = this.CurrentTokenIndex, bHasToken = this.bHasToken };
+			++this.SavedStateCount;
+		}
+
+		public void RestoreState()
+		{
+			if (this.SavedStateCount == 0)
+			{
+				throw new UhtIceException("Attempt to restore a state when none have been saved");
+			}
+
+			--this.SavedStateCount;
+			this.CurrentTokenIndex = this.SavedStates[this.SavedStateCount].TokenIndex;
+			this.bHasToken = this.SavedStates[this.SavedStateCount].bHasToken;
+		}
+
+		public void AbandonState()
+		{
+			if (this.SavedStateCount == 0)
+			{
+				throw new UhtIceException("Attempt to abandon a state when none have been saved");
+			}
+
+			--this.SavedStateCount;
+		}
+
+		/// <inheritdoc/>
+		public void EnableRecording()
+		{
+			throw new NotImplementedException();
+		}
+
+		/// <inheritdoc/>
+		public void DisableRecording()
+		{
+			throw new NotImplementedException();
+		}
+
+		/// <inheritdoc/>
+		public void RecordToken(ref UhtToken Token)
+		{
+			throw new NotImplementedException();
+		}
+
+		/// <inheritdoc/>
+		public List<UhtToken> RecordedTokens
+		{
+			get
+			{
+				throw new NotImplementedException();
+			}
+		}
+		#endregion
+
+		#region IUHTMessageLineNumber implementation
+		int IUhtMessageLineNumber.MessageLineNumber
+		{
+			get
+			{
+				if (this.Tokens.Length == 0)
+				{
+					return -1;
+				}
+				if (this.CurrentTokenIndex < 0)
+				{
+					return this.Tokens.Span[0].InputLine;
+				}
+				if (this.CurrentTokenIndex < this.Tokens.Length)
+				{
+					return this.Tokens.Span[this.CurrentTokenIndex].InputLine;
+				}
+				return this.Tokens.Span[this.Tokens.Length - 1].InputLine;
+			}
+		}
+		#endregion
+
+		private ref UhtToken GetTokenInternal()
+		{
+			if (this.CurrentTokenIndex < this.Tokens.Length)
+			{
+				++this.CurrentTokenIndex;
+			}
+			if (this.CurrentTokenIndex < this.Tokens.Length)
+			{
+				this.CurrentToken = this.Tokens.Span[this.CurrentTokenIndex];
+			}
+			else if (this.Tokens.Length == 0)
+			{
+				this.CurrentToken = new UhtToken();
+			}
+			else
+			{
+				UhtToken LastToken = this.Tokens.Span[this.Tokens.Length - 1];
+				int EndPos = LastToken.InputEndPos;
+				this.CurrentToken = new UhtToken(this.EndTokenType, EndPos, LastToken.InputLine, EndPos, LastToken.InputLine, new StringView());
+			}
+			this.bHasToken = true;
+			return ref this.CurrentToken;
+		}
+	}
+}
diff --git a/Engine/Source/Programs/Shared/EpicGames.UHT/Types/Properties/UhtArrayProperty.cs b/Engine/Source/Programs/Shared/EpicGames.UHT/Types/Properties/UhtArrayProperty.cs
new file mode 100644
--- /dev/null
+++ b/Engine/Source/Programs/Shared/EpicGames.UHT/Types/Properties/UhtArrayProperty.cs
@@ -0,0 +1,268 @@
+// Copyright Epic Games, Inc. All Rights Reserved.
+
+using EpicGames.Core;
+using EpicGames.UHT.Parsers;
+using EpicGames.UHT.Tables;
+using EpicGames.UHT.Tokenizer;
+using EpicGames.UHT.Utils;
+using System.Collections.Generic;
+using System.Text;
+
+namespace EpicGames.UHT.Types
+{
+
+	/// <summary>
+	/// Represents the FArrayProperty engine type
+	/// </summary>
+	[UnrealHeaderTool]
+	[UhtEngineClass(Name = "ArrayProperty", IsProperty = true)]
+	public class UhtArrayProperty : UhtContainerBaseProperty
+	{
+		/// <inheritdoc/>
+		public override string EngineClassName { get => "ArrayProperty"; }
+
+		/// <inheritdoc/>
+		protected override string CppTypeText { get => "TArray"; }
+
+		/// <inheritdoc/>
+		protected override string PGetMacroText { get => "TARRAY"; }
+
+		/// <inheritdoc/>
+		protected override UhtPGetArgumentType PGetTypeArgument { get => UhtPGetArgumentType.TypeText; }
+
+		/// <summary>
+		/// Construct a new array property
+		/// </summary>
+		/// <param name="PropertySettings">Property settings</param>
+		/// <param name="Value">Inner property value</param>
+		public UhtArrayProperty(UhtPropertySettings PropertySettings, UhtProperty Value) : base(PropertySettings, Value)
+		{
+			// If the creation of the value property set more flags, then copy those flags to ourselves
+			this.PropertyFlags |= Value.PropertyFlags & EPropertyFlags.UObjectWrapper;
+
+			if (Value.MetaData.ContainsKey(UhtNames.NativeConst))
+			{
+				this.MetaData.Add(UhtNames.NativeConstTemplateArg, "");
+				Value.MetaData.Remove(UhtNames.NativeConst);
+			}
+
+			this.PropertyCaps |= UhtPropertyCaps.PassCppArgsByRef;
+			this.PropertyCaps &= ~(UhtPropertyCaps.CanBeContainerValue | UhtPropertyCaps.CanBeContainerKey);
+			this.PropertyCaps = (this.PropertyCaps & ~(UhtPropertyCaps.IsParameterSupportedByBlueprint | UhtPropertyCaps.IsMemberSupportedByBlueprint | UhtPropertyCaps.CanExposeOnSpawn)) |
+				(this.ValueProperty.PropertyCaps & (UhtPropertyCaps.IsParameterSupportedByBlueprint | UhtPropertyCaps.IsMemberSupportedByBlueprint | UhtPropertyCaps.CanExposeOnSpawn));
+			this.ValueProperty.SourceName = this.SourceName;
+			this.ValueProperty.EngineName = this.EngineName;
+			this.ValueProperty.PropertyFlags = this.PropertyFlags & EPropertyFlags.PropagateToArrayInner;
+			this.ValueProperty.Outer = this;
+			this.ValueProperty.MetaData.Clear();
+		}
+
+		/// <inheritdoc/>
+		protected override bool ResolveSelf(UhtResolvePhase Phase)
+		{
+			bool bResults = base.ResolveSelf(Phase);
+			switch (Phase)
+			{
+				case UhtResolvePhase.Final:
+					this.PropertyFlags |= ResolveAndReturnNewFlags(this.ValueProperty, Phase);
+					this.MetaData.Add(this.ValueProperty.MetaData);
+					this.ValueProperty.MetaData.Clear();
+					PropagateFlagsFromInnerAndHandlePersistentInstanceMetadata(this, this.MetaData, this.ValueProperty);
+					break;
+			}
+			return bResults;
+		}
+
+		/// <inheritdoc/>
+		public override void CollectReferencesInternal(IUhtReferenceCollector Collector, bool bTemplateProperty)
+		{
+			this.ValueProperty.CollectReferencesInternal(Collector, true);
+		}
+
+		/// <inheritdoc/>
+		public override string? GetForwardDeclarations()
+		{
+			return this.ValueProperty.GetForwardDeclarations();
+		}
+
+		/// <inheritdoc/>
+		public override IEnumerable<UhtType> EnumerateReferencedTypes()
+		{
+			foreach (UhtType Type in this.ValueProperty.EnumerateReferencedTypes())
+			{
+				yield return Type;
+			}
+		}
+
+		/// <inheritdoc/>
+		public override StringBuilder AppendText(StringBuilder Builder, UhtPropertyTextType TextType, bool bIsTemplateArgument)
+		{
+			switch (TextType)
+			{
+				case UhtPropertyTextType.SparseShort:
+					Builder.Append("TArray");
+					break;
+
+				case UhtPropertyTextType.FunctionThunkParameterArgType:
+					Builder.AppendFunctionThunkParameterArrayType(this.ValueProperty);
+					break;
+
+				default:
+					Builder.Append("TArray<").AppendPropertyText(this.ValueProperty, TextType, true);
+					if (Builder[Builder.Length - 1] == '>')
+					{
+						// if our internal property type is a template class, add a space between the closing brackets b/c VS.NET cannot parse this correctly
+						Builder.Append(" ");
+					}
+					Builder.Append('>');
+					break;
+			}
+			return Builder;
+		}
+
+		/// <inheritdoc/>
+		public override StringBuilder AppendMemberDecl(StringBuilder Builder, IUhtPropertyMemberContext Context, string Name, string NameSuffix, int Tabs)
+		{
+			Builder.AppendMemberDecl(this.ValueProperty, Context, Name, GetNameSuffix(NameSuffix, "_Inner"), Tabs);
+			return AppendMemberDecl(Builder, Context, Name, NameSuffix, Tabs, "FArrayPropertyParams");
+		}
+
+		/// <inheritdoc/>
+		public override StringBuilder AppendMemberDef(StringBuilder Builder, IUhtPropertyMemberContext Context, string Name, string NameSuffix, string? Offset, int Tabs)
+		{
+			Builder.AppendMemberDef(this.ValueProperty, Context, Name, GetNameSuffix(NameSuffix, "_Inner"), "0", Tabs);
+			AppendMemberDefStart(Builder, Context, Name, NameSuffix, Offset, Tabs, "FArrayPropertyParams", "UECodeGen_Private::EPropertyGenFlags::Array");
+			Builder.Append(this.Allocator == UhtPropertyAllocator.MemoryImage ? "EArrayPropertyFlags::UsesMemoryImageAllocator" : "EArrayPropertyFlags::None").Append(", ");
+			AppendMemberDefEnd(Builder, Context, Name, NameSuffix);
+			return Builder;
+		}
+
+		/// <inheritdoc/>
+		public override StringBuilder AppendMemberPtr(StringBuilder Builder, IUhtPropertyMemberContext Context, string Name, string NameSuffix, int Tabs)
+		{
+			Builder.AppendMemberPtr(this.ValueProperty, Context, Name, GetNameSuffix(NameSuffix, "_Inner"), Tabs);
+			base.AppendMemberPtr(Builder, Context, Name, NameSuffix, Tabs);
+			return Builder;
+		}
+
+		/// <inheritdoc/>
+		public override void AppendObjectHashes(StringBuilder Builder, int StartingLength, IUhtPropertyMemberContext Context)
+		{
+			this.ValueProperty.AppendObjectHashes(Builder, StartingLength, Context);
+		}
+
+		/// <inheritdoc/>
+		public override StringBuilder AppendNullConstructorArg(StringBuilder Builder, bool bIsInitializer)
+		{
+			Builder.AppendPropertyText(this, UhtPropertyTextType.Construction).Append("()");
+			return Builder;
+		}
+
+		/// <inheritdoc/>
+		public override bool SanitizeDefaultValue(IUhtTokenReader DefaultValueReader, StringBuilder InnerDefaultValue)
+		{
+			return false;
+		}
+
+		/// <inheritdoc/>
+		protected override void ValidateFunctionArgument(UhtFunction Function, UhtValidationOptions Options)
+		{
+			base.ValidateFunctionArgument(Function, Options);
+
+			if (Function.FunctionFlags.HasAnyFlags(EFunctionFlags.Net))
+			{
+				if (!Function.FunctionFlags.HasAnyFlags(EFunctionFlags.NetRequest))
+				{
+					if (this.RefQualifier != UhtPropertyRefQualifier.ConstRef && !this.bIsStaticArray)
+					{
+						this.LogError("Replicated TArray parameters must be passed by const reference");
+					}
+				}
+			}
+		}
+
+		/// <inheritdoc/>
+		public override void Validate(UhtStruct OuterStruct, UhtProperty OutermostProperty, UhtValidationOptions Options)
+		{
+			base.Validate(OuterStruct, OutermostProperty, Options);
+
+			if (this.Allocator != UhtPropertyAllocator.Default)
+			{
+				if (this.PropertyFlags.HasAnyFlags(EPropertyFlags.Net))
+				{
+					this.LogError("Replicated arrays with MemoryImageAllocators are not yet supported");
+				}
+			}
+		}
+
+		/// <inheritdoc/>
+		public override bool IsSameType(UhtProperty Other)
+		{
+			if (Other is UhtArrayProperty OtherArray)
+			{
+				return this.ValueProperty.IsSameType(OtherArray.ValueProperty);
+			}
+			return false;
+		}
+
+		/// <inheritdoc/>
+		public override string? GetRigVMType(ref UhtRigVMParameterFlags ParameterFlags)
+		{
+			ParameterFlags |= UhtRigVMParameterFlags.IsArray;
+			StringBuilder Builder = new StringBuilder();
+			Builder.AppendPropertyText(this, UhtPropertyTextType.RigVMTemplateArg);
+			return Builder.ToString();
+		}
+
+#region Keyword
+		[UhtPropertyType(Keyword = "TArray")]
+		private static UhtProperty? ArrayProperty(UhtPropertyResolvePhase ResolvePhase, UhtPropertySettings PropertySettings, IUhtTokenReader TokenReader, UhtToken MatchedToken)
+		{
+			using (var TokenContext = new UhtMessageContext(TokenReader, "TArray"))
+			{
+				if (!TokenReader.SkipExpectedType(MatchedToken.Value, PropertySettings.PropertyCategory == UhtPropertyCategory.Member))
+				{
+					return null;
+				}
+				TokenReader.Require('<');
+
+				// Parse the value type
+				UhtProperty? Value = UhtPropertyParser.ParseTemplateParam(ResolvePhase, PropertySettings, PropertySettings.SourceName, TokenReader);
+				if (Value == null)
+				{
+					return null;
+				}
+
+				if (!Value.PropertyCaps.HasAnyFlags(UhtPropertyCaps.CanBeContainerValue))
+				{
+					TokenReader.LogError($"The type \'{Value.GetUserFacingDecl()}\' can not be used as a value in a TArray");
+				}
+
+				if (TokenReader.TryOptional(','))
+				{
+					// If we found a comma, read the next thing, assume it's an allocator, and report that
+					UhtToken AllocatorToken = TokenReader.GetIdentifier();
+					if (AllocatorToken.IsIdentifier("FMemoryImageAllocator"))
+					{
+						PropertySettings.Allocator = UhtPropertyAllocator.MemoryImage;
+					}
+					else if (AllocatorToken.IsIdentifier("TMemoryImageAllocator"))
+					{
+						TokenReader.RequireList('<', '>', "TMemoryImageAllocator template arguments");
+						PropertySettings.Allocator = UhtPropertyAllocator.MemoryImage;
+					}
+					else
+					{
+						throw new UhtException(TokenReader, $"Found '{AllocatorToken.Value.ToString()}' - explicit allocators are not supported in TArray properties.");
+					}
+				}
+				TokenReader.Require('>');
+
+				//@TODO: Prevent sparse delegate types from being used in a container
+
+				return new UhtArrayProperty(PropertySettings, Value);
+			}
+		}
+#endregion
+	}
+}
diff --git a/Engine/Source/Programs/Shared/EpicGames.UHT/Types/Properties/UhtBooleanProperty.cs b/Engine/Source/Programs/Shared/EpicGames.UHT/Types/Properties/UhtBooleanProperty.cs
new file mode 100644
--- /dev/null
+++ b/Engine/Source/Programs/Shared/EpicGames.UHT/Types/Properties/UhtBooleanProperty.cs
@@ -0,0 +1,301 @@
+// Copyright Epic Games, Inc. All Rights Reserved.
+
+using EpicGames.Core;
+using EpicGames.UHT.Tables;
+using EpicGames.UHT.Tokenizer;
+using EpicGames.UHT.Utils;
+using System.Text;
+
+namespace EpicGames.UHT.Types
+{
+
+	/// <summary>
+	/// Represents the FBoolProperty engine type
+	/// </summary>
+	[UhtEngineClass(Name = "BoolProperty", IsProperty = true)]
+	public abstract class UhtBooleanProperty : UhtProperty
+	{
+		/// <inheritdoc/>
+		public override string EngineClassName { get => "BoolProperty"; }
+
+		/// <summary>
+		/// If true, the boolean is a native bool and not a UBOOL
+		/// </summary>
+		protected abstract bool bIsNativeBool { get; }
+
+		/// <summary>
+		/// Construct a new boolean property
+		/// </summary>
+		/// <param name="PropertySettings">Property settings</param>
+		protected UhtBooleanProperty(UhtPropertySettings PropertySettings) : base(PropertySettings)
+		{
+			this.PropertyCaps |= UhtPropertyCaps.RequiresNullConstructorArg | UhtPropertyCaps.IsParameterSupportedByBlueprint | UhtPropertyCaps.IsMemberSupportedByBlueprint;
+		}
+
+		/// <inheritdoc/>
+		public override StringBuilder AppendText(StringBuilder Builder, UhtPropertyTextType TextType, bool bIsTemplateArgument = false)
+		{
+			switch (TextType)
+			{
+				case UhtPropertyTextType.UserFacing:
+				case UhtPropertyTextType.Construction:
+				case UhtPropertyTextType.FunctionThunkParameterArrayType:
+				case UhtPropertyTextType.FunctionThunkReturn:
+				case UhtPropertyTextType.RigVMTemplateArg:
+				case UhtPropertyTextType.ExportMember:
+				case UhtPropertyTextType.GetterSetterArg:
+					Builder.Append(this.CppTypeText);
+					break;
+
+				case UhtPropertyTextType.Sparse:
+				case UhtPropertyTextType.SparseShort:
+				case UhtPropertyTextType.ClassFunction:
+				case UhtPropertyTextType.EventFunction:
+				case UhtPropertyTextType.InterfaceFunction:
+				case UhtPropertyTextType.EventParameterMember:
+				case UhtPropertyTextType.EventParameterFunctionMember:
+					Builder.Append("bool");
+					break;
+			}
+			return Builder;
+		}
+
+		/// <inheritdoc/>
+		public override StringBuilder AppendMemberDecl(StringBuilder Builder, IUhtPropertyMemberContext Context, string Name, string NameSuffix, int Tabs)
+		{
+			Builder.AppendMetaDataDecl(this, Context, Name, NameSuffix, Tabs);
+			if (!(this.Outer is UhtProperty))
+			{
+				Builder.AppendTabs(Tabs).Append("static void ").AppendNameDecl(Context, Name, NameSuffix).Append("_SetBit(void* Obj);\r\n");
+			}
+			Builder.AppendTabs(Tabs).Append("static const UECodeGen_Private::").Append("FBoolPropertyParams").Append(' ').AppendNameDecl(Context, Name, NameSuffix).Append(";\r\n");
+			return Builder;
+		}
+
+		/// <inheritdoc/>
+		public override StringBuilder AppendMemberDef(StringBuilder Builder, IUhtPropertyMemberContext Context, string Name, string NameSuffix, string? Offset, int Tabs)
+		{
+			Builder.AppendMetaDataDef(this, Context, Name, NameSuffix, Tabs);
+			if (this.Outer == Context.OuterStruct)
+			{
+				Builder.AppendTabs(Tabs).Append("void ").AppendNameDef(Context, Name, NameSuffix).Append("_SetBit(void* Obj)\r\n");
+				Builder.AppendTabs(Tabs).Append("{\r\n");
+				Builder.AppendTabs(Tabs + 1).Append("((").Append(Context.OuterStructSourceName).Append("*)Obj)->").Append(this.SourceName).Append(" = 1;\r\n");
+				Builder.AppendTabs(Tabs).Append("}\r\n");
+			}
+
+			AppendMemberDefStart(Builder, Context, Name, NameSuffix, Offset, Tabs, "FBoolPropertyParams",
+				bIsNativeBool ?
+				"UECodeGen_Private::EPropertyGenFlags::Bool | UECodeGen_Private::EPropertyGenFlags::NativeBool" :
+				"UECodeGen_Private::EPropertyGenFlags::Bool ", 
+				false, false);
+
+			Builder.Append("sizeof(").Append(this.CppTypeText).Append("), ");
+
+			if (this.Outer == Context.OuterStruct)
+			{
+				Builder
+					.Append("sizeof(").Append(Context.OuterStructSourceName).Append("), ")
+					.Append("&").AppendNameDef(Context, Name, NameSuffix).Append("_SetBit, ");
+			}
+			else
+			{
+				Builder.Append("0, nullptr, ");
+			}
+
+			AppendMemberDefEnd(Builder, Context, Name, NameSuffix);
+			return Builder;
+		}
+
+		/// <inheritdoc/>
+		public override StringBuilder AppendFullDecl(StringBuilder Builder, UhtPropertyTextType TextType, bool bSkipParameterName = false)
+		{
+			AppendText(Builder, TextType);
+
+			//@todo we currently can't have out bools.. so this isn't really necessary, but eventually out bools may be supported, so leave here for now
+			if (TextType.IsParameter() && this.PropertyFlags.HasAnyFlags(EPropertyFlags.OutParm))
+			{
+				Builder.Append('&');
+			}
+
+			Builder.Append(' ');
+
+			if (!bSkipParameterName)
+			{
+				Builder.Append(this.SourceName);
+			}
+
+			if (this.ArrayDimensions != null)
+			{
+				Builder.Append('[').Append(this.ArrayDimensions).Append(']');
+			}
+			else if (TextType == UhtPropertyTextType.ExportMember && !(this is UhtBoolProperty))
+			{
+				Builder.Append(":1");
+			}
+			return Builder;
+		}
+
+		/// <inheritdoc/>
+		public override StringBuilder AppendNullConstructorArg(StringBuilder Builder, bool bIsInitializer)
+		{
+			Builder.Append("false");
+			return Builder;
+		}
+
+		/// <inheritdoc/>
+		public override bool SanitizeDefaultValue(IUhtTokenReader DefaultValueReader, StringBuilder InnerDefaultValue)
+		{
+			UhtToken Identifier = DefaultValueReader.GetIdentifier();
+			if (Identifier.IsValue("true") || Identifier.IsValue("false"))
+			{
+				InnerDefaultValue.Append(Identifier.Value.ToString());
+				return true;
+			}
+			return false;
+		}
+
+		/// <inheritdoc/>
+		public override bool IsSameType(UhtProperty Other)
+		{
+			return Other is UhtBooleanProperty;
+		}
+
+		/// <inheritdoc/>
+		public override string? GetRigVMType(ref UhtRigVMParameterFlags ParameterFlags)
+		{
+			return CppTypeText;
+		}
+	}
+
+	/// <summary>
+	/// Native boolean type specialization
+	/// </summary>
+	[UnrealHeaderTool]
+	public class UhtBoolProperty : UhtBooleanProperty
+	{
+		/// <inheritdoc/>
+		protected override string CppTypeText { get => "bool"; }
+
+		/// <inheritdoc/>
+		protected override string PGetMacroText { get => "UBOOL"; }
+
+		/// <inheritdoc/>
+		protected override bool bIsNativeBool { get => true; }
+
+		/// <summary>
+		/// Construct a new boolean property
+		/// </summary>
+		/// <param name="PropertySettings">Property settings</param>
+		public UhtBoolProperty(UhtPropertySettings PropertySettings) : base(PropertySettings)
+		{
+			this.PropertyCaps |= UhtPropertyCaps.CanExposeOnSpawn;
+		}
+
+		#region Keyword
+		[UhtPropertyType(Keyword = "bool", Options = UhtPropertyTypeOptions.Simple | UhtPropertyTypeOptions.Immediate)]
+		private static UhtProperty? BoolProperty(UhtPropertyResolvePhase ResolvePhase, UhtPropertySettings PropertySettings, IUhtTokenReader TokenReader, UhtToken MatchedToken)
+		{
+			if (PropertySettings.bIsBitfield)
+			{
+				TokenReader.LogError("bool bitfields are not supported.");
+				return null;
+			}
+			return new UhtBoolProperty(PropertySettings);
+		}
+		#endregion
+	}
+
+	/// <summary>
+	/// uint8 bit type specialization
+	/// </summary>
+	public class UhtBool8Property : UhtBooleanProperty
+	{
+		/// <inheritdoc/>
+		protected override string CppTypeText { get => "uint8"; }
+
+		/// <inheritdoc/>
+		protected override string PGetMacroText { get => "UBOOL8"; }
+
+		/// <inheritdoc/>
+		protected override bool bIsNativeBool { get => false; }
+
+		/// <summary>
+		/// Construct a new boolean property
+		/// </summary>
+		/// <param name="PropertySettings">Property settings</param>
+		public UhtBool8Property(UhtPropertySettings PropertySettings) : base(PropertySettings)
+		{
+			this.PropertyCaps |= UhtPropertyCaps.CanExposeOnSpawn;
+		}
+	}
+
+	/// <summary>
+	/// uint16 bit type specialization
+	/// </summary>
+	public class UhtBool16Property : UhtBooleanProperty
+	{
+		/// <inheritdoc/>
+		protected override string CppTypeText { get => "uint16"; }
+
+		/// <inheritdoc/>
+		protected override string PGetMacroText { get => "UBOOL16"; }
+
+		/// <inheritdoc/>
+		protected override bool bIsNativeBool { get => false; }
+
+		/// <summary>
+		/// Construct a new boolean property
+		/// </summary>
+		/// <param name="PropertySettings">Property settings</param>
+		public UhtBool16Property(UhtPropertySettings PropertySettings) : base(PropertySettings)
+		{
+		}
+	}
+
+	/// <summary>
+	/// uint32 bit type specialization
+	/// </summary>
+	public class UhtBool32Property : UhtBooleanProperty
+	{
+		/// <inheritdoc/>
+		protected override string CppTypeText { get => "uint32"; }
+
+		/// <inheritdoc/>
+		protected override string PGetMacroText { get => "UBOOL32"; }
+
+		/// <inheritdoc/>
+		protected override bool bIsNativeBool { get => false; }
+
+		/// <summary>
+		/// Construct a new boolean property
+		/// </summary>
+		/// <param name="PropertySettings">Property settings</param>
+		public UhtBool32Property(UhtPropertySettings PropertySettings) : base(PropertySettings)
+		{
+		}
+	}
+
+	/// <summary>
+	/// uint64 bit type specialization
+	/// </summary>
+	public class UhtBool64Property : UhtBooleanProperty
+	{
+		/// <inheritdoc/>
+		protected override string CppTypeText { get => "uint64"; }
+
+		/// <inheritdoc/>
+		protected override string PGetMacroText { get => "UBOOL64"; }
+
+		/// <inheritdoc/>
+		protected override bool bIsNativeBool { get => false; }
+
+		/// <summary>
+		/// Construct a new boolean property
+		/// </summary>
+		/// <param name="PropertySettings">Property settings</param>
+		public UhtBool64Property(UhtPropertySettings PropertySettings) : base(PropertySettings)
+		{
+		}
+	}
+}
diff --git a/Engine/Source/Programs/Shared/EpicGames.UHT/Types/Properties/UhtByteProperty.cs b/Engine/Source/Programs/Shared/EpicGames.UHT/Types/Properties/UhtByteProperty.cs
new file mode 100644
--- /dev/null
+++ b/Engine/Source/Programs/Shared/EpicGames.UHT/Types/Properties/UhtByteProperty.cs
@@ -0,0 +1,169 @@
+// Copyright Epic Games, Inc. All Rights Reserved.
+
+using EpicGames.Core;
+using EpicGames.UHT.Tables;
+using EpicGames.UHT.Tokenizer;
+using EpicGames.UHT.Utils;
+using System.Collections.Generic;
+using System.Text;
+using System.Text.Json.Serialization;
+
+namespace EpicGames.UHT.Types
+{
+	[UnrealHeaderTool]
+	[UhtEngineClass(Name = "ByteProperty", IsProperty = true)]
+	public class UhtByteProperty : UhtNumericProperty
+	{
+		[JsonConverter(typeof(UhtNullableTypeSourceNameJsonConverter<UhtEnum>))]
+		public UhtEnum? Enum { get; set; }
+
+		/// <inheritdoc/>
+		public override string EngineClassName { get => "ByteProperty"; }
+
+		/// <inheritdoc/>
+		protected override string CppTypeText { get => "uint8"; }
+
+		/// <inheritdoc/>
+		protected override string PGetMacroText { get => this.Enum == null || this.Enum.CppForm != UhtEnumCppForm.EnumClass ? "PROPERTY" : "ENUM"; }
+
+		/// <inheritdoc/>
+		protected override UhtPGetArgumentType PGetTypeArgument { get => this.Enum == null || this.Enum.CppForm != UhtEnumCppForm.EnumClass ? UhtPGetArgumentType.EngineClass : UhtPGetArgumentType.TypeText; }
+
+		public UhtByteProperty(UhtPropertySettings PropertySettings, UhtPropertyIntType IntType, UhtEnum? Enum = null) : base(PropertySettings, IntType)
+		{
+			this.Enum = Enum;
+			this.PropertyCaps |= UhtPropertyCaps.CanExposeOnSpawn | UhtPropertyCaps.IsParameterSupportedByBlueprint | UhtPropertyCaps.IsMemberSupportedByBlueprint;
+		}
+
+		/// <inheritdoc/>
+		public override IEnumerable<UhtType> EnumerateReferencedTypes()
+		{
+			if (this.Enum != null)
+			{
+				yield return this.Enum;
+			}
+		}
+
+		/// <inheritdoc/>
+		public override void CollectReferencesInternal(IUhtReferenceCollector Collector, bool bTemplateProperty)
+		{
+			Collector.AddCrossModuleReference(this.Enum, true);
+		}
+
+		/// <inheritdoc/>
+		public override StringBuilder AppendText(StringBuilder Builder, UhtPropertyTextType TextType, bool bIsTemplateArgument)
+		{
+			if (this.Enum != null)
+			{
+				return UhtEnumProperty.AppendEnumText(Builder, this, this.Enum, TextType, bIsTemplateArgument);
+			}
+			else
+			{
+				return base.AppendText(Builder, TextType, bIsTemplateArgument);
+			}
+		}
+
+		/// <inheritdoc/>
+		public override StringBuilder AppendMemberDecl(StringBuilder Builder, IUhtPropertyMemberContext Context, string Name, string NameSuffix, int Tabs)
+		{
+			return AppendMemberDecl(Builder, Context, Name, NameSuffix, Tabs, "FBytePropertyParams");
+		}
+
+		/// <inheritdoc/>
+		public override StringBuilder AppendMemberDef(StringBuilder Builder, IUhtPropertyMemberContext Context, string Name, string NameSuffix, string? Offset, int Tabs)
+		{
+			AppendMemberDefStart(Builder, Context, Name, NameSuffix, Offset, Tabs, "FBytePropertyParams", "UECodeGen_Private::EPropertyGenFlags::Byte");
+			AppendMemberDefRef(Builder, Context, this.Enum, true, true);
+			AppendMemberDefEnd(Builder, Context, Name, NameSuffix);
+			return Builder;
+		}
+
+		/// <inheritdoc/>
+		public override void AppendObjectHashes(StringBuilder Builder, int StartingLength, IUhtPropertyMemberContext Context)
+		{
+			Builder.AppendObjectHash(StartingLength, this, Context, this.Enum);
+		}
+
+		/// <inheritdoc/>
+		public override StringBuilder AppendFunctionThunkParameterArg(StringBuilder Builder)
+		{
+			if (this.Enum != null)
+			{
+				return UhtEnumProperty.AppendEnumFunctionThunkParameterArg(Builder, this, this.Enum);
+			}
+			else
+			{
+				return base.AppendFunctionThunkParameterArg(Builder);
+			}
+		}
+
+		/// <inheritdoc/>
+		public override bool SanitizeDefaultValue(IUhtTokenReader DefaultValueReader, StringBuilder InnerDefaultValue)
+		{
+			if (this.Enum != null)
+			{
+				return UhtEnumProperty.SanitizeEnumDefaultValue(this, this.Enum, DefaultValueReader, InnerDefaultValue);
+			}
+
+			int Value = DefaultValueReader.GetConstIntExpression();
+			InnerDefaultValue.Append(Value.ToString());
+			return Value >= byte.MinValue && Value <= byte.MaxValue;
+		}
+
+		/// <inheritdoc/>
+		public override bool IsSameType(UhtProperty Other)
+		{
+			if (Other is UhtByteProperty OtherByte)
+			{
+				return this.Enum == OtherByte.Enum;
+			}
+			else if (Other is UhtEnumProperty OtherEnum)
+			{
+				return this.Enum == OtherEnum.Enum;
+			}
+			return false;
+		}
+
+		/// <inheritdoc/>
+		protected override void ValidateFunctionArgument(UhtFunction Function, UhtValidationOptions Options)
+		{
+			base.ValidateFunctionArgument(Function, Options);
+
+			if (Function.FunctionFlags.HasAnyFlags(EFunctionFlags.BlueprintEvent | EFunctionFlags.BlueprintCallable))
+			{
+				if (this.Enum != null && this.Enum.UnderlyingType != UhtEnumUnderlyingType.uint8 && this.Enum.UnderlyingType != UhtEnumUnderlyingType.Unspecified)
+				{
+					this.LogError("Invalid enum param for Blueprints - currently only uint8 supported");
+				}
+			}
+		}
+
+		/// <inheritdoc/>
+		public override string? GetRigVMType(ref UhtRigVMParameterFlags ParameterFlags)
+		{
+			if (this.Enum != null)
+			{
+				return UhtEnumProperty.GetEnumRigVMType(this, this.Enum, ref ParameterFlags);
+			}
+			else
+			{
+				return CppTypeText;
+			}
+		}
+
+		#region Keyword
+		[UhtPropertyType(Keyword = "uint8", Options = UhtPropertyTypeOptions.Simple | UhtPropertyTypeOptions.Immediate)]
+		public static UhtProperty? UInt8Property(UhtPropertyResolvePhase ResolvePhase, UhtPropertySettings PropertySettings, IUhtTokenReader TokenReader, UhtToken MatchedToken)
+		{
+			if (PropertySettings.bIsBitfield)
+			{
+				return new UhtBool8Property(PropertySettings);
+			}
+			else
+			{
+				return new UhtByteProperty(PropertySettings, UhtPropertyIntType.Sized);
+			}
+		}
+		#endregion
+	}
+}
diff --git a/Engine/Source/Programs/Shared/EpicGames.UHT/Types/Properties/UhtClassProperty.cs b/Engine/Source/Programs/Shared/EpicGames.UHT/Types/Properties/UhtClassProperty.cs
new file mode 100644
--- /dev/null
+++ b/Engine/Source/Programs/Shared/EpicGames.UHT/Types/Properties/UhtClassProperty.cs
@@ -0,0 +1,163 @@
+// Copyright Epic Games, Inc. All Rights Reserved.
+
+using EpicGames.Core;
+using EpicGames.UHT.Tables;
+using EpicGames.UHT.Tokenizer;
+using EpicGames.UHT.Utils;
+using System.Text;
+
+namespace EpicGames.UHT.Types
+{
+	[UnrealHeaderTool]
+	[UhtEngineClass(Name = "ClassProperty", IsProperty = true)]
+	public class UhtClassProperty : UhtObjectProperty
+	{
+		/// <inheritdoc/>
+		public override string EngineClassName { get => "ClassProperty"; }
+
+		/// <inheritdoc/>
+		protected override bool bPGetPassAsNoPtr { get => this.PropertyFlags.HasAnyFlags(EPropertyFlags.UObjectWrapper); }
+
+		/// <inheritdoc/>
+		protected override UhtPGetArgumentType PGetTypeArgument { get => UhtPGetArgumentType.TypeText; }
+
+		public UhtClassProperty(UhtPropertySettings PropertySettings, UhtClass PropertyClass, UhtClass MetaClass, EPropertyFlags ExtraFlags = EPropertyFlags.None)
+			: base(PropertySettings, PropertyClass, MetaClass)
+		{
+			this.PropertyFlags |= ExtraFlags;
+			this.PropertyCaps |= UhtPropertyCaps.CanHaveConfig;
+			this.PropertyCaps &= ~(UhtPropertyCaps.CanBeInstanced);
+		}
+
+		/// <inheritdoc/>
+		public override string? GetForwardDeclarations()
+		{
+			return this.MetaClass != null ? $"class {this.MetaClass.SourceName};" : null;
+		}
+
+		/// <inheritdoc/>
+		private StringBuilder AppendSubClassText(StringBuilder Builder)
+		{
+			Builder.Append("TSubclassOf<").Append(this.MetaClass?.SourceName).Append("> "); //COMPATIBILITY-TODO - Extra space in old UHT
+			return Builder;
+		}
+
+		/// <inheritdoc/>
+		private StringBuilder AppendText(StringBuilder Builder)
+		{
+			if (this.PropertyFlags.HasAnyFlags(EPropertyFlags.UObjectWrapper))
+			{
+				AppendSubClassText(Builder);
+			}
+			else
+			{
+				Builder.Append(this.Class.SourceName).Append("*");
+			}
+			return Builder;
+		}
+
+		/// <inheritdoc/>
+		public override StringBuilder AppendText(StringBuilder Builder, UhtPropertyTextType TextType, bool bIsTemplateArgument)
+		{
+			switch (TextType)
+			{
+				case UhtPropertyTextType.UserFacing:
+				case UhtPropertyTextType.Sparse:
+				case UhtPropertyTextType.SparseShort:
+				case UhtPropertyTextType.ClassFunction:
+				case UhtPropertyTextType.EventFunction:
+				case UhtPropertyTextType.InterfaceFunction:
+				case UhtPropertyTextType.ExportMember:
+				case UhtPropertyTextType.Construction:
+				case UhtPropertyTextType.FunctionThunkParameterArrayType:
+				case UhtPropertyTextType.RigVMTemplateArg:
+				case UhtPropertyTextType.GetterSetterArg:
+					AppendText(Builder);
+					break;
+
+				case UhtPropertyTextType.FunctionThunkReturn:
+					if (this.PropertyFlags.HasAnyFlags(EPropertyFlags.ConstParm))
+					{
+						Builder.Append("const ");
+					}
+					if (this.PropertyFlags.HasAnyFlags(EPropertyFlags.UObjectWrapper))
+					{
+						AppendText(Builder);
+					}
+					else
+					{
+						Builder.Append("UClass*");
+					}
+					break;
+
+				case UhtPropertyTextType.EventParameterMember:
+				case UhtPropertyTextType.EventParameterFunctionMember:
+					if (this.PropertyFlags.HasAnyFlags(EPropertyFlags.ReturnParm) && !this.PropertyFlags.HasAnyFlags(EPropertyFlags.UObjectWrapper))
+					{
+						Builder.Append("UClass*");
+					}
+					else
+					{
+						AppendText(Builder);
+					}
+					break;
+
+				case UhtPropertyTextType.FunctionThunkParameterArgType:
+					if (this.PropertyFlags.HasAllFlags(EPropertyFlags.OutParm | EPropertyFlags.UObjectWrapper))
+					{
+						AppendSubClassText(Builder);
+					}
+					else
+					{
+						Builder.Append(this.Class.SourceName);
+					}
+					break;
+			}
+			return Builder;
+		}
+
+		/// <inheritdoc/>
+		public override StringBuilder AppendMemberDecl(StringBuilder Builder, IUhtPropertyMemberContext Context, string Name, string NameSuffix, int Tabs)
+		{
+			return base.AppendMemberDecl(Builder, Context, Name, NameSuffix, Tabs, "FClassPropertyParams");
+		}
+
+		/// <inheritdoc/>
+		public override StringBuilder AppendMemberDef(StringBuilder Builder, IUhtPropertyMemberContext Context, string Name, string NameSuffix, string? Offset, int Tabs)
+		{
+			AppendMemberDefStart(Builder, Context, Name, NameSuffix, Offset, Tabs, "FClassPropertyParams", "UECodeGen_Private::EPropertyGenFlags::Class");
+			AppendMemberDefRef(Builder, Context, this.Class, false);
+			AppendMemberDefRef(Builder, Context, this.MetaClass, false);
+			AppendMemberDefEnd(Builder, Context, Name, NameSuffix);
+			return Builder;
+		}
+
+		/// <inheritdoc/>
+		public override string? GetRigVMType(ref UhtRigVMParameterFlags ParameterFlags)
+		{
+			if (this.PropertyFlags.HasAnyFlags(EPropertyFlags.UObjectWrapper))
+			{
+				return $"TSubclassOf<{this.MetaClass?.SourceName}> "; //COMPATIBILITY-TODO - Extra space in old UHT
+			}
+			else
+			{
+				return $"{this.Class}*";
+			}
+		}
+
+		#region Keyword
+		[UhtPropertyType(Keyword = "TSubclassOf")]
+		public static UhtProperty? SubclassOfProperty(UhtPropertyResolvePhase ResolvePhase, UhtPropertySettings PropertySettings, IUhtTokenReader TokenReader, UhtToken MatchedToken)
+		{
+			UhtClass? MetaClass = UhtObjectPropertyBase.ParseTemplateClass(PropertySettings, TokenReader, MatchedToken);
+			if (MetaClass == null)
+			{
+				return null;
+			}
+
+			// With TSubclassOf, MetaClass is used as a class limiter.  
+			return new UhtClassProperty(PropertySettings, MetaClass.Session.UClass, MetaClass, EPropertyFlags.UObjectWrapper);
+		}
+		#endregion
+	}
+}
diff --git a/Engine/Source/Programs/Shared/EpicGames.UHT/Types/Properties/UhtClassPtrProperties.cs b/Engine/Source/Programs/Shared/EpicGames.UHT/Types/Properties/UhtClassPtrProperties.cs
new file mode 100644
--- /dev/null
+++ b/Engine/Source/Programs/Shared/EpicGames.UHT/Types/Properties/UhtClassPtrProperties.cs
@@ -0,0 +1,78 @@
+// Copyright Epic Games, Inc. All Rights Reserved.
+
+using EpicGames.Core;
+using EpicGames.UHT.Tables;
+using EpicGames.UHT.Utils;
+using System.Text;
+
+namespace EpicGames.UHT.Types
+{
+	[UhtEngineClass(Name = "ClassPtrProperty", IsProperty = true)]
+	public class UhtClassPtrProperty : UhtClassProperty
+	{
+		/// <inheritdoc/>
+		public override string EngineClassName { get => "ClassPtrProperty"; }
+
+		/// <inheritdoc/>
+		protected override string CppTypeText { get => "ObjectPtr"; }
+
+		/// <inheritdoc/>
+		protected override string PGetMacroText { get => "OBJECTPTR"; }
+
+		/// <inheritdoc/>
+		protected override UhtPGetArgumentType PGetTypeArgument { get => UhtPGetArgumentType.TypeText; }
+
+		public UhtClassPtrProperty(UhtPropertySettings PropertySettings, UhtClass PropertyClass, UhtClass MetaClass, EPropertyFlags ExtraFlags = EPropertyFlags.None)
+			: base(PropertySettings, PropertyClass, MetaClass)
+		{
+			this.PropertyFlags |= ExtraFlags | EPropertyFlags.UObjectWrapper;
+		}
+
+		public override StringBuilder AppendText(StringBuilder Builder, UhtPropertyTextType TextType, bool bIsTemplateArgument)
+		{
+			switch (TextType)
+			{
+				default:
+					Builder.Append("TObjectPtr<").Append(this.MetaClass?.SourceName).Append(">");
+					break;
+			}
+			return Builder;
+		}
+
+		/// <inheritdoc/>
+		public override StringBuilder AppendMemberDecl(StringBuilder Builder, IUhtPropertyMemberContext Context, string Name, string NameSuffix, int Tabs)
+		{
+			return AppendMemberDecl(Builder, Context, Name, NameSuffix, Tabs, "FClassPtrPropertyParams");
+		}
+
+		/// <inheritdoc/>
+		public override StringBuilder AppendMemberDef(StringBuilder Builder, IUhtPropertyMemberContext Context, string Name, string NameSuffix, string? Offset, int Tabs)
+		{
+			AppendMemberDefStart(Builder, Context, Name, NameSuffix, Offset, Tabs, "FClassPtrPropertyParams", 
+				"UECodeGen_Private::EPropertyGenFlags::Class | UECodeGen_Private::EPropertyGenFlags::ObjectPtr");
+			AppendMemberDefRef(Builder, Context, this.Class, false);
+			AppendMemberDefRef(Builder, Context, this.MetaClass, false);
+			AppendMemberDefEnd(Builder, Context, Name, NameSuffix);
+			return Builder;
+		}
+
+		/// <inheritdoc/>
+		public override void Validate(UhtStruct OuterStruct, UhtProperty OutermostProperty, UhtValidationOptions Options)
+		{
+			base.Validate(OuterStruct, OutermostProperty, Options);
+
+			// UFunctions with a smart pointer as input parameter wont compile anyway, because of missing P_GET_... macro.
+			// UFunctions with a smart pointer as return type will crash when called via blueprint, because they are not supported in VM.
+			if (this.PropertyCategory != UhtPropertyCategory.Member)
+			{
+				OuterStruct.LogError("UFunctions cannot take a TObjectPtr as a parameter.");
+			}
+		}
+
+		/// <inheritdoc/>
+		public override string? GetRigVMType(ref UhtRigVMParameterFlags ParameterFlags)
+		{
+			return $"TObjectPtr<{this.MetaClass?.SourceName}>";
+		}
+	}
+}
diff --git a/Engine/Source/Programs/Shared/EpicGames.UHT/Types/Properties/UhtContainerProperty.cs b/Engine/Source/Programs/Shared/EpicGames.UHT/Types/Properties/UhtContainerProperty.cs
new file mode 100644
--- /dev/null
+++ b/Engine/Source/Programs/Shared/EpicGames.UHT/Types/Properties/UhtContainerProperty.cs
@@ -0,0 +1,55 @@
+// Copyright Epic Games, Inc. All Rights Reserved.
+
+using EpicGames.Core;
+
+namespace EpicGames.UHT.Types
+{
+	public abstract class UhtContainerBaseProperty : UhtProperty
+	{
+
+		/// <inheritdoc/>
+		public UhtProperty ValueProperty { get; set; }
+
+		public UhtContainerBaseProperty(UhtPropertySettings PropertySettings, UhtProperty Value) : base(PropertySettings)
+		{
+			this.ValueProperty = Value;
+			this.PropertyCaps = (this.PropertyCaps & ~(UhtPropertyCaps.CanBeInstanced | UhtPropertyCaps.CanHaveConfig)) | (this.ValueProperty.PropertyCaps & (UhtPropertyCaps.CanBeInstanced | UhtPropertyCaps.CanHaveConfig));
+		}
+
+		/// <inheritdoc/>
+		public override void Validate(UhtStruct OuterStruct, UhtProperty OutermostProperty, UhtValidationOptions Options)
+		{
+			base.Validate(OuterStruct, OutermostProperty, Options);
+			this.ValueProperty.Validate(OuterStruct, OutermostProperty, Options | UhtValidationOptions.IsValue);
+		}
+
+		protected static void PropagateFlagsFromInnerAndHandlePersistentInstanceMetadata(UhtProperty Container, UhtMetaData? MetaData, UhtProperty Inner)
+		{
+			// Copy some of the property flags to the container property.
+			if (Inner.PropertyFlags.HasAnyFlags(EPropertyFlags.ContainsInstancedReference | EPropertyFlags.InstancedReference))
+			{
+				Container.PropertyFlags |= EPropertyFlags.ContainsInstancedReference;
+				Container.PropertyFlags &= ~(EPropertyFlags.InstancedReference | EPropertyFlags.PersistentInstance); //this was propagated to the inner
+
+				if (MetaData != null && Inner.PropertyFlags.HasAnyFlags(EPropertyFlags.PersistentInstance))
+				{
+					Inner.MetaData.Add(MetaData);
+				}
+			}
+		}
+
+		protected static EPropertyFlags ResolveAndReturnNewFlags(UhtProperty Child, UhtResolvePhase Phase)
+		{
+			EPropertyFlags OldFlags = Child.PropertyFlags;
+			Child.Resolve(Phase);
+			EPropertyFlags NewFlags = Child.PropertyFlags;
+			return NewFlags & ~OldFlags;
+		}
+
+		/// <inheritdoc/>
+		public override bool ScanForInstancedReferenced(bool bDeepScan)
+		{
+			return this.ValueProperty.ScanForInstancedReferenced(bDeepScan);
+		}
+	}
+}
diff --git a/Engine/Source/Programs/Shared/EpicGames.UHT/Types/Properties/UhtDelegateProperty.cs b/Engine/Source/Programs/Shared/EpicGames.UHT/Types/Properties/UhtDelegateProperty.cs
new file mode 100644
--- /dev/null
+++ b/Engine/Source/Programs/Shared/EpicGames.UHT/Types/Properties/UhtDelegateProperty.cs
@@ -0,0 +1,157 @@
+// Copyright Epic Games, Inc. All Rights Reserved.
+
+using EpicGames.Core;
+using EpicGames.UHT.Tables;
+using EpicGames.UHT.Tokenizer;
+using EpicGames.UHT.Utils;
+using System.Text;
+using System.Text.Json.Serialization;
+
+namespace EpicGames.UHT.Types
+{
+	[UhtEngineClass(Name = "DelegateProperty", IsProperty = true)]
+	public class UhtDelegateProperty : UhtProperty
+	{
+		/// <inheritdoc/>
+		public override string EngineClassName { get => "DelegateProperty"; }
+
+		/// <inheritdoc/>
+		protected override string CppTypeText { get => "FScriptDelegate"; }
+
+		/// <inheritdoc/>
+		protected override string PGetMacroText { get => "PROPERTY"; }
+
+		/// <inheritdoc/>
+		protected override UhtPGetArgumentType PGetTypeArgument { get => UhtPGetArgumentType.EngineClass; }
+
+		[JsonConverter(typeof(UhtTypeSourceNameJsonConverter<UhtFunction>))]
+		public UhtFunction Function { get; set; }
+
+		public UhtDelegateProperty(UhtPropertySettings PropertySettings, UhtFunction Function) : base(PropertySettings)
+		{
+			this.Function = Function;
+			this.HeaderFile.AddReferencedHeader(Function);
+			this.PropertyCaps |= UhtPropertyCaps.PassCppArgsByRef | UhtPropertyCaps.IsParameterSupportedByBlueprint | UhtPropertyCaps.IsMemberSupportedByBlueprint;
+		}
+
+		/// <inheritdoc/>
+		protected override bool ResolveSelf(UhtResolvePhase Phase)
+		{
+			bool bResults = base.ResolveSelf(Phase);
+			switch (Phase)
+			{
+				case UhtResolvePhase.Final:
+					this.PropertyFlags |= EPropertyFlags.InstancedReference & ~this.DisallowPropertyFlags;
+					break;
+			}
+			return bResults;
+		}
+
+		/// <inheritdoc/>
+		public override bool ScanForInstancedReferenced(bool bDeepScan)
+		{
+			return true;
+		}
+
+		/// <inheritdoc/>
+		public override void CollectReferencesInternal(IUhtReferenceCollector Collector, bool bTemplateProperty)
+		{
+			Collector.AddCrossModuleReference(this.Function, true);
+		}
+
+		/// <inheritdoc/>
+		public override StringBuilder AppendText(StringBuilder Builder, UhtPropertyTextType TextType, bool bIsTemplateArgument)
+		{
+			switch (TextType)
+			{
+				case UhtPropertyTextType.ExportMember:
+				case UhtPropertyTextType.RigVMTemplateArg:
+				case UhtPropertyTextType.EventParameterFunctionMember:
+					Builder.Append(this.CppTypeText);
+					break;
+
+				default:
+					Builder.Append(this.Function.SourceName);
+					break;
+			}
+			return Builder;
+		}
+
+		/// <inheritdoc/>
+		public override StringBuilder AppendMemberDecl(StringBuilder Builder, IUhtPropertyMemberContext Context, string Name, string NameSuffix, int Tabs)
+		{
+			return AppendMemberDecl(Builder, Context, Name, NameSuffix, Tabs, "FDelegatePropertyParams");
+		}
+
+		/// <inheritdoc/>
+		public override StringBuilder AppendMemberDef(StringBuilder Builder, IUhtPropertyMemberContext Context, string Name, string NameSuffix, string? Offset, int Tabs)
+		{
+			AppendMemberDefStart(Builder, Context, Name, NameSuffix, Offset, Tabs, "FDelegatePropertyParams", "UECodeGen_Private::EPropertyGenFlags::Delegate");
+			AppendMemberDefRef(Builder, Context, this.Function, true);
+			AppendMemberDefEnd(Builder, Context, Name, NameSuffix);
+			return Builder;
+		}
+
+		/// <inheritdoc/>
+		public override StringBuilder AppendFunctionThunkParameterArg(StringBuilder Builder)
+		{
+			return Builder.Append(this.Function.SourceName).Append("(").AppendFunctionThunkParameterName(this).Append(")");
+		}
+
+		/// <inheritdoc/>
+		public override void AppendObjectHashes(StringBuilder Builder, int StartingLength, IUhtPropertyMemberContext Context)
+		{
+			Builder.AppendObjectHash(StartingLength, this, Context, this.Function);
+		}
+
+		/// <inheritdoc/>
+		public override StringBuilder AppendNullConstructorArg(StringBuilder Builder, bool bIsInitializer)
+		{
+			Builder.AppendPropertyText(this, UhtPropertyTextType.Construction).Append("()");
+			return Builder;
+		}
+
+		/// <inheritdoc/>
+		public override bool SanitizeDefaultValue(IUhtTokenReader DefaultValueReader, StringBuilder InnerDefaultValue)
+		{
+			return false;
+		}
+
+		/// <inheritdoc/>
+		public override bool IsSameType(UhtProperty Other)
+		{
+			if (Other is UhtDelegateProperty OtherObject)
+			{
+				return this.Function == OtherObject.Function;
+			}
+			return false;
+		}
+
+		/// <inheritdoc/>
+		protected override void ValidateFunctionArgument(UhtFunction Function, UhtValidationOptions Options)
+		{
+			base.ValidateFunctionArgument(Function, Options);
+
+			if (Function.FunctionFlags.HasAnyFlags(EFunctionFlags.Net))
+			{
+				if (Function.FunctionFlags.HasAnyFlags(EFunctionFlags.NetRequest))
+				{
+					if (!this.PropertyFlags.HasAnyFlags(EPropertyFlags.RepSkip))
+					{
+						this.LogError("Service request functions cannot contain delegate parameters, unless marked NotReplicated");
+					}
+				}
+				else
+				{
+					this.LogError("Replicated functions cannot contain delegate parameters (this would be insecure)");
+				}
+			}
+		}
+
+		/// <inheritdoc/>
+		public override string? GetRigVMType(ref UhtRigVMParameterFlags ParameterFlags)
+		{
+			return this.CppTypeText;
+		}
+	}
+}
diff --git a/Engine/Source/Programs/Shared/EpicGames.UHT/Types/Properties/UhtDoubleProperty.cs b/Engine/Source/Programs/Shared/EpicGames.UHT/Types/Properties/UhtDoubleProperty.cs
new file mode 100644
--- /dev/null
+++ b/Engine/Source/Programs/Shared/EpicGames.UHT/Types/Properties/UhtDoubleProperty.cs
@@ -0,0 +1,58 @@
+// Copyright Epic Games, Inc. All Rights Reserved.
+
+using EpicGames.UHT.Tables;
+using EpicGames.UHT.Tokenizer;
+using EpicGames.UHT.Utils;
+using System.Text;
+
+namespace EpicGames.UHT.Types
+{
+	[UnrealHeaderTool]
+	[UhtEngineClass(Name = "DoubleProperty", IsProperty = true)]
+	public class UhtDoubleProperty : UhtNumericProperty
+	{
+		/// <inheritdoc/>
+		public override string EngineClassName { get => "DoubleProperty"; }
+
+		/// <inheritdoc/>
+		protected override string CppTypeText { get => "double"; }
+
+		public UhtDoubleProperty(UhtPropertySettings PropertySettings) : base(PropertySettings, UhtPropertyIntType.None)
+		{
+			this.PropertyCaps |= UhtPropertyCaps.IsParameterSupportedByBlueprint | UhtPropertyCaps.IsMemberSupportedByBlueprint;
+		}
+
+		/// <inheritdoc/>
+		public override StringBuilder AppendMemberDecl(StringBuilder Builder, IUhtPropertyMemberContext Context, string Name, string NameSuffix, int Tabs)
+		{
+			return AppendMemberDecl(Builder, Context, Name, NameSuffix, Tabs, "FDoublePropertyParams");
+		}
+
+		/// <inheritdoc/>
+		public override StringBuilder AppendMemberDef(StringBuilder Builder, IUhtPropertyMemberContext Context, string Name, string NameSuffix, string? Offset, int Tabs)
+		{
+			AppendMemberDefStart(Builder, Context, Name, NameSuffix, Offset, Tabs, "FDoublePropertyParams", "UECodeGen_Private::EPropertyGenFlags::Double");
+			AppendMemberDefEnd(Builder, Context, Name, NameSuffix);
+			return Builder;
+		}
+
+		public override bool SanitizeDefaultValue(IUhtTokenReader DefaultValueReader, StringBuilder InnerDefaultValue)
+		{
+			InnerDefaultValue.AppendFormat("{0:F6}", DefaultValueReader.GetConstFloatExpression());
+			return true;
+		}
+
+		public override bool IsSameType(UhtProperty Other)
+		{
+			return Other is UhtDoubleProperty;
+		}
+
+		#region Keyword
+		[UhtPropertyType(Keyword = "double", Options = UhtPropertyTypeOptions.Simple | UhtPropertyTypeOptions.Immediate)]
+		public static UhtProperty? DoubleProperty(UhtPropertyResolvePhase ResolvePhase, UhtPropertySettings PropertySettings, IUhtTokenReader TokenReader, UhtToken MatchedToken)
+		{
+			return new UhtDoubleProperty(PropertySettings);
+		}
+		#endregion
+	}
+}
diff --git a/Engine/Source/Programs/Shared/EpicGames.UHT/Types/Properties/UhtEnumProperty.cs b/Engine/Source/Programs/Shared/EpicGames.UHT/Types/Properties/UhtEnumProperty.cs
new file mode 100644
--- /dev/null
+++ b/Engine/Source/Programs/Shared/EpicGames.UHT/Types/Properties/UhtEnumProperty.cs
@@ -0,0 +1,307 @@
+// Copyright Epic Games, Inc. All Rights Reserved.
+
+using EpicGames.Core;
+using EpicGames.UHT.Tables;
+using EpicGames.UHT.Tokenizer;
+using EpicGames.UHT.Utils;
+using System.Collections.Generic;
+using System.Text;
+using System.Text.Json.Serialization;
+
+namespace EpicGames.UHT.Types
+{
+	[UnrealHeaderTool]
+	[UhtEngineClass(Name = "EnumProperty", IsProperty = true)]
+	public class UhtEnumProperty : UhtProperty
+	{
+		/// <inheritdoc/>
+		public override string EngineClassName { get => this.Enum.CppForm != UhtEnumCppForm.EnumClass ? "ByteProperty" : "EnumProperty"; }
+
+		/// <inheritdoc/>
+		protected override string CppTypeText { get => "invalid"; }
+
+		/// <inheritdoc/>
+		protected override string PGetMacroText { get => this.Enum.CppForm != UhtEnumCppForm.EnumClass ? "PROPERTY" : "ENUM"; }
+
+		/// <inheritdoc/>
+		protected override UhtPGetArgumentType PGetTypeArgument { get => this.Enum.CppForm != UhtEnumCppForm.EnumClass ? UhtPGetArgumentType.EngineClass : UhtPGetArgumentType.TypeText; }
+
+		[JsonConverter(typeof(UhtTypeSourceNameJsonConverter<UhtEnum>))]
+		public UhtEnum Enum { get; set; }
+
+		public UhtProperty? UnderlyingProperty { get; set; }
+
+		public UhtEnumUnderlyingType UnderlyingType => this.Enum.UnderlyingType != UhtEnumUnderlyingType.Unspecified ? this.Enum.UnderlyingType : UhtEnumUnderlyingType.int32;
+		public UhtPropertyIntType UnderlyingTypeSize => this.Enum.UnderlyingType != UhtEnumUnderlyingType.Unspecified ? UhtPropertyIntType.Sized : UhtPropertyIntType.Unsized;
+
+		public UhtEnumProperty(UhtPropertySettings PropertySettings, UhtEnum Enum) : base(PropertySettings)
+		{
+			this.PropertyCaps |= UhtPropertyCaps.RequiresNullConstructorArg | UhtPropertyCaps.CanExposeOnSpawn | UhtPropertyCaps.IsParameterSupportedByBlueprint | UhtPropertyCaps.IsMemberSupportedByBlueprint;
+			this.Enum = Enum;
+			this.HeaderFile.AddReferencedHeader(Enum);
+			if (this.Enum.CppForm == UhtEnumCppForm.EnumClass)
+			{
+				this.UnderlyingProperty = CreateUnderlyingProperty();
+			}
+			else
+			{
+				this.UnderlyingProperty = null;
+			}
+		}
+
+		/// <inheritdoc/>
+		public override void CollectReferencesInternal(IUhtReferenceCollector Collector, bool bTemplateProperty)
+		{
+			Collector.AddCrossModuleReference(this.Enum, true);
+		}
+
+		/// <inheritdoc/>
+		public override string? GetForwardDeclarations()
+		{
+			if (this.Enum.CppForm != UhtEnumCppForm.EnumClass)
+			{
+				return null;
+			}
+			return $"enum class {this.Enum.SourceName} : {this.UnderlyingType};";
+		}
+
+		/// <inheritdoc/>
+		public override IEnumerable<UhtType> EnumerateReferencedTypes()
+		{
+			yield return this.Enum;
+		}
+
+		public static StringBuilder AppendEnumText(StringBuilder Builder, UhtProperty Property, UhtEnum Enum, UhtPropertyTextType TextType, bool bIsTemplateArgument)
+		{
+			switch (TextType)
+			{
+				case UhtPropertyTextType.ClassFunction:
+				case UhtPropertyTextType.InterfaceFunction:
+				case UhtPropertyTextType.EventFunction:
+					if (Enum.CppForm == UhtEnumCppForm.EnumClass || (!bIsTemplateArgument && (Property.PropertyFlags.HasAnyFlags(EPropertyFlags.ReturnParm) || !Property.PropertyFlags.HasAnyFlags(EPropertyFlags.OutParm))))
+					{
+						Builder.Append(Enum.CppType);
+					}
+					else
+					{
+						Builder.Append("TEnumAsByte<").Append(Enum.CppType).Append(">");
+					}
+					break;
+
+				case UhtPropertyTextType.FunctionThunkParameterArgType:
+					Builder.Append(Enum.SourceName);
+					break;
+
+				default:
+					if (Enum.CppForm == UhtEnumCppForm.EnumClass)
+					{
+						Builder.Append(Enum.CppType);
+					}
+					else
+					{
+						Builder.Append("TEnumAsByte<").Append(Enum.CppType).Append(">");
+					}
+					break;
+			}
+			return Builder;
+		}
+
+		/// <summary>
+		/// Append the text for a function thunk call argument
+		/// </summary>
+		/// <param name="Builder">Output builder</param>
+		/// 
+		/// <returns>Output builder</returns>
+		public static StringBuilder AppendEnumFunctionThunkParameterArg(StringBuilder Builder, UhtProperty Property, UhtEnum Enum)
+		{
+			if (!Property.PropertyFlags.HasAnyFlags(EPropertyFlags.OutParm))
+			{
+				Builder.Append(Enum.CppType).Append('(').AppendFunctionThunkParameterName(Property).Append(')');
+			}
+			else if (Enum.CppForm == UhtEnumCppForm.EnumClass)
+			{
+				Builder.Append('(').Append(Enum.CppType).Append("&)(").AppendFunctionThunkParameterName(Property).Append(')');
+			}
+			else
+			{
+				Builder.Append("(TEnumAsByte<").Append(Enum.CppType).Append(">&)(").AppendFunctionThunkParameterName(Property).Append(')');
+			}
+			return Builder;
+		}
+
+		public static string? GetEnumRigVMType(UhtProperty Property, UhtEnum Enum, ref UhtRigVMParameterFlags ParameterFlags)
+		{
+			ParameterFlags |= UhtRigVMParameterFlags.IsEnum;
+			return Enum.CppType.ToString();
+		}
+
+		public override StringBuilder AppendText(StringBuilder Builder, UhtPropertyTextType TextType, bool bIsTemplateArgument)
+		{
+			return AppendEnumText(Builder, this, this.Enum, TextType, bIsTemplateArgument);
+		}
+
+		/// <inheritdoc/>
+		public override StringBuilder AppendMemberDecl(StringBuilder Builder, IUhtPropertyMemberContext Context, string Name, string NameSuffix, int Tabs)
+		{
+			if (this.UnderlyingProperty != null)
+			{
+				Builder.AppendMemberDecl(this.UnderlyingProperty, Context, Name, GetNameSuffix(NameSuffix, "_Underlying"), Tabs);
+			}
+			return AppendMemberDecl(Builder, Context, Name, NameSuffix, Tabs, this.Enum.CppForm != UhtEnumCppForm.EnumClass ? "FBytePropertyParams" : "FEnumPropertyParams");
+		}
+
+		/// <inheritdoc/>
+		public override StringBuilder AppendMemberDef(StringBuilder Builder, IUhtPropertyMemberContext Context, string Name, string NameSuffix, string? Offset, int Tabs)
+		{
+			if (this.Enum.CppForm == UhtEnumCppForm.EnumClass)
+			{
+				if (this.UnderlyingProperty != null)
+				{
+					Builder.AppendMemberDef(this.UnderlyingProperty, Context, Name, GetNameSuffix(NameSuffix, "_Underlying"), "0", Tabs);
+				}
+				AppendMemberDefStart(Builder, Context, Name, NameSuffix, Offset, Tabs, "FEnumPropertyParams", "UECodeGen_Private::EPropertyGenFlags::Enum");
+				AppendMemberDefRef(Builder, Context, this.Enum, true);
+				AppendMemberDefEnd(Builder, Context, Name, NameSuffix);
+				return Builder;
+			}
+			else
+			{
+				AppendMemberDefStart(Builder, Context, Name, NameSuffix, Offset, Tabs, "FBytePropertyParams", "UECodeGen_Private::EPropertyGenFlags::Byte", true, true);
+				AppendMemberDefRef(Builder, Context, this.Enum, true);
+				AppendMemberDefEnd(Builder, Context, Name, NameSuffix);
+			}
+			return Builder;
+		}
+
+		/// <inheritdoc/>
+		public override StringBuilder AppendMemberPtr(StringBuilder Builder, IUhtPropertyMemberContext Context, string Name, string NameSuffix, int Tabs)
+		{
+			if (this.UnderlyingProperty != null)
+			{
+				Builder.AppendMemberPtr(this.UnderlyingProperty, Context, Name, GetNameSuffix(NameSuffix, "_Underlying"), Tabs);
+			}
+			base.AppendMemberPtr(Builder, Context, Name, NameSuffix, Tabs);
+			return Builder;
+		}
+
+		/// <inheritdoc/>
+		public override StringBuilder AppendFunctionThunkParameterArg(StringBuilder Builder)
+		{
+			return AppendEnumFunctionThunkParameterArg(Builder, this, this.Enum);
+		}
+
+		/// <inheritdoc/>
+		public override void AppendObjectHashes(StringBuilder Builder, int StartingLength, IUhtPropertyMemberContext Context)
+		{
+			Builder.AppendObjectHash(StartingLength, this, Context, this.Enum);
+		}
+
+		/// <inheritdoc/>
+		public override StringBuilder AppendNullConstructorArg(StringBuilder Builder, bool bIsInitializer)
+		{
+			if (this.Enum.CppForm != UhtEnumCppForm.EnumClass)
+			{
+				Builder.Append("0");
+			}
+			else
+			{
+				Builder.Append('(').Append(this.Enum.CppType).Append(")0");
+			}
+			return Builder;
+		}
+
+		/// <inheritdoc/>
+		public override bool SanitizeDefaultValue(IUhtTokenReader DefaultValueReader, StringBuilder InnerDefaultValue)
+		{
+			return SanitizeEnumDefaultValue(this, this.Enum, DefaultValueReader, InnerDefaultValue);
+		}
+
+		/// <inheritdoc/>
+		public override bool IsSameType(UhtProperty Other)
+		{
+			if (Other is UhtEnumProperty OtherEnum)
+			{
+				return this.Enum == OtherEnum.Enum;
+			}
+			else if (Other is UhtByteProperty OtherByte)
+			{
+				return this.Enum == OtherByte.Enum;
+			}
+			return false;
+		}
+
+		public static bool SanitizeEnumDefaultValue(UhtProperty Property, UhtEnum Enum, IUhtTokenReader DefaultValueReader, StringBuilder InnerDefaultValue)
+		{
+			UhtTokenList CppIdentifier = DefaultValueReader.GetCppIdentifier(UhtCppIdentifierOptions.None);
+			UhtTokenList StartingPoint = CppIdentifier.Next != null ? CppIdentifier.Next : CppIdentifier;
+			StartingPoint.Join(InnerDefaultValue, "::");
+			UhtTokenListCache.Return(CppIdentifier);
+
+			int EntryIndex = Enum.GetIndexByName(InnerDefaultValue.ToString());
+			if (EntryIndex == -1)
+			{
+				return false;
+			}
+			if (Enum.MetaData.ContainsKey(UhtNames.Hidden, EntryIndex))
+			{
+				Property.LogError($"Hidden enum entries cannot be used as default values: '{Property.SourceName}' '{InnerDefaultValue}'");
+			}
+			return true;
+		}
+
+		private UhtProperty CreateUnderlyingProperty()
+		{
+			UhtPropertySettings PropertySettings = new UhtPropertySettings(this, 0, this.PropertyCategory, 0);
+			PropertySettings.SourceName = "UnderlyingType";
+			switch (this.UnderlyingType)
+			{
+				case UhtEnumUnderlyingType.int8:
+					return new UhtInt8Property(PropertySettings, this.UnderlyingTypeSize);
+				case UhtEnumUnderlyingType.int16:
+					return new UhtInt16Property(PropertySettings, this.UnderlyingTypeSize);
+				case UhtEnumUnderlyingType.int32:
+					return new UhtIntProperty(PropertySettings, this.UnderlyingTypeSize);
+				case UhtEnumUnderlyingType.int64:
+					return new UhtInt64Property(PropertySettings, this.UnderlyingTypeSize);
+				case UhtEnumUnderlyingType.uint8:
+					return new UhtByteProperty(PropertySettings, this.UnderlyingTypeSize);
+				case UhtEnumUnderlyingType.uint16:
+					return new UhtUInt16Property(PropertySettings, this.UnderlyingTypeSize);
+				case UhtEnumUnderlyingType.uint32:
+					return new UhtUInt32Property(PropertySettings, this.UnderlyingTypeSize);
+				case UhtEnumUnderlyingType.uint64:
+					return new UhtUInt64Property(PropertySettings, this.UnderlyingTypeSize);
+				default:
+					throw new UhtIceException("Unexpected underlying enum type");
+			}
+		}
+
+		/// <inheritdoc/>
+		public override string? GetRigVMType(ref UhtRigVMParameterFlags ParameterFlags)
+		{
+			return GetEnumRigVMType(this, this.Enum, ref ParameterFlags);
+		}
+
+		#region Keyword
+		[UhtPropertyType(Keyword = "TEnumAsByte")]
+		public static UhtProperty? EnumProperty(UhtPropertyResolvePhase ResolvePhase, UhtPropertySettings PropertySettings, IUhtTokenReader TokenReader, UhtToken MatchedToken)
+		{
+			UhtType Outer = PropertySettings.Outer;
+			UhtEnum? Enum = null;
+			using (var TokenContext = new UhtMessageContext(TokenReader, "TEnumAsByte"))
+			{
+				TokenReader
+					.Require("TEnumAsByte")
+					.Require('<')
+					.Optional("enum")
+					.RequireCppIdentifier(UhtCppIdentifierOptions.None, (UhtTokenList TokenList) =>
+					{
+						Enum = Outer.FindType(UhtFindOptions.Enum | UhtFindOptions.SourceName, TokenList, TokenReader) as UhtEnum;
+					})
+					.Require('>');
+			}
+			return Enum != null ? new UhtByteProperty(PropertySettings, UhtPropertyIntType.None, Enum) : null;
+		}
+		#endregion
+	}
+}
diff --git a/Engine/Source/Programs/Shared/EpicGames.UHT/Types/Properties/UhtFieldPathProperty.cs b/Engine/Source/Programs/Shared/EpicGames.UHT/Types/Properties/UhtFieldPathProperty.cs
new file mode 100644
--- /dev/null
+++ b/Engine/Source/Programs/Shared/EpicGames.UHT/Types/Properties/UhtFieldPathProperty.cs
@@ -0,0 +1,127 @@
+// Copyright Epic Games, Inc. All Rights Reserved.
+
+using EpicGames.Core;
+using EpicGames.UHT.Tables;
+using EpicGames.UHT.Tokenizer;
+using EpicGames.UHT.Utils;
+using System.Text;
+
+namespace EpicGames.UHT.Types
+{
+	[UnrealHeaderTool]
+	[UhtEngineClass(Name = "FieldPathProperty", IsProperty = true)]
+	public class UhtFieldPathProperty : UhtProperty
+	{
+		/// <inheritdoc/>
+		public override string EngineClassName { get => "FieldPathProperty"; }
+
+		/// <inheritdoc/>
+		protected override string CppTypeText { get => "TFieldPath"; }
+
+		/// <inheritdoc/>
+		protected override string PGetMacroText { get => "STRUCT"; }
+
+		/// <inheritdoc/>
+		protected override UhtPGetArgumentType PGetTypeArgument { get => UhtPGetArgumentType.TypeText; }
+
+		/// <summary>
+		/// Class name without the prefix
+		/// </summary>
+		public string FieldClassName { get; set; }
+
+		public UhtFieldPathProperty(UhtPropertySettings PropertySettings, string FieldClassName) : base(PropertySettings)
+		{
+			this.FieldClassName = FieldClassName;
+			this.PropertyCaps |= UhtPropertyCaps.IsParameterSupportedByBlueprint | UhtPropertyCaps.IsMemberSupportedByBlueprint;
+		}
+
+		/// <inheritdoc/>
+		public override string? GetForwardDeclarations()
+		{
+			return $"class F{this.FieldClassName};";
+		}
+
+		/// <inheritdoc/>
+		public override StringBuilder AppendText(StringBuilder Builder, UhtPropertyTextType TextType, bool bIsTemplateArgument)
+		{
+			switch (TextType)
+			{
+				case UhtPropertyTextType.SparseShort:
+					Builder.Append("TFieldPath");
+					break;
+
+				default:
+					Builder.Append("TFieldPath<F").Append(this.FieldClassName).Append('>');
+					break;
+			}
+			return Builder;
+		}
+
+		/// <inheritdoc/>
+		public override StringBuilder AppendMemberDecl(StringBuilder Builder, IUhtPropertyMemberContext Context, string Name, string NameSuffix, int Tabs)
+		{
+			return AppendMemberDecl(Builder, Context, Name, NameSuffix, Tabs, "FFieldPathPropertyParams");
+		}
+
+		/// <inheritdoc/>
+		public override StringBuilder AppendMemberDef(StringBuilder Builder, IUhtPropertyMemberContext Context, string Name, string NameSuffix, string? Offset, int Tabs)
+		{
+			AppendMemberDefStart(Builder, Context, Name, NameSuffix, Offset, Tabs, "FFieldPathPropertyParams", "UECodeGen_Private::EPropertyGenFlags::FieldPath");
+			Builder.Append("&F").Append(this.FieldClassName).Append("::StaticClass, ");
+			AppendMemberDefEnd(Builder, Context, Name, NameSuffix);
+			return Builder;
+		}
+
+		/// <inheritdoc/>
+		public override StringBuilder AppendNullConstructorArg(StringBuilder Builder, bool bIsInitializer)
+		{
+			Builder.Append("nullptr");
+			return Builder;
+		}
+
+		/// <inheritdoc/>
+		public override bool SanitizeDefaultValue(IUhtTokenReader DefaultValueReader, StringBuilder InnerDefaultValue)
+		{
+			return false;
+		}
+
+		/// <inheritdoc/>
+		public override bool IsSameType(UhtProperty Other)
+		{
+			if (Other is UhtFieldPathProperty OtherFieldPath)
+			{
+				return this.FieldClassName == OtherFieldPath.FieldClassName;
+			}
+			return false;
+		}
+
+		/// <inheritdoc/>
+		public override string? GetRigVMType(ref UhtRigVMParameterFlags ParameterFlags)
+		{
+			return null;
+		}
+
+#region Keyword
+		[UhtPropertyType(Keyword = "TFieldPath", Options = UhtPropertyTypeOptions.Immediate)]
+		public static UhtProperty? FieldPathProperty(UhtPropertyResolvePhase ResolvePhase, UhtPropertySettings PropertySettings, IUhtTokenReader TokenReader, UhtToken MatchedToken)
+		{
+			using (var TokenContext = new UhtMessageContext(TokenReader, "TFieldPath"))
+			{
+				UhtToken Identifier = new UhtToken();
+				TokenReader
+					.Require("TFieldPath")
+					.Require('<')
+					.RequireIdentifier((ref UhtToken Token) => Identifier = Token)
+					.Require('>');
+
+				StringView FieldClassName = new StringView(Identifier.Value, 1);
+				if (!UhtEngineClassTable.Instance.IsValidPropertyTypeName(FieldClassName))
+				{
+					throw new UhtException($"Undefined property type: {Identifier.Value}");
+				}
+				return new UhtFieldPathProperty(PropertySettings, FieldClassName.ToString());
+			}
+		}
+#endregion
+	}
+}
diff --git a/Engine/Source/Programs/Shared/EpicGames.UHT/Types/Properties/UhtFloatProperty.cs b/Engine/Source/Programs/Shared/EpicGames.UHT/Types/Properties/UhtFloatProperty.cs
new file mode 100644
--- /dev/null
+++ b/Engine/Source/Programs/Shared/EpicGames.UHT/Types/Properties/UhtFloatProperty.cs
@@ -0,0 +1,60 @@
+// Copyright Epic Games, Inc. All Rights Reserved.
+
+using EpicGames.UHT.Tables;
+using EpicGames.UHT.Tokenizer;
+using EpicGames.UHT.Utils;
+using System.Text;
+
+namespace EpicGames.UHT.Types
+{
+	[UnrealHeaderTool]
+	[UhtEngineClass(Name = "FloatProperty", IsProperty = true)]
+	public class UhtFloatProperty : UhtNumericProperty
+	{
+		/// <inheritdoc/>
+		public override string EngineClassName { get => "FloatProperty"; }
+
+		/// <inheritdoc/>
+		protected override string CppTypeText { get => "float"; }
+
+		public UhtFloatProperty(UhtPropertySettings PropertySettings) : base(PropertySettings, UhtPropertyIntType.None)
+		{
+			this.PropertyCaps |= UhtPropertyCaps.CanExposeOnSpawn | UhtPropertyCaps.IsParameterSupportedByBlueprint | UhtPropertyCaps.IsMemberSupportedByBlueprint;
+		}
+
+		/// <inheritdoc/>
+		public override StringBuilder AppendMemberDecl(StringBuilder Builder, IUhtPropertyMemberContext Context, string Name, string NameSuffix, int Tabs)
+		{
+			return AppendMemberDecl(Builder, Context, Name, NameSuffix, Tabs, "FFloatPropertyParams");
+		}
+
+		/// <inheritdoc/>
+		public override StringBuilder AppendMemberDef(StringBuilder Builder, IUhtPropertyMemberContext Context, string Name, string NameSuffix, string? Offset, int Tabs)
+		{
+			AppendMemberDefStart(Builder, Context, Name, NameSuffix, Offset, Tabs, "FFloatPropertyParams", "UECodeGen_Private::EPropertyGenFlags::Float");
+			AppendMemberDefEnd(Builder, Context, Name, NameSuffix);
+			return Builder;
+		}
+
+		/// <inheritdoc/>
+		public override bool SanitizeDefaultValue(IUhtTokenReader DefaultValueReader, StringBuilder InnerDefaultValue)
+		{
+			InnerDefaultValue.AppendFormat("{0:F6}", DefaultValueReader.GetConstFloatExpression());
+			return true;
+		}
+
+		/// <inheritdoc/>
+		public override bool IsSameType(UhtProperty Other)
+		{
+			return Other is UhtFloatProperty;
+		}
+
+		#region Keywords
+		[UhtPropertyType(Keyword = "float", Options = UhtPropertyTypeOptions.Simple | UhtPropertyTypeOptions.Immediate)]
+		public static UhtProperty? FloatProperty(UhtPropertyResolvePhase ResolvePhase, UhtPropertySettings PropertySettings, IUhtTokenReader TokenReader, UhtToken MatchedToken)
+		{
+			return new UhtFloatProperty(PropertySettings);
+		}
+		#endregion
+	}
+}
