UHT in UBT initial version.  Currently disabled.

diff --git a/Engine/Source/Programs/Shared/EpicGames.UHT/Types/Properties/UhtInt16Property.cs b/Engine/Source/Programs/Shared/EpicGames.UHT/Types/Properties/UhtInt16Property.cs
new file mode 100644
--- /dev/null
+++ b/Engine/Source/Programs/Shared/EpicGames.UHT/Types/Properties/UhtInt16Property.cs
@@ -0,0 +1,56 @@
+// Copyright Epic Games, Inc. All Rights Reserved.
+
+using EpicGames.UHT.Tables;
+using EpicGames.UHT.Tokenizer;
+using EpicGames.UHT.Utils;
+using System.Text;
+
+namespace EpicGames.UHT.Types
+{
+	[UnrealHeaderTool]
+	[UhtEngineClass(Name = "Int16Property", IsProperty = true)]
+	public class UhtInt16Property : UhtNumericProperty
+	{
+		/// <inheritdoc/>
+		public override string EngineClassName { get => "Int16Property"; }
+
+		/// <inheritdoc/>
+		protected override string CppTypeText { get => "int16"; }
+
+		public UhtInt16Property(UhtPropertySettings PropertySettings, UhtPropertyIntType IntType) : base(PropertySettings, IntType)
+		{
+		}
+
+		/// <inheritdoc/>
+		public override StringBuilder AppendMemberDecl(StringBuilder Builder, IUhtPropertyMemberContext Context, string Name, string NameSuffix, int Tabs)
+		{
+			return AppendMemberDecl(Builder, Context, Name, NameSuffix, Tabs, "FInt16PropertyParams");
+		}
+
+		/// <inheritdoc/>
+		public override StringBuilder AppendMemberDef(StringBuilder Builder, IUhtPropertyMemberContext Context, string Name, string NameSuffix, string? Offset, int Tabs)
+		{
+			AppendMemberDefStart(Builder, Context, Name, NameSuffix, Offset, Tabs, "FInt16PropertyParams", "UECodeGen_Private::EPropertyGenFlags::Int16");
+			AppendMemberDefEnd(Builder, Context, Name, NameSuffix);
+			return Builder;
+		}
+
+		public override bool SanitizeDefaultValue(IUhtTokenReader DefaultValueReader, StringBuilder InnerDefaultValue)
+		{
+			return false;
+		}
+
+		public override bool IsSameType(UhtProperty Other)
+		{
+			return Other is UhtInt16Property;
+		}
+
+		#region Keyword
+		[UhtPropertyType(Keyword = "int16", Options = UhtPropertyTypeOptions.Simple | UhtPropertyTypeOptions.Immediate)]
+		public static UhtProperty? Int16Property(UhtPropertyResolvePhase ResolvePhase, UhtPropertySettings PropertySettings, IUhtTokenReader TokenReader, UhtToken MatchedToken)
+		{
+			return new UhtInt16Property(PropertySettings, UhtPropertyIntType.Sized);
+		}
+		#endregion
+	}
+}
diff --git a/Engine/Source/Programs/Shared/EpicGames.UHT/Types/Properties/UhtInt64Property.cs b/Engine/Source/Programs/Shared/EpicGames.UHT/Types/Properties/UhtInt64Property.cs
new file mode 100644
--- /dev/null
+++ b/Engine/Source/Programs/Shared/EpicGames.UHT/Types/Properties/UhtInt64Property.cs
@@ -0,0 +1,60 @@
+// Copyright Epic Games, Inc. All Rights Reserved.
+
+using EpicGames.UHT.Tables;
+using EpicGames.UHT.Tokenizer;
+using EpicGames.UHT.Utils;
+using System.Text;
+
+namespace EpicGames.UHT.Types
+{
+	[UnrealHeaderTool]
+	[UhtEngineClass(Name = "Int64Property", IsProperty = true)]
+	public class UhtInt64Property : UhtNumericProperty
+	{
+		/// <inheritdoc/>
+		public override string EngineClassName { get => "Int64Property"; }
+
+		/// <inheritdoc/>
+		protected override string CppTypeText { get => "int64"; }
+
+		public UhtInt64Property(UhtPropertySettings PropertySettings, UhtPropertyIntType IntType) : base(PropertySettings, IntType)
+		{
+			this.PropertyCaps |= UhtPropertyCaps.CanExposeOnSpawn | UhtPropertyCaps.IsParameterSupportedByBlueprint | UhtPropertyCaps.IsMemberSupportedByBlueprint;
+		}
+
+		/// <inheritdoc/>
+		public override StringBuilder AppendMemberDecl(StringBuilder Builder, IUhtPropertyMemberContext Context, string Name, string NameSuffix, int Tabs)
+		{
+			return AppendMemberDecl(Builder, Context, Name, NameSuffix, Tabs, "FInt64PropertyParams");
+		}
+
+		/// <inheritdoc/>
+		public override StringBuilder AppendMemberDef(StringBuilder Builder, IUhtPropertyMemberContext Context, string Name, string NameSuffix, string? Offset, int Tabs)
+		{
+			AppendMemberDefStart(Builder, Context, Name, NameSuffix, Offset, Tabs, "FInt64PropertyParams", "UECodeGen_Private::EPropertyGenFlags::Int64");
+			AppendMemberDefEnd(Builder, Context, Name, NameSuffix);
+			return Builder;
+		}
+
+		/// <inheritdoc/>
+		public override bool SanitizeDefaultValue(IUhtTokenReader DefaultValueReader, StringBuilder InnerDefaultValue)
+		{
+			InnerDefaultValue.Append(DefaultValueReader.GetConstLongExpression().ToString());
+			return true;
+		}
+
+		/// <inheritdoc/>
+		public override bool IsSameType(UhtProperty Other)
+		{
+			return Other is UhtInt64Property;
+		}
+
+		#region Keyword
+		[UhtPropertyType(Keyword = "int64", Options = UhtPropertyTypeOptions.Simple | UhtPropertyTypeOptions.Immediate)]
+		public static UhtProperty? Int64Property(UhtPropertyResolvePhase ResolvePhase, UhtPropertySettings PropertySettings, IUhtTokenReader TokenReader, UhtToken MatchedToken)
+		{
+			return new UhtInt64Property(PropertySettings, UhtPropertyIntType.Sized);
+		}
+		#endregion
+	}
+}
diff --git a/Engine/Source/Programs/Shared/EpicGames.UHT/Types/Properties/UhtInt8Property.cs b/Engine/Source/Programs/Shared/EpicGames.UHT/Types/Properties/UhtInt8Property.cs
new file mode 100644
--- /dev/null
+++ b/Engine/Source/Programs/Shared/EpicGames.UHT/Types/Properties/UhtInt8Property.cs
@@ -0,0 +1,56 @@
+// Copyright Epic Games, Inc. All Rights Reserved.
+
+using EpicGames.UHT.Tables;
+using EpicGames.UHT.Tokenizer;
+using EpicGames.UHT.Utils;
+using System.Text;
+
+namespace EpicGames.UHT.Types
+{
+	[UnrealHeaderTool]
+	[UhtEngineClass(Name = "Int8Property", IsProperty = true)]
+	public class UhtInt8Property : UhtNumericProperty
+	{
+		/// <inheritdoc/>
+		public override string EngineClassName { get => "Int8Property"; }
+
+		/// <inheritdoc/>
+		protected override string CppTypeText { get => "int8"; }
+
+		public UhtInt8Property(UhtPropertySettings PropertySettings, UhtPropertyIntType IntType) : base(PropertySettings, IntType)
+		{
+		}
+
+		/// <inheritdoc/>
+		public override StringBuilder AppendMemberDecl(StringBuilder Builder, IUhtPropertyMemberContext Context, string Name, string NameSuffix, int Tabs)
+		{
+			return AppendMemberDecl(Builder, Context, Name, NameSuffix, Tabs, "FInt8PropertyParams");
+		}
+
+		/// <inheritdoc/>
+		public override StringBuilder AppendMemberDef(StringBuilder Builder, IUhtPropertyMemberContext Context, string Name, string NameSuffix, string? Offset, int Tabs)
+		{
+			AppendMemberDefStart(Builder, Context, Name, NameSuffix, Offset, Tabs, "FInt8PropertyParams", "UECodeGen_Private::EPropertyGenFlags::Int8");
+			AppendMemberDefEnd(Builder, Context, Name, NameSuffix);
+			return Builder;
+		}
+
+		public override bool SanitizeDefaultValue(IUhtTokenReader DefaultValueReader, StringBuilder InnerDefaultValue)
+		{
+			return false;
+		}
+
+		public override bool IsSameType(UhtProperty Other)
+		{
+			return Other is UhtInt8Property;
+		}
+
+		#region Keyword
+		[UhtPropertyType(Keyword = "int8", Options = UhtPropertyTypeOptions.Simple | UhtPropertyTypeOptions.Immediate)]
+		public static UhtProperty? Int8Property(UhtPropertyResolvePhase ResolvePhase, UhtPropertySettings PropertySettings, IUhtTokenReader TokenReader, UhtToken MatchedToken)
+		{
+			return new UhtInt8Property(PropertySettings, UhtPropertyIntType.Sized);
+		}
+		#endregion
+	}
+}
diff --git a/Engine/Source/Programs/Shared/EpicGames.UHT/Types/Properties/UhtIntProperty.cs b/Engine/Source/Programs/Shared/EpicGames.UHT/Types/Properties/UhtIntProperty.cs
new file mode 100644
--- /dev/null
+++ b/Engine/Source/Programs/Shared/EpicGames.UHT/Types/Properties/UhtIntProperty.cs
@@ -0,0 +1,78 @@
+// Copyright Epic Games, Inc. All Rights Reserved.
+
+using EpicGames.UHT.Tables;
+using EpicGames.UHT.Tokenizer;
+using EpicGames.UHT.Utils;
+using System.Text;
+
+namespace EpicGames.UHT.Types
+{
+	[UnrealHeaderTool]
+	[UhtEngineClass(Name = "IntProperty", IsProperty = true)]
+	public class UhtIntProperty : UhtNumericProperty
+	{
+		/// <inheritdoc/>
+		public override string EngineClassName { get => "IntProperty"; }
+
+		/// <inheritdoc/>
+		protected override string CppTypeText { get => "int32"; }
+
+		public UhtIntProperty(UhtPropertySettings PropertySettings, UhtPropertyIntType IntType) : base(PropertySettings, IntType)
+		{
+			this.PropertyCaps |= UhtPropertyCaps.CanExposeOnSpawn | UhtPropertyCaps.IsParameterSupportedByBlueprint | UhtPropertyCaps.IsMemberSupportedByBlueprint;
+		}
+
+		/// <inheritdoc/>
+		public override StringBuilder AppendMemberDecl(StringBuilder Builder, IUhtPropertyMemberContext Context, string Name, string NameSuffix, int Tabs)
+		{
+			return AppendMemberDecl(Builder, Context, Name, NameSuffix, Tabs, 
+				this.IntType == UhtPropertyIntType.Unsized ? "FUnsizedIntPropertyParams" : "FIntPropertyParams");
+		}
+
+		/// <inheritdoc/>
+		public override StringBuilder AppendMemberDef(StringBuilder Builder, IUhtPropertyMemberContext Context, string Name, string NameSuffix, string? Offset, int Tabs)
+		{
+			AppendMemberDefStart(Builder, Context, Name, NameSuffix, Offset, Tabs,
+				this.IntType == UhtPropertyIntType.Unsized ? "FUnsizedIntPropertyParams" : "FIntPropertyParams",
+				"UECodeGen_Private::EPropertyGenFlags::Int");
+			AppendMemberDefEnd(Builder, Context, Name, NameSuffix);
+			return Builder;
+		}
+
+		/// <inheritdoc/>
+		public override bool SanitizeDefaultValue(IUhtTokenReader DefaultValueReader, StringBuilder InnerDefaultValue)
+		{
+			InnerDefaultValue.Append(DefaultValueReader.GetConstIntExpression().ToString());
+			return true;
+		}
+
+		/// <inheritdoc/>
+		public override bool IsSameType(UhtProperty Other)
+		{
+			return Other is UhtIntProperty;
+		}
+
+		#region Keywords
+		[UhtPropertyType(Keyword = "int32", Options = UhtPropertyTypeOptions.Simple | UhtPropertyTypeOptions.Immediate)]
+		public static UhtProperty? Int32Property(UhtPropertyResolvePhase ResolvePhase, UhtPropertySettings PropertySettings, IUhtTokenReader TokenReader, UhtToken MatchedToken)
+		{
+			return new UhtIntProperty(PropertySettings, UhtPropertyIntType.Sized);
+		}
+
+		[UhtPropertyType(Keyword = "int", Options = UhtPropertyTypeOptions.Simple | UhtPropertyTypeOptions.Immediate)]
+		public static UhtProperty? IntProperty(UhtPropertyResolvePhase ResolvePhase, UhtPropertySettings PropertySettings, IUhtTokenReader TokenReader, UhtToken MatchedToken)
+		{
+			return new UhtIntProperty(PropertySettings, UhtPropertyIntType.Unsized);
+		}
+
+		[UhtPropertyType(Keyword = "signed", Options = UhtPropertyTypeOptions.Immediate)]
+		public static UhtProperty? SignedProperty(UhtPropertyResolvePhase ResolvePhase, UhtPropertySettings PropertySettings, IUhtTokenReader TokenReader, UhtToken MatchedToken)
+		{
+			TokenReader
+				.Require("signed")
+				.Optional("int");
+			return new UhtIntProperty(PropertySettings, UhtPropertyIntType.Unsized);
+		}
+		#endregion
+	}
+}
diff --git a/Engine/Source/Programs/Shared/EpicGames.UHT/Types/Properties/UhtInterfaceProperty.cs b/Engine/Source/Programs/Shared/EpicGames.UHT/Types/Properties/UhtInterfaceProperty.cs
new file mode 100644
--- /dev/null
+++ b/Engine/Source/Programs/Shared/EpicGames.UHT/Types/Properties/UhtInterfaceProperty.cs
@@ -0,0 +1,193 @@
+// Copyright Epic Games, Inc. All Rights Reserved.
+
+using EpicGames.Core;
+using EpicGames.UHT.Tables;
+using EpicGames.UHT.Tokenizer;
+using EpicGames.UHT.Utils;
+using System.Collections.Generic;
+using System.Diagnostics.CodeAnalysis;
+using System.Text;
+using System.Text.Json.Serialization;
+
+namespace EpicGames.UHT.Types
+{
+	[UnrealHeaderTool]
+	[UhtEngineClass(Name = "InterfaceProperty", IsProperty = true)]
+	public class UhtInterfaceProperty : UhtProperty
+	{
+		/// <inheritdoc/>
+		public override string EngineClassName { get => "InterfaceProperty"; }
+
+		/// <inheritdoc/>
+		protected override string CppTypeText { get => "TScriptInterface"; }
+
+		/// <inheritdoc/>
+		protected override string PGetMacroText { get => "TINTERFACE"; }
+
+		/// <inheritdoc/>
+		protected override UhtPGetArgumentType PGetTypeArgument { get => UhtPGetArgumentType.TypeText; }
+
+		[JsonConverter(typeof(UhtTypeSourceNameJsonConverter<UhtClass>))]
+		public UhtClass InterfaceClass { get; set; }
+
+		public UhtInterfaceProperty(UhtPropertySettings PropertySettings, UhtClass InterfaceClass) : base(PropertySettings)
+		{
+			this.InterfaceClass = InterfaceClass;
+			this.PropertyFlags |= EPropertyFlags.UObjectWrapper;
+			this.PropertyCaps |= UhtPropertyCaps.CanExposeOnSpawn | UhtPropertyCaps.IsParameterSupportedByBlueprint | UhtPropertyCaps.IsMemberSupportedByBlueprint | UhtPropertyCaps.PassCppArgsByRef;
+			this.PropertyCaps &= ~(UhtPropertyCaps.CanHaveConfig | UhtPropertyCaps.CanBeContainerKey);
+		}
+
+		/// <inheritdoc/>
+		protected override bool ResolveSelf(UhtResolvePhase Phase)
+		{
+			bool bResults = base.ResolveSelf(Phase);
+			switch (Phase)
+			{
+				case UhtResolvePhase.Final:
+					if (this.InterfaceClass.HierarchyHasAnyClassFlags(EClassFlags.DefaultToInstanced))
+					{
+						this.PropertyFlags |= (EPropertyFlags.InstancedReference | EPropertyFlags.ExportObject) & ~this.DisallowPropertyFlags;
+					}
+					break;
+			}
+			return bResults;
+		}
+
+		/// <inheritdoc/>
+		public override bool ScanForInstancedReferenced(bool bDeepScan)
+		{
+			return !this.DisallowPropertyFlags.HasAnyFlags(EPropertyFlags.InstancedReference) && this.InterfaceClass.HierarchyHasAnyClassFlags(EClassFlags.DefaultToInstanced);
+		}
+
+		/// <inheritdoc/>
+		public override void CollectReferencesInternal(IUhtReferenceCollector Collector, bool bTemplateProperty)
+		{
+			Collector.AddCrossModuleReference(this.InterfaceClass, false);
+		}
+
+		/// <inheritdoc/>
+		public override string? GetForwardDeclarations()
+		{
+			UhtClass? ExportClass = this.InterfaceClass;
+			while (ExportClass != null && !ExportClass.ClassFlags.HasAnyFlags(EClassFlags.Native))
+			{
+				ExportClass = ExportClass.SuperClass;
+			}
+			return ExportClass != null ? $"class {ExportClass.SourceName};" : null;
+		}
+
+		/// <inheritdoc/>
+		public override IEnumerable<UhtType> EnumerateReferencedTypes()
+		{
+			yield return this.InterfaceClass;
+		}
+
+		/// <inheritdoc/>
+		public override StringBuilder AppendText(StringBuilder Builder, UhtPropertyTextType TextType, bool bIsTemplateArgument)
+		{
+			switch (TextType)
+			{
+				case UhtPropertyTextType.SparseShort:
+					Builder.Append("TScriptInterface");
+					break;
+
+				case UhtPropertyTextType.FunctionThunkParameterArgType:
+					Builder.Append(this.InterfaceClass.SourceName);
+					break;
+
+				default:
+					Builder.Append("TScriptInterface<").Append(this.InterfaceClass.SourceName).Append('>');
+					break;
+			}
+			return Builder;
+		}
+
+		/// <inheritdoc/>
+		public override StringBuilder AppendMemberDecl(StringBuilder Builder, IUhtPropertyMemberContext Context, string Name, string NameSuffix, int Tabs)
+		{
+			return AppendMemberDecl(Builder, Context, Name, NameSuffix, Tabs, "FInterfacePropertyParams");
+		}
+
+		/// <inheritdoc/>
+		public override StringBuilder AppendMemberDef(StringBuilder Builder, IUhtPropertyMemberContext Context, string Name, string NameSuffix, string? Offset, int Tabs)
+		{
+			// FScriptInterface<USomeInterface> is valid so in that case we need to pass in the interface class and not the alternate object (which in the end is the same object)
+			AppendMemberDefStart(Builder, Context, Name, NameSuffix, Offset, Tabs, "FInterfacePropertyParams", "UECodeGen_Private::EPropertyGenFlags::Interface");
+			AppendMemberDefRef(Builder, Context, this.InterfaceClass.AlternateObject != null ? this.InterfaceClass.AlternateObject : this.InterfaceClass, false);
+			AppendMemberDefEnd(Builder, Context, Name, NameSuffix);
+			return Builder;
+		}
+
+		/// <inheritdoc/>
+		public override StringBuilder AppendNullConstructorArg(StringBuilder Builder, bool bIsInitializer)
+		{
+			Builder.Append("NULL");
+			return Builder;
+		}
+
+		/// <inheritdoc/>
+		public override bool SanitizeDefaultValue(IUhtTokenReader DefaultValueReader, StringBuilder InnerDefaultValue)
+		{
+			return false;
+		}
+
+		/// <inheritdoc/>
+		protected override void ValidateMember(UhtStruct Struct, UhtValidationOptions Options)
+		{
+			base.ValidateMember(Struct, Options);
+
+			if (this.PointerType == UhtPointerType.Native)
+			{
+				this.LogError($"UPROPERTY pointers cannot be interfaces - did you mean TScriptInterface<{this.InterfaceClass.SourceName}>?");
+			}
+		}
+
+		/// <inheritdoc/>
+		public override bool IsSameType(UhtProperty Other)
+		{
+			if (Other is UhtInterfaceProperty OtherObject)
+			{
+				return this.InterfaceClass == OtherObject.InterfaceClass;
+			}
+			return false;
+		}
+
+		/// <inheritdoc/>
+		public override bool MustBeConstArgument([NotNullWhen(true)] out UhtType? ErrorType)
+		{
+			ErrorType = this.InterfaceClass;
+			return this.InterfaceClass.ClassFlags.HasAnyFlags(EClassFlags.Const);
+		}
+
+		/// <inheritdoc/>
+		public override string? GetRigVMType(ref UhtRigVMParameterFlags ParameterFlags)
+		{
+			return null;
+		}
+
+		#region Keywords
+		[UhtPropertyType(Keyword = "FScriptInterface", Options = UhtPropertyTypeOptions.Simple)] // This can't be immediate due to the reference to UInterface
+		public static UhtProperty? FScriptInterfaceProperty(UhtPropertyResolvePhase ResolvePhase, UhtPropertySettings PropertySettings, IUhtTokenReader TokenReader, UhtToken MatchedToken)
+		{
+			return new UhtInterfaceProperty(PropertySettings, PropertySettings.Outer.Session.IInterface);
+		}
+
+		[UhtPropertyType(Keyword = "TScriptInterface")]
+		public static UhtProperty? TScriptInterfaceProperty(UhtPropertyResolvePhase ResolvePhase, UhtPropertySettings PropertySettings, IUhtTokenReader TokenReader, UhtToken MatchedToken)
+		{
+			UhtClass? PropertyClass = UhtObjectPropertyBase.ParseTemplateObject(PropertySettings, TokenReader, MatchedToken, false);
+			if (PropertyClass == null)
+			{
+				return null;
+			}
+
+			if (PropertyClass.ClassFlags.HasAnyFlags(EClassFlags.Interface))
+			{
+				return new UhtInterfaceProperty(PropertySettings, PropertyClass);
+			}
+			return new UhtObjectProperty(PropertySettings, PropertyClass, EPropertyFlags.UObjectWrapper);
+		}
+		#endregion
+	}
+}
diff --git a/Engine/Source/Programs/Shared/EpicGames.UHT/Types/Properties/UhtLargeWorldCoordinatesRealProperty.cs b/Engine/Source/Programs/Shared/EpicGames.UHT/Types/Properties/UhtLargeWorldCoordinatesRealProperty.cs
new file mode 100644
--- /dev/null
+++ b/Engine/Source/Programs/Shared/EpicGames.UHT/Types/Properties/UhtLargeWorldCoordinatesRealProperty.cs
@@ -0,0 +1,66 @@
+// Copyright Epic Games, Inc. All Rights Reserved.
+
+using EpicGames.UHT.Tables;
+using EpicGames.UHT.Tokenizer;
+using EpicGames.UHT.Utils;
+using System.Text;
+
+namespace EpicGames.UHT.Types
+{
+	[UnrealHeaderTool]
+	[UhtEngineClass(Name = "LargeWorldCoordinatesRealProperty", IsProperty = true)]
+	public class UhtLargeWorldCoordinatesRealProperty : UhtNumericProperty
+	{
+		/// <inheritdoc/>
+		public override string EngineClassName { get => "LargeWorldCoordinatesRealProperty"; }
+
+		/// <inheritdoc/>
+		protected override string CppTypeText { get => "double"; }
+
+		public UhtLargeWorldCoordinatesRealProperty(UhtPropertySettings PropertySettings) : base(PropertySettings, UhtPropertyIntType.None)
+		{
+			this.PropertyCaps |= UhtPropertyCaps.IsParameterSupportedByBlueprint | UhtPropertyCaps.IsMemberSupportedByBlueprint;
+		}
+
+		/// <inheritdoc/>
+		public override StringBuilder AppendMemberDecl(StringBuilder Builder, IUhtPropertyMemberContext Context, string Name, string NameSuffix, int Tabs)
+		{
+			return AppendMemberDecl(Builder, Context, Name, NameSuffix, Tabs, "FLargeWorldCoordinatesRealPropertyParams");
+		}
+
+		/// <inheritdoc/>
+		public override StringBuilder AppendMemberDef(StringBuilder Builder, IUhtPropertyMemberContext Context, string Name, string NameSuffix, string? Offset, int Tabs)
+		{
+			AppendMemberDefStart(Builder, Context, Name, NameSuffix, Offset, Tabs, 
+				"FLargeWorldCoordinatesRealPropertyParams", 
+				"UECodeGen_Private::EPropertyGenFlags::LargeWorldCoordinatesReal");
+			AppendMemberDefEnd(Builder, Context, Name, NameSuffix);
+			return Builder;
+		}
+
+		/// <inheritdoc/>
+		public override bool SanitizeDefaultValue(IUhtTokenReader DefaultValueReader, StringBuilder InnerDefaultValue)
+		{
+			InnerDefaultValue.AppendFormat("{0:F6}", DefaultValueReader.GetConstFloatExpression());
+			return true;
+		}
+
+		/// <inheritdoc/>
+		public override bool IsSameType(UhtProperty Other)
+		{
+			return Other is UhtLargeWorldCoordinatesRealProperty;
+		}
+
+		#region Keyword
+		[UhtPropertyType(Keyword = "FLargeWorldCoordinatesReal", Options = UhtPropertyTypeOptions.Simple | UhtPropertyTypeOptions.Immediate)]
+		public static UhtProperty? LargeWorldCoordinatesRealProperty(UhtPropertyResolvePhase ResolvePhase, UhtPropertySettings PropertySettings, IUhtTokenReader TokenReader, UhtToken MatchedToken)
+		{
+			if (!PropertySettings.Outer.HeaderFile.bIsNoExportTypes)
+			{
+				TokenReader.LogError("FLargeWorldCoordinatesReal is intended for LWC support only and should not be used outside of NoExportTypes.h");
+			}
+			return new UhtLargeWorldCoordinatesRealProperty(PropertySettings);
+		}
+		#endregion
+	}
+}
diff --git a/Engine/Source/Programs/Shared/EpicGames.UHT/Types/Properties/UhtLazyObjectPtrProperty.cs b/Engine/Source/Programs/Shared/EpicGames.UHT/Types/Properties/UhtLazyObjectPtrProperty.cs
new file mode 100644
--- /dev/null
+++ b/Engine/Source/Programs/Shared/EpicGames.UHT/Types/Properties/UhtLazyObjectPtrProperty.cs
@@ -0,0 +1,123 @@
+// Copyright Epic Games, Inc. All Rights Reserved.
+
+using EpicGames.Core;
+using EpicGames.UHT.Tables;
+using EpicGames.UHT.Tokenizer;
+using EpicGames.UHT.Utils;
+using System.Text;
+
+namespace EpicGames.UHT.Types
+{
+	[UnrealHeaderTool]
+	[UhtEngineClass(Name = "LazyObjectProperty", IsProperty = true)]
+	public class UhtLazyObjectPtrProperty : UhtObjectPropertyBase
+	{
+		/// <inheritdoc/>
+		public override string EngineClassName { get => "LazyObjectProperty"; }
+
+		/// <inheritdoc/>
+		protected override string CppTypeText { get => "TLazyObjectPtr"; }
+
+		/// <inheritdoc/>
+		protected override string PGetMacroText { get => "LAZYOBJECT"; }
+
+		/// <inheritdoc/>
+		protected override bool bPGetPassAsNoPtr { get => true; }
+
+		/// <inheritdoc/>
+		protected override UhtPGetArgumentType PGetTypeArgument { get => UhtPGetArgumentType.TypeText; }
+
+		public UhtLazyObjectPtrProperty(UhtPropertySettings PropertySettings, UhtClass PropertyClass)
+			: base(PropertySettings, PropertyClass, null)
+		{
+			this.PropertyFlags |= EPropertyFlags.UObjectWrapper;
+			this.PropertyCaps |= UhtPropertyCaps.PassCppArgsByRef | UhtPropertyCaps.RequiresNullConstructorArg;
+		}
+
+		/// <inheritdoc/>
+		public override StringBuilder AppendText(StringBuilder Builder, UhtPropertyTextType TextType, bool bIsTemplateArgument)
+		{
+			switch (TextType)
+			{
+				case UhtPropertyTextType.FunctionThunkParameterArgType:
+					Builder.Append("TLazyObjectPtr<").Append(this.Class.SourceName).Append("> "); //COMPATIBILITY-TODO - Extra space
+					break;
+
+				default:
+					Builder.Append("TLazyObjectPtr<").Append(this.Class.SourceName).Append(">");
+					break;
+			}
+			return Builder;
+		}
+
+		/// <inheritdoc/>
+		public override StringBuilder AppendMemberDecl(StringBuilder Builder, IUhtPropertyMemberContext Context, string Name, string NameSuffix, int Tabs)
+		{
+			return AppendMemberDecl(Builder, Context, Name, NameSuffix, Tabs, "FLazyObjectPropertyParams");
+		}
+
+		/// <inheritdoc/>
+		public override StringBuilder AppendMemberDef(StringBuilder Builder, IUhtPropertyMemberContext Context, string Name, string NameSuffix, string? Offset, int Tabs)
+		{
+			AppendMemberDefStart(Builder, Context, Name, NameSuffix, Offset, Tabs, "FLazyObjectPropertyParams", "UECodeGen_Private::EPropertyGenFlags::LazyObject");
+			AppendMemberDefRef(Builder, Context, this.Class, false);
+			AppendMemberDefEnd(Builder, Context, Name, NameSuffix);
+			return Builder;
+		}
+
+		/// <inheritdoc/>
+		public override StringBuilder AppendNullConstructorArg(StringBuilder Builder, bool bIsInitializer)
+		{
+			Builder.Append("NULL");
+			return Builder;
+		}
+
+		/// <inheritdoc/>
+		public override bool IsSameType(UhtProperty Other)
+		{
+			if (Other is UhtLazyObjectPtrProperty OtherObject)
+			{
+				return this.Class == OtherObject.Class && this.MetaClass == OtherObject.MetaClass;
+			}
+			return false;
+		}
+
+		/// <inheritdoc/>
+		public override void Validate(UhtStruct OuterStruct, UhtProperty OutermostProperty, UhtValidationOptions Options)
+		{
+			base.Validate(OuterStruct, OutermostProperty, Options);
+
+			// UFunctions with a smart pointer as input parameter wont compile anyway, because of missing P_GET_... macro.
+			// UFunctions with a smart pointer as return type will crash when called via blueprint, because they are not supported in VM.
+			if (this.PropertyCategory != UhtPropertyCategory.Member)
+			{
+				OuterStruct.LogError("UFunctions cannot take a lazy pointer as a parameter.");
+			}
+		}
+
+		/// <inheritdoc/>
+		public override string? GetRigVMType(ref UhtRigVMParameterFlags ParameterFlags)
+		{
+			return $"TLazyObjectPtr<{this.Class.SourceName}>";
+		}
+
+		#region Keyword
+		[UhtPropertyType(Keyword = "TLazyObjectPtr")]
+		public static UhtProperty? LazyObjectPtrProperty(UhtPropertyResolvePhase ResolvePhase, UhtPropertySettings PropertySettings, IUhtTokenReader TokenReader, UhtToken MatchedToken)
+		{
+			UhtClass? PropertyClass = ParseTemplateObject(PropertySettings, TokenReader, MatchedToken, true);
+			if (PropertyClass == null)
+			{
+				return null;
+			}
+
+			if (PropertyClass.IsChildOf(PropertyClass.Session.UClass))
+			{
+				TokenReader.LogError("Class variables cannot be lazy, they are always strong.");
+			}
+
+			return new UhtLazyObjectPtrProperty(PropertySettings, PropertyClass);
+		}
+		#endregion
+	}
+}
diff --git a/Engine/Source/Programs/Shared/EpicGames.UHT/Types/Properties/UhtMapProperty.cs b/Engine/Source/Programs/Shared/EpicGames.UHT/Types/Properties/UhtMapProperty.cs
new file mode 100644
--- /dev/null
+++ b/Engine/Source/Programs/Shared/EpicGames.UHT/Types/Properties/UhtMapProperty.cs
@@ -0,0 +1,301 @@
+// Copyright Epic Games, Inc. All Rights Reserved.
+
+using EpicGames.Core;
+using EpicGames.UHT.Parsers;
+using EpicGames.UHT.Tables;
+using EpicGames.UHT.Tokenizer;
+using EpicGames.UHT.Utils;
+using System.Collections.Generic;
+using System.Text;
+
+namespace EpicGames.UHT.Types
+{
+	[UnrealHeaderTool]
+	[UhtEngineClass(Name = "MapProperty", IsProperty = true)]
+	public class UhtMapProperty : UhtContainerBaseProperty
+	{
+		/// <inheritdoc/>
+		public override string EngineClassName { get => "MapProperty"; }
+
+		/// <inheritdoc/>
+		protected override string CppTypeText { get => "TMap"; }
+
+		/// <inheritdoc/>
+		protected override string PGetMacroText { get => "TMAP"; }
+
+		/// <inheritdoc/>
+		protected override UhtPGetArgumentType PGetTypeArgument { get => UhtPGetArgumentType.TypeText; }
+
+		public UhtProperty KeyProperty { get; set; }
+
+		public UhtMapProperty(UhtPropertySettings PropertySettings, UhtProperty Key, UhtProperty Value) : base(PropertySettings, Value)
+		{
+			// If the creation of the value property set more flags, then copy those flags to ourselves
+			this.PropertyFlags |= Value.PropertyFlags & EPropertyFlags.UObjectWrapper;
+
+			// Make sure the 'UObjectWrapper' flag is maintained so that both 'TMap<TSubclassOf<...>, ...>' and 'TMap<UClass*, TSubclassOf<...>>' works correctly
+			Key.PropertyFlags = (Value.PropertyFlags & ~EPropertyFlags.UObjectWrapper) | (Key.PropertyFlags & EPropertyFlags.UObjectWrapper);
+			Key.DisallowPropertyFlags = ~(EPropertyFlags.ContainsInstancedReference | EPropertyFlags.InstancedReference | EPropertyFlags.UObjectWrapper);
+
+			this.KeyProperty = Key;
+			this.PropertyCaps |= UhtPropertyCaps.PassCppArgsByRef;
+			this.PropertyCaps &= ~(UhtPropertyCaps.CanBeContainerValue | UhtPropertyCaps.CanBeContainerKey);
+			this.PropertyCaps = (this.PropertyCaps & ~(UhtPropertyCaps.IsParameterSupportedByBlueprint | UhtPropertyCaps.IsMemberSupportedByBlueprint | UhtPropertyCaps.CanExposeOnSpawn));
+			this.PropertyCaps |= (this.ValueProperty.PropertyCaps & (UhtPropertyCaps.IsParameterSupportedByBlueprint | UhtPropertyCaps.IsMemberSupportedByBlueprint | UhtPropertyCaps.CanExposeOnSpawn)) &
+				(this.KeyProperty.PropertyCaps & (UhtPropertyCaps.IsParameterSupportedByBlueprint | UhtPropertyCaps.IsMemberSupportedByBlueprint));
+			this.PropertyFlags = Value.PropertyFlags;
+			this.ValueProperty.SourceName = this.SourceName;
+			this.ValueProperty.EngineName = this.EngineName;
+			this.ValueProperty.PropertyFlags &= EPropertyFlags.PropagateToMapValue;
+			this.ValueProperty.Outer = this;
+			this.ValueProperty.MetaData.Clear();
+			this.KeyProperty.SourceName = $"{this.SourceName}_Key";
+			this.KeyProperty.EngineName = $"{this.EngineName}_Key";
+			this.KeyProperty.PropertyFlags &= EPropertyFlags.PropagateToMapKey;
+			this.KeyProperty.Outer = this;
+			this.KeyProperty.MetaData.Clear();
+
+			// With old UHT, Deprecated was applied after property create.
+			// With the new, it is applied prior to creation.  Deprecated in old 
+			// was not on the key.
+			this.KeyProperty.PropertyFlags &= ~EPropertyFlags.Deprecated;
+		}
+
+		/// <inheritdoc/>
+		protected override bool ResolveSelf(UhtResolvePhase Phase)
+		{
+			bool bResults = base.ResolveSelf(Phase);
+			switch (Phase)
+			{
+				case UhtResolvePhase.Final:
+					this.KeyProperty.Resolve(Phase);
+					this.KeyProperty.MetaData.Clear();
+
+					EPropertyFlags NewFlags = ResolveAndReturnNewFlags(this.ValueProperty, Phase);
+					this.PropertyFlags |= NewFlags;
+					this.KeyProperty.PropertyFlags |= NewFlags;
+					this.MetaData.Add(this.ValueProperty.MetaData);
+					this.ValueProperty.MetaData.Clear();
+
+					PropagateFlagsFromInnerAndHandlePersistentInstanceMetadata(this, null, this.KeyProperty);
+					PropagateFlagsFromInnerAndHandlePersistentInstanceMetadata(this, this.MetaData, this.ValueProperty);
+					break;
+			}
+			return bResults;
+		}
+
+		/// <inheritdoc/>
+		public override void CollectReferencesInternal(IUhtReferenceCollector Collector, bool bTemplateProperty)
+		{
+			this.ValueProperty.CollectReferencesInternal(Collector, true);
+			this.KeyProperty.CollectReferencesInternal(Collector, true);
+		}
+
+		/// <inheritdoc/>
+		public override string? GetForwardDeclarations()
+		{
+			return $"{this.KeyProperty.GetForwardDeclarations()} {this.ValueProperty.GetForwardDeclarations()}";
+		}
+
+		/// <inheritdoc/>
+		public override IEnumerable<UhtType> EnumerateReferencedTypes()
+		{
+			foreach (UhtType Type in this.ValueProperty.EnumerateReferencedTypes())
+			{
+				yield return Type;
+			}
+			foreach (UhtType Type in this.KeyProperty.EnumerateReferencedTypes())
+			{
+				yield return Type;
+			}
+		}
+
+		/// <inheritdoc/>
+		public override StringBuilder AppendText(StringBuilder Builder, UhtPropertyTextType TextType, bool bIsTemplateArgument)
+		{
+			switch (TextType)
+			{
+				case UhtPropertyTextType.SparseShort:
+					Builder.Append("TMap");
+					break;
+
+				case UhtPropertyTextType.FunctionThunkParameterArgType:
+					Builder.AppendFunctionThunkParameterArrayType(this.KeyProperty).Append(",").AppendFunctionThunkParameterArrayType(this.ValueProperty);
+					break;
+
+				default:
+					Builder.Append("TMap<").AppendPropertyText(this.KeyProperty, TextType, true).Append(',').AppendPropertyText(this.ValueProperty, TextType, true).Append('>');
+					break;
+			}
+			return Builder;
+		}
+
+		/// <inheritdoc/>
+		public override StringBuilder AppendMemberDecl(StringBuilder Builder, IUhtPropertyMemberContext Context, string Name, string NameSuffix, int Tabs)
+		{
+			Builder.AppendMemberDecl(this.ValueProperty, Context, Name, GetNameSuffix(NameSuffix, "_ValueProp"), Tabs);
+			Builder.AppendMemberDecl(this.KeyProperty, Context, Name, GetNameSuffix(NameSuffix, "_Key_KeyProp"), Tabs);
+			return AppendMemberDecl(Builder, Context, Name, NameSuffix, Tabs, "FMapPropertyParams");
+		}
+
+		/// <inheritdoc/>
+		public override StringBuilder AppendMemberDef(StringBuilder Builder, IUhtPropertyMemberContext Context, string Name, string NameSuffix, string? Offset, int Tabs)
+		{
+			Builder.AppendMemberDef(this.ValueProperty, Context, Name, GetNameSuffix(NameSuffix, "_ValueProp"), "1", Tabs);
+			Builder.AppendMemberDef(this.KeyProperty, Context, Name, GetNameSuffix(NameSuffix, "_Key_KeyProp"), "0", Tabs);
+			AppendMemberDefStart(Builder, Context, Name, NameSuffix, Offset, Tabs, "FMapPropertyParams", "UECodeGen_Private::EPropertyGenFlags::Map");
+			Builder.Append(this.Allocator == UhtPropertyAllocator.MemoryImage ? "EMapPropertyFlags::UsesMemoryImageAllocator" : "EMapPropertyFlags::None").Append(", ");
+			AppendMemberDefEnd(Builder, Context, Name, NameSuffix);
+			return Builder;
+		}
+
+		/// <inheritdoc/>
+		public override StringBuilder AppendMemberPtr(StringBuilder Builder, IUhtPropertyMemberContext Context, string Name, string NameSuffix, int Tabs)
+		{
+			Builder.AppendMemberPtr(this.ValueProperty, Context, Name, GetNameSuffix(NameSuffix, "_ValueProp"), Tabs);
+			Builder.AppendMemberPtr(this.KeyProperty, Context, Name, GetNameSuffix(NameSuffix, "_Key_KeyProp"), Tabs);
+			base.AppendMemberPtr(Builder, Context, Name, NameSuffix, Tabs);
+			return Builder;
+		}
+
+		/// <inheritdoc/>
+		public override void AppendObjectHashes(StringBuilder Builder, int StartingLength, IUhtPropertyMemberContext Context)
+		{
+			this.KeyProperty.AppendObjectHashes(Builder, StartingLength, Context);
+			this.ValueProperty.AppendObjectHashes(Builder, StartingLength, Context);
+		}
+
+		/// <inheritdoc/>
+		public override StringBuilder AppendNullConstructorArg(StringBuilder Builder, bool bIsInitializer)
+		{
+			Builder.AppendPropertyText(this, UhtPropertyTextType.Construction).Append("()");
+			return Builder;
+		}
+
+		/// <inheritdoc/>
+		public override bool SanitizeDefaultValue(IUhtTokenReader DefaultValueReader, StringBuilder InnerDefaultValue)
+		{
+			return false;
+		}
+
+		/// <inheritdoc/>
+		public override bool IsSameType(UhtProperty Other)
+		{
+			if (Other is UhtMapProperty OtherMap)
+			{
+				return this.ValueProperty.IsSameType(OtherMap.ValueProperty) &&
+					this.KeyProperty.IsSameType(OtherMap.KeyProperty);
+			}
+			return false;
+		}
+
+		/// <inheritdoc/>
+		public override void ValidateDeprecated()
+		{
+			base.ValidateDeprecated();
+			this.KeyProperty.ValidateDeprecated();
+		}
+
+		/// <inheritdoc/>
+		public override void Validate(UhtStruct OuterStruct, UhtProperty OutermostProperty, UhtValidationOptions Options)
+		{
+			base.Validate(OuterStruct, OutermostProperty, Options);
+			this.KeyProperty.Validate(OuterStruct, OutermostProperty, Options | UhtValidationOptions.IsKey);
+		}
+
+		///<inheritdoc/>
+		public override bool ValidateStructPropertyOkForNet(UhtProperty ReferencingProperty)
+		{
+			if (!this.PropertyFlags.HasAnyFlags(EPropertyFlags.RepSkip))
+			{
+				ReferencingProperty.LogError($"Maps are not supported for Replication or RPCs.  Map '{this.SourceName}' in '{this.Outer?.SourceName}'.  Origin '{ReferencingProperty.SourceName}'");
+				return false;
+			}
+			return true;
+		}
+
+		/// <inheritdoc/>
+		protected override void ValidateFunctionArgument(UhtFunction Function, UhtValidationOptions Options)
+		{
+			base.ValidateFunctionArgument(Function, Options);
+
+			if (Function.FunctionFlags.HasAnyFlags(EFunctionFlags.Net))
+			{
+				if (!Function.FunctionFlags.HasAnyFlags(EFunctionFlags.NetRequest | EFunctionFlags.NetResponse))
+				{
+					this.LogError("Maps are not supported in an RPC.");
+				}
+			}
+		}
+
+		/// <inheritdoc/>
+		public override string? GetRigVMType(ref UhtRigVMParameterFlags ParameterFlags)
+		{
+			return null;
+		}
+
+#region Keyword
+		[UhtPropertyType(Keyword = "TMap")]
+		public static UhtProperty? MapProperty(UhtPropertyResolvePhase ResolvePhase, UhtPropertySettings PropertySettings, IUhtTokenReader TokenReader, UhtToken MatchedToken)
+		{
+			using (var TokenContext = new UhtMessageContext(TokenReader, "TMap"))
+			{
+				if (!TokenReader.SkipExpectedType(MatchedToken.Value, PropertySettings.PropertyCategory == UhtPropertyCategory.Member))
+				{
+					return null;
+				}
+				TokenReader.Require('<');
+
+				// Parse the key type
+				UhtProperty? Key = UhtPropertyParser.ParseTemplateParam(ResolvePhase, PropertySettings, "Key", TokenReader);
+				if (Key == null)
+				{
+					return null;
+				}
+
+				if (!Key.PropertyCaps.HasAnyFlags(UhtPropertyCaps.CanBeContainerKey))
+				{
+					TokenReader.LogError($"The type \'{Key.GetUserFacingDecl()}\' can not be used as a key in a TMap");
+				}
+
+				if (PropertySettings.PropertyFlags.HasAnyFlags(EPropertyFlags.Net))
+				{
+					TokenReader.LogError("Replicated maps are not supported.");
+				}
+
+				TokenReader.Require(',');
+
+				// Parse the value type
+				UhtProperty? Value = UhtPropertyParser.ParseTemplateParam(ResolvePhase, PropertySettings, "Value", TokenReader);
+				if (Value == null)
+				{
+					return null;
+				}
+
+				if (!Value.PropertyCaps.HasAnyFlags(UhtPropertyCaps.CanBeContainerValue))
+				{
+					TokenReader.LogError($"The type \'{Value.GetUserFacingDecl()}\' can not be used as a value in a TMap");
+				}
+
+				if (TokenReader.TryOptional(','))
+				{
+					UhtToken AllocatorToken = TokenReader.GetIdentifier();
+					if (AllocatorToken.IsIdentifier("FMemoryImageSetAllocator"))
+					{
+						PropertySettings.Allocator = UhtPropertyAllocator.MemoryImage;
+					}
+					else
+					{
+						TokenReader.LogError($"Found '{AllocatorToken.Value}' - explicit allocators are not supported in TMap properties.");
+					}
+				}
+				TokenReader.Require('>');
+
+				//@TODO: Prevent sparse delegate types from being used in a container
+
+				return new UhtMapProperty(PropertySettings, Key, Value);
+			}
+		}
+#endregion
+	}
+}
diff --git a/Engine/Source/Programs/Shared/EpicGames.UHT/Types/Properties/UhtMulticastDelegateProperty.cs b/Engine/Source/Programs/Shared/EpicGames.UHT/Types/Properties/UhtMulticastDelegateProperty.cs
new file mode 100644
--- /dev/null
+++ b/Engine/Source/Programs/Shared/EpicGames.UHT/Types/Properties/UhtMulticastDelegateProperty.cs
@@ -0,0 +1,124 @@
+// Copyright Epic Games, Inc. All Rights Reserved.
+
+using EpicGames.Core;
+using EpicGames.UHT.Tables;
+using EpicGames.UHT.Tokenizer;
+using EpicGames.UHT.Utils;
+using System.Text;
+using System.Text.Json.Serialization;
+
+namespace EpicGames.UHT.Types
+{
+	[UhtEngineClass(Name = "MulticastDelegateProperty", IsProperty = true)]
+	public abstract class UhtMulticastDelegateProperty : UhtProperty
+	{
+		/// <inheritdoc/>
+		public override string EngineClassName { get => "MulticastDelegateProperty"; }
+
+		/// <inheritdoc/>
+		protected override string CppTypeText { get => "FMulticastScriptDelegate"; }
+
+		/// <inheritdoc/>
+		protected override string PGetMacroText { get => "PROPERTY"; }
+
+		/// <inheritdoc/>
+		protected override UhtPGetArgumentType PGetTypeArgument { get => UhtPGetArgumentType.TypeText; }
+
+		[JsonConverter(typeof(UhtTypeSourceNameJsonConverter<UhtFunction>))]
+		public UhtFunction Function { get; set; }
+
+		public UhtMulticastDelegateProperty(UhtPropertySettings PropertySettings, UhtFunction Function) : base(PropertySettings)
+		{
+			this.Function = Function;
+			this.HeaderFile.AddReferencedHeader(Function);
+			this.PropertyCaps |= UhtPropertyCaps.PassCppArgsByRef | UhtPropertyCaps.CanBeBlueprintAssignable | UhtPropertyCaps.CanBeBlueprintCallable |
+				UhtPropertyCaps.CanBeBlueprintAuthorityOnly | UhtPropertyCaps.IsMemberSupportedByBlueprint;
+			this.PropertyCaps &= ~(UhtPropertyCaps.IsParameterSupportedByBlueprint);
+		}
+
+		/// <inheritdoc/>
+		protected override bool ResolveSelf(UhtResolvePhase Phase)
+		{
+			bool bResults = base.ResolveSelf(Phase);
+			switch (Phase)
+			{
+				case UhtResolvePhase.Final:
+					this.PropertyFlags |= EPropertyFlags.InstancedReference & ~this.DisallowPropertyFlags;
+					break;
+			}
+			return bResults;
+		}
+
+		/// <inheritdoc/>
+		public override bool ScanForInstancedReferenced(bool bDeepScan)
+		{
+			return true;
+		}
+
+		/// <inheritdoc/>
+		public override void CollectReferencesInternal(IUhtReferenceCollector Collector, bool bTemplateProperty)
+		{
+			Collector.AddCrossModuleReference(this.Function, true);
+		}
+
+		/// <inheritdoc/>
+		public override StringBuilder AppendText(StringBuilder Builder, UhtPropertyTextType TextType, bool bIsTemplateArgument)
+		{
+			switch (TextType)
+			{
+				case UhtPropertyTextType.EventParameterFunctionMember:
+				case UhtPropertyTextType.ExportMember:
+				case UhtPropertyTextType.RigVMTemplateArg:
+				case UhtPropertyTextType.FunctionThunkParameterArgType:
+					Builder.Append("FMulticastScriptDelegate");
+					break;
+
+				default:
+					Builder.Append(this.Function.SourceName);
+					break;
+			}
+			return Builder;
+		}
+
+		/// <inheritdoc/>
+		public override void AppendObjectHashes(StringBuilder Builder, int StartingLength, IUhtPropertyMemberContext Context)
+		{
+			Builder.AppendObjectHash(StartingLength, this, Context, this.Function);
+		}
+
+		/// <inheritdoc/>
+		public override StringBuilder AppendFunctionThunkParameterArg(StringBuilder Builder)
+		{
+			return Builder.Append(this.Function.SourceName).Append("(").AppendFunctionThunkParameterName(this).Append(")");
+		}
+
+		/// <inheritdoc/>
+		public override StringBuilder AppendNullConstructorArg(StringBuilder Builder, bool bIsInitializer)
+		{
+			Builder.AppendPropertyText(this, UhtPropertyTextType.Construction).Append("()");
+			return Builder;
+		}
+
+		/// <inheritdoc/>
+		public override bool SanitizeDefaultValue(IUhtTokenReader DefaultValueReader, StringBuilder InnerDefaultValue)
+		{
+			return false;
+		}
+
+		/// <inheritdoc/>
+		public override bool IsSameType(UhtProperty Other)
+		{
+			if (Other is UhtMulticastDelegateProperty OtherObject)
+			{
+				return this.Function == OtherObject.Function;
+			}
+			return false;
+		}
+
+		/// <inheritdoc/>
+		public override string? GetRigVMType(ref UhtRigVMParameterFlags ParameterFlags)
+		{
+			return this.CppTypeText;
+		}
+	}
+}
diff --git a/Engine/Source/Programs/Shared/EpicGames.UHT/Types/Properties/UhtMulticastInlineDelegateProperty.cs b/Engine/Source/Programs/Shared/EpicGames.UHT/Types/Properties/UhtMulticastInlineDelegateProperty.cs
new file mode 100644
--- /dev/null
+++ b/Engine/Source/Programs/Shared/EpicGames.UHT/Types/Properties/UhtMulticastInlineDelegateProperty.cs
@@ -0,0 +1,48 @@
+// Copyright Epic Games, Inc. All Rights Reserved.
+
+using EpicGames.UHT.Tables;
+using System.Text;
+
+namespace EpicGames.UHT.Types
+{
+	[UhtEngineClass(Name = "MulticastInlineDelegateProperty", IsProperty = true)]
+	public class UhtMulticastInlineDelegateProperty : UhtMulticastDelegateProperty
+	{
+		/// <inheritdoc/>
+		public override string EngineClassName { get => "MulticastInlineDelegateProperty"; }
+
+		public UhtMulticastInlineDelegateProperty(UhtPropertySettings PropertySettings, UhtFunction Function) : base(PropertySettings, Function)
+		{
+		}
+
+		public override StringBuilder AppendText(StringBuilder Builder, UhtPropertyTextType TextType, bool bIsTemplateArgument)
+		{
+			switch (TextType)
+			{
+				case UhtPropertyTextType.FunctionThunkParameterArgType:
+					Builder.Append("FMulticastInlineDelegateProperty");
+					break;
+
+				default:
+					base.AppendText(Builder, TextType, bIsTemplateArgument);
+					break;
+			}
+			return Builder;
+		}
+
+		/// <inheritdoc/>
+		public override StringBuilder AppendMemberDecl(StringBuilder Builder, IUhtPropertyMemberContext Context, string Name, string NameSuffix, int Tabs)
+		{
+			return AppendMemberDecl(Builder, Context, Name, NameSuffix, Tabs, "FMulticastDelegatePropertyParams");
+		}
+
+		/// <inheritdoc/>
+		public override StringBuilder AppendMemberDef(StringBuilder Builder, IUhtPropertyMemberContext Context, string Name, string NameSuffix, string? Offset, int Tabs)
+		{
+			AppendMemberDefStart(Builder, Context, Name, NameSuffix, Offset, Tabs, "FMulticastDelegatePropertyParams", "UECodeGen_Private::EPropertyGenFlags::InlineMulticastDelegate");
+			AppendMemberDefRef(Builder, Context, this.Function, true);
+			AppendMemberDefEnd(Builder, Context, Name, NameSuffix);
+			return Builder;
+		}
+	}
+}
diff --git a/Engine/Source/Programs/Shared/EpicGames.UHT/Types/Properties/UhtMulticastSparseDelegateProperty.cs b/Engine/Source/Programs/Shared/EpicGames.UHT/Types/Properties/UhtMulticastSparseDelegateProperty.cs
new file mode 100644
--- /dev/null
+++ b/Engine/Source/Programs/Shared/EpicGames.UHT/Types/Properties/UhtMulticastSparseDelegateProperty.cs
@@ -0,0 +1,48 @@
+// Copyright Epic Games, Inc. All Rights Reserved.
+
+using EpicGames.UHT.Tables;
+using System.Text;
+
+namespace EpicGames.UHT.Types
+{
+	[UhtEngineClass(Name = "MulticastSparseDelegateProperty", IsProperty = true)]
+	public class UhtMulticastSparseDelegateProperty : UhtMulticastDelegateProperty
+	{
+		/// <inheritdoc/>
+		public override string EngineClassName { get => "MulticastSparseDelegateProperty"; }
+
+		public UhtMulticastSparseDelegateProperty(UhtPropertySettings PropertySettings, UhtFunction Function) : base(PropertySettings, Function)
+		{
+		}
+
+		public override StringBuilder AppendText(StringBuilder Builder, UhtPropertyTextType TextType, bool bIsTemplateArgument)
+		{
+			switch (TextType)
+			{
+				case UhtPropertyTextType.FunctionThunkParameterArgType:
+					Builder.Append("FMulticastSparseDelegateProperty");
+					break;
+
+				default:
+					base.AppendText(Builder, TextType, bIsTemplateArgument);
+					break;
+			}
+			return Builder;
+		}
+
+		/// <inheritdoc/>
+		public override StringBuilder AppendMemberDecl(StringBuilder Builder, IUhtPropertyMemberContext Context, string Name, string NameSuffix, int Tabs)
+		{
+			return AppendMemberDecl(Builder, Context, Name, NameSuffix, Tabs, "FMulticastDelegatePropertyParams");
+		}
+
+		/// <inheritdoc/>
+		public override StringBuilder AppendMemberDef(StringBuilder Builder, IUhtPropertyMemberContext Context, string Name, string NameSuffix, string? Offset, int Tabs)
+		{
+			AppendMemberDefStart(Builder, Context, Name, NameSuffix, Offset, Tabs, "FMulticastDelegatePropertyParams", "UECodeGen_Private::EPropertyGenFlags::SparseMulticastDelegate");
+			AppendMemberDefRef(Builder, Context, this.Function, true);
+			AppendMemberDefEnd(Builder, Context, Name, NameSuffix);
+			return Builder;
+		}
+	}
+}
diff --git a/Engine/Source/Programs/Shared/EpicGames.UHT/Types/Properties/UhtNameProperty.cs b/Engine/Source/Programs/Shared/EpicGames.UHT/Types/Properties/UhtNameProperty.cs
new file mode 100644
--- /dev/null
+++ b/Engine/Source/Programs/Shared/EpicGames.UHT/Types/Properties/UhtNameProperty.cs
@@ -0,0 +1,95 @@
+// Copyright Epic Games, Inc. All Rights Reserved.
+
+using EpicGames.Core;
+using EpicGames.UHT.Tables;
+using EpicGames.UHT.Tokenizer;
+using EpicGames.UHT.Utils;
+using System.Text;
+
+namespace EpicGames.UHT.Types
+{
+	[UnrealHeaderTool]
+	[UhtEngineClass(Name = "NameProperty", IsProperty = true)]
+	public class UhtNameProperty : UhtProperty
+	{
+		/// <inheritdoc/>
+		public override string EngineClassName { get => "NameProperty"; }
+
+		/// <inheritdoc/>
+		protected override string CppTypeText { get => "FName"; }
+
+		/// <inheritdoc/>
+		protected override string PGetMacroText { get => "PROPERTY"; }
+
+		/// <inheritdoc/>
+		protected override UhtPGetArgumentType PGetTypeArgument { get => UhtPGetArgumentType.EngineClass; }
+
+		public UhtNameProperty(UhtPropertySettings PropertySettings) : base(PropertySettings)
+		{
+			this.PropertyCaps |= UhtPropertyCaps.CanExposeOnSpawn | UhtPropertyCaps.IsParameterSupportedByBlueprint | UhtPropertyCaps.IsMemberSupportedByBlueprint;
+		}
+
+		/// <inheritdoc/>
+		public override StringBuilder AppendNullConstructorArg(StringBuilder Builder, bool bIsInitializer)
+		{
+			Builder.Append("NAME_None");
+			return Builder;
+		}
+
+		/// <inheritdoc/>
+		public override StringBuilder AppendMemberDecl(StringBuilder Builder, IUhtPropertyMemberContext Context, string Name, string NameSuffix, int Tabs)
+		{
+			return AppendMemberDecl(Builder, Context, Name, NameSuffix, Tabs, "FNamePropertyParams");
+		}
+
+		/// <inheritdoc/>
+		public override StringBuilder AppendMemberDef(StringBuilder Builder, IUhtPropertyMemberContext Context, string Name, string NameSuffix, string? Offset, int Tabs)
+		{
+			AppendMemberDefStart(Builder, Context, Name, NameSuffix, Offset, Tabs, "FNamePropertyParams", "UECodeGen_Private::EPropertyGenFlags::Name");
+			AppendMemberDefEnd(Builder, Context, Name, NameSuffix);
+			return Builder;
+		}
+
+		/// <inheritdoc/>
+		public override bool SanitizeDefaultValue(IUhtTokenReader DefaultValueReader, StringBuilder InnerDefaultValue)
+		{
+			if (DefaultValueReader.TryOptional("NAME_None"))
+			{
+				InnerDefaultValue.Append("None");
+			}
+			else if (DefaultValueReader.TryOptional("FName"))
+			{
+				DefaultValueReader.Require('(');
+				StringView Value = DefaultValueReader.GetWrappedConstString();
+				DefaultValueReader.Require(')');
+				InnerDefaultValue.Append(Value);
+			}
+			else
+			{
+				StringView Value = DefaultValueReader.GetWrappedConstString();
+				InnerDefaultValue.Append(Value);
+			}
+			return true;
+		}
+
+		/// <inheritdoc/>
+		public override bool IsSameType(UhtProperty Other)
+		{
+			return Other is UhtNameProperty;
+		}
+
+		/// <inheritdoc/>
+		public override string? GetRigVMType(ref UhtRigVMParameterFlags ParameterFlags)
+		{
+			return this.CppTypeText;
+		}
+
+		#region Keyword
+		[UhtPropertyType(Keyword = "FName", Options = UhtPropertyTypeOptions.Simple | UhtPropertyTypeOptions.Immediate)]
+		public static UhtProperty? NameProperty(UhtPropertyResolvePhase ResolvePhase, UhtPropertySettings PropertySettings, IUhtTokenReader TokenReader, UhtToken MatchedToken)
+		{
+			return new UhtNameProperty(PropertySettings);
+		}
+		#endregion
+	}
+}
diff --git a/Engine/Source/Programs/Shared/EpicGames.UHT/Types/Properties/UhtNumericProperty.cs b/Engine/Source/Programs/Shared/EpicGames.UHT/Types/Properties/UhtNumericProperty.cs
new file mode 100644
--- /dev/null
+++ b/Engine/Source/Programs/Shared/EpicGames.UHT/Types/Properties/UhtNumericProperty.cs
@@ -0,0 +1,39 @@
+// Copyright Epic Games, Inc. All Rights Reserved.
+
+using System.Text;
+
+namespace EpicGames.UHT.Types
+{
+	public abstract class UhtNumericProperty : UhtProperty
+	{
+		/// <inheritdoc/>
+		protected override string PGetMacroText { get => "PROPERTY"; }
+
+		/// <inheritdoc/>
+		protected override UhtPGetArgumentType PGetTypeArgument { get => UhtPGetArgumentType.EngineClass; }
+
+		/// <summary>
+		/// Describes the integer as either being sized on unsized
+		/// </summary>
+		public readonly UhtPropertyIntType IntType;
+
+		protected UhtNumericProperty(UhtPropertySettings PropertySettings, UhtPropertyIntType IntType) : base(PropertySettings)
+		{
+			this.IntType = IntType;
+			this.PropertyCaps |= UhtPropertyCaps.RequiresNullConstructorArg;
+		}
+
+		/// <inheritdoc/>
+		public override StringBuilder AppendNullConstructorArg(StringBuilder Builder, bool bIsInitializer)
+		{
+			Builder.Append("0");
+			return Builder;
+		}
+
+		/// <inheritdoc/>
+		public override string? GetRigVMType(ref UhtRigVMParameterFlags ParameterFlags)
+		{
+			return CppTypeText;
+		}
+	}
+}
diff --git a/Engine/Source/Programs/Shared/EpicGames.UHT/Types/Properties/UhtObjectProperty.cs b/Engine/Source/Programs/Shared/EpicGames.UHT/Types/Properties/UhtObjectProperty.cs
new file mode 100644
--- /dev/null
+++ b/Engine/Source/Programs/Shared/EpicGames.UHT/Types/Properties/UhtObjectProperty.cs
@@ -0,0 +1,120 @@
+// Copyright Epic Games, Inc. All Rights Reserved.
+
+using EpicGames.Core;
+using EpicGames.UHT.Tables;
+using EpicGames.UHT.Utils;
+using System.Text;
+
+namespace EpicGames.UHT.Types
+{
+	[UhtEngineClass(Name = "ObjectProperty", IsProperty = true)]
+	public class UhtObjectProperty : UhtObjectPropertyBase
+	{
+		/// <inheritdoc/>
+		public override string EngineClassName => "ObjectProperty";
+
+		/// <inheritdoc/>
+		protected override UhtPGetArgumentType PGetTypeArgument { get => UhtPGetArgumentType.TypeText; }
+
+		public UhtObjectProperty(UhtPropertySettings PropertySettings, UhtClass PropertyClass, EPropertyFlags ExtraFlags = EPropertyFlags.None)
+			: base(PropertySettings, PropertyClass, null)
+		{
+			this.PropertyFlags |= ExtraFlags;
+			this.PropertyCaps |= UhtPropertyCaps.RequiresNullConstructorArg | UhtPropertyCaps.CanBeInstanced | UhtPropertyCaps.CanExposeOnSpawn |
+				UhtPropertyCaps.IsParameterSupportedByBlueprint | UhtPropertyCaps.IsMemberSupportedByBlueprint;
+		}
+
+		protected UhtObjectProperty(UhtPropertySettings PropertySettings, UhtClass PropertyClass, UhtClass MetaClass, EPropertyFlags ExtraFlags = EPropertyFlags.None)
+			: base(PropertySettings, PropertyClass, MetaClass)
+		{
+			this.PropertyFlags |= ExtraFlags;
+			this.PropertyCaps |= UhtPropertyCaps.RequiresNullConstructorArg | UhtPropertyCaps.CanBeInstanced | UhtPropertyCaps.CanExposeOnSpawn |
+				UhtPropertyCaps.IsParameterSupportedByBlueprint | UhtPropertyCaps.IsMemberSupportedByBlueprint;
+		}
+
+		/// <inheritdoc/>
+		protected override bool ResolveSelf(UhtResolvePhase Phase)
+		{
+			bool bResults = base.ResolveSelf(Phase);
+			switch (Phase)
+			{
+				case UhtResolvePhase.Final:
+					if (this.Class.HierarchyHasAnyClassFlags(EClassFlags.DefaultToInstanced))
+					{
+						this.PropertyFlags |= EPropertyFlags.InstancedReference;
+						this.MetaData.Add(UhtNames.EditInline, true);
+					}
+					break;
+			}
+			return bResults;
+		}
+
+		/// <inheritdoc/>
+		public override bool ScanForInstancedReferenced(bool bDeepScan)
+		{
+			return this.Class.HierarchyHasAnyClassFlags(EClassFlags.DefaultToInstanced);
+		}
+
+		/// <inheritdoc/>
+		public override StringBuilder AppendText(StringBuilder Builder, UhtPropertyTextType TextType, bool bIsTemplateArgument)
+		{
+			switch (TextType)
+			{
+				case UhtPropertyTextType.FunctionThunkReturn:
+					if (this.PropertyFlags.HasAnyFlags(EPropertyFlags.ConstParm))
+					{
+						Builder.Append("const ");
+					}
+					Builder.Append(this.Class.SourceName).Append("*");
+					break;
+
+				case UhtPropertyTextType.FunctionThunkParameterArgType:
+					Builder.Append(this.Class.SourceName);
+					break;
+
+				default:
+					Builder.Append(this.Class.SourceName).Append("*");
+					break;
+			}
+			return Builder;
+		}
+
+		/// <inheritdoc/>
+		public override StringBuilder AppendMemberDecl(StringBuilder Builder, IUhtPropertyMemberContext Context, string Name, string NameSuffix, int Tabs)
+		{
+			return AppendMemberDecl(Builder, Context, Name, NameSuffix, Tabs, "FObjectPropertyParams");
+		}
+
+		/// <inheritdoc/>
+		public override StringBuilder AppendMemberDef(StringBuilder Builder, IUhtPropertyMemberContext Context, string Name, string NameSuffix, string? Offset, int Tabs)
+		{
+			AppendMemberDefStart(Builder, Context, Name, NameSuffix, Offset, Tabs, "FObjectPropertyParams", "UECodeGen_Private::EPropertyGenFlags::Object");
+			AppendMemberDefRef(Builder, Context, this.Class, false);
+			AppendMemberDefEnd(Builder, Context, Name, NameSuffix);
+			return Builder;
+		}
+
+		/// <inheritdoc/>
+		public override StringBuilder AppendNullConstructorArg(StringBuilder Builder, bool bIsInitializer)
+		{
+			Builder.Append("NULL");
+			return Builder;
+		}
+
+		/// <inheritdoc/>
+		public override bool IsSameType(UhtProperty Other)
+		{
+			if (Other is UhtObjectProperty OtherObject)
+			{
+				return this.Class == OtherObject.Class && this.MetaClass == OtherObject.MetaClass;
+			}
+			return false;
+		}
+
+		/// <inheritdoc/>
+		public override string? GetRigVMType(ref UhtRigVMParameterFlags ParameterFlags)
+		{
+			return $"{this.Class}*";
+		}
+	}
+}
diff --git a/Engine/Source/Programs/Shared/EpicGames.UHT/Types/Properties/UhtObjectPropertyBase.cs b/Engine/Source/Programs/Shared/EpicGames.UHT/Types/Properties/UhtObjectPropertyBase.cs
new file mode 100644
--- /dev/null
+++ b/Engine/Source/Programs/Shared/EpicGames.UHT/Types/Properties/UhtObjectPropertyBase.cs
@@ -0,0 +1,240 @@
+// Copyright Epic Games, Inc. All Rights Reserved.
+
+using EpicGames.Core;
+using EpicGames.UHT.Tables;
+using EpicGames.UHT.Tokenizer;
+using EpicGames.UHT.Utils;
+using System.Collections.Generic;
+using System.Diagnostics.CodeAnalysis;
+using System.Text;
+using System.Text.Json.Serialization;
+
+namespace EpicGames.UHT.Types
+{
+	[UhtEngineClass(Name = "ObjectPropertyBase", IsProperty = true)]
+	public abstract class UhtObjectPropertyBase : UhtProperty
+	{
+		/// <inheritdoc/>
+		public override string EngineClassName { get => "ObjectPropertyBase"; }
+
+		/// <inheritdoc/>
+		protected override string CppTypeText { get => "Object"; }
+
+		/// <inheritdoc/>
+		protected override string PGetMacroText { get => "OBJECT"; }
+
+		[JsonConverter(typeof(UhtTypeSourceNameJsonConverter<UhtClass>))]
+		public UhtClass Class { get; set; }
+
+		[JsonConverter(typeof(UhtNullableTypeSourceNameJsonConverter<UhtClass>))]
+		public UhtClass? MetaClass { get; set; }
+
+		public UhtObjectPropertyBase(UhtPropertySettings PropertySettings, UhtClass PropertyClass, UhtClass? MetaClass) : base(PropertySettings)
+		{
+			this.Class = PropertyClass;
+			this.MetaClass = MetaClass;
+
+			// This applies to EVERYTHING including raw pointer
+			// Imply const if it's a parameter that is a pointer to a const class
+			// NOTE: We shouldn't be automatically adding const param because in some cases with functions and blueprint native event, the 
+			// generated code won't match.  For now, just disabled the auto add in that case and check for the error in the validation code.
+			// Otherwise, the user is not warned and they will get compile errors.
+			if (PropertySettings.PropertyCategory != UhtPropertyCategory.Member && this.Class.ClassFlags.HasAnyFlags(EClassFlags.Const) && !PropertySettings.Options.HasAnyFlags(UhtPropertyOptions.NoAutoConst))
+			{
+				this.PropertyFlags |= EPropertyFlags.ConstParm;
+			}
+
+			this.PropertyCaps &= ~(UhtPropertyCaps.CanHaveConfig);
+		}
+
+		/// <inheritdoc/>
+		protected override bool ResolveSelf(UhtResolvePhase Phase)
+		{
+			bool bResults = base.ResolveSelf(Phase);
+			switch (Phase)
+			{
+				case UhtResolvePhase.Final:
+					if (this.Class.HierarchyHasAnyClassFlags(EClassFlags.DefaultToInstanced))
+					{
+						this.PropertyFlags |= (EPropertyFlags.InstancedReference | EPropertyFlags.ExportObject) & ~this.DisallowPropertyFlags;
+					}
+					break;
+			}
+			return bResults;
+		}
+
+		/// <inheritdoc/>
+		public override bool ScanForInstancedReferenced(bool bDeepScan)
+		{
+			return !this.DisallowPropertyFlags.HasAnyFlags(EPropertyFlags.InstancedReference) && this.Class.HierarchyHasAnyClassFlags(EClassFlags.DefaultToInstanced);
+		}
+
+		/// <inheritdoc/>
+		public override void CollectReferencesInternal(IUhtReferenceCollector Collector, bool bTemplateProperty)
+		{
+			Collector.AddCrossModuleReference(this.Class, false);
+			Collector.AddCrossModuleReference(this.MetaClass, false);
+		}
+
+		/// <inheritdoc/>
+		public override string? GetForwardDeclarations()
+		{
+			return $"class {this.Class.SourceName};";
+		}
+
+		/// <inheritdoc/>
+		public override IEnumerable<UhtType> EnumerateReferencedTypes()
+		{
+			yield return this.Class;
+			if (this.MetaClass != null)
+			{
+				yield return this.MetaClass;
+			}
+		}
+
+		/// <inheritdoc/>
+		public override bool SanitizeDefaultValue(IUhtTokenReader DefaultValueReader, StringBuilder InnerDefaultValue)
+		{
+			int Value;
+			if (DefaultValueReader.TryOptional("NULL") ||
+				DefaultValueReader.TryOptional("nullptr") ||
+				(DefaultValueReader.TryOptionalConstInt(out Value) && Value == 0))
+			{
+				InnerDefaultValue.Append("None");
+				return true;
+			}
+			return false;
+		}
+
+		/// <inheritdoc/>
+		public override void ValidateDeprecated()
+		{
+			ValidateDeprecatedClass(this.Class);
+			ValidateDeprecatedClass(this.MetaClass);
+		}
+
+		/// <inheritdoc/>
+		public override bool MustBeConstArgument([NotNullWhen(true)] out UhtType? ErrorType)
+		{
+			ErrorType = this.Class;
+			return this.Class.ClassFlags.HasAnyFlags(EClassFlags.Const);
+		}
+
+		#region Parsing support methods
+		public static UhtClass? ParseTemplateObject(UhtPropertySettings PropertySettings, IUhtTokenReader TokenReader, UhtToken MatchedToken, bool bReturnUInterface)
+		{
+			if (TokenReader.TryOptional("const"))
+			{
+				PropertySettings.MetaData.Add(UhtNames.NativeConst, "");
+			}
+			if (!TokenReader.SkipExpectedType(MatchedToken.Value, PropertySettings.PropertyCategory == UhtPropertyCategory.Member))
+			{
+				return null;
+			}
+			bool bIsNativeConstTemplateArg = false;
+			UhtToken Identifier = new UhtToken();
+			TokenReader
+				.Require('<')
+				.Optional("const", () => { bIsNativeConstTemplateArg = true; })
+				.Optional("class")
+				.RequireIdentifier((ref UhtToken Token) => { Identifier = Token; })
+				.Optional("const", () => { bIsNativeConstTemplateArg = true; })
+				.Require('>');
+
+			if (bIsNativeConstTemplateArg)
+			{
+				PropertySettings.MetaData.Add(UhtNames.NativeConstTemplateArg, "");
+			}
+			UhtConfig.Instance.RedirectTypeIdentifier(ref Identifier);
+			UhtClass? Return = PropertySettings.Outer.FindType(UhtFindOptions.SourceName | UhtFindOptions.Class, ref Identifier, TokenReader) as UhtClass;
+			if (Return != null && Return.AlternateObject != null && bReturnUInterface)
+			{
+				Return = Return.AlternateObject as UhtClass;
+			}
+			return Return;
+		}
+
+		public static UhtClass? ParseTemplateClass(UhtPropertySettings PropertySettings, IUhtTokenReader TokenReader, UhtToken MatchedToken)
+		{
+			UhtToken Identifier = new UhtToken();
+
+			if (TokenReader.TryOptional("const"))
+			{
+				PropertySettings.MetaData.Add(UhtNames.NativeConst, "");
+			}
+
+			TokenReader.Optional("class");
+
+			if (!TokenReader.SkipExpectedType(MatchedToken.Value, PropertySettings.PropertyCategory == UhtPropertyCategory.Member))
+			{
+				return null;
+			}
+
+			TokenReader
+				.Require('<')
+				.Optional("class")
+				.RequireIdentifier((ref UhtToken Token) => { Identifier = Token; })
+				.Require('>');
+
+			UhtConfig.Instance.RedirectTypeIdentifier(ref Identifier);
+			UhtClass? Return = PropertySettings.Outer.FindType(UhtFindOptions.SourceName | UhtFindOptions.Class, ref Identifier, TokenReader) as UhtClass;
+			if (Return != null && Return.AlternateObject != null)
+			{
+				Return = Return.AlternateObject as UhtClass;
+			}
+			return Return;
+		}
+
+		public static void ConditionalLogPointerUsage(UhtPropertySettings PropertySettings, UhtPointerMemberBehavior EngineBehavior, UhtPointerMemberBehavior NonEngineBehavior,
+			string PointerTypeDesc, IUhtTokenReader TokenReader, int TypeStartPos, string? AlternativeTypeDesc)
+		{
+			if (PropertySettings.PropertyCategory != UhtPropertyCategory.Member)
+			{
+				return;
+			}
+
+			UhtPackage Package = PropertySettings.Outer.Package;
+			bool bEngineBehavior = Package.bIsPartOfEngine && !Package.bIsPlugin;
+			UhtPointerMemberBehavior Behavior = bEngineBehavior ? EngineBehavior : NonEngineBehavior;
+
+			if (Behavior == UhtPointerMemberBehavior.AllowSilently)
+			{
+				return;
+			}
+
+			string Type = TokenReader.GetStringView(TypeStartPos, TokenReader.InputPos - TypeStartPos).ToString();
+			Type.Replace("\n", " ");
+			Type.Replace("\r", "");
+			Type.Replace("\t", " ");
+
+			switch (Behavior)
+			{
+				case UhtPointerMemberBehavior.Disallow:
+					if (!string.IsNullOrEmpty(AlternativeTypeDesc))
+					{
+						TokenReader.LogError($"{PointerTypeDesc} usage in member declaration detected [[[{Type}]]].  This is disallowed for the target/module, consider {AlternativeTypeDesc} as an alternative.");
+					}
+					else
+					{
+						TokenReader.LogError($"{PointerTypeDesc} usage in member declaration detected [[[{Type}]]].");
+					}
+					break;
+
+				case UhtPointerMemberBehavior.AllowAndLog:
+					if (!string.IsNullOrEmpty(AlternativeTypeDesc))
+					{
+						TokenReader.LogTrace($"{PointerTypeDesc} usage in member declaration detected [[[{Type}]]].  Consider {AlternativeTypeDesc} as an alternative.");
+					}
+					else
+					{
+						TokenReader.LogTrace("{PointerTypeDesc} usage in member declaration detected [[[{Type}]]].");
+					}
+					break;
+
+				default:
+					throw new UhtIceException("Unknown enum value");
+			}
+		}
+		#endregion
+	}
+}
diff --git a/Engine/Source/Programs/Shared/EpicGames.UHT/Types/Properties/UhtObjectPtrProperty.cs b/Engine/Source/Programs/Shared/EpicGames.UHT/Types/Properties/UhtObjectPtrProperty.cs
new file mode 100644
--- /dev/null
+++ b/Engine/Source/Programs/Shared/EpicGames.UHT/Types/Properties/UhtObjectPtrProperty.cs
@@ -0,0 +1,118 @@
+// Copyright Epic Games, Inc. All Rights Reserved.
+
+using EpicGames.Core;
+using EpicGames.UHT.Tables;
+using EpicGames.UHT.Tokenizer;
+using EpicGames.UHT.Utils;
+using System.Text;
+
+namespace EpicGames.UHT.Types
+{
+	[UnrealHeaderTool]
+	[UhtEngineClass(Name = "ObjectPtrProperty", IsProperty = true)]
+	public class UhtObjectPtrProperty : UhtObjectProperty
+	{
+		/// <inheritdoc/>
+		public override string EngineClassName { get => "ObjectPtrProperty"; }
+
+		/// <inheritdoc/>
+		protected override string CppTypeText { get => "ObjectPtr"; }
+
+		/// <inheritdoc/>
+		protected override string PGetMacroText { get => "OBJECTPTR"; }
+
+		/// <inheritdoc/>
+		protected override UhtPGetArgumentType PGetTypeArgument { get => UhtPGetArgumentType.TypeText; }
+
+		public UhtObjectPtrProperty(UhtPropertySettings PropertySettings, UhtClass PropertyClass, EPropertyFlags ExtraFlags = EPropertyFlags.None)
+			: base(PropertySettings, PropertyClass)
+		{
+			this.PropertyFlags |= ExtraFlags | EPropertyFlags.UObjectWrapper;
+			this.PropertyCaps |= UhtPropertyCaps.PassCppArgsByRef;
+		}
+
+		/// <inheritdoc/>
+		public override StringBuilder AppendText(StringBuilder Builder, UhtPropertyTextType TextType, bool bIsTemplateArgument)
+		{
+			switch (TextType)
+			{
+				case UhtPropertyTextType.GetterSetterArg:
+					Builder.Append(this.Class.SourceName).Append("*");
+					break;
+
+				case UhtPropertyTextType.FunctionThunkReturn:
+					if (this.PropertyFlags.HasAnyFlags(EPropertyFlags.ConstParm))
+					{
+						Builder.Append("const ");
+					}
+					Builder.Append(this.Class.SourceName).Append("*");
+					break;
+
+				default:
+					Builder.Append("TObjectPtr<").Append(this.Class.SourceName).Append(">");
+					break;
+			}
+			return Builder;
+		}
+		
+		/// <inheritdoc/>
+		public override StringBuilder AppendMemberDecl(StringBuilder Builder, IUhtPropertyMemberContext Context, string Name, string NameSuffix, int Tabs)
+		{
+			return AppendMemberDecl(Builder, Context, Name, NameSuffix, Tabs, "FObjectPtrPropertyParams");
+		}
+
+		/// <inheritdoc/>
+		public override StringBuilder AppendMemberDef(StringBuilder Builder, IUhtPropertyMemberContext Context, string Name, string NameSuffix, string? Offset, int Tabs)
+		{
+			AppendMemberDefStart(Builder, Context, Name, NameSuffix, Offset, Tabs, "FObjectPtrPropertyParams", "UECodeGen_Private::EPropertyGenFlags::Object | UECodeGen_Private::EPropertyGenFlags::ObjectPtr");
+			AppendMemberDefRef(Builder, Context, this.Class, false);
+			AppendMemberDefEnd(Builder, Context, Name, NameSuffix);
+			return Builder;
+		}
+
+		/// <inheritdoc/>
+		public override void Validate(UhtStruct OuterStruct, UhtProperty OutermostProperty, UhtValidationOptions Options)
+		{
+			base.Validate(OuterStruct, OutermostProperty, Options);
+
+			// UFunctions with a smart pointer as input parameter wont compile anyway, because of missing P_GET_... macro.
+			// UFunctions with a smart pointer as return type will crash when called via blueprint, because they are not supported in VM.
+			if (this.PropertyCategory != UhtPropertyCategory.Member)
+			{
+				OuterStruct.LogError("UFunctions cannot take a TObjectPtr as a parameter.");
+			}
+		}
+
+		/// <inheritdoc/>
+		public override string? GetRigVMType(ref UhtRigVMParameterFlags ParameterFlags)
+		{
+			return $"TObjectPtr<{this.Class.SourceName}>";
+		}
+
+		#region Keyword
+		[UhtPropertyType(Keyword = "TObjectPtr")]
+		public static UhtProperty? ObjectPtrProperty(UhtPropertyResolvePhase ResolvePhase, UhtPropertySettings PropertySettings, IUhtTokenReader TokenReader, UhtToken MatchedToken)
+		{
+			int TypeStartPos = TokenReader.PeekToken().InputStartPos;
+
+			UhtClass? PropertyClass = ParseTemplateObject(PropertySettings, TokenReader, MatchedToken, true);
+			if (PropertyClass == null)
+			{
+				return null;
+			}
+
+			ConditionalLogPointerUsage(PropertySettings, UhtConfig.Instance.EngineObjectPtrMemberBehavior, UhtConfig.Instance.NonEngineObjectPtrMemberBehavior, "ObjectPtr", TokenReader, TypeStartPos, null);
+
+			if (PropertyClass.IsChildOf(PropertyClass.Session.UClass))
+			{
+				// UObject specifies that there is no limiter
+				return new UhtClassPtrProperty(PropertySettings, PropertyClass, PropertyClass.Session.UObject);
+			}
+			else
+			{
+				return new UhtObjectPtrProperty(PropertySettings, PropertyClass);
+			}
+		}
+		#endregion
+	}
+}
diff --git a/Engine/Source/Programs/Shared/EpicGames.UHT/Types/Properties/UhtSetProperty.cs b/Engine/Source/Programs/Shared/EpicGames.UHT/Types/Properties/UhtSetProperty.cs
new file mode 100644
--- /dev/null
+++ b/Engine/Source/Programs/Shared/EpicGames.UHT/Types/Properties/UhtSetProperty.cs
@@ -0,0 +1,250 @@
+// Copyright Epic Games, Inc. All Rights Reserved.
+
+using EpicGames.Core;
+using EpicGames.UHT.Parsers;
+using EpicGames.UHT.Tables;
+using EpicGames.UHT.Tokenizer;
+using EpicGames.UHT.Utils;
+using System.Collections.Generic;
+using System.Text;
+
+namespace EpicGames.UHT.Types
+{
+	[UnrealHeaderTool]
+	[UhtEngineClass(Name = "SetProperty", IsProperty = true)]
+	public class UhtSetProperty : UhtContainerBaseProperty
+	{
+		/// <inheritdoc/>
+		public override string EngineClassName { get => "SetProperty"; }
+
+		/// <inheritdoc/>
+		protected override string CppTypeText { get => "TSet"; }
+
+		/// <inheritdoc/>
+		protected override string PGetMacroText { get => "TSET"; }
+
+		/// <inheritdoc/>
+		protected override UhtPGetArgumentType PGetTypeArgument { get => UhtPGetArgumentType.TypeText; }
+
+		public UhtSetProperty(UhtPropertySettings PropertySettings, UhtProperty Value) : base(PropertySettings, Value)
+		{
+			// If the creation of the value property set more flags, then copy those flags to ourselves
+			this.PropertyFlags |= Value.PropertyFlags & EPropertyFlags.UObjectWrapper;
+
+			this.PropertyCaps |= UhtPropertyCaps.PassCppArgsByRef;
+			this.PropertyCaps &= ~(UhtPropertyCaps.CanBeContainerValue | UhtPropertyCaps.CanBeContainerKey);
+			this.PropertyCaps = (this.PropertyCaps & ~(UhtPropertyCaps.IsParameterSupportedByBlueprint | UhtPropertyCaps.IsMemberSupportedByBlueprint | UhtPropertyCaps.CanExposeOnSpawn)) |
+				(this.ValueProperty.PropertyCaps & (UhtPropertyCaps.IsParameterSupportedByBlueprint | UhtPropertyCaps.IsMemberSupportedByBlueprint | UhtPropertyCaps.CanExposeOnSpawn));
+			this.ValueProperty.SourceName = this.SourceName;
+			this.ValueProperty.EngineName = this.EngineName;
+			this.ValueProperty.PropertyFlags = this.PropertyFlags & EPropertyFlags.PropagateToSetElement;
+			this.ValueProperty.Outer = this;
+			this.ValueProperty.MetaData.Clear();
+		}
+
+		/// <inheritdoc/>
+		protected override bool ResolveSelf(UhtResolvePhase Phase)
+		{
+			bool bResults = base.ResolveSelf(Phase);
+			switch (Phase)
+			{
+				case UhtResolvePhase.Final:
+					this.PropertyFlags |= ResolveAndReturnNewFlags(this.ValueProperty, Phase);
+					this.MetaData.Add(this.ValueProperty.MetaData);
+					this.ValueProperty.MetaData.Clear();
+					PropagateFlagsFromInnerAndHandlePersistentInstanceMetadata(this, this.MetaData, this.ValueProperty);
+					break;
+			}
+			return bResults;
+		}
+
+		/// <inheritdoc/>
+		public override void CollectReferencesInternal(IUhtReferenceCollector Collector, bool bTemplateProperty)
+		{
+			this.ValueProperty.CollectReferencesInternal(Collector, true);
+		}
+
+		/// <inheritdoc/>
+		public override string? GetForwardDeclarations()
+		{
+			return this.ValueProperty.GetForwardDeclarations();
+		}
+
+		/// <inheritdoc/>
+		public override IEnumerable<UhtType> EnumerateReferencedTypes()
+		{
+			foreach (UhtType Type in this.ValueProperty.EnumerateReferencedTypes())
+			{
+				yield return Type;
+			}
+		}
+
+		/// <inheritdoc/>
+		public override StringBuilder AppendText(StringBuilder Builder, UhtPropertyTextType TextType, bool bIsTemplateArgument)
+		{
+			switch (TextType)
+			{
+				case UhtPropertyTextType.SparseShort:
+					Builder.Append("TSet");
+					break;
+
+				case UhtPropertyTextType.FunctionThunkParameterArgType:
+					Builder.AppendFunctionThunkParameterArrayType(this.ValueProperty);
+					break;
+
+				default:
+					Builder.Append("TSet<").AppendPropertyText(this.ValueProperty, TextType, true);
+					if (Builder[Builder.Length - 1] == '>')
+					{
+						// if our internal property type is a template class, add a space between the closing brackets b/c VS.NET cannot parse this correctly
+						Builder.Append(" ");
+					}
+					Builder.Append('>');
+					break;
+			}
+			return Builder;
+		}
+
+		/// <inheritdoc/>
+		public override StringBuilder AppendMemberDecl(StringBuilder Builder, IUhtPropertyMemberContext Context, string Name, string NameSuffix, int Tabs)
+		{
+			Builder.AppendMemberDecl(this.ValueProperty, Context, Name, GetNameSuffix(NameSuffix, "_ElementProp"), Tabs);
+			return AppendMemberDecl(Builder, Context, Name, NameSuffix, Tabs, "FSetPropertyParams");
+		}
+
+		/// <inheritdoc/>
+		public override StringBuilder AppendMemberDef(StringBuilder Builder, IUhtPropertyMemberContext Context, string Name, string NameSuffix, string? Offset, int Tabs)
+		{
+			Builder.AppendMemberDef(this.ValueProperty, Context, Name, GetNameSuffix(NameSuffix, "_ElementProp"), "0", Tabs);
+			Builder.AppendMetaDataDef(this, Context, Name, NameSuffix, Tabs);
+
+			if (this.ValueProperty is UhtStructProperty StructProperty)
+			{
+				Builder
+					.AppendTabs(Tabs)
+					.Append("static_assert(TModels<CGetTypeHashable, ")
+					.Append(StructProperty.ScriptStruct.SourceName)
+					.Append(">::Value, \"The structure '")
+					.Append(StructProperty.ScriptStruct.SourceName)
+					.Append("' is used in a TSet but does not have a GetValueTypeHash defined\");\r\n");
+			}
+
+			AppendMemberDefStart(Builder, Context, Name, NameSuffix, Offset, Tabs, "FSetPropertyParams", "UECodeGen_Private::EPropertyGenFlags::Set", false);
+			AppendMemberDefEnd(Builder, Context, Name, NameSuffix);
+			return Builder;
+		}
+
+		/// <inheritdoc/>
+		public override StringBuilder AppendMemberPtr(StringBuilder Builder, IUhtPropertyMemberContext Context, string Name, string NameSuffix, int Tabs)
+		{
+			Builder.AppendMemberPtr(this.ValueProperty, Context, Name, GetNameSuffix(NameSuffix, "_ElementProp"), Tabs);
+			base.AppendMemberPtr(Builder, Context, Name, NameSuffix, Tabs);
+			return Builder;
+		}
+
+		/// <inheritdoc/>
+		public override void AppendObjectHashes(StringBuilder Builder, int StartingLength, IUhtPropertyMemberContext Context)
+		{
+			this.ValueProperty.AppendObjectHashes(Builder, StartingLength, Context);
+		}
+
+		/// <inheritdoc/>
+		public override StringBuilder AppendNullConstructorArg(StringBuilder Builder, bool bIsInitializer)
+		{
+			Builder.AppendPropertyText(this, UhtPropertyTextType.Construction).Append("()");
+			return Builder;
+		}
+
+		/// <inheritdoc/>
+		public override bool SanitizeDefaultValue(IUhtTokenReader DefaultValueReader, StringBuilder InnerDefaultValue)
+		{
+			return false;
+		}
+
+		/// <inheritdoc/>
+		public override bool IsSameType(UhtProperty Other)
+		{
+			if (Other is UhtSetProperty OtherSet)
+			{
+				return this.ValueProperty.IsSameType(OtherSet.ValueProperty);
+			}
+			return false;
+		}
+
+		///<inheritdoc/>
+		public override bool ValidateStructPropertyOkForNet(UhtProperty ReferencingProperty)
+		{
+			if (!this.PropertyFlags.HasAnyFlags(EPropertyFlags.RepSkip))
+			{
+				ReferencingProperty.LogError($"Sets are not supported for Replication or RPCs.  Set '{this.SourceName}' in '{this.Outer?.SourceName}'.  Origin '{ReferencingProperty.SourceName}'");
+				return false;
+			}
+			return true;
+		}
+
+		/// <inheritdoc/>
+		protected override void ValidateFunctionArgument(UhtFunction Function, UhtValidationOptions Options)
+		{
+			base.ValidateFunctionArgument(Function, Options);
+
+			if (Function.FunctionFlags.HasAnyFlags(EFunctionFlags.Net))
+			{
+				if (!Function.FunctionFlags.HasAnyFlags(EFunctionFlags.NetRequest | EFunctionFlags.NetResponse))
+				{
+					this.LogError("Sets are not supported in an RPC.");
+				}
+			}
+		}
+
+		/// <inheritdoc/>
+		public override string? GetRigVMType(ref UhtRigVMParameterFlags ParameterFlags)
+		{
+			return null;
+		}
+
+#region Keyword
+		[UhtPropertyType(Keyword = "TSet")]
+		public static UhtProperty? SetProperty(UhtPropertyResolvePhase ResolvePhase, UhtPropertySettings PropertySettings, IUhtTokenReader TokenReader, UhtToken MatchedToken)
+		{
+			using (var TokenContext = new UhtMessageContext(TokenReader, "TSet"))
+			{
+				if (!TokenReader.SkipExpectedType(MatchedToken.Value, PropertySettings.PropertyCategory == UhtPropertyCategory.Member))
+				{
+					return null;
+				}
+				TokenReader.Require('<');
+
+				// Parse the value type
+				UhtProperty? Value = UhtPropertyParser.ParseTemplateParam(ResolvePhase, PropertySettings, "Value", TokenReader);
+				if (Value == null)
+				{
+					return null;
+				}
+
+				if (!Value.PropertyCaps.HasAnyFlags(UhtPropertyCaps.CanBeContainerKey))
+				{
+					TokenReader.LogError($"The type \'{Value.GetUserFacingDecl()}\' can not be used as a key in a TSet");
+				}
+
+				if (PropertySettings.PropertyFlags.HasAnyFlags(EPropertyFlags.Net))
+				{
+					TokenReader.LogError("Replicated sets are not supported.");
+				}
+
+				if (TokenReader.TryOptional(','))
+				{
+					// If we found a comma, read the next thing, assume it's a keyfuncs, and report that
+					UhtToken KeyFuncToken = TokenReader.GetIdentifier();
+					throw new UhtException(TokenReader, $"Found '{KeyFuncToken.Value}' - explicit KeyFuncs are not supported in TSet properties.");
+				}
+
+				TokenReader.Require('>');
+
+				//@TODO: Prevent sparse delegate types from being used in a container
+
+				return new UhtSetProperty(PropertySettings, Value);
+			}
+		}
+#endregion
+	}
+}
diff --git a/Engine/Source/Programs/Shared/EpicGames.UHT/Types/Properties/UhtSoftClassProperty.cs b/Engine/Source/Programs/Shared/EpicGames.UHT/Types/Properties/UhtSoftClassProperty.cs
new file mode 100644
--- /dev/null
+++ b/Engine/Source/Programs/Shared/EpicGames.UHT/Types/Properties/UhtSoftClassProperty.cs
@@ -0,0 +1,91 @@
+// Copyright Epic Games, Inc. All Rights Reserved.
+
+using EpicGames.UHT.Tables;
+using EpicGames.UHT.Tokenizer;
+using EpicGames.UHT.Utils;
+using System.Text;
+
+namespace EpicGames.UHT.Types
+{
+	[UnrealHeaderTool]
+	[UhtEngineClass(Name = "SoftClassProperty", IsProperty = true)]
+	public class UhtSoftClassProperty : UhtSoftObjectProperty
+	{
+		/// <inheritdoc/>
+		public override string EngineClassName { get => "SoftClassProperty"; }
+
+		/// <inheritdoc/>
+		protected override string CppTypeText { get => "SoftClassPtr"; }
+
+		/// <inheritdoc/>
+		protected override string PGetMacroText { get => "SOFTCLASS"; }
+
+		/// <inheritdoc/>
+		protected override UhtPGetArgumentType PGetTypeArgument { get => UhtPGetArgumentType.TypeText; }
+
+		public UhtSoftClassProperty(UhtPropertySettings PropertySettings, UhtClass PropertyClass, UhtClass MetaClass)
+			: base(PropertySettings, PropertyClass, MetaClass)
+		{
+		}
+
+		/// <inheritdoc/>
+		public override string? GetForwardDeclarations()
+		{
+			return this.MetaClass != null ? $"class {this.MetaClass.SourceName};" : null;
+		}
+
+		/// <inheritdoc/>
+		public override void CollectReferencesInternal(IUhtReferenceCollector Collector, bool bTemplateProperty)
+		{
+			Collector.AddCrossModuleReference(this.MetaClass, false);
+		}
+
+		/// <inheritdoc/>
+		public override StringBuilder AppendText(StringBuilder Builder, UhtPropertyTextType TextType, bool bIsTemplateArgument)
+		{
+			switch (TextType)
+			{
+				default:
+					Builder.Append("TSoftClassPtr<").Append(this.MetaClass?.SourceName).Append("> ");
+					break;
+			}
+			return Builder;
+		}
+
+		/// <inheritdoc/>
+		public override StringBuilder AppendMemberDecl(StringBuilder Builder, IUhtPropertyMemberContext Context, string Name, string NameSuffix, int Tabs)
+		{
+			return AppendMemberDecl(Builder, Context, Name, NameSuffix, Tabs, "FSoftClassPropertyParams");
+		}
+
+		/// <inheritdoc/>
+		public override StringBuilder AppendMemberDef(StringBuilder Builder, IUhtPropertyMemberContext Context, string Name, string NameSuffix, string? Offset, int Tabs)
+		{
+			AppendMemberDefStart(Builder, Context, Name, NameSuffix, Offset, Tabs, "FSoftClassPropertyParams", "UECodeGen_Private::EPropertyGenFlags::SoftClass");
+			AppendMemberDefRef(Builder, Context, this.MetaClass, false);
+			AppendMemberDefEnd(Builder, Context, Name, NameSuffix);
+			return Builder;
+		}
+
+		/// <inheritdoc/>
+		public override string? GetRigVMType(ref UhtRigVMParameterFlags ParameterFlags)
+		{
+			return $"TSoftClassPtr<{this.MetaClass?.SourceName}>";
+		}
+
+		#region Keyword
+		[UhtPropertyType(Keyword = "TSoftClassPtr")]
+		public static UhtProperty? SoftClassPtrProperty(UhtPropertyResolvePhase ResolvePhase, UhtPropertySettings PropertySettings, IUhtTokenReader TokenReader, UhtToken MatchedToken)
+		{
+			UhtClass? MetaClass = ParseTemplateClass(PropertySettings, TokenReader, MatchedToken);
+			if (MetaClass == null)
+			{
+				return null;
+			}
+
+			// With TSubclassOf, MetaClass is used as a class limiter.  
+			return new UhtSoftClassProperty(PropertySettings, MetaClass.Session.UClass, MetaClass);
+		}
+		#endregion
+	}
+}
diff --git a/Engine/Source/Programs/Shared/EpicGames.UHT/Types/Properties/UhtSoftObjectProperty.cs b/Engine/Source/Programs/Shared/EpicGames.UHT/Types/Properties/UhtSoftObjectProperty.cs
new file mode 100644
--- /dev/null
+++ b/Engine/Source/Programs/Shared/EpicGames.UHT/Types/Properties/UhtSoftObjectProperty.cs
@@ -0,0 +1,112 @@
+// Copyright Epic Games, Inc. All Rights Reserved.
+
+using EpicGames.Core;
+using EpicGames.UHT.Tables;
+using EpicGames.UHT.Tokenizer;
+using EpicGames.UHT.Utils;
+using System.Text;
+
+namespace EpicGames.UHT.Types
+{
+	[UnrealHeaderTool]
+	[UhtEngineClass(Name = "SoftObjectProperty", IsProperty = true)]
+	public class UhtSoftObjectProperty : UhtObjectPropertyBase
+	{
+		/// <inheritdoc/>
+		public override string EngineClassName { get => "SoftObjectProperty"; }
+
+		/// <inheritdoc/>
+		protected override string CppTypeText { get => "SoftObjectPtr"; }
+
+		/// <inheritdoc/>
+		protected override string PGetMacroText { get => "SOFTOBJECT"; }
+
+		/// <inheritdoc/>
+		protected override UhtPGetArgumentType PGetTypeArgument { get => UhtPGetArgumentType.TypeText; }
+
+		public UhtSoftObjectProperty(UhtPropertySettings PropertySettings, UhtClass PropertyClass)
+			: base(PropertySettings, PropertyClass, null)
+		{
+			this.PropertyFlags |= EPropertyFlags.UObjectWrapper;
+			this.PropertyCaps |= UhtPropertyCaps.PassCppArgsByRef | UhtPropertyCaps.RequiresNullConstructorArg | UhtPropertyCaps.CanExposeOnSpawn |
+				UhtPropertyCaps.CanHaveConfig | UhtPropertyCaps.IsParameterSupportedByBlueprint | UhtPropertyCaps.IsMemberSupportedByBlueprint;
+		}
+
+		protected UhtSoftObjectProperty(UhtPropertySettings PropertySettings, UhtClass PropertyClass, UhtClass MetaClass)
+			: base(PropertySettings, PropertyClass, MetaClass)
+		{
+			this.PropertyFlags |= EPropertyFlags.UObjectWrapper;
+			this.PropertyCaps |= UhtPropertyCaps.PassCppArgsByRef | UhtPropertyCaps.RequiresNullConstructorArg | UhtPropertyCaps.CanExposeOnSpawn |
+				UhtPropertyCaps.CanHaveConfig | UhtPropertyCaps.IsParameterSupportedByBlueprint | UhtPropertyCaps.IsMemberSupportedByBlueprint;
+		}
+
+		/// <inheritdoc/>
+		public override StringBuilder AppendNullConstructorArg(StringBuilder Builder, bool bIsInitializer)
+		{
+			Builder.Append("NULL");
+			return Builder;
+		}
+
+		/// <inheritdoc/>
+		public override StringBuilder AppendText(StringBuilder Builder, UhtPropertyTextType TextType, bool bIsTemplateArgument)
+		{
+			switch (TextType)
+			{
+				default:
+					Builder.Append("TSoftObjectPtr<").Append(this.Class.SourceName).Append(">");
+					break;
+			}
+			return Builder;
+		}
+
+		/// <inheritdoc/>
+		public override StringBuilder AppendMemberDecl(StringBuilder Builder, IUhtPropertyMemberContext Context, string Name, string NameSuffix, int Tabs)
+		{
+			return AppendMemberDecl(Builder, Context, Name, NameSuffix, Tabs, "FSoftObjectPropertyParams");
+		}
+
+		/// <inheritdoc/>
+		public override StringBuilder AppendMemberDef(StringBuilder Builder, IUhtPropertyMemberContext Context, string Name, string NameSuffix, string? Offset, int Tabs)
+		{
+			AppendMemberDefStart(Builder, Context, Name, NameSuffix, Offset, Tabs, "FSoftObjectPropertyParams", "UECodeGen_Private::EPropertyGenFlags::SoftObject");
+			AppendMemberDefRef(Builder, Context, this.Class, false);
+			AppendMemberDefEnd(Builder, Context, Name, NameSuffix);
+			return Builder;
+		}
+
+		/// <inheritdoc/>
+		public override bool IsSameType(UhtProperty Other)
+		{
+			if (Other is UhtSoftObjectProperty OtherObject)
+			{
+				return this.Class == OtherObject.Class && this.MetaClass == OtherObject.MetaClass;
+			}
+			return false;
+		}
+
+		/// <inheritdoc/>
+		public override string? GetRigVMType(ref UhtRigVMParameterFlags ParameterFlags)
+		{
+			return $"TSoftObjectPtr<{this.Class.SourceName}>";
+		}
+
+		#region Keyword
+		[UhtPropertyType(Keyword = "TSoftObjectPtr")]
+		public static UhtProperty? SoftObjectPtrProperty(UhtPropertyResolvePhase ResolvePhase, UhtPropertySettings PropertySettings, IUhtTokenReader TokenReader, UhtToken MatchedToken)
+		{
+			UhtClass? PropertyClass = UhtObjectPropertyBase.ParseTemplateObject(PropertySettings, TokenReader, MatchedToken, true);
+			if (PropertyClass == null)
+			{
+				return null;
+			}
+
+			if (PropertyClass.IsChildOf(PropertyClass.Session.UClass))
+			{
+				TokenReader.LogError("Class variables cannot be stored in TSoftObjectPtr, use TSoftClassPtr instead.");
+			}
+
+			return new UhtSoftObjectProperty(PropertySettings, PropertyClass);
+		}
+		#endregion
+	}
+}
diff --git a/Engine/Source/Programs/Shared/EpicGames.UHT/Types/Properties/UhtStrProperty.cs b/Engine/Source/Programs/Shared/EpicGames.UHT/Types/Properties/UhtStrProperty.cs
new file mode 100644
--- /dev/null
+++ b/Engine/Source/Programs/Shared/EpicGames.UHT/Types/Properties/UhtStrProperty.cs
@@ -0,0 +1,134 @@
+// Copyright Epic Games, Inc. All Rights Reserved.
+
+using EpicGames.Core;
+using EpicGames.UHT.Tables;
+using EpicGames.UHT.Tokenizer;
+using EpicGames.UHT.Utils;
+using System.Text;
+
+namespace EpicGames.UHT.Types
+{
+	[UnrealHeaderTool]
+	[UhtEngineClass(Name = "StrProperty", IsProperty = true)]
+	public class UhtStrProperty : UhtProperty
+	{
+		/// <inheritdoc/>
+		public override string EngineClassName { get => "StrProperty"; }
+
+		/// <inheritdoc/>
+		protected override string CppTypeText { get => "FString"; }
+
+		/// <inheritdoc/>
+		protected override string PGetMacroText { get => "PROPERTY"; }
+
+		/// <inheritdoc/>
+		protected override UhtPGetArgumentType PGetTypeArgument { get => UhtPGetArgumentType.EngineClass; }
+
+		public UhtStrProperty(UhtPropertySettings PropertySettings) : base(PropertySettings)
+		{
+			this.PropertyCaps |= UhtPropertyCaps.PassCppArgsByRef | UhtPropertyCaps.CanExposeOnSpawn | UhtPropertyCaps.IsParameterSupportedByBlueprint | UhtPropertyCaps.IsMemberSupportedByBlueprint;
+		}
+
+		/// <inheritdoc/>
+		public override StringBuilder AppendNullConstructorArg(StringBuilder Builder, bool bIsInitializer)
+		{
+			Builder.Append("TEXT(\"\")");
+			return Builder;
+		}
+
+		/// <inheritdoc/>
+		public override StringBuilder AppendMemberDecl(StringBuilder Builder, IUhtPropertyMemberContext Context, string Name, string NameSuffix, int Tabs)
+		{
+			return AppendMemberDecl(Builder, Context, Name, NameSuffix, Tabs, "FStrPropertyParams");
+		}
+
+		/// <inheritdoc/>
+		public override StringBuilder AppendMemberDef(StringBuilder Builder, IUhtPropertyMemberContext Context, string Name, string NameSuffix, string? Offset, int Tabs)
+		{
+			AppendMemberDefStart(Builder, Context, Name, NameSuffix, Offset, Tabs, "FStrPropertyParams", "UECodeGen_Private::EPropertyGenFlags::Str");
+			AppendMemberDefEnd(Builder, Context, Name, NameSuffix);
+			return Builder;
+		}
+
+		/// <inheritdoc/>
+		public override bool SanitizeDefaultValue(IUhtTokenReader DefaultValueReader, StringBuilder InnerDefaultValue)
+		{
+			if (DefaultValueReader.TryOptional("FString"))
+			{
+				DefaultValueReader.Require('(');
+				StringView Value = DefaultValueReader.GetWrappedConstString();
+				DefaultValueReader.Require(')');
+				InnerDefaultValue.Append(Value);
+			}
+			else
+			{
+				StringView Value = DefaultValueReader.GetWrappedConstString();
+				InnerDefaultValue.Append(Value);
+			}
+			return true;
+		}
+
+		/// <inheritdoc/>
+		public override bool IsSameType(UhtProperty Other)
+		{
+			return Other is UhtStrProperty;
+		}
+
+		/// <inheritdoc/>
+		protected override void ValidateFunctionArgument(UhtFunction Function, UhtValidationOptions Options)
+		{
+			base.ValidateFunctionArgument(Function, Options);
+
+			if (Function.FunctionFlags.HasAnyFlags(EFunctionFlags.Net))
+			{
+				if (!Function.FunctionFlags.HasAnyFlags(EFunctionFlags.NetRequest))
+				{
+					if (this.RefQualifier != UhtPropertyRefQualifier.ConstRef && !this.bIsStaticArray)
+					{
+						this.LogError("Replicated FString parameters must be passed by const reference");
+					}
+				}
+			}
+		}
+
+		/// <inheritdoc/>
+		public override string? GetRigVMType(ref UhtRigVMParameterFlags ParameterFlags)
+		{
+			return this.CppTypeText;
+		}
+
+		[UhtPropertyType(Keyword = "FString")]
+		[UhtPropertyType(Keyword = "FMemoryImageString", Options = UhtPropertyTypeOptions.Immediate)]
+		public static UhtProperty? StrProperty(UhtPropertyResolvePhase ResolvePhase, UhtPropertySettings PropertySettings, IUhtTokenReader TokenReader, UhtToken MatchedToken)
+		{
+			if (!TokenReader.SkipExpectedType(MatchedToken.Value, PropertySettings.PropertyCategory == UhtPropertyCategory.Member))
+			{
+				return null;
+			}
+			UhtStrProperty Out = new UhtStrProperty(PropertySettings);
+			if (Out.PropertyCategory != UhtPropertyCategory.Member)
+			{
+				if (TokenReader.TryOptional('&'))
+				{
+					if (Out.PropertyFlags.HasAnyFlags(EPropertyFlags.ConstParm))
+					{
+						// 'const FString& Foo' came from 'FString' in .uc, no flags
+						Out.PropertyFlags &= ~EPropertyFlags.ConstParm;
+
+						// We record here that we encountered a const reference, because we need to remove that information from flags for code generation purposes.
+						Out.RefQualifier = UhtPropertyRefQualifier.ConstRef;
+					}
+					else
+					{
+						// 'FString& Foo' came from 'out FString' in .uc
+						Out.PropertyFlags |= EPropertyFlags.OutParm;
+
+						// And we record here that we encountered a non-const reference here too.
+						Out.RefQualifier = UhtPropertyRefQualifier.NonConstRef;
+					}
+				}
+			}
+			return Out;
+		}
+	}
+}
