UHT in UBT initial version.  Currently disabled.

diff --git a/Engine/Source/Programs/Shared/EpicGames.UHT/Types/Properties/UhtStructProperty.cs b/Engine/Source/Programs/Shared/EpicGames.UHT/Types/Properties/UhtStructProperty.cs
new file mode 100644
--- /dev/null
+++ b/Engine/Source/Programs/Shared/EpicGames.UHT/Types/Properties/UhtStructProperty.cs
@@ -0,0 +1,412 @@
+// Copyright Epic Games, Inc. All Rights Reserved.
+
+using EpicGames.Core;
+using EpicGames.UHT.Tables;
+using EpicGames.UHT.Tokenizer;
+using EpicGames.UHT.Utils;
+using System.Collections.Generic;
+using System.Text;
+using System.Text.Json.Serialization;
+
+namespace EpicGames.UHT.Types
+{
+	[UnrealHeaderTool]
+	[UhtEngineClass(Name = "StructProperty", IsProperty = true)]
+	public class UhtStructProperty : UhtProperty
+	{
+		/// <inheritdoc/>
+		public override string EngineClassName { get => "StructProperty"; }
+
+		/// <inheritdoc/>
+		protected override string CppTypeText { get => "invalid"; }
+
+		/// <inheritdoc/>
+		protected override string PGetMacroText { get => "STRUCT"; }
+
+		/// <inheritdoc/>
+		protected override UhtPGetArgumentType PGetTypeArgument { get => UhtPGetArgumentType.TypeText; }
+
+		[JsonConverter(typeof(UhtTypeSourceNameJsonConverter<UhtScriptStruct>))]
+		public UhtScriptStruct ScriptStruct { get; set; }
+
+		public UhtStructProperty(UhtPropertySettings PropertySettings, UhtScriptStruct ScriptStruct) : base(PropertySettings)
+		{
+			this.ScriptStruct = ScriptStruct;
+			this.HeaderFile.AddReferencedHeader(ScriptStruct);
+			if (this.ScriptStruct.bHasNoOpConstructor)
+			{
+				this.PropertyCaps |= UhtPropertyCaps.RequiresNullConstructorArg;
+			}
+			if (this.ScriptStruct.MetaData.GetBoolean(UhtNames.BlueprintType))
+			{
+				this.PropertyCaps |= UhtPropertyCaps.CanExposeOnSpawn | UhtPropertyCaps.IsParameterSupportedByBlueprint | UhtPropertyCaps.IsMemberSupportedByBlueprint;
+			}
+			else if (this.ScriptStruct.MetaData.GetBooleanHierarchical(UhtNames.BlueprintType))
+			{
+				this.PropertyCaps |= UhtPropertyCaps.IsParameterSupportedByBlueprint | UhtPropertyCaps.IsMemberSupportedByBlueprint;
+			}
+		}
+
+		/// <inheritdoc/>
+		protected override bool ResolveSelf(UhtResolvePhase Phase)
+		{
+			bool bResults = base.ResolveSelf(Phase);
+			switch (Phase)
+			{
+				case UhtResolvePhase.Final:
+					if (ScanForInstancedReferenced(true))
+					{
+						this.PropertyFlags |= EPropertyFlags.ContainsInstancedReference;
+					}
+					break;
+			}
+			return bResults;
+		}
+
+		/// <inheritdoc/>
+		public override bool ScanForInstancedReferenced(bool bDeepScan)
+		{
+			return this.ScriptStruct.ScanForInstancedReferenced(bDeepScan);
+		}
+
+		/// <inheritdoc/>
+		public override void CollectReferencesInternal(IUhtReferenceCollector Collector, bool bTemplateProperty)
+		{
+			Collector.AddCrossModuleReference(this.ScriptStruct, true);
+		}
+
+		/// <inheritdoc/>
+		public override string? GetForwardDeclarations()
+		{
+			return !this.ScriptStruct.bIsCoreType ? $"struct {this.ScriptStruct.SourceName};" : null;
+		}
+
+		/// <inheritdoc/>
+		public override IEnumerable<UhtType> EnumerateReferencedTypes()
+		{
+			yield return this.ScriptStruct;
+		}
+
+		/// <inheritdoc/>
+		public override StringBuilder AppendText(StringBuilder Builder, UhtPropertyTextType TextType, bool bIsTemplateArgument)
+		{
+			switch (TextType)
+			{
+				default:
+					Builder.Append(this.ScriptStruct.SourceName);
+					break;
+			}
+			return Builder;
+		}
+
+		/// <inheritdoc/>
+		public override StringBuilder AppendMemberDecl(StringBuilder Builder, IUhtPropertyMemberContext Context, string Name, string NameSuffix, int Tabs)
+		{
+			return AppendMemberDecl(Builder, Context, Name, NameSuffix, Tabs, "FStructPropertyParams");
+		}
+
+		/// <inheritdoc/>
+		public override StringBuilder AppendMemberDef(StringBuilder Builder, IUhtPropertyMemberContext Context, string Name, string NameSuffix, string? Offset, int Tabs)
+		{
+			AppendMemberDefStart(Builder, Context, Name, NameSuffix, Offset, Tabs, "FStructPropertyParams", "UECodeGen_Private::EPropertyGenFlags::Struct");
+			AppendMemberDefRef(Builder, Context, this.ScriptStruct, true);
+			AppendMemberDefEnd(Builder, Context, Name, NameSuffix);
+			return Builder;
+		}
+
+		/// <inheritdoc/>
+		public override void AppendObjectHashes(StringBuilder Builder, int StartingLength, IUhtPropertyMemberContext Context)
+		{
+			Builder.AppendObjectHash(StartingLength, this, Context, this.ScriptStruct);
+		}
+
+		/// <inheritdoc/>
+		public override StringBuilder AppendNullConstructorArg(StringBuilder Builder, bool bIsInitializer)
+		{
+			bool bHasNoOpConstructor = this.ScriptStruct.bHasNoOpConstructor;
+			if (bIsInitializer && bHasNoOpConstructor)
+			{
+				Builder.Append("ForceInit");
+			}
+			else
+			{
+				Builder.AppendPropertyText(this, UhtPropertyTextType.Construction);
+				if (bHasNoOpConstructor)
+				{
+					Builder.Append("(ForceInit)");
+				}
+				else
+				{
+					Builder.Append("()");
+				}
+			}
+			return Builder;
+		}
+
+		/// <inheritdoc/>
+		public override bool SanitizeDefaultValue(IUhtTokenReader DefaultValueReader, StringBuilder InnerDefaultValue)
+		{
+			UhtStructDefaultValue StructDefaultValue;
+			if (!UhtStructDefaultValueTable.Instance.TryGet(this.ScriptStruct.SourceName, out StructDefaultValue))
+			{
+				StructDefaultValue = UhtStructDefaultValueTable.Instance.Default;
+			}
+			return StructDefaultValue.Delegate(this, DefaultValueReader, InnerDefaultValue);
+		}
+
+		/// <inheritdoc/>
+		public override bool IsSameType(UhtProperty Other)
+		{
+			if (Other is UhtStructProperty OtherObject)
+			{
+				return this.ScriptStruct == OtherObject.ScriptStruct;
+			}
+			return false;
+		}
+
+		/// <inheritdoc/>
+		protected override void ValidateFunctionArgument(UhtFunction Function, UhtValidationOptions Options)
+		{
+			base.ValidateFunctionArgument(Function, Options);
+
+			if (Function.FunctionFlags.HasAnyFlags(EFunctionFlags.Net))
+			{
+				if (!Function.FunctionFlags.HasAnyFlags(EFunctionFlags.NetRequest | EFunctionFlags.NetResponse))
+				{
+					this.Session.ValidateScriptStructOkForNet(this, this.ScriptStruct);
+				}
+			}
+		}
+
+		/// <inheritdoc/>
+		protected override void ValidateMember(UhtStruct Struct, UhtValidationOptions Options)
+		{
+			if (this.PropertyFlags.HasAnyFlags(EPropertyFlags.Net))
+			{
+				this.Session.ValidateScriptStructOkForNet(this, this.ScriptStruct);
+			}
+		}
+
+		///<inheritdoc/>
+		public override bool ValidateStructPropertyOkForNet(UhtProperty ReferencingProperty)
+		{
+			return ReferencingProperty.Session.ValidateScriptStructOkForNet(ReferencingProperty, this.ScriptStruct);
+		}
+
+		/// <inheritdoc/>
+		public override string? GetRigVMType(ref UhtRigVMParameterFlags ParameterFlags)
+		{
+			return this.ScriptStruct.SourceName;
+		}
+
+		#region Structure default value sanitizers
+		[UhtStructDefaultValue(Name = "FVector")]
+		public static bool VectorStructDefaultValue(UhtStructProperty Property, IUhtTokenReader DefaultValueReader, StringBuilder InnerDefaultValue)
+		{
+			string Format = "{0:F6},{1:F6},{2:F6}";
+
+			DefaultValueReader.Require("FVector");
+			if (DefaultValueReader.TryOptional("::"))
+			{
+				switch (DefaultValueReader.GetIdentifier().Value.ToString())
+				{
+					case "ZeroVector": return true;
+					case "UpVector": InnerDefaultValue.AppendFormat(Format, 0, 0, 1); return true;
+					case "ForwardVector": InnerDefaultValue.AppendFormat(Format, 1, 0, 0); return true;
+					case "RightVector": InnerDefaultValue.AppendFormat(Format, 0, 1, 0); return true;
+					default: return false;
+				}
+			}
+			else
+			{
+				DefaultValueReader.Require("(");
+				if (DefaultValueReader.TryOptional("ForceInit"))
+				{
+				}
+				else
+				{
+					double X, Y, Z;
+					X = Y = Z = DefaultValueReader.GetConstDoubleExpression();
+					if (DefaultValueReader.TryOptional(','))
+					{
+						Y = DefaultValueReader.GetConstDoubleExpression();
+						DefaultValueReader.Require(',');
+						Z = DefaultValueReader.GetConstDoubleExpression();
+					}
+					InnerDefaultValue.AppendFormat(Format, X, Y, Z);
+				}
+				DefaultValueReader.Require(")");
+				return true;
+			}
+		}
+
+		[UhtStructDefaultValue(Name = "FRotator")]
+		public static bool RotatorStructDefaultValue(UhtStructProperty Property, IUhtTokenReader DefaultValueReader, StringBuilder InnerDefaultValue)
+		{
+			DefaultValueReader.Require("FRotator");
+			if (DefaultValueReader.TryOptional("::"))
+			{
+				switch (DefaultValueReader.GetIdentifier().Value.ToString())
+				{
+					case "ZeroRotator": return true;
+					default: return false;
+				}
+			}
+			else
+			{
+				DefaultValueReader.Require("(");
+				if (DefaultValueReader.TryOptional("ForceInit"))
+				{
+				}
+				else
+				{
+					double X = DefaultValueReader.GetConstDoubleExpression();
+					DefaultValueReader.Require(',');
+					double Y = DefaultValueReader.GetConstDoubleExpression();
+					DefaultValueReader.Require(',');
+					double Z = DefaultValueReader.GetConstDoubleExpression();
+					InnerDefaultValue.AppendFormat("{0:F6},{1:F6},{2:F6}", X, Y, Z);
+				}
+				DefaultValueReader.Require(")");
+				return true;
+			}
+		}
+
+		[UhtStructDefaultValue(Name = "FVector2D")]
+		public static bool Vector2DStructDefaultValue(UhtStructProperty Property, IUhtTokenReader DefaultValueReader, StringBuilder InnerDefaultValue)
+		{
+			string Format = "(X={0:F3},Y={1:F3})";
+
+			DefaultValueReader.Require("FVector2D");
+			if (DefaultValueReader.TryOptional("::"))
+			{
+				switch (DefaultValueReader.GetIdentifier().Value.ToString())
+				{
+					case "ZeroVector": return true;
+					case "UnitVector": InnerDefaultValue.AppendFormat(Format, 1.0, 1.0); return true;
+					default: return false;
+				}
+			}
+			else
+			{
+				DefaultValueReader.Require("(");
+				if (DefaultValueReader.TryOptional("ForceInit"))
+				{
+				}
+				else
+				{
+					double X = DefaultValueReader.GetConstDoubleExpression();
+					DefaultValueReader.Require(',');
+					double Y = DefaultValueReader.GetConstDoubleExpression();
+					InnerDefaultValue.AppendFormat(Format, X, Y);
+				}
+				DefaultValueReader.Require(")");
+				return true;
+			}
+		}
+
+		[UhtStructDefaultValue(Name = "FLinearColor")]
+		public static bool LinearColorStructDefaultValue(UhtStructProperty Property, IUhtTokenReader DefaultValueReader, StringBuilder InnerDefaultValue)
+		{
+			string Format = "(R={0:F6},G={1:F6},B={2:F6},A={3:F6})";
+
+			DefaultValueReader.Require("FLinearColor");
+			if (DefaultValueReader.TryOptional("::"))
+			{
+				switch (DefaultValueReader.GetIdentifier().Value.ToString())
+				{
+					case "White": InnerDefaultValue.AppendFormat(Format, 1.0, 1.0, 1.0, 1.0); ; return true;
+					case "Gray": InnerDefaultValue.AppendFormat(Format, 0.5, 0.5, 0.5, 1.0); ; return true;
+					case "Black": InnerDefaultValue.AppendFormat(Format, 0.0, 0.0, 0.0, 1.0); ; return true;
+					case "Transparent": InnerDefaultValue.AppendFormat(Format, 0.0, 0.0, 0.0, 0.0); ; return true;
+					case "Red": InnerDefaultValue.AppendFormat(Format, 1.0, 0.0, 0.0, 1.0); ; return true;
+					case "Green": InnerDefaultValue.AppendFormat(Format, 0.0, 1.0, 0.0, 1.0); ; return true;
+					case "Blue": InnerDefaultValue.AppendFormat(Format, 0.0, 0.0, 1.0, 1.0); ; return true;
+					case "Yellow": InnerDefaultValue.AppendFormat(Format, 1.0, 1.0, 0.0, 1.0); ; return true;
+					default: return false;
+				}
+			}
+			else
+			{
+				DefaultValueReader.Require("(");
+				if (DefaultValueReader.TryOptional("ForceInit"))
+				{
+				}
+				else
+				{
+					double R = DefaultValueReader.GetConstDoubleExpression();
+					DefaultValueReader.Require(',');
+					double G = DefaultValueReader.GetConstDoubleExpression();
+					DefaultValueReader.Require(',');
+					double B = DefaultValueReader.GetConstDoubleExpression();
+					double A = 1.0;
+					if (DefaultValueReader.TryOptional(','))
+					{
+						A = DefaultValueReader.GetConstDoubleExpression();
+					}
+					InnerDefaultValue.AppendFormat(Format, R, G, B, A);
+				}
+				DefaultValueReader.Require(")");
+				return true;
+			}
+		}
+
+		[UhtStructDefaultValue(Name = "FColor")]
+		public static bool ColorStructDefaultValue(UhtStructProperty Property, IUhtTokenReader DefaultValueReader, StringBuilder InnerDefaultValue)
+		{
+			string Format = "(R={0},G={1},B={2},A={3})";
+
+			DefaultValueReader.Require("FLinearColor");
+			if (DefaultValueReader.TryOptional("::"))
+			{
+				switch (DefaultValueReader.GetIdentifier().Value.ToString())
+				{
+					case "White": InnerDefaultValue.AppendFormat(Format, 255, 255, 255, 255); ; return true;
+					case "Black": InnerDefaultValue.AppendFormat(Format, 0, 0, 0, 255); ; return true;
+					case "Red": InnerDefaultValue.AppendFormat(Format, 255, 0, 0, 255); ; return true;
+					case "Green": InnerDefaultValue.AppendFormat(Format, 0, 255, 0, 255); ; return true;
+					case "Blue": InnerDefaultValue.AppendFormat(Format, 0, 0, 255, 255); ; return true;
+					case "Yellow": InnerDefaultValue.AppendFormat(Format, 255, 255, 0, 255); ; return true;
+					case "Cyan": InnerDefaultValue.AppendFormat(Format, 0, 255, 255, 255); ; return true;
+					case "Magenta": InnerDefaultValue.AppendFormat(Format, 255, 0, 255, 255); ; return true;
+					default: return false;
+				}
+			}
+			else
+			{
+				DefaultValueReader.Require("(");
+				if (DefaultValueReader.TryOptional("ForceInit"))
+				{
+				}
+				else
+				{
+					int R = DefaultValueReader.GetConstIntExpression();
+					DefaultValueReader.Require(',');
+					int G = DefaultValueReader.GetConstIntExpression();
+					DefaultValueReader.Require(',');
+					int B = DefaultValueReader.GetConstIntExpression();
+					int A = 255;
+					if (DefaultValueReader.TryOptional(','))
+					{
+						A = DefaultValueReader.GetConstIntExpression();
+					}
+					InnerDefaultValue.AppendFormat(Format, R, G, B, A);
+				}
+				DefaultValueReader.Require(")");
+				return true;
+			}
+		}
+
+		[UhtStructDefaultValue(Options = UhtStructDefaultValueOptions.Default)]
+		public static bool DefaultStructDefaultValue(UhtStructProperty Property, IUhtTokenReader DefaultValueReader, StringBuilder InnerDefaultValue)
+		{
+			DefaultValueReader
+				.Require(Property.ScriptStruct.SourceName)
+				.Require('(')
+				.Require(')');
+			InnerDefaultValue.Append("()");
+			return true;
+		}
+		#endregion
+	}
+}
diff --git a/Engine/Source/Programs/Shared/EpicGames.UHT/Types/Properties/UhtTextProperty.cs b/Engine/Source/Programs/Shared/EpicGames.UHT/Types/Properties/UhtTextProperty.cs
new file mode 100644
--- /dev/null
+++ b/Engine/Source/Programs/Shared/EpicGames.UHT/Types/Properties/UhtTextProperty.cs
@@ -0,0 +1,688 @@
+// Copyright Epic Games, Inc. All Rights Reserved.
+
+using EpicGames.Core;
+using EpicGames.UHT.Tables;
+using EpicGames.UHT.Tokenizer;
+using EpicGames.UHT.Utils;
+using System;
+using System.Text;
+
+namespace EpicGames.UHT.Types
+{
+	[UnrealHeaderTool]
+	[UhtEngineClass(Name = "TextProperty", IsProperty = true)]
+	public class UhtTextProperty : UhtProperty
+	{
+		/// <inheritdoc/>
+		public override string EngineClassName { get => "TextProperty"; }
+
+		/// <inheritdoc/>
+		protected override string CppTypeText { get => "FText"; }
+
+		/// <inheritdoc/>
+		protected override string PGetMacroText { get => "PROPERTY"; }
+
+		/// <inheritdoc/>
+		protected override UhtPGetArgumentType PGetTypeArgument { get => UhtPGetArgumentType.EngineClass; }
+
+		public UhtTextProperty(UhtPropertySettings PropertySettings) : base(PropertySettings)
+		{
+			this.PropertyCaps |= UhtPropertyCaps.PassCppArgsByRef | UhtPropertyCaps.CanExposeOnSpawn | UhtPropertyCaps.IsParameterSupportedByBlueprint | UhtPropertyCaps.IsMemberSupportedByBlueprint;
+			this.PropertyCaps &= ~(UhtPropertyCaps.CanBeContainerKey);
+		}
+
+		/// <inheritdoc/>
+		public override StringBuilder AppendMemberDecl(StringBuilder Builder, IUhtPropertyMemberContext Context, string Name, string NameSuffix, int Tabs)
+		{
+			return AppendMemberDecl(Builder, Context, Name, NameSuffix, Tabs, "FTextPropertyParams");
+		}
+
+		/// <inheritdoc/>
+		public override StringBuilder AppendMemberDef(StringBuilder Builder, IUhtPropertyMemberContext Context, string Name, string NameSuffix, string? Offset, int Tabs)
+		{
+			AppendMemberDefStart(Builder, Context, Name, NameSuffix, Offset, Tabs, "FTextPropertyParams", "UECodeGen_Private::EPropertyGenFlags::Text");
+			AppendMemberDefEnd(Builder, Context, Name, NameSuffix);
+			return Builder;
+		}
+
+		/// <inheritdoc/>
+		public override StringBuilder AppendNullConstructorArg(StringBuilder Builder, bool bIsInitializer)
+		{
+			Builder.Append("FText::GetEmpty()");
+			return Builder;
+		}
+
+		/// <inheritdoc/>
+		public override bool SanitizeDefaultValue(IUhtTokenReader DefaultValueReader, StringBuilder InnerDefaultValue)
+		{
+			if (DefaultValueReader.TryOptional("FText"))
+			{
+				if (DefaultValueReader.TryOptional('('))
+				{
+					DefaultValueReader.Require(')');
+					return true;
+				}
+				else if (DefaultValueReader.TryOptional("::"))
+				{
+					// Handle legacy cases of FText::FromString being used as default values
+					// These should be replaced with INVTEXT as FText::FromString can produce inconsistent keys
+					if (DefaultValueReader.TryOptional("FromString"))
+					{
+						DefaultValueReader.Require('(');
+						StringView Value = DefaultValueReader.GetWrappedConstString();
+						DefaultValueReader.Require(')');
+						InnerDefaultValue.Append("\"").Append(Value).Append("\"");
+						this.LogWarning("FText::FromString should be replaced with INVTEXT for default parameter values");
+					}
+					else
+					{
+						DefaultValueReader.Require("GetEmpty").Require('(').Require(')');
+					}
+					return true;
+				}
+				else
+				{
+					return false;
+				}
+			}
+
+			UhtToken Token = DefaultValueReader.GetToken();
+
+			if (Token.IsIdentifier("LOCTEXT"))
+			{
+				//ETSTODO - Add a test case for this
+				this.LogError($"LOCTEXT default parameter values are not supported; use NSLOCTEXT instead: {this.SourceName}");
+				return false;
+			}
+			else
+			{
+				return SanitizeDefaultValue(this, DefaultValueReader, ref Token, InnerDefaultValue, false);
+			}
+		}
+
+		/// <inheritdoc/>
+		public override bool IsSameType(UhtProperty Other)
+		{
+			return Other is UhtTextProperty;
+		}
+
+		/// <inheritdoc/>
+		public override string? GetRigVMType(ref UhtRigVMParameterFlags ParameterFlags)
+		{
+			return this.CppTypeText;
+		}
+
+		#region Parsing keywords and default parsers		
+		[UhtPropertyType(Keyword = "FText", Options = UhtPropertyTypeOptions.Simple | UhtPropertyTypeOptions.Immediate)]
+		private static UhtProperty? TextProperty(UhtPropertyResolvePhase ResolvePhase, UhtPropertySettings PropertySettings, IUhtTokenReader TokenReader, UhtToken MatchedToken)
+		{
+			return new UhtTextProperty(PropertySettings);
+		}
+
+		[UhtPropertyType(Keyword = "Text", Options = UhtPropertyTypeOptions.Simple | UhtPropertyTypeOptions.Immediate | UhtPropertyTypeOptions.CaseInsensitive)]
+		public static UhtProperty? MissingPrefixTextProperty(UhtPropertyResolvePhase ResolvePhase, UhtPropertySettings PropertySettings, IUhtTokenReader TokenReader, UhtToken MatchedToken)
+		{
+			throw new UhtException(TokenReader, "'Text' is missing a prefix, expecting 'FText'");
+		}
+
+		[UhtLocTextDefaultValue(Name = "INVTEXT")]
+		private static bool InvTextDefaultValue(UhtTextProperty Property, IUhtTokenReader DefaultValueReader, ref UhtToken MacroToken, StringBuilder InnerDefaultValue)
+		{
+			DefaultValueReader.Require('(');
+			StringView String = DefaultValueReader.GetConstQuotedString();
+			DefaultValueReader.Require(')');
+			InnerDefaultValue.Append("INVTEXT(").Append(String).Append(')');
+			return true;
+		}
+
+		[UhtLocTextDefaultValue(Name = "LOCTEXT")]
+		private static bool LocTextDefaultValue(UhtTextProperty Property, IUhtTokenReader DefaultValueReader, ref UhtToken MacroToken, StringBuilder InnerDefaultValue)
+		{
+			Property.LogError($"LOCTEXT default parameter values are not supported; use NSLOCTEXT instead: {Property.SourceName}");
+			return false;
+		}
+
+		[UhtLocTextDefaultValue(Name = "NSLOCTEXT")]
+		private static bool NsLocTextDefaultValue(UhtTextProperty Property, IUhtTokenReader DefaultValueReader, ref UhtToken MacroToken, StringBuilder InnerDefaultValue)
+		{
+			DefaultValueReader.Require('(');
+			StringView NamespaceString = DefaultValueReader.GetConstQuotedString();
+			DefaultValueReader.Require(',');
+			StringView KeyString = DefaultValueReader.GetConstQuotedString();
+			DefaultValueReader.Require(',');
+			StringView SourceString = DefaultValueReader.GetConstQuotedString();
+			DefaultValueReader.Require(')');
+
+			// Strip out the package name
+			ReadOnlySpan<char> StrippedNS = NamespaceString.Span;
+			if (StrippedNS.Length > 1)
+			{
+				StrippedNS = StrippedNS.Slice(1, StrippedNS.Length - 2).Trim();
+				if (StrippedNS.Length > 0)
+				{
+					if (StrippedNS[StrippedNS.Length - 1] == ']')
+					{
+						int Index = StrippedNS.LastIndexOf('[');
+						if (Index != -1)
+						{
+							StrippedNS = StrippedNS.Slice(0, Index).TrimEnd();
+						}
+					}
+				}
+			}
+
+			InnerDefaultValue.Append("NSLOCTEXT(\"").Append(StrippedNS).Append("\", ").Append(KeyString).Append(", ").Append(SourceString).Append(')');
+			return true;
+		}
+
+		[UhtLocTextDefaultValue(Name = "LOCTABLE")]
+		private static bool LocTableDefaultValue(UhtTextProperty Property, IUhtTokenReader DefaultValueReader, ref UhtToken MacroToken, StringBuilder InnerDefaultValue)
+		{
+			DefaultValueReader.Require('(');
+			StringView NamespaceString = DefaultValueReader.GetConstQuotedString();
+			DefaultValueReader.Require(',');
+			StringView KeyString = DefaultValueReader.GetConstQuotedString();
+			DefaultValueReader.Require(')');
+
+			//ETSTODO - Validate the namespace string?
+			InnerDefaultValue.Append("LOCTABLE(").Append(NamespaceString).Append(", ").Append(KeyString).Append(')');
+			return true;
+		}
+
+		[UhtLocTextDefaultValue(Name = "LOCGEN_FORMAT_NAMED")]
+		private static bool LocGenFormatNamedDefaultValue(UhtTextProperty Property, IUhtTokenReader DefaultValueReader, ref UhtToken MacroToken, StringBuilder InnerDefaultValue)
+		{
+			InnerDefaultValue.Append("LOCGEN_FORMAT_NAMED(");
+			DefaultValueReader.Require('(');
+			if (!SanitizeDefaultValue(Property, DefaultValueReader, InnerDefaultValue, false))
+			{
+				return false;
+			}
+
+			// RequireList assumes that we have already parsed the ','.  So if it is there, consume it.
+			DefaultValueReader.TryOptional(',');
+
+			bool bSuccess = true;
+			DefaultValueReader.RequireList(')', ',', false, () =>
+			{
+				StringView Value = DefaultValueReader.GetWrappedConstString();
+				DefaultValueReader.Require(',');
+				InnerDefaultValue.Append(", \"").Append(Value).Append("\", ");
+				if (!SanitizeDefaultValue(Property, DefaultValueReader, InnerDefaultValue, true))
+				{
+					bSuccess = false;
+				}
+			});
+			InnerDefaultValue.Append(')');
+			return bSuccess;
+		}
+
+		[UhtLocTextDefaultValue(Name = "LOCGEN_FORMAT_ORDERED")]
+		private static bool LocGenFormatOrderedDefaultValue(UhtTextProperty Property, IUhtTokenReader DefaultValueReader, ref UhtToken MacroToken, StringBuilder InnerDefaultValue)
+		{
+			InnerDefaultValue.Append("LOCGEN_FORMAT_ORDERED(");
+			DefaultValueReader.Require('(');
+			if (!SanitizeDefaultValue(Property, DefaultValueReader, InnerDefaultValue, false))
+			{
+				return false;
+			}
+
+			// RequireList assumes that we have already parsed the ','.  So if it is there, consume it.
+			DefaultValueReader.TryOptional(',');
+
+			bool bSuccess = true;
+			DefaultValueReader.RequireList(')', ',', false, () =>
+			{
+				InnerDefaultValue.Append(", ");
+				if (!SanitizeDefaultValue(Property, DefaultValueReader, InnerDefaultValue, true))
+				{
+					bSuccess = false;
+				}
+			});
+			InnerDefaultValue.Append(')');
+			return bSuccess;
+		}
+
+		[UhtLocTextDefaultValue(Name = "LOCGEN_NUMBER")]
+		private static bool LocGenNumberDefaultValue(UhtTextProperty Property, IUhtTokenReader DefaultValueReader, ref UhtToken MacroToken, StringBuilder InnerDefaultValue)
+		{
+			return LocGenNumberOrPercentDefaultValue(Property, DefaultValueReader, ref MacroToken, InnerDefaultValue, NumberStyle.None);
+		}
+
+		[UhtLocTextDefaultValue(Name = "LOCGEN_NUMBER_GROUPED")]
+		private static bool LocGenNumberGroupedDefaultValue(UhtTextProperty Property, IUhtTokenReader DefaultValueReader, ref UhtToken MacroToken, StringBuilder InnerDefaultValue)
+		{
+			return LocGenNumberOrPercentDefaultValue(Property, DefaultValueReader, ref MacroToken, InnerDefaultValue, NumberStyle.Grouped);
+		}
+
+		[UhtLocTextDefaultValue(Name = "LOCGEN_NUMBER_UNGROUPED")]
+		private static bool LocGenNumberUngroupedDefaultValue(UhtTextProperty Property, IUhtTokenReader DefaultValueReader, ref UhtToken MacroToken, StringBuilder InnerDefaultValue)
+		{
+			return LocGenNumberOrPercentDefaultValue(Property, DefaultValueReader, ref MacroToken, InnerDefaultValue, NumberStyle.Ungrouped);
+		}
+
+		[UhtLocTextDefaultValue(Name = "LOCGEN_NUMBER_CUSTOM")]
+		private static bool LocGenNumberCustomDefaultValue(UhtTextProperty Property, IUhtTokenReader DefaultValueReader, ref UhtToken MacroToken, StringBuilder InnerDefaultValue)
+		{
+			return LocGenNumberOrPercentDefaultValue(Property, DefaultValueReader, ref MacroToken, InnerDefaultValue, NumberStyle.Custom);
+		}
+
+		[UhtLocTextDefaultValue(Name = "LOCGEN_PERCENT")]
+		private static bool LocGenPercentDefaultValue(UhtTextProperty Property, IUhtTokenReader DefaultValueReader, ref UhtToken MacroToken, StringBuilder InnerDefaultValue)
+		{
+			return LocGenNumberOrPercentDefaultValue(Property, DefaultValueReader, ref MacroToken, InnerDefaultValue, NumberStyle.None);
+		}
+
+		[UhtLocTextDefaultValue(Name = "LOCGEN_PERCENT_GROUPED")]
+		private static bool LocGenPercentGroupedDefaultValue(UhtTextProperty Property, IUhtTokenReader DefaultValueReader, ref UhtToken MacroToken, StringBuilder InnerDefaultValue)
+		{
+			return LocGenNumberOrPercentDefaultValue(Property, DefaultValueReader, ref MacroToken, InnerDefaultValue, NumberStyle.Grouped);
+		}
+
+		[UhtLocTextDefaultValue(Name = "LOCGEN_PERCENT_UNGROUPED")]
+		private static bool LocGenPercentUngroupedDefaultValue(UhtTextProperty Property, IUhtTokenReader DefaultValueReader, ref UhtToken MacroToken, StringBuilder InnerDefaultValue)
+		{
+			return LocGenNumberOrPercentDefaultValue(Property, DefaultValueReader, ref MacroToken, InnerDefaultValue, NumberStyle.Ungrouped);
+		}
+
+		[UhtLocTextDefaultValue(Name = "LOCGEN_PERCENT_CUSTOM")]
+		private static bool LocGenPercentCustomDefaultValue(UhtTextProperty Property, IUhtTokenReader DefaultValueReader, ref UhtToken MacroToken, StringBuilder InnerDefaultValue)
+		{
+			return LocGenNumberOrPercentDefaultValue(Property, DefaultValueReader, ref MacroToken, InnerDefaultValue, NumberStyle.Custom);
+		}
+
+		[UhtLocTextDefaultValue(Name = "LOCGEN_CURRENCY")]
+		private static bool LocGenCurrencyDefaultValue(UhtTextProperty Property, IUhtTokenReader DefaultValueReader, ref UhtToken MacroToken, StringBuilder InnerDefaultValue)
+		{
+			DefaultValueReader.Require('(');
+
+			double BaseValue;
+			UhtToken Token = DefaultValueReader.GetToken();
+			if (!Token.GetConstDouble(out BaseValue))
+			{
+				return false;
+			}
+
+			DefaultValueReader.Require(',');
+			StringView CurrencyCode = DefaultValueReader.GetConstQuotedString();
+			DefaultValueReader.Require(',');
+			StringView CultureNameString = DefaultValueReader.GetConstQuotedString();
+			DefaultValueReader.Require(')');
+
+			// With UHT, we end up outputting just the integer part of the specified value.  It is using default locale information in UHT
+			InnerDefaultValue.Append("LOCGEN_CURRENCY(").Append((int)BaseValue).Append(", ").Append(CurrencyCode).Append(", ").Append(CultureNameString).Append(")");
+			return true;
+		}
+
+		[UhtLocTextDefaultValue(Name = "LOCGEN_DATE_LOCAL")]
+		private static bool LocGenDateLocalDefaultValue(UhtTextProperty Property, IUhtTokenReader DefaultValueReader, ref UhtToken MacroToken, StringBuilder InnerDefaultValue)
+		{
+			return LocGenDateTimeDefaultValue(Property, DefaultValueReader, ref MacroToken, InnerDefaultValue, true, false, false, false);
+		}
+
+		[UhtLocTextDefaultValue(Name = "LOCGEN_DATE_UTC")]
+		private static bool LocGenDateUtcDefaultValue(UhtTextProperty Property, IUhtTokenReader DefaultValueReader, ref UhtToken MacroToken, StringBuilder InnerDefaultValue)
+		{
+			return LocGenDateTimeDefaultValue(Property, DefaultValueReader, ref MacroToken, InnerDefaultValue, true, false, true, false);
+		}
+
+		[UhtLocTextDefaultValue(Name = "LOCGEN_TIME_LOCAL")]
+		private static bool LocGenTimeLocalDefaultValue(UhtTextProperty Property, IUhtTokenReader DefaultValueReader, ref UhtToken MacroToken, StringBuilder InnerDefaultValue)
+		{
+			return LocGenDateTimeDefaultValue(Property, DefaultValueReader, ref MacroToken, InnerDefaultValue, false, true, false, false);
+		}
+
+		[UhtLocTextDefaultValue(Name = "LOCGEN_TIME_UTC")]
+		private static bool LocGenTimeUtcDefaultValue(UhtTextProperty Property, IUhtTokenReader DefaultValueReader, ref UhtToken MacroToken, StringBuilder InnerDefaultValue)
+		{
+			return LocGenDateTimeDefaultValue(Property, DefaultValueReader, ref MacroToken, InnerDefaultValue, false, true, true, false);
+		}
+
+		[UhtLocTextDefaultValue(Name = "LOCGEN_DATETIME_LOCAL")]
+		private static bool LocGenDateTimeLocalDefaultValue(UhtTextProperty Property, IUhtTokenReader DefaultValueReader, ref UhtToken MacroToken, StringBuilder InnerDefaultValue)
+		{
+			return LocGenDateTimeDefaultValue(Property, DefaultValueReader, ref MacroToken, InnerDefaultValue, true, true, false, false);
+		}
+
+		[UhtLocTextDefaultValue(Name = "LOCGEN_DATETIME_UTC")]
+		private static bool LocGenDateTimeUtcDefaultValue(UhtTextProperty Property, IUhtTokenReader DefaultValueReader, ref UhtToken MacroToken, StringBuilder InnerDefaultValue)
+		{
+			return LocGenDateTimeDefaultValue(Property, DefaultValueReader, ref MacroToken, InnerDefaultValue, true, true, true, false);
+		}
+
+		[UhtLocTextDefaultValue(Name = "LOCGEN_DATETIME_CUSTOM_LOCAL")]
+		private static bool LocGenDateTimeCustomLocalDefaultValue(UhtTextProperty Property, IUhtTokenReader DefaultValueReader, ref UhtToken MacroToken, StringBuilder InnerDefaultValue)
+		{
+			return LocGenDateTimeDefaultValue(Property, DefaultValueReader, ref MacroToken, InnerDefaultValue, true, true, false, true);
+		}
+
+		[UhtLocTextDefaultValue(Name = "LOCGEN_DATETIME_CUSTOM_UTC")]
+		private static bool LocGenDateTimeCustomUtcDefaultValue(UhtTextProperty Property, IUhtTokenReader DefaultValueReader, ref UhtToken MacroToken, StringBuilder InnerDefaultValue)
+		{
+			return LocGenDateTimeDefaultValue(Property, DefaultValueReader, ref MacroToken, InnerDefaultValue, true, true, true, true);
+		}
+
+		[UhtLocTextDefaultValue(Name = "LOCGEN_TOUPPER")]
+		private static bool LocGenToUpperDefaultValue(UhtTextProperty Property, IUhtTokenReader DefaultValueReader, ref UhtToken MacroToken, StringBuilder InnerDefaultValue)
+		{
+			return LocGenTransformDefaultValue(Property, DefaultValueReader, ref MacroToken, InnerDefaultValue);
+		}
+
+		[UhtLocTextDefaultValue(Name = "LOCGEN_TOLOWER")]
+		private static bool LocGenToLowerDefaultValue(UhtTextProperty Property, IUhtTokenReader DefaultValueReader, ref UhtToken MacroToken, StringBuilder InnerDefaultValue)
+		{
+			return LocGenTransformDefaultValue(Property, DefaultValueReader, ref MacroToken, InnerDefaultValue);
+		}
+		#endregion
+
+		#region Default value helper methods
+		private static bool SanitizeDefaultValue(UhtTextProperty Property, IUhtTokenReader DefaultValueReader, StringBuilder InnerDefaultValue, bool bAllowNumerics)
+		{
+			UhtToken Token = DefaultValueReader.GetToken();
+			return SanitizeDefaultValue(Property, DefaultValueReader, ref Token, InnerDefaultValue, bAllowNumerics);
+		}
+
+		private static bool SanitizeDefaultValue(UhtTextProperty Property, IUhtTokenReader DefaultValueReader, ref UhtToken Token, StringBuilder InnerDefaultValue, bool bAllowNumerics)
+		{
+			switch (Token.TokenType)
+			{
+				case UhtTokenType.Identifier:
+					UhtLocTextDefaultValue LocTextDefaultValue;
+					if (UhtLocTextDefaultValueTable.Instance.TryGet(Token.Value, out LocTextDefaultValue))
+					{
+						if (LocTextDefaultValue.Delegate(Property, DefaultValueReader, ref Token, InnerDefaultValue))
+						{
+							return true;
+						}
+					}
+					return false;
+
+				case UhtTokenType.StringConst:
+					InnerDefaultValue.Append(Token.Value);
+					return true;
+
+				case UhtTokenType.DecimalConst:
+					if (bAllowNumerics)
+					{
+						FormatDecimal(ref Token, InnerDefaultValue);
+						return true;
+					}
+					return false;
+
+				case UhtTokenType.FloatConst:
+					if (bAllowNumerics)
+					{
+						FormatFloat(ref Token, InnerDefaultValue);
+						return true;
+					}
+					return false;
+
+				default:
+					return false;
+			}
+		}
+
+		private static void FormatDecimal(ref UhtToken Token, StringBuilder InnerDefaultValue)
+		{
+			ReadOnlySpan<char> Span = Token.Value.Span;
+			int SuffixStart = Span.Length;
+			bool bIsUnsigned = false;
+			while (SuffixStart > 0)
+			{
+				char C = Span[SuffixStart - 1];
+				if (UhtFCString.IsUnsignedMarker(C))
+				{
+					bIsUnsigned = true;
+					SuffixStart--;
+				}
+				else if (UhtFCString.IsLongMarker(C))
+				{
+					SuffixStart--;
+				}
+				else
+				{
+					break;
+				}
+			}
+			if (Token.GetConstLong(out long Value))
+			{
+				if (bIsUnsigned)
+				{
+					InnerDefaultValue.Append((ulong)Value).Append('u');
+				}
+				else
+				{
+					InnerDefaultValue.Append(Value);
+				}
+			}
+		}
+
+		private static void FormatFloat(ref UhtToken Token, StringBuilder InnerDefaultValue)
+		{
+			char C = Token.Value.Span[Token.Value.Length - 1];
+			bool bIsFloat = UhtFCString.IsFloatMarker(C);
+			if (Token.GetConstDouble(out double Value))
+			{
+				InnerDefaultValue.AppendFormat("{0:F6}", Value);
+				if (bIsFloat)
+				{
+					InnerDefaultValue.Append(C);
+				}
+			}
+		}
+
+		enum NumberStyle
+		{
+			None,
+			Grouped,
+			Ungrouped,
+			Custom,
+		}
+
+		private static bool LocGenNumberOfPercentBoolean(IUhtTokenReader DefaultValueReader, StringBuilder InnerDefaultValue)
+		{
+			return false;
+		}
+
+		private static bool LocGenNumberOrPercentDefaultValue(UhtTextProperty Property, IUhtTokenReader DefaultValueReader, ref UhtToken MacroToken, StringBuilder InnerDefaultValue, NumberStyle NumberStyle)
+		{
+			InnerDefaultValue.Append(MacroToken.Value.ToString());
+			DefaultValueReader.Require('(');
+			InnerDefaultValue.Append('(');
+
+			UhtToken Token = DefaultValueReader.GetToken();
+			switch (Token.TokenType)
+			{
+				case UhtTokenType.DecimalConst:
+					FormatDecimal(ref Token, InnerDefaultValue);
+					break;
+
+				case UhtTokenType.FloatConst:
+					FormatFloat(ref Token, InnerDefaultValue);
+					break;
+
+				default:
+					return false;
+			}
+			DefaultValueReader.Require(',');
+			InnerDefaultValue.Append(", ");
+
+			if (NumberStyle == NumberStyle.Custom)
+			{
+				while (true)
+				{
+					UhtToken CustomToken = DefaultValueReader.GetToken();
+					if (!CustomToken.IsIdentifier())
+					{
+						return false;
+					}
+					InnerDefaultValue.Append(CustomToken.Value);
+
+					DefaultValueReader.Require('(');
+					InnerDefaultValue.Append('(');
+
+					switch (CustomToken.Value.ToString())
+					{
+						case "SetAlwaysSign":
+						case "SetUseGrouping":
+							{
+								UhtToken BooleanToken = DefaultValueReader.GetToken();
+								if (BooleanToken.IsIdentifier("true"))
+								{
+									InnerDefaultValue.Append("true");
+								}
+								else if (BooleanToken.IsIdentifier("false"))
+								{
+									InnerDefaultValue.Append("false");
+								}
+								else
+								{
+									return false;
+								}
+							}
+							break;
+
+						case "SetRoundingMode":
+							{
+								DefaultValueReader.Require("ERoundingMode");
+								DefaultValueReader.Require("::");
+								StringView Identifier = DefaultValueReader.GetIdentifier().Value;
+								InnerDefaultValue.Append("ERoundingMode::").Append(Identifier);
+							}
+							break;
+
+						case "SetMinimumIntegralDigits":
+						case "SetMaximumIntegralDigits":
+						case "SetMinimumFractionalDigits":
+						case "MaximumFractionalDigits":
+							{
+								UhtToken NumericToken = DefaultValueReader.GetToken();
+								if (NumericToken.GetConstInt(out int Value))
+								{
+									InnerDefaultValue.Append(Value);
+								}
+								else
+								{
+									return false;
+								}
+							}
+							break;
+
+						default:
+							return false;
+					}
+
+					DefaultValueReader.Require(')');
+					InnerDefaultValue.Append(')');
+
+					if (!DefaultValueReader.TryOptional('.'))
+					{
+						break;
+					}
+					InnerDefaultValue.Append('.');
+				}
+
+				DefaultValueReader.Require(',');
+				InnerDefaultValue.Append(", ");
+			}
+
+			StringView Culture = DefaultValueReader.GetConstQuotedString();
+			InnerDefaultValue.Append("\"\""); // UHT doesn't write out the culture
+
+			DefaultValueReader.Require(')');
+			InnerDefaultValue.Append(')');
+			return true;
+		}
+
+		private static void FormatDateTimeStyle(IUhtTokenReader DefaultValueReader, StringBuilder InnerDefaultValue)
+		{
+			DefaultValueReader.Require(',');
+			DefaultValueReader.Require("EDateTimeStyle");
+			DefaultValueReader.Require("::");
+			StringView Identifier = DefaultValueReader.GetIdentifier().Value;
+			InnerDefaultValue.Append(", EDateTimeStyle::").Append(Identifier);
+		}
+
+		private static bool LocGenDateTimeDefaultValue(UhtTextProperty Property, IUhtTokenReader DefaultValueReader, ref UhtToken MacroToken, StringBuilder InnerDefaultValue, bool bDate, bool bTime, bool bUtc, bool bCustom)
+		{
+			DefaultValueReader.Require('(');
+			InnerDefaultValue.Append(MacroToken.Value).Append('(');
+
+			long UnixTimestampValue = 0;
+			UhtToken TimestampToken = DefaultValueReader.GetToken();
+			switch (TimestampToken.TokenType)
+			{
+				case UhtTokenType.FloatConst:
+					if (TimestampToken.GetConstDouble(out double Value))
+					{
+						UnixTimestampValue = (long)Value;
+						break;
+					}
+					else
+					{
+						return false;
+					}
+
+				case UhtTokenType.DecimalConst:
+					if (!TimestampToken.GetConstLong(out UnixTimestampValue))
+					{
+						return false;
+					}
+					break;
+
+				default:
+					return false;
+
+			}
+			InnerDefaultValue.Append(UnixTimestampValue);
+
+			if (bCustom)
+			{
+				DefaultValueReader.Require(',');
+				StringView CustomPattern = DefaultValueReader.GetConstQuotedString();
+				InnerDefaultValue.Append(", ").Append(CustomPattern);
+			}
+			else
+			{
+				if (bDate)
+				{
+					FormatDateTimeStyle(DefaultValueReader, InnerDefaultValue);
+				}
+				if (bTime)
+				{
+					FormatDateTimeStyle(DefaultValueReader, InnerDefaultValue);
+				}
+			}
+
+			if (bUtc)
+			{
+				DefaultValueReader.Require(',');
+				StringView TimeZone = DefaultValueReader.GetConstQuotedString();
+				InnerDefaultValue.Append(", ").Append(TimeZone);
+			}
+
+			DefaultValueReader.Require(',');
+			StringView CultureName = DefaultValueReader.GetConstQuotedString();
+			InnerDefaultValue.Append(", ").Append("\"\""); // We don't write out the culture
+
+			DefaultValueReader.Require(')');
+			InnerDefaultValue.Append(")");
+			return true;
+		}
+
+		private static bool LocGenTransformDefaultValue(UhtTextProperty Property, IUhtTokenReader DefaultValueReader, ref UhtToken MacroToken, StringBuilder InnerDefaultValue)
+		{
+			InnerDefaultValue.Append(MacroToken.Value).Append('(');
+			DefaultValueReader.Require('(');
+			if (!SanitizeDefaultValue(Property, DefaultValueReader, InnerDefaultValue, false))
+			{
+				return false;
+			}
+			DefaultValueReader.Require(')');
+			InnerDefaultValue.Append(')');
+			return true;
+		}
+
+		#endregion
+	}
+}
diff --git a/Engine/Source/Programs/Shared/EpicGames.UHT/Types/Properties/UhtUInt16Property.cs b/Engine/Source/Programs/Shared/EpicGames.UHT/Types/Properties/UhtUInt16Property.cs
new file mode 100644
--- /dev/null
+++ b/Engine/Source/Programs/Shared/EpicGames.UHT/Types/Properties/UhtUInt16Property.cs
@@ -0,0 +1,63 @@
+// Copyright Epic Games, Inc. All Rights Reserved.
+
+using EpicGames.UHT.Tables;
+using EpicGames.UHT.Tokenizer;
+using EpicGames.UHT.Utils;
+using System.Text;
+
+namespace EpicGames.UHT.Types
+{
+	[UnrealHeaderTool]
+	[UhtEngineClass(Name = "UInt16Property", IsProperty = true)]
+	public class UhtUInt16Property : UhtNumericProperty
+	{
+		/// <inheritdoc/>
+		public override string EngineClassName { get => "UInt16Property"; }
+
+		/// <inheritdoc/>
+		protected override string CppTypeText { get => "uint16"; }
+
+		public UhtUInt16Property(UhtPropertySettings PropertySettings, UhtPropertyIntType IntType) : base(PropertySettings, IntType)
+		{
+		}
+
+		/// <inheritdoc/>
+		public override StringBuilder AppendMemberDecl(StringBuilder Builder, IUhtPropertyMemberContext Context, string Name, string NameSuffix, int Tabs)
+		{
+			return AppendMemberDecl(Builder, Context, Name, NameSuffix, Tabs, "FFInt16PropertyParams");
+		}
+
+		/// <inheritdoc/>
+		public override StringBuilder AppendMemberDef(StringBuilder Builder, IUhtPropertyMemberContext Context, string Name, string NameSuffix, string? Offset, int Tabs)
+		{
+			AppendMemberDefStart(Builder, Context, Name, NameSuffix, Offset, Tabs, "FFInt16PropertyParams", "UECodeGen_Private::EPropertyGenFlags::UInt16");
+			AppendMemberDefEnd(Builder, Context, Name, NameSuffix);
+			return Builder;
+		}
+
+		public override bool SanitizeDefaultValue(IUhtTokenReader DefaultValueReader, StringBuilder InnerDefaultValue)
+		{
+			return false;
+		}
+
+		public override bool IsSameType(UhtProperty Other)
+		{
+			return Other is UhtUInt16Property;
+		}
+
+		#region Keyword
+		[UhtPropertyType(Keyword = "uint16", Options = UhtPropertyTypeOptions.Simple | UhtPropertyTypeOptions.Immediate)]
+		public static UhtProperty? UInt16Property(UhtPropertyResolvePhase ResolvePhase, UhtPropertySettings PropertySettings, IUhtTokenReader TokenReader, UhtToken MatchedToken)
+		{
+			if (PropertySettings.bIsBitfield)
+			{
+				return new UhtBool8Property(PropertySettings);
+			}
+			else
+			{
+				return new UhtUInt16Property(PropertySettings, UhtPropertyIntType.Sized);
+			}
+		}
+		#endregion
+	}
+}
diff --git a/Engine/Source/Programs/Shared/EpicGames.UHT/Types/Properties/UhtUInt32Property.cs b/Engine/Source/Programs/Shared/EpicGames.UHT/Types/Properties/UhtUInt32Property.cs
new file mode 100644
--- /dev/null
+++ b/Engine/Source/Programs/Shared/EpicGames.UHT/Types/Properties/UhtUInt32Property.cs
@@ -0,0 +1,75 @@
+// Copyright Epic Games, Inc. All Rights Reserved.
+
+using EpicGames.UHT.Tables;
+using EpicGames.UHT.Tokenizer;
+using EpicGames.UHT.Utils;
+using System.Text;
+
+namespace EpicGames.UHT.Types
+{
+	[UnrealHeaderTool]
+	[UhtEngineClass(Name = "UInt32Property", IsProperty = true)]
+	public class UhtUInt32Property : UhtNumericProperty
+	{
+		/// <inheritdoc/>
+		public override string EngineClassName { get => "UInt32Property"; }
+
+		/// <inheritdoc/>
+		protected override string CppTypeText { get => "uint32"; }
+
+		public UhtUInt32Property(UhtPropertySettings PropertySettings, UhtPropertyIntType IntType) : base(PropertySettings, IntType)
+		{
+		}
+
+		/// <inheritdoc/>
+		public override StringBuilder AppendMemberDecl(StringBuilder Builder, IUhtPropertyMemberContext Context, string Name, string NameSuffix, int Tabs)
+		{
+			return AppendMemberDecl(Builder, Context, Name, NameSuffix, Tabs,
+				this.IntType == UhtPropertyIntType.Unsized ? "FUnsizedFIntPropertyParams" : "FUInt32PropertyParams");
+		}
+
+		/// <inheritdoc/>
+		public override StringBuilder AppendMemberDef(StringBuilder Builder, IUhtPropertyMemberContext Context, string Name, string NameSuffix, string? Offset, int Tabs)
+		{
+			AppendMemberDefStart(Builder, Context, Name, NameSuffix, Offset, Tabs,
+				this.IntType == UhtPropertyIntType.Unsized ? "FUnsizedFIntPropertyParams" : "FUInt32PropertyParams",
+				"UECodeGen_Private::EPropertyGenFlags::UInt32");
+			AppendMemberDefEnd(Builder, Context, Name, NameSuffix);
+			return Builder;
+		}
+
+		public override bool SanitizeDefaultValue(IUhtTokenReader DefaultValueReader, StringBuilder InnerDefaultValue)
+		{
+			return false;
+		}
+
+		public override bool IsSameType(UhtProperty Other)
+		{
+			return Other is UhtUInt32Property;
+		}
+
+		#region Keyword
+		[UhtPropertyType(Keyword = "uint32", Options = UhtPropertyTypeOptions.Simple | UhtPropertyTypeOptions.Immediate)]
+		public static UhtProperty? UInt32Property(UhtPropertyResolvePhase ResolvePhase, UhtPropertySettings PropertySettings, IUhtTokenReader TokenReader, UhtToken MatchedToken)
+		{
+			if (PropertySettings.bIsBitfield)
+			{
+				return new UhtBool8Property(PropertySettings);
+			}
+			else
+			{
+				return new UhtUInt32Property(PropertySettings, UhtPropertyIntType.Sized);
+			}
+		}
+
+		[UhtPropertyType(Keyword = "unsigned", Options = UhtPropertyTypeOptions.Immediate)]
+		public static UhtProperty? UnsignedProperty(UhtPropertyResolvePhase ResolvePhase, UhtPropertySettings PropertySettings, IUhtTokenReader TokenReader, UhtToken MatchedToken)
+		{
+			TokenReader
+				.Require("unsigned")
+				.Optional("int");
+			return new UhtUInt32Property(PropertySettings, UhtPropertyIntType.Unsized);
+		}
+		#endregion
+	}
+}
diff --git a/Engine/Source/Programs/Shared/EpicGames.UHT/Types/Properties/UhtUInt64Property.cs b/Engine/Source/Programs/Shared/EpicGames.UHT/Types/Properties/UhtUInt64Property.cs
new file mode 100644
--- /dev/null
+++ b/Engine/Source/Programs/Shared/EpicGames.UHT/Types/Properties/UhtUInt64Property.cs
@@ -0,0 +1,63 @@
+// Copyright Epic Games, Inc. All Rights Reserved.
+
+using EpicGames.UHT.Tables;
+using EpicGames.UHT.Tokenizer;
+using EpicGames.UHT.Utils;
+using System.Text;
+
+namespace EpicGames.UHT.Types
+{
+	[UnrealHeaderTool]
+	[UhtEngineClass(Name = "UInt64Property", IsProperty = true)]
+	public class UhtUInt64Property : UhtNumericProperty
+	{
+		/// <inheritdoc/>
+		public override string EngineClassName { get => "UInt64Property"; }
+
+		/// <inheritdoc/>
+		protected override string CppTypeText { get => "uint64"; }
+
+		public UhtUInt64Property(UhtPropertySettings PropertySettings, UhtPropertyIntType IntType) : base(PropertySettings, IntType)
+		{
+		}
+
+		/// <inheritdoc/>
+		public override StringBuilder AppendMemberDecl(StringBuilder Builder, IUhtPropertyMemberContext Context, string Name, string NameSuffix, int Tabs)
+		{
+			return AppendMemberDecl(Builder, Context, Name, NameSuffix, Tabs, "FFInt64PropertyParams");
+		}
+
+		/// <inheritdoc/>
+		public override StringBuilder AppendMemberDef(StringBuilder Builder, IUhtPropertyMemberContext Context, string Name, string NameSuffix, string? Offset, int Tabs)
+		{
+			AppendMemberDefStart(Builder, Context, Name, NameSuffix, Offset, Tabs, "FFInt64PropertyParams", "UECodeGen_Private::EPropertyGenFlags::UInt64");
+			AppendMemberDefEnd(Builder, Context, Name, NameSuffix);
+			return Builder;
+		}
+
+		public override bool SanitizeDefaultValue(IUhtTokenReader DefaultValueReader, StringBuilder InnerDefaultValue)
+		{
+			return false;
+		}
+
+		public override bool IsSameType(UhtProperty Other)
+		{
+			return Other is UhtUInt64Property;
+		}
+
+		#region Keyword
+		[UhtPropertyType(Keyword = "uint64", Options = UhtPropertyTypeOptions.Simple | UhtPropertyTypeOptions.Immediate)]
+		public static UhtProperty? UInt64Property(UhtPropertyResolvePhase ResolvePhase, UhtPropertySettings PropertySettings, IUhtTokenReader TokenReader, UhtToken MatchedToken)
+		{
+			if (PropertySettings.bIsBitfield)
+			{
+				return new UhtBool8Property(PropertySettings);
+			}
+			else
+			{
+				return new UhtUInt64Property(PropertySettings, UhtPropertyIntType.Sized);
+			}
+		}
+		#endregion
+	}
+}
diff --git a/Engine/Source/Programs/Shared/EpicGames.UHT/Types/Properties/UhtVoidProperty.cs b/Engine/Source/Programs/Shared/EpicGames.UHT/Types/Properties/UhtVoidProperty.cs
new file mode 100644
--- /dev/null
+++ b/Engine/Source/Programs/Shared/EpicGames.UHT/Types/Properties/UhtVoidProperty.cs
@@ -0,0 +1,82 @@
+// Copyright Epic Games, Inc. All Rights Reserved.
+
+using EpicGames.UHT.Tables;
+using EpicGames.UHT.Tokenizer;
+using EpicGames.UHT.Utils;
+using System;
+using System.Text;
+
+namespace EpicGames.UHT.Types
+{
+	[UnrealHeaderTool]
+	public class UhtVoidProperty : UhtProperty
+	{
+		/// <inheritdoc/>
+		public override string EngineClassName { get => "UHTVoidProperty"; }
+
+		/// <inheritdoc/>
+		protected override string CppTypeText { get => "void"; }
+
+		/// <inheritdoc/>
+		protected override string PGetMacroText { get => "invalid"; }
+
+		public UhtVoidProperty(UhtPropertySettings PropertySettings) : base(PropertySettings)
+		{
+		}
+
+		/// <inheritdoc/>
+		public override StringBuilder AppendMemberDecl(StringBuilder Builder, IUhtPropertyMemberContext Context, string Name, string NameSuffix, int Tabs)
+		{
+			throw new NotImplementedException();
+		}
+
+		/// <inheritdoc/>
+		public override StringBuilder AppendMemberDef(StringBuilder Builder, IUhtPropertyMemberContext Context, string Name, string NameSuffix, string? Offset, int Tabs)
+		{
+			throw new NotImplementedException();
+		}
+
+		/// <inheritdoc/>
+		public override StringBuilder AppendFunctionThunkParameterGet(StringBuilder Builder)
+		{
+			throw new NotImplementedException();
+		}
+
+		/// <inheritdoc/>
+		public override StringBuilder AppendNullConstructorArg(StringBuilder Builder, bool bIsInitializer)
+		{
+			throw new NotImplementedException();
+		}
+
+		/// <inheritdoc/>
+		public override bool SanitizeDefaultValue(IUhtTokenReader DefaultValueReader, StringBuilder InnerDefaultValue)
+		{
+			throw new NotImplementedException();
+		}
+
+		/// <inheritdoc/>
+		public override bool IsSameType(UhtProperty Other)
+		{
+			return Other is UhtVoidProperty;
+		}
+
+		/// <inheritdoc/>
+		public override string? GetRigVMType(ref UhtRigVMParameterFlags ParameterFlags)
+		{
+			throw new NotImplementedException();
+		}
+
+		#region Keyword
+		[UhtPropertyType(Keyword = "void", Options = UhtPropertyTypeOptions.Simple | UhtPropertyTypeOptions.Immediate)]
+		public static UhtProperty? VoidProperty(UhtPropertyResolvePhase ResolvePhase, UhtPropertySettings PropertySettings, IUhtTokenReader TokenReader, UhtToken MatchedToken)
+		{
+			if (PropertySettings.PropertyCategory != UhtPropertyCategory.Return)
+			{
+				TokenReader.LogError("void type is only valid as a return type");
+				return null;
+			}
+			return new UhtVoidProperty(PropertySettings);
+		}
+		#endregion
+	}
+}
diff --git a/Engine/Source/Programs/Shared/EpicGames.UHT/Types/Properties/UhtWeakObjectPtrProperty.cs b/Engine/Source/Programs/Shared/EpicGames.UHT/Types/Properties/UhtWeakObjectPtrProperty.cs
new file mode 100644
--- /dev/null
+++ b/Engine/Source/Programs/Shared/EpicGames.UHT/Types/Properties/UhtWeakObjectPtrProperty.cs
@@ -0,0 +1,139 @@
+// Copyright Epic Games, Inc. All Rights Reserved.
+
+using EpicGames.Core;
+using EpicGames.UHT.Tables;
+using EpicGames.UHT.Tokenizer;
+using EpicGames.UHT.Utils;
+using System.Text;
+
+namespace EpicGames.UHT.Types
+{
+	[UnrealHeaderTool]
+	[UhtEngineClass(Name = "WeakObjectProperty", IsProperty = true)]
+	public class UhtWeakObjectPtrProperty : UhtObjectPropertyBase
+	{
+		/// <inheritdoc/>
+		public override string EngineClassName { get => "WeakObjectProperty"; }
+
+		/// <inheritdoc/>
+		protected override string PGetMacroText { get => this.PropertyFlags.HasAnyFlags(EPropertyFlags.AutoWeak) ? "AUTOWEAKOBJECT" : "WEAKOBJECT"; }
+
+		/// <inheritdoc/>
+		protected override bool bPGetPassAsNoPtr { get => true; }
+
+		/// <inheritdoc/>
+		protected override UhtPGetArgumentType PGetTypeArgument { get => UhtPGetArgumentType.TypeText; }
+
+		public UhtWeakObjectPtrProperty(UhtPropertySettings PropertySettings, UhtClass PropertyClass, EPropertyFlags ExtraFlags = EPropertyFlags.None)
+			: base(PropertySettings, PropertyClass, null)
+		{
+			this.PropertyFlags |= EPropertyFlags.UObjectWrapper | ExtraFlags;
+			this.PropertyCaps |= UhtPropertyCaps.PassCppArgsByRef | UhtPropertyCaps.RequiresNullConstructorArg | UhtPropertyCaps.IsMemberSupportedByBlueprint;
+			this.PropertyCaps &= ~(UhtPropertyCaps.IsParameterSupportedByBlueprint);
+		}
+
+		/// <inheritdoc/>
+		public override StringBuilder AppendText(StringBuilder Builder, UhtPropertyTextType TextType, bool bIsTemplateArgument)
+		{
+			switch (TextType)
+			{
+				default:
+					if (this.PropertyFlags.HasAnyFlags(EPropertyFlags.AutoWeak))
+					{
+						Builder.Append("TAutoWeakObjectPtr<").Append(this.Class.SourceName).Append(">");
+					}
+					else
+					{
+						Builder.Append("TWeakObjectPtr<").Append(this.Class.SourceName).Append(">");
+					}
+					break;
+			}
+			return Builder;
+		}
+
+		/// <inheritdoc/>
+		public override StringBuilder AppendMemberDecl(StringBuilder Builder, IUhtPropertyMemberContext Context, string Name, string NameSuffix, int Tabs)
+		{
+			return AppendMemberDecl(Builder, Context, Name, NameSuffix, Tabs, "FWeakObjectPropertyParams");
+		}
+
+		/// <inheritdoc/>
+		public override StringBuilder AppendMemberDef(StringBuilder Builder, IUhtPropertyMemberContext Context, string Name, string NameSuffix, string? Offset, int Tabs)
+		{
+			AppendMemberDefStart(Builder, Context, Name, NameSuffix, Offset, Tabs, "FWeakObjectPropertyParams", "UECodeGen_Private::EPropertyGenFlags::WeakObject");
+			AppendMemberDefRef(Builder, Context, this.Class, false);
+			AppendMemberDefEnd(Builder, Context, Name, NameSuffix);
+			return Builder;
+		}
+
+		/// <inheritdoc/>
+		public override StringBuilder AppendNullConstructorArg(StringBuilder Builder, bool bIsInitializer)
+		{
+			Builder.Append("NULL");
+			return Builder;
+		}
+
+		/// <inheritdoc/>
+		public override bool IsSameType(UhtProperty Other)
+		{
+			if (Other is UhtWeakObjectPtrProperty OtherObject)
+			{
+				return this.Class == OtherObject.Class && this.MetaClass == OtherObject.MetaClass;
+			}
+			return false;
+		}
+
+		/// <inheritdoc/>
+		public override string? GetRigVMType(ref UhtRigVMParameterFlags ParameterFlags)
+		{
+			if (this.PropertyFlags.HasAnyFlags(EPropertyFlags.AutoWeak))
+			{
+				return $"TAutoWeakObjectPtr<{this.Class.SourceName}>";
+			}
+			else
+			{
+				return $"TWeakObjectPtr<{this.Class.SourceName}>";
+			}
+		}
+
+		#region Keywords
+		protected static UhtProperty CreateWeakProperty(UhtPropertySettings PropertySettings, IUhtTokenReader TokenReader, UhtClass PropertyClass, EPropertyFlags ExtraFlags = EPropertyFlags.None)
+		{
+			if (PropertyClass.IsChildOf(PropertyClass.Session.UClass))
+			{
+				TokenReader.LogError("Class variables cannot be weak, they are always strong.");
+			}
+
+			if (PropertySettings.DisallowPropertyFlags.HasAnyFlags(EPropertyFlags.AutoWeak))
+			{
+				return new UhtObjectProperty(PropertySettings, PropertyClass, ExtraFlags | EPropertyFlags.UObjectWrapper);
+			}
+			return new UhtWeakObjectPtrProperty(PropertySettings, PropertyClass, ExtraFlags);
+		}
+
+		[UhtPropertyType(Keyword = "TWeakObjectPtr")]
+		public static UhtProperty? WeakObjectProperty(UhtPropertyResolvePhase ResolvePhase, UhtPropertySettings PropertySettings, IUhtTokenReader TokenReader, UhtToken MatchedToken)
+		{
+			UhtClass? PropertyClass = ParseTemplateObject(PropertySettings, TokenReader, MatchedToken, true);
+			if (PropertyClass == null)
+			{
+				return null;
+			}
+
+			return CreateWeakProperty(PropertySettings, TokenReader, PropertyClass);
+		}
+
+		[UhtPropertyType(Keyword = "TAutoWeakObjectPtr")]
+		public static UhtProperty? AutoWeakObjectProperty(UhtPropertyResolvePhase ResolvePhase, UhtPropertySettings PropertySettings, IUhtTokenReader TokenReader, UhtToken MatchedToken)
+		{
+			UhtClass? PropertyClass = ParseTemplateObject(PropertySettings, TokenReader, MatchedToken, true);
+			if (PropertyClass == null)
+			{
+				return null;
+			}
+
+			return CreateWeakProperty(PropertySettings, TokenReader, PropertyClass, EPropertyFlags.AutoWeak);
+		}
+		#endregion
+	}
+}
