UHT in UBT initial version.  Currently disabled.

diff --git a/Engine/Source/Programs/Shared/EpicGames.UHT/Types/UhtClass.cs b/Engine/Source/Programs/Shared/EpicGames.UHT/Types/UhtClass.cs
new file mode 100644
--- /dev/null
+++ b/Engine/Source/Programs/Shared/EpicGames.UHT/Types/UhtClass.cs
@@ -0,0 +1,1240 @@
+// Copyright Epic Games, Inc. All Rights Reserved.
+
+using EpicGames.Core;
+using EpicGames.UHT.Parsers;
+using EpicGames.UHT.Tables;
+using EpicGames.UHT.Tokenizer;
+using EpicGames.UHT.Utils;
+using System;
+using System.Collections.Generic;
+using System.Diagnostics;
+using System.Text;
+using System.Text.Json.Serialization;
+
+namespace EpicGames.UHT.Types
+{
+	/// <summary>
+	/// Series of flags not part of the engine's class flags that affect code generation or verification
+	/// </summary>
+	[Flags]
+	public enum UhtClassExportFlags : UInt32
+	{
+		None = 0,
+
+		/// <summary>
+		/// If set, the class itself has replicated properties.
+		/// </summary>
+		SelfHasReplicatedProperties = 1 << 0,
+
+		/// <summary>
+		/// If set, some super class has replicated properties.
+		/// </summary>
+		SuperHasReplicatedProperties = 1 << 1,
+
+		/// <summary>
+		/// If set, either the class itself or a super class has replicated properties
+		/// </summary>
+		HasReplciatedProperties = SelfHasReplicatedProperties | SuperHasReplicatedProperties,
+
+		HasCustomConstructor = 1 << 2,
+
+		HasDefaultConstructor = 1 << 3,
+
+		HasObjectInitializerConstructor = 1 << 4,
+
+		HasCustomVTableHelperConstructor = 1 << 5,
+
+		HasConstructor = 1 << 6,
+
+		HasGetLifetimeReplicatedProps = 1 << 7,
+	}
+
+	/// <summary>
+	/// Helper methods for testing flags.  These methods perform better than the generic HasFlag which hits
+	/// the GC and stalls.
+	/// </summary>
+	public static class UhtClassExportFlagsExtensions
+	{
+
+		/// <summary>
+		/// Test to see if any of the specified flags are set
+		/// </summary>
+		/// <param name="InFlags">Current flags</param>
+		/// <param name="TestFlags">Flags to test for</param>
+		/// <returns>True if any of the flags are set</returns>
+		public static bool HasAnyFlags(this UhtClassExportFlags InFlags, UhtClassExportFlags TestFlags)
+		{
+			return (InFlags & TestFlags) != 0;
+		}
+
+		/// <summary>
+		/// Test to see if all of the specified flags are set
+		/// </summary>
+		/// <param name="InFlags">Current flags</param>
+		/// <param name="TestFlags">Flags to test for</param>
+		/// <returns>True if all the flags are set</returns>
+		public static bool HasAllFlags(this UhtClassExportFlags InFlags, UhtClassExportFlags TestFlags)
+		{
+			return (InFlags & TestFlags) == TestFlags;
+		}
+
+		/// <summary>
+		/// Test to see if a specific set of flags have a specific value.
+		/// </summary>
+		/// <param name="InFlags">Current flags</param>
+		/// <param name="TestFlags">Flags to test for</param>
+		/// <param name="MatchFlags">Expected value of the tested flags</param>
+		/// <returns>True if the given flags have a specific value.</returns>
+		public static bool HasExactFlags(this UhtClassExportFlags InFlags, UhtClassExportFlags TestFlags, UhtClassExportFlags MatchFlags)
+		{
+			return (InFlags & TestFlags) == MatchFlags;
+		}
+	}
+
+	public enum UhtClassType
+	{
+		Class,
+		Interface,
+		NativeInterface,
+	}
+
+	[Flags]
+	public enum UhtSerializerArchiveType
+	{
+		None = 0,
+		Archive = 1 << 0,
+		StructuredArchiveRecord = 1 << 1,
+		All = Archive | StructuredArchiveRecord,
+	}
+
+	/// <summary>
+	/// Helper methods for testing flags.  These methods perform better than the generic HasFlag which hits
+	/// the GC and stalls.
+	/// </summary>
+	public static class UhtSerializerArchiveTypeExtensions
+	{
+
+		/// <summary>
+		/// Test to see if any of the specified flags are set
+		/// </summary>
+		/// <param name="InFlags">Current flags</param>
+		/// <param name="TestFlags">Flags to test for</param>
+		/// <returns>True if any of the flags are set</returns>
+		public static bool HasAnyFlags(this UhtSerializerArchiveType InFlags, UhtSerializerArchiveType TestFlags)
+		{
+			return (InFlags & TestFlags) != 0;
+		}
+
+		/// <summary>
+		/// Test to see if all of the specified flags are set
+		/// </summary>
+		/// <param name="InFlags">Current flags</param>
+		/// <param name="TestFlags">Flags to test for</param>
+		/// <returns>True if all the flags are set</returns>
+		public static bool HasAllFlags(this UhtSerializerArchiveType InFlags, UhtSerializerArchiveType TestFlags)
+		{
+			return (InFlags & TestFlags) == TestFlags;
+		}
+
+		/// <summary>
+		/// Test to see if a specific set of flags have a specific value.
+		/// </summary>
+		/// <param name="InFlags">Current flags</param>
+		/// <param name="TestFlags">Flags to test for</param>
+		/// <param name="MatchFlags">Expected value of the tested flags</param>
+		/// <returns>True if the given flags have a specific value.</returns>
+		public static bool HasExactFlags(this UhtSerializerArchiveType InFlags, UhtSerializerArchiveType TestFlags, UhtSerializerArchiveType MatchFlags)
+		{
+			return (InFlags & TestFlags) == MatchFlags;
+		}
+	}
+
+	public struct UhtDeclaration
+	{
+		public UhtCompilerDirective CompilerDirectives;
+		public UhtToken[] Tokens;
+	}
+
+	public struct UhtFoundDeclaration
+	{
+		public UhtCompilerDirective CompilerDirectives;
+		public UhtToken[] Tokens;
+		public int NameTokenIndex;
+		public bool bIsVirtual;
+	}
+
+	[UhtEngineClass(Name = "Class")]
+	public class UhtClass : UhtStruct
+	{
+		private static UhtSpecifierValidatorTable ClassSpecifierValidatorTable = UhtSpecifierValidatorTables.Instance.Get(UhtTableNames.Class);
+		private static UhtSpecifierValidatorTable InterfaceSpecifierValidatorTable = UhtSpecifierValidatorTables.Instance.Get(UhtTableNames.Interface);
+		private static UhtSpecifierValidatorTable NativeInterfaceSpecifierValidatorTable = UhtSpecifierValidatorTables.Instance.Get(UhtTableNames.NativeInterface);
+
+		public string Config { get; set; } = String.Empty;
+		public string EnclosingDefine { get; set; } = String.Empty;
+
+		[JsonConverter(typeof(UhtTypeSourceNameJsonConverter<UhtClass>))]
+		public UhtClass ClassWithin { get; set; }
+
+		[JsonConverter(typeof(JsonStringEnumConverter))]
+		public EClassFlags ClassFlags { get; set; } = EClassFlags.None;
+
+		[JsonConverter(typeof(JsonStringEnumConverter))]
+		public EClassCastFlags ClassCastFlags { get; set; } = EClassCastFlags.None;
+
+		[JsonConverter(typeof(JsonStringEnumConverter))]
+		public UhtClassExportFlags ClassExportFlags { get; set; } = UhtClassExportFlags.None;
+
+		public UhtClassType ClassType = UhtClassType.Class;
+
+		[JsonConverter(typeof(JsonStringEnumConverter))]
+		public UhtSerializerArchiveType SerializerArchiveType { get; set; } = UhtSerializerArchiveType.None;
+
+		public int PrologLineNumber { get; set; } = -1;
+
+		[JsonIgnore]
+		public IList<UhtDeclaration>? Declarations => this.DeclarationsInternal;
+		public UhtClass? NativeInterface = null;
+
+		public int GeneratedBodyLineNumber { get; set; } = -1;
+
+		[JsonConverter(typeof(JsonStringEnumConverter))]
+		public UhtAccessSpecifier GeneratedBodyAccessSpecifier { get; set; } = UhtAccessSpecifier.None;
+
+		public bool bHasGeneratedBody = false; // If not true, then GENERATED_UCLASS_BODY was used
+
+		[JsonIgnore]
+		public override UhtEngineType EngineType
+		{
+			get
+			{
+				switch (this.ClassType)
+				{
+					case UhtClassType.Class:
+						return UhtEngineType.Class;
+					case UhtClassType.Interface:
+						return UhtEngineType.Interface;
+					case UhtClassType.NativeInterface:
+						return UhtEngineType.NativeInterface;
+					default:
+						throw new UhtIceException("Unknown class type");
+				}
+			}
+		}
+
+		/// <inheritdoc/>
+		public override string EngineClassName { get => "Class"; }
+
+		///<inheritdoc/>
+		[JsonIgnore]
+		public override bool bDeprecated => this.ClassFlags.HasAnyFlags(EClassFlags.Deprecated);
+
+		///<inheritdoc/>
+		[JsonIgnore]
+		protected override UhtSpecifierValidatorTable? SpecifierValidatorTable
+		{
+			get
+			{
+				switch (this.ClassType)
+				{
+					case UhtClassType.Class:
+						return UhtClass.ClassSpecifierValidatorTable;
+					case UhtClassType.Interface:
+						return UhtClass.InterfaceSpecifierValidatorTable;
+					case UhtClassType.NativeInterface:
+						return UhtClass.NativeInterfaceSpecifierValidatorTable;
+					default:
+						throw new UhtIceException("Unknown class type");
+				}
+			}
+		}
+
+		[JsonConverter(typeof(UhtNullableTypeSourceNameJsonConverter<UhtClass>))]
+		public UhtClass? SuperClass => (UhtClass?)this.Super;
+
+		public UhtClass(UhtType Outer, int LineNumber) : base(Outer, LineNumber)
+		{
+			this.ClassWithin = this;
+		}
+
+		[JsonIgnore]
+		public bool bIsActorClass => IsChildOf(this.Session.AActor);
+
+		[JsonIgnore]
+		public override string EngineNamePrefix
+		{ 
+			get
+			{
+				switch (this.ClassType)
+				{
+					case UhtClassType.Class:  
+						if (bIsActorClass)
+						{
+							return this.ClassFlags.HasAnyFlags(EClassFlags.Deprecated) ? "ADEPRECATED_" : "A";
+						}
+						else
+						{
+							return this.ClassFlags.HasAnyFlags(EClassFlags.Deprecated) ? "UDEPRECATED_" : "U";
+						}
+
+					case UhtClassType.Interface:
+						return "U";
+
+					case UhtClassType.NativeInterface:
+						return "I";
+
+					default:
+						throw new NotImplementedException();
+				}
+			}
+		}
+
+		private List<UhtDeclaration>? DeclarationsInternal = null;
+
+		/// <summary>
+		/// Add the given list of tokens as a possible declaration
+		/// </summary>
+		/// <param name="CompilerDirectives">Currently active compiler directives</param>
+		/// <param name="Tokens">List of declaration tokens</param>
+		public void AddDeclaration(UhtCompilerDirective CompilerDirectives, List<UhtToken> Tokens)
+		{
+			if (this.DeclarationsInternal == null)
+			{
+				this.DeclarationsInternal = new List<UhtDeclaration>();
+			}
+			this.DeclarationsInternal.Add(new UhtDeclaration { CompilerDirectives = CompilerDirectives, Tokens = Tokens.ToArray() });
+		}
+
+		/// <summary>
+		/// Search the declarations for a possible declaration match with the given name
+		/// </summary>
+		/// <param name="Name">Name to be located</param>
+		/// <param name="FoundDeclaration">Information about the matched declaration</param>
+		/// <returns>True if a declaration was found with the name.</returns>
+		public bool TryGetDeclaration(string Name, out UhtFoundDeclaration FoundDeclaration)
+		{
+			if (Name.Length > 0)
+			{
+				if (this.Declarations != null)
+				{
+					foreach (UhtDeclaration Declaration in this.Declarations)
+					{
+						bool bIsVirtual = false;
+						for (int Index = 0; Index < Declaration.Tokens.Length; ++Index)
+						{
+							if (Declaration.Tokens[Index].IsIdentifier("virtual"))
+							{
+								bIsVirtual = true;
+							}
+							else if (Declaration.Tokens[Index].IsIdentifier(Name))
+							{
+								FoundDeclaration = new UhtFoundDeclaration
+								{
+									CompilerDirectives = Declaration.CompilerDirectives,
+									Tokens = Declaration.Tokens,
+									NameTokenIndex = Index,
+									bIsVirtual = bIsVirtual,
+								};
+								return true;
+							}
+						}
+					}
+				}
+			}
+			if (this.NativeInterface != null)
+			{
+				return this.NativeInterface.TryGetDeclaration(Name, out FoundDeclaration);
+			}
+			FoundDeclaration = new UhtFoundDeclaration();
+			return false;
+		}
+
+		public IEnumerable<UhtProperty> EnumerateReplicatedProperties(bool bIncludeSuper)
+		{
+			if (bIncludeSuper && this.SuperClass != null)
+			{
+				foreach (UhtProperty Property in this.SuperClass.EnumerateReplicatedProperties(true))
+				{
+					yield return Property;
+				}
+			}
+
+			foreach (UhtProperty Property in this.Properties)
+			{
+				if (Property.PropertyFlags.HasAnyFlags(EPropertyFlags.Net))
+				{
+					yield return Property;
+				}
+			}
+		}
+
+		/// <summary>
+		/// Checks to see if the class or any super class has the given flags.
+		/// </summary>
+		/// <param name="FlagsToCheck"></param>
+		/// <returns>True if the flags are found</returns>
+		public bool HierarchyHasAnyClassFlags(EClassFlags FlagsToCheck)
+		{
+			for (UhtClass? Current = this; Current != null; Current = Current.SuperClass)
+			{
+				if (Current.ClassFlags.HasAnyFlags(FlagsToCheck))
+				{
+					return true;
+				}
+			}
+			return false;
+		}
+
+		#region Resolution support
+		struct GetterSetterToResolve
+		{
+			public UhtProperty Property;
+			public bool bSetter;
+			public bool bFound;
+		}
+
+		protected override bool ResolveSelf(UhtResolvePhase Phase)
+		{
+			bool bResult = base.ResolveSelf(Phase);
+
+			switch (Phase)
+			{
+				case UhtResolvePhase.Bases:
+
+					// Check to see if this class matches an entry in the ClassCastFlags.  If it does, that
+					// becomes our ClassCastFlags.  If not, then we inherit the flags from the super.
+					if (Enum.TryParse<EClassCastFlags>(this.SourceName, false, out EClassCastFlags Found))
+					{
+						if (Found != EClassCastFlags.None)
+						{
+							this.ClassCastFlags = Found;
+						}
+					}
+					if (this.ClassCastFlags == EClassCastFlags.None && this.SuperClass != null)
+					{
+						this.ClassCastFlags = this.SuperClass.ClassCastFlags;
+					}
+
+					// force members to be 'blueprint read only' if in a const class
+					// This must be done after the Const flag has been propagated from the parent
+					// and before properties have been finalized.
+					if (this.ClassType == UhtClassType.Class && this.ClassFlags.HasAnyFlags(EClassFlags.Const))
+					{
+						foreach (UhtType Child in this.Children)
+						{
+							if (Child is UhtProperty Property)
+							{
+								Property.PropertyFlags |= EPropertyFlags.BlueprintReadOnly;
+								Property.PropertyExportFlags |= UhtPropertyExportFlags.ImpliedBlueprintPure;
+							}
+						}
+					}
+					break;
+
+				case UhtResolvePhase.Final:
+					bool bAutoGetterSetter = UhtConfig.Instance.bAllowAutomaticSettersAndGetters;
+					Dictionary<string, List<GetterSetterToResolve>>? GSToResolve = null;
+					foreach (UhtType Child in this.Children)
+					{
+						if (Child is UhtProperty Property)
+						{
+							if (Property.PropertyFlags.HasAnyFlags(EPropertyFlags.Net))
+							{
+								this.ClassExportFlags |= UhtClassExportFlags.SelfHasReplicatedProperties;
+								break;
+							}
+
+							if (Property.SourceName == "SourceModels")
+							{
+								Debugger.Break();
+							}
+
+							if (!Property.PropertyExportFlags.HasAnyFlags(UhtPropertyExportFlags.GetterSpecifiedNone))
+							{
+								if (Property.PropertyExportFlags.HasAnyFlags(UhtPropertyExportFlags.GetterSpecified))
+								{
+									GSToResolve = AddGetterSetter(GSToResolve, Property.Getter == null ? $"Get{Property.SourceName}" : Property.Getter, Property, false);
+								}
+								else if (bAutoGetterSetter)
+								{
+									GSToResolve = AddGetterSetter(GSToResolve, $"Get{Property.SourceName}", Property, false);
+								}
+							}
+
+							if (!Property.PropertyExportFlags.HasAnyFlags(UhtPropertyExportFlags.SetterSpecifiedNone))
+							{
+								if (Property.PropertyExportFlags.HasAnyFlags(UhtPropertyExportFlags.SetterSpecified))
+								{
+									GSToResolve = AddGetterSetter(GSToResolve, Property.Setter == null ? $"Set{Property.SourceName}" : Property.Setter, Property, true);
+								}
+								else if (bAutoGetterSetter)
+								{
+									GSToResolve = AddGetterSetter(GSToResolve, $"Set{Property.SourceName}", Property, true);
+								}
+							}
+						}
+					}
+
+					if (this.SuperClass != null)
+					{
+						if (this.SuperClass.ClassExportFlags.HasAnyFlags(UhtClassExportFlags.HasReplciatedProperties))
+						{
+							this.ClassExportFlags |= UhtClassExportFlags.SuperHasReplicatedProperties;
+						}
+					}
+
+					if (!this.ClassFlags.HasAnyFlags(EClassFlags.Interface))
+					{
+						if (this.ClassExportFlags.HasAnyFlags(UhtClassExportFlags.SelfHasReplicatedProperties))
+						{
+							if (TryGetDeclaration("GetLifetimeReplicatedProps", out UhtFoundDeclaration FoundDeclaration))
+							{
+								this.ClassExportFlags |= UhtClassExportFlags.HasGetLifetimeReplicatedProps;
+							}
+						}
+					}
+
+					// If we have any getters/setters to resolve 
+					if (GSToResolve != null)
+					{
+						if (this.Declarations != null)
+						{
+							foreach (UhtDeclaration Declaration in this.Declarations)
+							{
+
+								// Locate the index of the name
+								int NameIndex = 0;
+								for (; NameIndex < Declaration.Tokens.Length - 1; ++NameIndex)
+								{
+									if (Declaration.Tokens[NameIndex + 1].IsValue('('))
+									{
+										break;
+									}
+								}
+
+								string Name = Declaration.Tokens[NameIndex].Value.ToString();
+								if (GSToResolve.TryGetValue(Name, out List<GetterSetterToResolve>? GSList))
+								{
+									for (int Index = 0; Index < GSList.Count; ++Index)
+									{
+										GetterSetterToResolve GS = GSList[Index];
+										if (!GS.bFound)
+										{
+											if (TryMatchGetterSetter(GS.Property, GS.bSetter, Declaration, NameIndex))
+											{
+												GS.bFound = true;
+												GSList[Index] = GS;
+												GS.Property.PropertyExportFlags |= GS.bSetter ? UhtPropertyExportFlags.SetterFound : UhtPropertyExportFlags.GetterFound;
+												if (GS.bSetter)
+												{
+													GS.Property.Setter = Name;
+												}
+												else
+												{
+													GS.Property.Getter = Name;
+												}
+											}
+										}
+									}
+								}
+							}
+						}
+					}
+					break;
+			}
+
+			return bResult;
+		}
+
+		/// <inheritdoc/>
+		protected override void ResolveChildren(UhtResolvePhase Phase)
+		{
+			base.ResolveChildren(Phase);
+
+			switch (Phase)
+			{
+				case UhtResolvePhase.Final:
+					if (ScanForInstancedReferenced(false))
+					{
+						this.ClassFlags |= EClassFlags.HasInstancedReference;
+					}
+					break;
+			}
+		}
+
+		/// <inheritdoc/>
+		public override bool ScanForInstancedReferenced(bool bDeepScan)
+		{
+			if (this.ClassFlags.HasAnyFlags(EClassFlags.HasInstancedReference))
+			{
+				return true;
+			}
+
+			if (this.SuperClass != null && this.SuperClass.ScanForInstancedReferenced(bDeepScan))
+			{
+				return true;
+			}
+
+			return base.ScanForInstancedReferenced(bDeepScan);
+		}
+
+		/// <summary>
+		/// Given a token stream, verify that it matches the expected signature for a getter/setter
+		/// </summary>
+		/// <param name="Property">Property requesting the setter/getter</param>
+		/// <param name="bSetter">If true, a setter is expected</param>
+		/// <param name="Declaration">The declaration being tested</param>
+		/// <param name="NameIndex">The index of the token with the expected getter/setter name</param>
+		/// <returns>True if the declaration matches</returns>
+		private bool TryMatchGetterSetter(UhtProperty Property, bool bSetter, UhtDeclaration Declaration, int NameIndex)
+		{
+			UhtToken[] Tokens = Declaration.Tokens;
+			int TokenIndex = 0;
+			int TokenCount = Tokens.Length;
+
+			// Skip all of the API and virtual keywords
+			for (; TokenIndex < TokenCount && (Tokens[TokenIndex].IsValue("virtual") || Tokens[TokenIndex].Value.Span.EndsWith("_API")); ++TokenIndex)
+			{
+				// Nothing to do in the body
+			}
+
+			int TypeIndex = 0;
+			int TypeCount = 0;
+
+			// Verify the format of the function declaration and extract the type
+			if (bSetter)
+			{
+
+				// The return type must be 'void'
+				if (TokenIndex == TokenCount || !Tokens[TokenIndex].IsValue("void"))
+				{
+					return false;
+				}
+				TokenIndex++;
+
+				// Followed by the name
+				if (TokenIndex == TokenCount || TokenIndex != NameIndex)
+				{
+					return false;
+				}
+				TokenIndex++;
+
+				// Followed by the '('
+				if (TokenIndex == TokenCount || !Tokens[TokenIndex].IsValue('('))
+				{
+					return false;
+				}
+				TokenIndex++;
+
+				// Find the ')'
+				TypeIndex = TokenIndex;
+				while (true)
+				{
+					if (TokenIndex == TokenCount)
+					{
+						return false;
+					}
+					if (Tokens[TokenIndex].IsValue(')'))
+					{
+						TypeCount = TokenIndex - TypeIndex;
+						break;
+					}
+					TokenIndex++;
+				}
+			}
+			else
+			{
+				// Get the type range
+				TypeIndex = TokenIndex;
+				TypeCount = NameIndex - TypeIndex;
+				TokenIndex = NameIndex + 1;
+
+				// Followed by the '('
+				if (TokenIndex == TokenCount || !Tokens[TokenIndex].IsValue('('))
+				{
+					return false;
+				}
+				TokenIndex++;
+
+				// Followed by the ')'
+				if (TokenIndex == TokenCount || !Tokens[TokenIndex].IsValue(')'))
+				{
+					return false;
+				}
+				TokenIndex++;
+
+				// Followed by the "const"
+				if (TokenIndex == TokenCount || !Tokens[TokenIndex].IsValue("const"))
+				{
+					return false;
+				}
+				TokenIndex++;
+			}
+
+			// Verify the type
+			using (BorrowStringBuilder Borrower = new BorrowStringBuilder(StringBuilderCache.Small))
+			{
+				StringBuilder Builder = Borrower.StringBuilder;
+				Builder.AppendPropertyText(Property, UhtPropertyTextType.GetterSetterArg);
+				string ExpectedType = Builder.ToString();
+				int TypeEnd = TypeIndex + TypeCount;
+
+				// We will do a somewhat brute force string compare to avoid constructing an expected
+				// string or parsing the required string
+				ReadOnlySpan<char> ExpectedSpan = ExpectedType.AsSpan().Trim();
+
+				// skip const. We allow const in parameter and return values even if the property is not const
+				if (TypeIndex < TypeEnd && Tokens[TypeIndex].IsIdentifier("const"))
+				{
+					TypeIndex++;
+				}
+
+				// Loop until we are done
+				while (TypeIndex < TypeEnd && ExpectedSpan.Length > 0)
+				{
+					UhtToken Token = Tokens[TypeIndex];
+					ReadOnlySpan<char> TokenSpan = Token.Value.Span;
+
+					// Make sure we match the expected type
+					if (TokenSpan.Length > ExpectedSpan.Length ||
+						!ExpectedSpan.StartsWith(TokenSpan))
+					{
+						return false;
+					}
+
+					// Extract the remaining part of the string we expect to match
+					ExpectedSpan = ExpectedSpan.Slice(TokenSpan.Length).TrimStart();
+					TypeIndex++;
+				}
+
+				// Consume any '&' for reference support
+				if (TypeIndex < TypeEnd && Tokens[TypeIndex].IsSymbol('&'))
+				{
+					TypeIndex++;
+				}
+
+				// For setters, there can be an identifier
+				if (bSetter && TypeIndex < TypeEnd && Tokens[TypeIndex].IsIdentifier())
+				{
+					TypeIndex++;
+				}
+
+				// We must be at the end for a match
+				if (TypeIndex != TypeEnd)
+				{
+					return false;
+				}
+			}
+
+			// Validate the #if blocks.  This should really be done during validation
+			string FuncType = bSetter ? "setter" : "getter";
+			if (Declaration.CompilerDirectives.HasAnyFlags(UhtCompilerDirective.WithEditor))
+			{
+				if (Property.PropertyExportFlags.HasAnyFlags(UhtPropertyExportFlags.GetterSpecified))
+				{
+					Property.LogError(Tokens[0].InputLine, $"Property {Property.SourceName} {FuncType} function {Tokens[NameIndex].Value.ToString()} "
+						+ "cannot be declared within WITH_EDITOR block. Use WITH_EDITORONLY_DATA instead.");
+				}
+				return false;
+			}
+			if (Declaration.CompilerDirectives.HasAnyFlags(UhtCompilerDirective.WithEditorOnlyData) && !Property.PropertyFlags.HasAnyFlags(EPropertyFlags.EditorOnly))
+			{
+				if (Property.PropertyExportFlags.HasAnyFlags(UhtPropertyExportFlags.GetterSpecified))
+				{
+					Property.LogError(Tokens[0].InputLine, $"Property {Property.SourceName} is not editor-only but its {FuncType} function '{Tokens[NameIndex].Value.ToString()}' is");
+				}
+				return false;
+			}
+			return true;
+		}
+
+		/// <summary>
+		/// Add a requested getter/setter function name
+		/// </summary>
+		/// <param name="GSToResolve">Dictionary containing the lookup by function name.  Will be created if null.</param>
+		/// <param name="Name">Name of the getter/setter</param>
+		/// <param name="Property">Property requesting the getter/setter</param>
+		/// <param name="bSetter">True if this is a setter</param>
+		/// <returns>Resulting dictionary</returns>
+		private Dictionary<string, List<GetterSetterToResolve>> AddGetterSetter(Dictionary<string, List<GetterSetterToResolve>>? GSToResolve, string Name, UhtProperty Property, bool bSetter)
+		{
+			if (GSToResolve == null)
+			{
+				GSToResolve = new Dictionary<string, List<GetterSetterToResolve>>();
+			}
+			List<GetterSetterToResolve>? GSList = null;
+			if (!GSToResolve.TryGetValue(Name, out GSList))
+			{
+				GSList = new List<GetterSetterToResolve>();
+				GSToResolve.Add(Name, GSList);
+			}
+			GSList.Add(new GetterSetterToResolve { Property = Property, bSetter = bSetter });
+			return GSToResolve;
+		}
+		#endregion
+
+		#region Validation support
+		protected override UhtValidationOptions Validate(UhtValidationOptions Options)
+		{
+			Options = base.Validate(Options);
+
+			// Classes must start with a valid prefix
+			string ExpectedClassName = this.EngineNamePrefix + this.EngineName;
+			if (ExpectedClassName != this.SourceName)
+			{
+				this.LogError($"Class '{this.SourceName}' has an invalid Unreal prefix, expecting '{ExpectedClassName}'");
+			}
+
+			// If we have a super class
+			if (this.SuperClass != null)
+			{
+				if (!this.ClassFlags.HasAnyFlags(EClassFlags.NotPlaceable) && this.SuperClass.ClassFlags.HasAnyFlags(EClassFlags.NotPlaceable))
+				{
+					this.LogError("The 'placeable' specifier cannot override a 'nonplaceable' base class. Classes are assumed to be placeable by default. "
+						+ "Consider whether using the 'abstract' specifier on the base class would work.");
+				}
+
+				// Native interfaces don't require a super class, but UHT technically allows it.
+				if (this.ClassType == UhtClassType.Interface || this.ClassType == UhtClassType.NativeInterface)
+				{
+					if (this != this.Session.UInterface && !this.SuperClass.ClassFlags.HasAnyFlags(EClassFlags.Interface))
+					{
+						this.LogError($"Interface class '{this.SourceName}' cannot inherit from non-interface class '{this.SuperClass.SourceName}'");
+					}
+				}
+			}
+
+			// If we don't have a super class
+			else
+			{
+				if (this.ClassType == UhtClassType.Interface)
+				{
+					this.LogError($"Interface '{this.SourceName}' must derive from an interface");
+				}
+			}
+
+			if (this.ClassFlags.HasAnyFlags(EClassFlags.NeedsDeferredDependencyLoading) && !this.IsChildOf(this.Session.UClass))
+			{
+				// CLASS_NeedsDeferredDependencyLoading can only be set on classes derived from UClass
+				this.LogError($"'NeedsDeferredDependencyLoading' is set on '{this.SourceName}' but the flag can only be used with classes derived from UClass.");
+			}
+
+			ValidateProperties();
+
+			// Perform any validation specific to the class type
+			switch (this.ClassType)
+			{
+				case UhtClassType.Class:
+					{
+						if (this.ClassFlags.HasAnyFlags(EClassFlags.EditInlineNew))
+						{
+							// don't allow actor classes to be declared EditInlineNew
+							if (this.IsChildOf(this.Session.AActor))
+							{
+								this.LogError("Invalid class attribute: Creating actor instances via the property window is not allowed");
+							}
+						}
+
+						// Make sure both RequiredAPI and MinimalAPI aren't specified
+						if (this.ClassFlags.HasAllFlags(EClassFlags.MinimalAPI | EClassFlags.RequiredAPI))
+						{
+							this.LogError("MinimalAPI cannot be specified when the class is fully exported using a MODULENAME_API macro");
+						}
+
+						// Make sure that all interface functions are implemented property
+						if (this.Bases != null)
+						{
+
+							// Iterate over all the bases looking for any interfaces
+							foreach (UhtStruct BaseStruct in this.Bases)
+							{
+
+								// Skip children that aren't interfaces or if a common ancestor
+								UhtClass? BaseClass = BaseStruct as UhtClass;
+								if (BaseClass == null || BaseClass.ClassType == UhtClassType.Class || IsChildOf(BaseClass) )
+								{
+									continue;
+								}
+
+								// Loop through the function to make sure they are implemented
+								List<UhtType> BaseChildren = BaseClass.AlternateObject != null ? BaseClass.AlternateObject.Children : BaseClass.Children;
+								foreach (UhtType BaseChild in BaseChildren)
+								{
+									UhtFunction? BaseFunction = BaseChild as UhtFunction;
+									if (BaseFunction == null)
+									{
+										continue; 
+									}
+
+									// Delegate signature functions are simple stubs and aren't required to be implemented (they are not callable)
+									bool bImplemented = BaseFunction.FunctionFlags.HasAnyFlags(EFunctionFlags.Delegate);
+
+									// Try to find an existing function
+									foreach (UhtType Child in this.Children)
+									{
+										UhtFunction? Function = Child as UhtFunction;
+										if (Function == null || Function.SourceName != BaseFunction.SourceName)
+										{
+											continue;
+										}
+
+										if (BaseFunction.FunctionFlags.HasAnyFlags(EFunctionFlags.Event) && !Function.FunctionFlags.HasAnyFlags(EFunctionFlags.Event))
+										{
+											Function.LogError($"Implementation of function '{Function.SourceName}' must be declared as 'event' to match declaration in interface '{BaseClass.SourceName}'");
+										}
+
+										if (BaseFunction.FunctionFlags.HasAnyFlags(EFunctionFlags.Delegate) && !Function.FunctionFlags.HasAnyFlags(EFunctionFlags.Delegate))
+										{
+											Function.LogError($"Implementation of function '{Function.SourceName}' must be declared as 'delegate' to match declaration in interface '{BaseClass.SourceName}'");
+										}
+
+										bImplemented = true;
+
+										if (BaseFunction.Children.Count != Function.Children.Count)
+										{
+											Function.LogError($"Implementation of function '{Function.SourceName}' conflicts with interface '{BaseClass.SourceName}' - different number of parameters ({Function.Children.Count}/{BaseFunction.Children.Count})");
+											continue;
+										}
+
+										for (int Index = 0; Index < Function.Children.Count; ++Index)
+										{
+											UhtProperty BaseProperty = (UhtProperty)BaseFunction.Children[Index];
+											UhtProperty Property = (UhtProperty)Function.Children[Index];
+											if (!BaseProperty.MatchesType(Property))
+											{
+												if (BaseProperty.PropertyFlags.HasAnyFlags(EPropertyFlags.ReturnParm))
+												{
+													Function.LogError($"Implementation of function '{Function.SourceName}' conflicts only by return type with interface '{BaseClass.SourceName}'");
+												}
+												else
+												{
+													Function.LogError($"Implementation of function '{Function.SourceName}' conflicts type with interface '{BaseClass.SourceName}' - parameter {Index} '{Property.SourceName}'");
+												}
+											}
+										}
+									}
+
+									// Verify that if this has blueprint-callable functions that are not implementable events, we've implemented them as a UFunction in the target class
+									if (!bImplemented
+										&& BaseFunction.FunctionFlags.HasAnyFlags(EFunctionFlags.BlueprintCallable)
+										&& !BaseFunction.FunctionFlags.HasAnyFlags(EFunctionFlags.BlueprintEvent)
+										&& !BaseClass.MetaData.ContainsKey(UhtNames.CannotImplementInterfaceInBlueprint)
+										&& (BaseClass.AlternateObject == null || !BaseClass.AlternateObject.MetaData.ContainsKey(UhtNames.CannotImplementInterfaceInBlueprint)))
+									{
+										this.LogError($"Missing UFunction implementation of function '{BaseFunction.SourceName}' from interface '{BaseClass.SourceName}'.  This function needs a UFUNCTION() declaration.");
+									}
+								}
+							}
+						}
+					}
+
+					if (!this.ClassFlags.HasAnyFlags(EClassFlags.Interface))
+					{
+						if (this.ClassExportFlags.HasAnyFlags(UhtClassExportFlags.SelfHasReplicatedProperties) &&
+							!this.ClassExportFlags.HasAnyFlags(UhtClassExportFlags.HasGetLifetimeReplicatedProps) &&
+							this.GeneratedCodeVersion != EGeneratedCodeVersion.V1)
+						{
+							this.LogError($"Class {this.SourceName} has Net flagged properties and should declare member function: void GetLifetimeReplicatedProps(TArray<FLifetimeProperty>& OutLifetimeProps) const override");
+						}
+					}
+					break;
+
+				case UhtClassType.Interface:
+					{
+						// Interface with blueprint data should declare explicitly Blueprintable or NotBlueprintable to be clear
+						// In the backward compatible case where they declare neither, both of these bools are false
+						bool bCanImplementInBlueprints = this.MetaData.GetBoolean(UhtNames.IsBlueprintBase);
+						bool bCannotImplementInBlueprints = (!bCanImplementInBlueprints && this.MetaData.ContainsKey(UhtNames.IsBlueprintBase))
+							|| this.MetaData.ContainsKey(UhtNames.CannotImplementInterfaceInBlueprint);
+
+						bool bCanImplementInterfaceInBlueprint = !this.MetaData.ContainsKey(UhtNames.CannotImplementInterfaceInBlueprint);
+						foreach (UhtType Child in this.Children)
+						{
+							if (Child is UhtFunction Function)
+							{
+								// Verify interfaces with respect to their blueprint accessible functions
+								if (Function.FunctionFlags.HasAnyFlags(EFunctionFlags.BlueprintEvent) && !Function.MetaData.GetBoolean(UhtNames.BlueprintInternalUseOnly))
+								{
+									// Ensure that blueprint events are only allowed in implementable interfaces. Internal only functions allowed
+									if (bCannotImplementInBlueprints)
+									{
+										Function.LogError("Interfaces that are not implementable in blueprints cannot have Blueprint Event members.");
+									}
+									if (!bCanImplementInBlueprints)
+									{
+										// We do not currently warn about this case as there are a large number of existing interfaces that do not specify
+										// Function.LogWarning(TEXT("Interfaces with Blueprint Events should declare Blueprintable on the interface."));
+									}
+								}
+
+								if (Function.FunctionFlags.HasExactFlags(EFunctionFlags.BlueprintCallable | EFunctionFlags.BlueprintEvent, EFunctionFlags.BlueprintCallable))
+								{
+									// Ensure that if this interface contains blueprint callable functions that are not blueprint defined, that it must be implemented natively
+									if (bCanImplementInBlueprints)
+									{
+										Function.LogError("Blueprint implementable interfaces cannot contain BlueprintCallable functions that are not "
+											+ "BlueprintImplementableEvents. Add NotBlueprintable to the interface if you wish to keep this function.");
+									}
+									if (!bCannotImplementInBlueprints)
+									{
+										// Lowered this case to a warning instead of error, they will not show up as blueprintable unless they also have events
+										Function.LogWarning("Interfaces with BlueprintCallable functions but no events should explicitly declare NotBlueprintable on the interface.");
+									}
+								}
+							}
+						}
+					}
+					break;
+
+				case UhtClassType.NativeInterface:
+					break;
+			}
+
+			return Options |= UhtValidationOptions.Shadowing | UhtValidationOptions.Deprecated;
+		}
+
+		void ValidateBlueprintPopertyGetter(UhtProperty Property, string FunctionName, UhtFunction? TargetFunction)
+		{
+			if (TargetFunction == null)
+			{
+				Property.LogError($"Blueprint Property getter function '{FunctionName}' not found");
+				return;
+			}
+
+			if (TargetFunction.ParameterProperties.Length != 0)
+			{
+				TargetFunction.LogError($"Blueprint Property getter function '{TargetFunction.SourceName}' must not have parameters.");
+			}
+
+			UhtProperty? ReturnProperty = TargetFunction.ReturnProperty;
+			if (ReturnProperty == null || !Property.IsSameType(ReturnProperty))
+			{
+				TargetFunction.LogError($"Blueprint Property getter function '{TargetFunction.SourceName}' must have return value of type '{Property.GetUserFacingDecl()}'.");
+			}
+
+			if (TargetFunction.FunctionFlags.HasAnyFlags(EFunctionFlags.Event))
+			{
+				TargetFunction.LogError($"Blueprint Property getter function '{TargetFunction.SourceName}' cannot be a blueprint event.");
+			}
+
+			if (!TargetFunction.FunctionFlags.HasAnyFlags(EFunctionFlags.BlueprintPure))
+			{
+				TargetFunction.LogError($"Blueprint Property getter function '{TargetFunction.SourceName}' must be pure.");
+			}
+		}
+
+		void ValidateBlueprintPopertySetter(UhtProperty Property, string FunctionName, UhtFunction? TargetFunction)
+		{
+			if (TargetFunction == null)
+			{
+				Property.LogError($"Blueprint Property setter function '{FunctionName}' not found");
+				return;
+			}
+
+			if (TargetFunction.ReturnProperty != null)
+			{
+				TargetFunction.LogError($"Blueprint Property setter function '{TargetFunction.SourceName}' must not have a return value.");
+			}
+
+			if (TargetFunction.ParameterProperties.Length != 1 || !Property.IsSameType((UhtProperty)TargetFunction.ParameterProperties.Span[0]))
+			{
+				TargetFunction.LogError($"Blueprint Property setter function '{TargetFunction.SourceName}' must have exactly one parameter of type '{Property.GetUserFacingDecl()}'.");
+			}
+
+			if (TargetFunction.FunctionFlags.HasAnyFlags(EFunctionFlags.Event))
+			{
+				TargetFunction.LogError($"Blueprint Property setter function '{TargetFunction.SourceName}' cannot be a blueprint event.");
+			}
+
+			if (!TargetFunction.FunctionFlags.HasAnyFlags(EFunctionFlags.BlueprintCallable))
+			{
+				TargetFunction.LogError($"Blueprint Property setter function '{TargetFunction.SourceName}' must be a blueprint callable.");
+			}
+
+			if (TargetFunction.FunctionFlags.HasAnyFlags(EFunctionFlags.BlueprintPure))
+			{
+				TargetFunction.LogError($"Blueprint Property setter function '{TargetFunction.SourceName}' must not be pure.");
+			}
+		}
+
+		void ValidateRepNotifyCallback(UhtProperty Property, string FunctionName, UhtFunction? TargetFunction)
+		{
+			if (TargetFunction == null)
+			{
+				Property.LogError($"Replication notification function '{FunctionName}' not found");
+				return;
+			}
+
+			if (TargetFunction.ReturnProperty != null)
+			{
+				TargetFunction.LogError($"Replication notification function '{TargetFunction.SourceName}' must not have a return value.");
+			}
+
+			bool bIsArrayProperty = Property.bIsStaticArray || Property is UhtArrayProperty;
+			int MaxParms = bIsArrayProperty ? 2 : 1;
+			ReadOnlyMemory<UhtType> Parameters = TargetFunction.ParameterProperties;
+
+			if (Parameters.Length > MaxParms)
+			{
+				TargetFunction.LogError($"Replication notification function '{TargetFunction.SourceName}' has too many parameters.");
+			}
+
+			if (Parameters.Length >= 1)
+			{
+				UhtProperty Parm = (UhtProperty)Parameters.Span[0];
+				// First parameter is always the old value:
+				if (!Property.IsSameType(Parm))
+				{
+					TargetFunction.LogError($"Replication notification function '{TargetFunction.SourceName}' first (optional) parameter must be of type '{Property.GetUserFacingDecl()}'.");
+				}
+			}
+
+			if (Parameters.Length >= 2)
+			{
+				UhtProperty Parm = (UhtProperty)Parameters.Span[1];
+				// First parameter is always the old value:
+				bool bIsValid = false;
+				if (Parm is UhtArrayProperty ArrayProperty)
+				{
+					bIsValid = ArrayProperty.ValueProperty is UhtByteProperty && Parm.PropertyFlags.HasAllFlags(EPropertyFlags.ConstParm | EPropertyFlags.ReferenceParm);
+				}
+				if (!bIsValid)
+				{
+					TargetFunction.LogError($"Replication notification function '{TargetFunction.SourceName}' second (optional) parameter must be of type 'const TArray<uint8>&'.");
+				}
+			}
+		}
+
+		void ValidateProperties()
+		{
+			foreach (UhtType Child in this.Children)
+			{
+				if (Child is UhtProperty Property)
+				{
+					if (Property.PropertyFlags.HasAnyFlags(EPropertyFlags.RepNotify))
+					{
+						string Name = Property.RepNotifyName != null ? Property.RepNotifyName : String.Empty;
+						ValidateRepNotifyCallback(Property, Name, (UhtFunction?)FindType(UhtFindOptions.EngineName | UhtFindOptions.Function, Name.ToString()));
+					}
+
+					if (Property.PropertyFlags.HasAnyFlags(EPropertyFlags.BlueprintVisible))
+					{
+						if (Property.MetaData.TryGetValue(UhtNames.BlueprintGetter, out string? Getter) && Getter.Length > 0)
+						{
+							ValidateBlueprintPopertyGetter(Property, Getter, (UhtFunction?)FindType(UhtFindOptions.EngineName | UhtFindOptions.Function, Getter));
+						}
+
+						if (!Property.PropertyFlags.HasAnyFlags(EPropertyFlags.BlueprintReadOnly))
+						{
+							if (Property.MetaData.TryGetValue(UhtNames.BlueprintSetter, out string? Setter) && Setter.Length > 0)
+							{
+								ValidateBlueprintPopertySetter(Property, Setter, (UhtFunction?)FindType(UhtFindOptions.EngineName | UhtFindOptions.Function, Setter));
+							}
+						}
+					}
+
+					// Validate if we are using editor only data in a class or struct definition
+					if (Property.PropertyFlags.HasAnyFlags(EPropertyFlags.EditorOnly) && this.ClassFlags.HasAnyFlags(EClassFlags.Optional))
+					{
+						Property.LogError("Cannot specify an editor only property inside an optional class.");
+					}
+
+					// Check for getter/setter
+					if (Property.PropertyExportFlags.HasExactFlags(UhtPropertyExportFlags.GetterSpecified | UhtPropertyExportFlags.GetterSpecifiedNone | UhtPropertyExportFlags.GetterFound, UhtPropertyExportFlags.GetterSpecified))
+					{
+						string ExpectedName = Property.Getter == null ? $"Get{Property.SourceName}" : Property.Getter;
+						using (BorrowStringBuilder Borrower = new BorrowStringBuilder(StringBuilderCache.Small))
+						{
+							StringBuilder Builder = Borrower.StringBuilder;
+							Builder.AppendPropertyText(Property, UhtPropertyTextType.GetterSetterArg);
+							string ExpectedType = Builder.ToString();
+							Property.LogError($"Property '{Property.SourceName}' expected a getter function '[const] {ExpectedType} [&] {ExpectedName}() const', but it was not found");
+						}
+					}
+					if (Property.PropertyExportFlags.HasExactFlags(UhtPropertyExportFlags.SetterSpecified | UhtPropertyExportFlags.SetterSpecifiedNone | UhtPropertyExportFlags.SetterFound, UhtPropertyExportFlags.SetterSpecified))
+					{
+						string ExpectedName = Property.Setter == null ? $"Set{Property.SourceName}" : Property.Setter;
+						using (BorrowStringBuilder Borrower = new BorrowStringBuilder(StringBuilderCache.Small))
+						{
+							StringBuilder Builder = Borrower.StringBuilder;
+							Builder.AppendPropertyText(Property, UhtPropertyTextType.GetterSetterArg);
+							string ExpectedType = Builder.ToString();
+							Property.LogError($"Property '{Property.SourceName}' expected a setter function 'void {ExpectedName}([const] {ExpectedType} [&] InArg)', but it was not found");
+						}
+					}
+				}
+			}
+		}
+		#endregion
+
+		/// <inheritdoc/>
+		public override void CollectReferences(IUhtReferenceCollector Collector)
+		{
+
+			// Ignore any intrinsics
+			if (this.ClassFlags.HasAnyFlags(EClassFlags.Intrinsic))
+			{
+				return;
+			}
+
+			// In the original UHT, we don't export IInterface (it wasn't even in the header file)
+			if (this.HeaderFile.bIsNoExportTypes && this.SourceName == "IInterface")
+			{
+				return;
+			}
+
+			// Add myself as declarations and cross module references
+			Collector.AddDeclaration(this, false);
+			Collector.AddDeclaration(this, true);
+			Collector.AddCrossModuleReference(this, false);
+			Collector.AddCrossModuleReference(this, true);
+
+			// Add the super class
+			if (this.SuperClass != null)
+			{
+				Collector.AddDeclaration(this.SuperClass, true);
+				Collector.AddCrossModuleReference(this.SuperClass, true);
+			}
+
+			// Add any other base classes
+			if (this.Bases != null)
+			{
+				foreach (UhtStruct Base in this.Bases)
+				{
+					if (Base is UhtClass BaseClass)
+					{
+						Collector.AddCrossModuleReference(BaseClass, false);
+					}
+				}
+			}
+
+			// Add the package
+			Collector.AddDeclaration(this.Package, true);
+			Collector.AddCrossModuleReference(this.Package, true);
+
+			// Collect children references
+			foreach (UhtType Child in this.Children)
+			{
+				Child.CollectReferences(Collector);
+			}
+
+			// Done at the end so any contained delegate functions are added first
+			// We also add interfaces in the export type order where the native interface
+			// is located.  But we add the interface...
+			switch (this.ClassType)
+			{
+				case UhtClassType.NativeInterface:
+					if (this.AlternateObject != null && this.AlternateObject is UhtField Field)
+					{
+						Collector.AddExportType(Field);
+					}
+					break;
+				case UhtClassType.Interface:
+					break;
+				case UhtClassType.Class:
+					Collector.AddExportType(this);
+					break;
+			}
+		}
+	}
+}
diff --git a/Engine/Source/Programs/Shared/EpicGames.UHT/Types/UhtEnum.cs b/Engine/Source/Programs/Shared/EpicGames.UHT/Types/UhtEnum.cs
new file mode 100644
--- /dev/null
+++ b/Engine/Source/Programs/Shared/EpicGames.UHT/Types/UhtEnum.cs
@@ -0,0 +1,189 @@
+// Copyright Epic Games, Inc. All Rights Reserved.
+
+using EpicGames.Core;
+using EpicGames.UHT.Tables;
+using EpicGames.UHT.Utils;
+using System;
+using System.Collections.Generic;
+using System.Text.Json.Serialization;
+
+namespace EpicGames.UHT.Types
+{
+	public enum UhtEnumCppForm
+	{
+		Regular,
+		Namespaced,
+		EnumClass
+	}
+
+	public enum UhtEnumUnderlyingType
+	{
+		Unspecified,
+		uint8,
+		uint16,
+		uint32,
+		uint64,
+		int8,
+		int16,
+		int32,
+		int64
+	}
+
+	public struct UhtEnumValue
+	{
+		public string Name { get; internal set; }
+		public long Value { get; internal set; }
+	}
+
+	[UhtEngineClass(Name = "Enum")]
+	public class UhtEnum : UhtField, IUhtMetaDataKeyConversion
+	{
+		private static UhtSpecifierValidatorTable EnumSpecifierValidatorTable = UhtSpecifierValidatorTables.Instance.Get(UhtTableNames.Enum);
+
+		[JsonConverter(typeof(JsonStringEnumConverter))]
+		public EEnumFlags EnumFlags { get; internal set; } = EEnumFlags.None;
+
+		[JsonConverter(typeof(JsonStringEnumConverter))]
+		public UhtEnumCppForm CppForm { get; internal set; } = UhtEnumCppForm.Regular;
+
+		[JsonConverter(typeof(JsonStringEnumConverter))]
+		public UhtEnumUnderlyingType UnderlyingType { get; internal set; } = UhtEnumUnderlyingType.uint8;
+
+		public string CppType { get; internal set; } = String.Empty;
+		public bool bIsEditorOnly { get; internal set; } = false;
+		public List<UhtEnumValue> EnumValues { get; set; }
+
+		[JsonIgnore]
+		public override UhtEngineType EngineType => UhtEngineType.Enum;
+
+		/// <inheritdoc/>
+		public override string EngineClassName { get => "Enum"; }
+
+		///<inheritdoc/>
+		[JsonIgnore]
+		protected override UhtSpecifierValidatorTable? SpecifierValidatorTable { get => UhtEnum.EnumSpecifierValidatorTable; }
+
+		public UhtEnum(UhtType Outer, int LineNumber) : base(Outer, LineNumber)
+		{
+			this.MetaData.KeyConversion = this;
+			this.EnumValues = new List<UhtEnumValue>();
+		}
+
+		public bool IsValidEnumValue(long Value)
+		{
+			foreach (UhtEnumValue EnumValue in EnumValues)
+			{
+				if (EnumValue.Value == Value)
+				{
+					return true;
+				}
+			}
+			return false;
+		}
+
+		public int GetIndexByName(string Name)
+		{
+			Name = CleanEnumValueName(Name);
+			for (int Index = 0; Index < EnumValues.Count; ++Index)
+			{
+				if (this.EnumValues[Index].Name == Name)
+				{
+					return Index;
+				}
+			}
+			return -1;
+		}
+
+		public string GetMetaDataKey(string Name, int NameIndex)
+		{
+			string EnumName = this.EnumValues[NameIndex].Name;
+			if (this.CppForm != UhtEnumCppForm.Regular)
+			{
+				int ScopeIndex = EnumName.IndexOf("::");
+				if (ScopeIndex >= 0)
+				{
+					EnumName = EnumName.Substring(ScopeIndex + 2);
+				}
+			}
+			return $"{EnumName}.{Name}";
+		}
+
+		public string GetFullEnumName(string ShortEnumName)
+		{
+			switch (this.CppForm)
+			{
+				case UhtEnumCppForm.Namespaced:
+				case UhtEnumCppForm.EnumClass:
+					return $"{this.SourceName}::{ShortEnumName}";
+
+				case UhtEnumCppForm.Regular:
+					return ShortEnumName;
+
+				default:
+					throw new UhtIceException("Unexpected EEnumCppForm value");
+			}
+		}
+
+		public int AddEnumValue(string ShortEnumName, long Value)
+		{
+			int EnumIndex = this.EnumValues.Count;
+			this.EnumValues.Add(new UhtEnumValue { Name = GetFullEnumName(ShortEnumName), Value = Value });
+			return EnumIndex;
+		}
+
+		private string CleanEnumValueName(string Name)
+		{
+			int LastColons = Name.LastIndexOf("::");
+			return LastColons == -1 ? GetFullEnumName(Name) : GetFullEnumName(Name.Substring(LastColons + 2));		
+		}
+
+		#region Validation support
+		protected override void ValidateDocumentationPolicy(UhtDocumentationPolicy Policy)
+		{
+			if (Policy.bClassOrStructCommentRequired)
+			{
+				if (!this.MetaData.ContainsKey(UhtNames.ToolTip))
+				{
+					this.LogError(this.MetaData.LineNumber, $"Enum '{this.SourceName}' does not provide a tooltip / comment (DocumentationPolicy)");
+				}
+			}
+
+			Dictionary<string, string> ToolTips = new Dictionary<string, string>(StringComparer.OrdinalIgnoreCase);
+			for (int EnumIndex = 0; EnumIndex < this.EnumValues.Count; ++EnumIndex)
+			{
+				string? EntryName;
+				if (!this.MetaData.TryGetValue(UhtNames.Name, EnumIndex, out EntryName))
+				{
+					continue;
+				}
+
+				string? ToolTip;
+				if (!this.MetaData.TryGetValue(UhtNames.ToolTip, EnumIndex, out ToolTip))
+				{
+					this.LogError(this.MetaData.LineNumber, $"Enum entry '{this.SourceName}::{EntryName}' does not provide a tooltip / comment (DocumentationPolicy)");
+					continue;
+				}
+
+				string? DupName;
+				if (ToolTips.TryGetValue(ToolTip, out DupName))
+				{
+					this.LogError(this.MetaData.LineNumber, $"Enum entries '{this.SourceName}::{EntryName}' and '{this.SourceName}::{DupName}' have identical tooltips / comments (DocumentationPolicy)");
+				}
+				else
+				{
+					ToolTips.Add(ToolTip, EntryName);
+				}
+			}
+		}
+		#endregion
+
+		/// <inheritdoc/>
+		public override void CollectReferences(IUhtReferenceCollector Collector)
+		{
+			Collector.AddExportType(this);
+			Collector.AddDeclaration(this, true);
+			Collector.AddCrossModuleReference(this, true);
+			Collector.AddCrossModuleReference(this.Package, true);
+		}
+	}
+}
diff --git a/Engine/Source/Programs/Shared/EpicGames.UHT/Types/UhtField.cs b/Engine/Source/Programs/Shared/EpicGames.UHT/Types/UhtField.cs
new file mode 100644
--- /dev/null
+++ b/Engine/Source/Programs/Shared/EpicGames.UHT/Types/UhtField.cs
@@ -0,0 +1,14 @@
+// Copyright Epic Games, Inc. All Rights Reserved.
+
+namespace EpicGames.UHT.Types
+{
+	public abstract class UhtField : UhtObject
+	{
+		/// <inheritdoc/>
+		public override string EngineClassName { get => "Field"; }
+
+		public UhtField(UhtType Outer, int LineNumber) : base(Outer, LineNumber)
+		{
+		}
+	}
+}
diff --git a/Engine/Source/Programs/Shared/EpicGames.UHT/Types/UhtFunction.cs b/Engine/Source/Programs/Shared/EpicGames.UHT/Types/UhtFunction.cs
new file mode 100644
--- /dev/null
+++ b/Engine/Source/Programs/Shared/EpicGames.UHT/Types/UhtFunction.cs
@@ -0,0 +1,963 @@
+// Copyright Epic Games, Inc. All Rights Reserved.
+
+using EpicGames.Core;
+using EpicGames.UHT.Tables;
+using EpicGames.UHT.Utils;
+using System;
+using System.Collections.Generic;
+using System.Text;
+using System.Text.Json.Serialization;
+
+namespace EpicGames.UHT.Types
+{
+	/// <summary>
+	/// Series of flags not part of the engine's function flags that affect code generation or verification
+	/// </summary>
+	[Flags]
+	public enum UhtFunctionExportFlags : UInt32
+	{
+		None = 0,
+
+		/// <summary>
+		/// Function declaration included "final" keyword.  Used to differentiate between functions that have FUNC_Final only because they're private
+		/// </summary>
+		Final = 1 << 0,
+
+		/// <summary>
+		/// Function should be exported as a public API function
+		/// </summary>
+		RequiredAPI = 1 << 1,
+
+		/// <summary>
+		/// Export as an inline static C++ function
+		/// </summary>
+		Inline = 1 << 2,
+
+		/// <summary>
+		/// Export as a real C++ static function, causing thunks to call via ClassName::FuncName instead of this->FuncName
+		/// </summary>
+		CppStatic = 1 << 3,
+
+		/// <summary>
+		/// Export no thunk function; the user will manually define a custom one
+		/// </summary>
+		CustomThunk = 1 << 4,
+
+		/// <summary>
+		/// Function is marked as virtual
+		/// </summary>
+		Virtual = 1 << 5,
+
+		/// <summary>
+		/// The unreliable specified was present
+		/// </summary>
+		Unreliable = 1 << 6,
+
+		/// <summary>
+		/// The function is a sealed event
+		/// </summary>
+		SealedEvent = 1 << 7,
+
+		/// <summary>
+		/// Blueprint pure is being forced to false
+		/// </summary>
+		ForceBlueprintImpure = 1 << 8,
+
+		/// <summary>
+		/// If set, the BlueprintPure was automatically set
+		/// </summary>
+		AutoBlueprintPure = 1 << 9,
+
+		/// <summary>
+		/// If set, a method matching the CppImplName was found.  The search is only performed if it
+		/// differs from the function name.
+		/// </summary>
+		ImplFound = 1 << 10,
+
+		/// <summary>
+		/// If the ImplFound flag is set, then if this flag is set, the method is virtual
+		/// </summary>
+		ImplVirtual = 1 << 11,
+
+		/// <summary>
+		/// If set, a method matching the CppValidationImplName was found.  The search is only performed if 
+		/// CppImplName differs from the function name
+		/// </summary>
+		ValidationImplFound = 1 << 12,
+
+		/// <summary>
+		/// If the ValidationImplFound flag is set, then if this flag is set, the method id virtual.
+		/// </summary>
+		ValidationImplVirtual = 1 << 13,
+
+		/// <summary>
+		/// Set true if the function itself is declared const.  The is required for the automatic setting of BlueprintPure.
+		/// </summary>
+		DeclaredConst = 1 << 14,
+
+		/// <summary>
+		/// Final flag was set automatically and should not be considered for validation
+		/// </summary>
+		AutoFinal = 1 << 15,
+	}
+
+	/// <summary>
+	/// Helper methods for testing flags.  These methods perform better than the generic HasFlag which hits
+	/// the GC and stalls.
+	/// </summary>
+	public static class UhtFunctionExportFlagsExtensions
+	{
+
+		/// <summary>
+		/// Test to see if any of the specified flags are set
+		/// </summary>
+		/// <param name="InFlags">Current flags</param>
+		/// <param name="TestFlags">Flags to test for</param>
+		/// <returns>True if any of the flags are set</returns>
+		public static bool HasAnyFlags(this UhtFunctionExportFlags InFlags, UhtFunctionExportFlags TestFlags)
+		{
+			return (InFlags & TestFlags) != 0;
+		}
+
+		/// <summary>
+		/// Test to see if all of the specified flags are set
+		/// </summary>
+		/// <param name="InFlags">Current flags</param>
+		/// <param name="TestFlags">Flags to test for</param>
+		/// <returns>True if all the flags are set</returns>
+		public static bool HasAllFlags(this UhtFunctionExportFlags InFlags, UhtFunctionExportFlags TestFlags)
+		{
+			return (InFlags & TestFlags) == TestFlags;
+		}
+
+		/// <summary>
+		/// Test to see if a specific set of flags have a specific value.
+		/// </summary>
+		/// <param name="InFlags">Current flags</param>
+		/// <param name="TestFlags">Flags to test for</param>
+		/// <param name="MatchFlags">Expected value of the tested flags</param>
+		/// <returns>True if the given flags have a specific value.</returns>
+		public static bool HasExactFlags(this UhtFunctionExportFlags InFlags, UhtFunctionExportFlags TestFlags, UhtFunctionExportFlags MatchFlags)
+		{
+			return (InFlags & TestFlags) == MatchFlags;
+		}
+	}
+
+	public enum UhtFunctionType
+	{
+		Function,
+		Delegate,
+		SparseDelegate,
+	}
+
+	public static class UhtFunctionTypeExtensions
+	{
+		public static bool IsDelegate(this UhtFunctionType FunctionType)
+		{
+			return FunctionType == UhtFunctionType.Delegate || FunctionType == UhtFunctionType.SparseDelegate;
+		}
+	}
+
+	public class UhtFunction : UhtStruct
+	{
+		private static UhtSpecifierValidatorTable FunctionSpecifierValidatorTable = UhtSpecifierValidatorTables.Instance.Get(UhtTableNames.Function);
+		private string? StrippedFunctionNameInternal = null;
+
+		public static string GeneratedDelegateSignatureSuffix = "__DelegateSignature";
+
+		[JsonConverter(typeof(JsonStringEnumConverter))]
+		public EFunctionFlags FunctionFlags { get; set; } = EFunctionFlags.None;
+
+		[JsonConverter(typeof(JsonStringEnumConverter))]
+		public UhtFunctionExportFlags FunctionExportFlags { get; set; } = UhtFunctionExportFlags.None;
+
+		public UhtFunctionType FunctionType = UhtFunctionType.Function;
+
+		public int FunctionLineNumber { get; set; } = 1;
+		public int MacroLineNumber { get; set; } = 1;
+		public string? SparseOwningClassName { get; set; } = null;
+		public string? SparseDelegateName { get; set; } = null;
+
+		[JsonConverter(typeof(UhtNullableTypeSourceNameJsonConverter<UhtFunction>))]
+		public UhtFunction? SuperFunction => (UhtFunction?)this.Super;
+
+		[JsonIgnore]
+		public override UhtEngineType EngineType => this.FunctionFlags.HasAnyFlags(EFunctionFlags.Delegate) ? UhtEngineType.Delegate : UhtEngineType.Function;
+
+		/// <inheritdoc/>
+		public override string EngineClassName
+		{
+			get
+			{
+				switch (this.FunctionType)
+				{
+					case UhtFunctionType.Function:
+						return "Function";
+					case UhtFunctionType.Delegate:
+						return "DelegateFunction";
+					case UhtFunctionType.SparseDelegate:
+						return "SparseDelegateFunction";
+					default:
+						throw new UhtIceException("Invalid function type");
+				}
+			}
+		}
+
+		[JsonIgnore]
+		public string StrippedFunctionName
+		{
+			get
+			{
+				return this.StrippedFunctionNameInternal != null ? this.StrippedFunctionNameInternal : this.EngineName;
+			}
+			set 
+			{
+				this.StrippedFunctionNameInternal = value;
+			}
+		}
+
+		///<inheritdoc/>
+		[JsonIgnore]
+		public override bool bDeprecated => this.MetaData.ContainsKey(UhtNames.DeprecatedFunction);
+
+		///<inheritdoc/>
+		[JsonIgnore]
+		protected override UhtSpecifierValidatorTable? SpecifierValidatorTable { get => UhtFunction.FunctionSpecifierValidatorTable; }
+
+		/// <summary>
+		/// Identifier for an RPC call to a platform service
+		/// </summary>
+		public UInt16 RPCId = 0;
+
+		/// <summary>
+		/// Identifier for an RPC call expecting a response
+		/// </summary>
+		public UInt16 RPCResponseId = 0;
+
+		/// <summary>
+		/// Name of the actual implementation
+		/// </summary>
+		public string CppImplName = string.Empty;
+
+		/// <summary>
+		/// Name of the actual validation implementation
+		/// </summary>
+		public string CppValidationImplName = string.Empty;
+
+		/// <summary>
+		/// Name of the wrapper function that marshals the arguments and does the indirect call
+		/// </summary>
+		public string MarshalAndCallName = string.Empty;
+
+		/// <summary>
+		/// Name for callback-style names
+		/// </summary>
+		public string UnMarshalAndCallName = string.Empty;
+
+		/// <summary>
+		/// Endpoint name
+		/// </summary>
+		public string EndpointName = string.Empty;
+
+		/// <summary>
+		/// True if the function has a return value.
+		/// </summary>
+		[JsonIgnore]
+		public bool bHasReturnProperty => this.Children.Count > 0 && ((UhtProperty)this.Children[this.Children.Count - 1]).PropertyFlags.HasAnyFlags(EPropertyFlags.ReturnParm);
+
+		/// <summary>
+		/// The return value property or null if the function doesn't have a return value
+		/// </summary>
+		[JsonIgnore]
+		public UhtProperty? ReturnProperty => this.bHasReturnProperty ? (UhtProperty)this.Children[this.Children.Count - 1] : null;
+
+		/// <summary>
+		/// True if the function has parameters.
+		/// </summary>
+		[JsonIgnore]
+		public bool bHasParameters => this.Children.Count > 0 && !((UhtProperty)this.Children[0]).PropertyFlags.HasAnyFlags(EPropertyFlags.ReturnParm);
+
+		/// <summary>
+		/// Return read only memory of all the function parameters
+		/// </summary>
+		[JsonIgnore]
+		public ReadOnlyMemory<UhtType> ParameterProperties => new ReadOnlyMemory<UhtType>(this.Children.ToArray(), 0, this.bHasReturnProperty ? this.Children.Count - 1 : this.Children.Count);
+
+		/// <summary>
+		/// True if the function has any outputs including a return value
+		/// </summary>
+		[JsonIgnore]
+		public bool bHasAnyOutputs
+		{
+			get
+			{
+				foreach (UhtType Type in this.Children)
+				{
+					if (Type is UhtProperty Property)
+					{
+						if (Property.PropertyFlags.HasAnyFlags(EPropertyFlags.ReturnParm | EPropertyFlags.OutParm))
+						{
+							return true;
+						}
+					}
+				}
+				return false;
+			}
+		}
+
+		/// <summary>
+		/// Construct a new instance of a function
+		/// </summary>
+		/// <param name="Outer">The parent object</param>
+		/// <param name="LineNumber">The line number where the function is defined</param>
+		public UhtFunction(UhtType Outer, int LineNumber) : base(Outer, LineNumber)
+		{
+		}
+
+		/// <inheritdoc/>
+		public override void AddChild(UhtType Type)
+		{
+			if (Type is UhtProperty Property)
+			{
+				if (!Property.PropertyFlags.HasAnyFlags(EPropertyFlags.Parm))
+				{
+					throw new UhtIceException("Only properties marked as parameters can be added to functions");
+				}
+				if (Property.PropertyFlags.HasAnyFlags(EPropertyFlags.ReturnParm) && this.bHasReturnProperty)
+				{
+					throw new UhtIceException("Attempt to set multiple return properties on function");
+				}
+			}
+			else
+			{
+				throw new UhtIceException("Invalid type added");
+			}
+			base.AddChild(Type);
+		}
+
+		#region Resolution support
+		protected override bool ResolveSelf(UhtResolvePhase Phase)
+		{
+			bool bResults = base.ResolveSelf(Phase);
+
+			switch (Phase)
+			{
+				case UhtResolvePhase.Bases:
+					if (this.FunctionType == UhtFunctionType.Function && this.Outer is UhtClass OuterClass)
+					{
+						// non-static functions in a const class must be const themselves
+						if (OuterClass.ClassFlags.HasAnyFlags(EClassFlags.Const))
+						{
+							this.FunctionFlags |= EFunctionFlags.Const;
+						}
+					}
+					break;
+
+				case UhtResolvePhase.Final:
+					if (this.Outer is UhtClass Class)
+					{
+						if (this.FunctionFlags.HasAnyFlags(EFunctionFlags.Native) &&
+							!this.FunctionExportFlags.HasAnyFlags(UhtFunctionExportFlags.CustomThunk) &&
+							this.CppImplName != this.SourceName)
+						{
+							UhtFoundDeclaration Declaration;
+							if (Class.TryGetDeclaration(this.CppImplName, out Declaration))
+							{
+								this.FunctionExportFlags |= UhtFunctionExportFlags.ImplFound;
+								if (Declaration.bIsVirtual)
+								{
+									this.FunctionExportFlags |= UhtFunctionExportFlags.ImplVirtual;
+								}
+							}
+							if (Class.TryGetDeclaration(this.CppValidationImplName, out Declaration))
+							{
+								this.FunctionExportFlags |= UhtFunctionExportFlags.ValidationImplFound;
+								if (Declaration.bIsVirtual)
+								{
+									this.FunctionExportFlags |= UhtFunctionExportFlags.ValidationImplVirtual;
+								}
+							}
+						}
+
+						if (this.FunctionType == UhtFunctionType.Function && this.FunctionExportFlags.HasAnyFlags(UhtFunctionExportFlags.DeclaredConst))
+						{
+							// @todo: the presence of const and one or more outputs does not guarantee that there are
+							// no side effects. On GCC and clang we could use __attribure__((pure)) or __attribute__((const))
+							// or we could just rely on the use marking things BlueprintPure. Either way, checking the C++
+							// const identifier to determine purity is not desirable. We should remove the following logic:
+
+							// If its a const BlueprintCallable function with some sort of output and is not being marked as an BlueprintPure=false function, mark it as BlueprintPure as well
+							if (this.bHasAnyOutputs && this.FunctionFlags.HasAnyFlags(EFunctionFlags.BlueprintCallable) && 
+								!this.FunctionExportFlags.HasAnyFlags(UhtFunctionExportFlags.ForceBlueprintImpure))
+							{
+								this.FunctionFlags |= EFunctionFlags.BlueprintPure;
+								this.FunctionExportFlags |= UhtFunctionExportFlags.AutoBlueprintPure; // Disable error for pure being set
+							}
+						}
+					}
+
+					// The following code is only performed on functions in a class.
+					if (this.Outer is UhtClass)
+					{
+						foreach (UhtType Type in this.Children)
+						{
+							if (Type is UhtProperty Property)
+							{
+								if (Property.PropertyFlags.HasExactFlags(EPropertyFlags.OutParm | EPropertyFlags.ReturnParm, EPropertyFlags.OutParm))
+								{
+									this.FunctionFlags |= EFunctionFlags.HasOutParms;
+								}
+								if (Property is UhtStructProperty StructProperty)
+								{
+									if (StructProperty.ScriptStruct.bHasDefaults)
+									{
+										this.FunctionFlags |= EFunctionFlags.HasDefaults;
+									}
+								}
+							}
+						}
+					}
+					break;
+			}
+			return bResults;
+		}
+		#endregion
+
+		#region Validation support
+		protected override UhtValidationOptions Validate(UhtValidationOptions Options)
+		{
+			Options = base.Validate(Options);
+
+			ValidateSpecifierFlags();
+
+			// Some parameter test are optionally done
+			bool bCheckForBlueprint = false;
+
+			// Function type specific validation
+			switch (this.FunctionType)
+			{
+				case UhtFunctionType.Function:
+					UhtClass? OuterClass = (UhtClass?)this.Outer;
+
+					if (this.FunctionFlags.HasAnyFlags(EFunctionFlags.BlueprintPure) && OuterClass != null && 
+						OuterClass.ClassFlags.HasAnyFlags(EClassFlags.Interface) && 
+						!this.FunctionExportFlags.HasAnyFlags(UhtFunctionExportFlags.AutoBlueprintPure))
+					{
+						this.LogError("BlueprintPure specifier is not allowed for interface functions");
+					}
+
+					// If this function is blueprint callable or blueprint pure, require a category 
+					if (this.FunctionFlags.HasAnyFlags(EFunctionFlags.BlueprintCallable | EFunctionFlags.BlueprintPure))
+					{
+						bool bBlueprintAccessor = this.MetaData.ContainsKey(UhtNames.BlueprintSetter) || this.MetaData.ContainsKey(UhtNames.BlueprintGetter);
+						bool bHasMenuCategory = this.MetaData.ContainsKey(UhtNames.Category);
+						bool bInternalOnly = this.MetaData.GetBoolean(UhtNames.BlueprintInternalUseOnly);
+
+						if (!bHasMenuCategory && !bInternalOnly && !this.bDeprecated && !bBlueprintAccessor)
+						{
+							// To allow for quick iteration, don't enforce the requirement that game functions have to be categorized
+							if (this.HeaderFile.Package.bIsPartOfEngine)
+							{
+								this.LogError("An explicit Category specifier is required for Blueprint accessible functions in an Engine module.");
+							}
+						}
+					}
+
+					// Process the virtual-ness
+					if (this.FunctionExportFlags.HasAnyFlags(UhtFunctionExportFlags.Virtual))
+					{
+
+						// if this is a BlueprintNativeEvent or BlueprintImplementableEvent in an interface, make sure it's not "virtual"
+						if (this.FunctionFlags.HasAnyFlags(EFunctionFlags.BlueprintEvent))
+						{
+							if (OuterClass != null && OuterClass.ClassFlags.HasAnyFlags(EClassFlags.Interface))
+							{
+								this.LogError("BlueprintImplementableEvents in Interfaces must not be declared 'virtual'");
+							}
+
+							// if this is a BlueprintNativeEvent, make sure it's not "virtual"
+							else if (this.FunctionFlags.HasAnyFlags(EFunctionFlags.Native))
+							{
+								this.LogError("BlueprintNativeEvent functions must be non-virtual.");
+							}
+
+							else
+							{
+								this.LogWarning("BlueprintImplementableEvents should not be virtual. Use BlueprintNativeEvent instead.");
+							}
+						}
+					}
+					else
+					{
+						// if this is a function in an Interface, it must be marked 'virtual' unless it's an event
+						if (OuterClass != null && OuterClass.ClassFlags.HasAnyFlags(EClassFlags.Interface) && !this.FunctionFlags.HasAnyFlags(EFunctionFlags.BlueprintEvent))
+						{
+							this.LogError("Interface functions that are not BlueprintImplementableEvents must be declared 'virtual'");
+						}
+					}
+
+					// This is a final (prebinding, non-overridable) function
+					if (this.FunctionFlags.HasAnyFlags(EFunctionFlags.Final) || this.FunctionExportFlags.HasAnyFlags(UhtFunctionExportFlags.Final))
+					{
+						if (OuterClass != null && OuterClass.ClassFlags.HasAnyFlags(EClassFlags.Interface))
+						{
+							this.LogError("Interface functions cannot be declared 'final'");
+						}
+						else if (this.FunctionFlags.HasAnyFlags(EFunctionFlags.BlueprintEvent) && !this.FunctionExportFlags.HasAnyFlags(UhtFunctionExportFlags.AutoFinal))
+						{
+							this.LogError("Blueprint events cannot be declared 'final'");
+						}
+					}
+
+					if (this.FunctionFlags.HasAnyFlags(EFunctionFlags.RequiredAPI) &&
+						OuterClass != null && OuterClass.ClassFlags.HasAnyFlags(EClassFlags.RequiredAPI))
+					{
+						this.LogError($"API must not be used on methods of a class that is marked with an API itself.");
+					}
+
+					// Make sure that blueprint pure functions have some type of output
+					if (this.FunctionFlags.HasAnyFlags(EFunctionFlags.BlueprintPure) && !this.bHasAnyOutputs)
+					{
+						this.LogError("BlueprintPure specifier is not allowed for functions with no return value and no output parameters.");
+					}
+
+					if (this.FunctionFlags.HasExactFlags(EFunctionFlags.Net | EFunctionFlags.NetRequest | EFunctionFlags.NetResponse, EFunctionFlags.Net) && this.ReturnProperty != null)
+					{
+						this.LogError("Replicated functions can't have return values");
+					}
+
+					bCheckForBlueprint = this.FunctionFlags.HasAnyFlags(EFunctionFlags.BlueprintCallable | EFunctionFlags.BlueprintEvent);
+
+					if (this.Outer is UhtStruct OuterStruct)
+					{
+						UhtStruct? OuterSuperStruct = OuterStruct.SuperStruct;
+						if (OuterSuperStruct != null)
+						{
+							UhtType? OverriddenFunction = OuterSuperStruct.FindType(UhtFindOptions.SourceName | UhtFindOptions.Function, this.SourceName);
+							if (OverriddenFunction != null)
+							{
+								// Native function overrides should be done in CPP text, not in a UFUNCTION() declaration (you can't change flags, and it'd otherwise be a burden to keep them identical)
+								this.LogError(
+									$"Override of UFUNCTION '{this.SourceName}' in parent '{OverriddenFunction.Outer?.SourceName}' cannot have a UFUNCTION() declaration above it; it will use the same parameters as the original declaration.");
+							}
+						}
+					}
+
+					// Make sure that the replication flags set on an overridden function match the parent function
+					UhtFunction? SuperFunction = this.SuperFunction;
+					if (SuperFunction != null)
+					{
+						if ((this.FunctionFlags & EFunctionFlags.NetFuncFlags) != (SuperFunction.FunctionFlags & EFunctionFlags.NetFuncFlags))
+						{
+							this.LogError($"Overridden function '{this.SourceName}' cannot specify different replication flags when overriding a function.");
+						}
+					}
+
+					// if this function is an RPC in state scope, verify that it is an override
+					// this is required because the networking code only checks the class for RPCs when initializing network data, not any states within it
+					if (this.FunctionFlags.HasAnyFlags(EFunctionFlags.Net) && SuperFunction == null && !(this.Outer is UhtClass))
+					{
+						this.LogError($"Function '{this.SourceName}' base implementation of RPCs cannot be in a state. Add a stub outside state scope.");
+					}
+
+					// Check implemented RPC functions
+					if (this.CppImplName != this.EngineName && 
+						!this.FunctionExportFlags.HasAnyFlags(UhtFunctionExportFlags.CustomThunk) &&
+						OuterClass!.GeneratedCodeVersion > EGeneratedCodeVersion.V1)
+					{
+						bool bIsNative = this.FunctionFlags.HasAnyFlags(EFunctionFlags.Native);
+						bool bIsNet = this.FunctionFlags.HasAnyFlags(EFunctionFlags.Net);
+						bool bIsNetValidate = this.FunctionFlags.HasAnyFlags(EFunctionFlags.NetValidate);
+						bool bIsNetResponse = this.FunctionFlags.HasAnyFlags(EFunctionFlags.NetResponse);
+						bool bIsBlueprintEvent = this.FunctionFlags.HasAnyFlags(EFunctionFlags.BlueprintEvent);
+
+						bool bNeedsImplementation = (bIsNet && !bIsNetResponse) || bIsBlueprintEvent || bIsNative;
+						bool bNeedsValidate = (bIsNative || bIsNet) && !bIsNetResponse && bIsNetValidate;
+
+						bool bHasImplementation = this.FunctionExportFlags.HasAnyFlags(UhtFunctionExportFlags.ImplFound);
+						bool bHasValidate = this.FunctionExportFlags.HasAnyFlags(UhtFunctionExportFlags.ValidationImplFound);
+						bool bHasVirtualImplementation = this.FunctionExportFlags.HasAnyFlags(UhtFunctionExportFlags.ImplVirtual);
+						bool bHasVirtualValidate = this.FunctionExportFlags.HasAnyFlags(UhtFunctionExportFlags.ValidationImplVirtual);
+
+						if (bNeedsImplementation || bNeedsValidate)
+						{
+							// Check if functions are missing.
+							if (bNeedsImplementation && !bHasImplementation)
+							{
+								LogRpcFunctionError(false, this.CppImplName);
+							}
+
+							if (bNeedsValidate && !bHasValidate)
+							{
+								LogRpcFunctionError(false, this.CppValidationImplName);
+							}
+
+							// If all needed functions are declared, check if they have virtual specifiers.
+							if (bNeedsImplementation && bHasImplementation && !bHasVirtualImplementation)
+							{
+								LogRpcFunctionError(true, this.CppImplName);
+							}
+
+							if (bNeedsValidate && bHasValidate && !bHasVirtualValidate)
+							{
+								LogRpcFunctionError(true, this.CppValidationImplName);
+							}
+						}
+					}
+					break;
+
+				case UhtFunctionType.Delegate:
+				case UhtFunctionType.SparseDelegate:
+
+					// Check for shadowing if requested
+					if (Options.HasAnyFlags(UhtValidationOptions.Shadowing))
+					{
+						UhtType? Existing = this.Outer!.FindType(UhtFindOptions.DelegateFunction | UhtFindOptions.EngineName | UhtFindOptions.SelfOnly, this.EngineName);
+						if (Existing != null && Existing != this)
+						{
+							this.LogError($"Can't override delegate signature function '{this.SourceName}'");
+						}
+					}
+					break;
+			}
+
+			// Make sure this isn't a duplicate type.  In some ways this has already been check above
+			UhtType? ExistingType = this.Outer!.FindType(UhtFindOptions.PropertyOrFunction | UhtFindOptions.EngineName | UhtFindOptions.SelfOnly, this.EngineName);
+			if (ExistingType != null && ExistingType != this)
+			{
+				this.LogError($"'{this.EngineName}' conflicts with '{ExistingType.FullName}'");
+			}
+
+			// Do a more generic validation of the arguments
+			foreach (UhtType Type in this.Children)
+			{
+				if (Type is UhtProperty Property)
+				{
+					// This code is technically incorrect but here for compatibility reasons.  Container types (TObjectPtr, TArray, etc...) should have the 
+					// const on the inside of the template arguments but it places the const on the outside which isn't correct.  This needs to be addressed
+					// in any updates to the type parser.  See EGetVarTypeOptions::NoAutoConst for the parser side of the problem.
+					if (!Property.PropertyFlags.HasAnyFlags(EPropertyFlags.ConstParm) && Property.MustBeConstArgument(out UhtType? ErrorType))
+					{
+						if (Property.PropertyFlags.HasAnyFlags(EPropertyFlags.ReturnParm))
+						{
+							Property.LogError($"Return value must be 'const' since '{ErrorType.SourceName}' is marked 'const'");
+						}
+						else
+						{
+							Property.LogError($"Argument '{Property.SourceName}' must be 'const' since '{ErrorType.SourceName}' is marked 'const'");
+						}
+					}
+
+					// Due to structures that might not be fully parsed at parse time, do blueprint validation here
+					if (bCheckForBlueprint)
+					{
+						if (Property.bIsStaticArray)
+						{
+							this.LogError($"Static array cannot be exposed to blueprint. Function: {this.SourceName} Parameter {Property.SourceName}");
+						}
+						if (!Property.PropertyCaps.HasAnyFlags(UhtPropertyCaps.IsParameterSupportedByBlueprint))
+						{
+							this.LogError($"Type '{Property.GetUserFacingDecl()}' is not supported by blueprint. Function:  {this.SourceName} Parameter {Property.SourceName}");
+						}
+					}
+				}
+			}
+
+			if (!this.bDeprecated)
+			{
+				Options |= UhtValidationOptions.Deprecated;
+			}
+			return Options;
+		}
+
+		private void LogRpcFunctionError(bool bVirtualError, string FunctionName)
+		{
+			StringBuilder Builder = new StringBuilder();
+
+			if (bVirtualError)
+			{
+				Builder.Append("Declared function \"");
+			}
+			else
+			{
+				Builder.Append($"Function {this.SourceName} was marked as ");
+				if (this.FunctionFlags.HasAnyFlags(EFunctionFlags.Native))
+				{
+					Builder.Append("Native, ");
+				}
+				if (this.FunctionFlags.HasAnyFlags(EFunctionFlags.Net))
+				{
+					Builder.Append("Net, ");
+				}
+				if (this.FunctionFlags.HasAnyFlags(EFunctionFlags.BlueprintEvent))
+				{
+					Builder.Append("BlueprintEvent, ");
+				}
+				if (this.FunctionFlags.HasAnyFlags(EFunctionFlags.NetValidate))
+				{
+					Builder.Append("NetValidate, ");
+				}
+				Builder.Length = Builder.Length - 2;
+				Builder.Append(". Declare function \"virtual ");
+			}
+
+			UhtProperty? ReturnType = this.ReturnProperty;
+			if (ReturnType != null)
+			{
+				Builder.AppendPropertyText(ReturnType, UhtPropertyTextType.EventFunction);
+				Builder.Append(' ');
+			}
+			else
+			{
+				Builder.Append("void ");
+			}
+
+			Builder.Append(this.Outer!.SourceName).Append("::").Append(FunctionName).Append('(');
+
+			bool bFirst = true;
+			foreach (UhtProperty Arg in this.ParameterProperties.Span)
+			{
+				if (!bFirst)
+				{
+					Builder.Append(", ");
+				}
+				bFirst = false;
+				Builder.AppendFullDecl(Arg, UhtPropertyTextType.EventFunction, true);
+			}
+
+			Builder.Append(')');
+			if (this.FunctionFlags.HasAnyFlags(EFunctionFlags.Const))
+			{
+				Builder.Append(" const");
+			}
+
+			if (bVirtualError)
+			{
+				Builder.Append("\" is not marked as virtual.");
+			}
+			else
+			{
+				Builder.Append("\"");
+			}
+			this.LogError(Builder.ToString());
+		}
+
+		private void ValidateSpecifierFlags()
+		{
+			if (this.FunctionFlags.HasAnyFlags(EFunctionFlags.Net))
+			{
+				if (this.FunctionFlags.HasAnyFlags(EFunctionFlags.BlueprintEvent | EFunctionFlags.Exec))
+				{
+					this.LogError("Replicated functions can not be blueprint events or exec");
+				}
+
+				bool bIsNetService = this.FunctionFlags.HasAnyFlags(EFunctionFlags.NetRequest | EFunctionFlags.NetResponse);
+				bool bIsNetReliable = this.FunctionFlags.HasAnyFlags(EFunctionFlags.NetReliable);
+
+				if (this.FunctionFlags.HasAnyFlags(EFunctionFlags.Static))
+				{
+					this.LogError("Static functions can't be replicated");
+				}
+
+				if (!bIsNetReliable && !this.FunctionExportFlags.HasAnyFlags(UhtFunctionExportFlags.Unreliable) && !bIsNetService)
+				{
+					this.LogError("Replicated function: 'reliable' or 'unreliable' is required");
+				}
+
+				if (bIsNetReliable && this.FunctionExportFlags.HasAnyFlags(UhtFunctionExportFlags.Unreliable) && !bIsNetService)
+				{
+					this.LogError("'reliable' and 'unreliable' are mutually exclusive");
+				}
+			}
+			else if (this.FunctionFlags.HasAnyFlags(EFunctionFlags.NetReliable))
+			{
+				this.LogError("'reliable' specified without 'client' or 'server'");
+			}
+			else if (this.FunctionExportFlags.HasAnyFlags(UhtFunctionExportFlags.Unreliable))
+			{
+				this.LogError("'unreliable' specified without 'client' or 'server'");
+			}
+
+			if (this.FunctionExportFlags.HasAnyFlags(UhtFunctionExportFlags.SealedEvent) && !this.FunctionFlags.HasAnyFlags(EFunctionFlags.Event))
+			{
+				this.LogError("SealedEvent may only be used on events");
+			}
+
+			if (this.FunctionExportFlags.HasAnyFlags(UhtFunctionExportFlags.SealedEvent) && this.FunctionFlags.HasAnyFlags(EFunctionFlags.BlueprintEvent))
+			{
+				this.LogError("SealedEvent cannot be used on Blueprint events");
+			}
+
+			if (this.FunctionExportFlags.HasAnyFlags(UhtFunctionExportFlags.ForceBlueprintImpure) && this.FunctionFlags.HasAnyFlags(EFunctionFlags.BlueprintPure))
+			{
+				this.LogError("BlueprintPure (or BlueprintPure=true) and BlueprintPure=false should not both appear on the same function, they are mutually exclusive");
+			}
+		}
+
+		protected override void ValidateDocumentationPolicy(UhtDocumentationPolicy Policy)
+		{
+			if (this.FunctionType != UhtFunctionType.Function)
+			{
+				return;
+			}
+
+			UhtClass? OuterClass = (UhtClass?)this.Outer;
+
+			if (Policy.bFunctionToolTipsRequired)
+			{
+				if (!this.MetaData.ContainsKey(UhtNames.ToolTip))
+				{
+					this.LogError($"Function '{OuterClass?.SourceName}::{this.SourceName}' does not provide a tooltip / comment (DocumentationPolicy).");
+				}
+			}
+
+			if (Policy.bParameterToolTipsRequired)
+			{
+				if (!this.MetaData.ContainsKey(UhtNames.Comment))
+				{
+					this.LogError($"Function '{OuterClass?.SourceName}::{this.SourceName}' does not provide a comment (DocumentationPolicy).");
+
+				}
+
+				Dictionary<StringView, StringView> ParamToolTips = GetParameterToolTipsFromFunctionComment(this.MetaData.GetValueOrDefault(UhtNames.Comment));
+				bool bHasAnyParamToolTips = ParamToolTips.Count > 1 || (ParamToolTips.Count == 1 && !ParamToolTips.ContainsKey(UhtNames.ReturnValue));
+
+				// only apply the validation for parameter tooltips if a function has any @param statements at all.
+				if (bHasAnyParamToolTips)
+				{
+					// ensure each parameter has a tooltip
+					HashSet<StringView> ExistingFields = new HashSet<StringView>(StringViewComparer.OrdinalIgnoreCase);
+					foreach (UhtType Parameter in this.ParameterProperties.Span)
+					{
+						ExistingFields.Add(Parameter.SourceName);
+						if (!ParamToolTips.ContainsKey(Parameter.SourceName))
+						{
+							this.LogError($"Function '{OuterClass?.SourceName}::{this.SourceName}' doesn't provide a tooltip for parameter '{Parameter.SourceName}' (DocumentationPolicy).");
+						}
+					}
+
+					// ensure we don't have parameter tooltips for parameters that don't exist
+					foreach (KeyValuePair<StringView, StringView> Kvp in ParamToolTips)
+					{
+						if (Kvp.Key == UhtNames.ReturnValue)
+						{
+							continue;
+						}
+
+						if (!ExistingFields.Contains(Kvp.Key))
+						{
+							this.LogError($"Function '{OuterClass?.SourceName}::{this.SourceName}' provides a tooltip for an unknown parameter '{Kvp.Key}'");
+						}
+					}
+
+					// check for duplicate tooltips
+					Dictionary<StringView, StringView> ToolTipToParam = new Dictionary<StringView, StringView>(StringViewComparer.OrdinalIgnoreCase);
+					foreach (KeyValuePair<StringView, StringView> Kvp in ParamToolTips)
+					{
+						if (Kvp.Key == UhtNames.ReturnValue)
+						{
+							continue;
+						}
+
+						if (ToolTipToParam.TryGetValue(Kvp.Value, out StringView Existing))
+						{
+							this.LogError($"Function '{OuterClass?.SourceName}::{this.SourceName}' uses identical tooltips for parameters '{Kvp.Key}' and '{Existing}' (DocumentationPolicy).");
+						}
+						else
+						{
+							ToolTipToParam.Add(Kvp.Value, Kvp.Key);
+						}
+					}
+				}
+			}
+		}
+
+		/// <summary>
+		/// Add a parameter tooltip to the collection
+		/// </summary>
+		/// <param name="Name">Name of the parameter</param>
+		/// <param name="Text">Text of the parameter tooltip</param>
+		/// <param name="ParamMap">Parameter map to add the tooltip</param>
+		private void AddParameterToolTip(StringView Name, StringView Text, Dictionary<StringView, StringView> ParamMap)
+		{
+			ParamMap.Add(Name, new StringView(UhtFCString.TabsToSpaces(Text.Memory.Trim(), 4, false)));
+		}
+
+		/// <summary>
+		/// Add a parameter tooltip to the collection where the name of the parameter is the first word of the line
+		/// </summary>
+		/// <param name="Text">Text of the parameter tooltip</param>
+		/// <param name="ParamMap">Parameter map to add the tooltip</param>
+		private void AddParameterToolTip(StringView Text, Dictionary<StringView, StringView> ParamMap)
+		{
+			ReadOnlyMemory<char> Trimmed = Text.Memory.Trim();
+			int NameEnd = Trimmed.Span.IndexOf(' ');
+			if (NameEnd >= 0)
+			{
+				AddParameterToolTip(new StringView(Trimmed, 0, NameEnd), new StringView(Trimmed, NameEnd + 1), ParamMap);
+			}
+		}
+
+		/// <summary>
+		/// Parse the function comment looking for parameter documentation.
+		/// </summary>
+		/// <param name="Input">The function input comment</param>
+		/// <returns>Dictionary of parameter names and the documentation.  The return value will have a name of "ReturnValue"</returns>
+		private Dictionary<StringView, StringView> GetParameterToolTipsFromFunctionComment(StringView Input)
+		{
+			const string ParamTag = "@param";
+			const string ReturnTag = "@return";
+
+			Dictionary<StringView, StringView> ParamMap = new Dictionary<StringView, StringView>(StringViewComparer.OrdinalIgnoreCase);
+
+			// Loop until we are out of text
+			while (Input.Length > 0)
+			{
+
+				// Extract the line
+				StringView Line = Input;
+				int LineEnd = Input.Span.IndexOf('\n');
+				if (LineEnd >= 0)
+				{
+					Line = new StringView(Input, 0, LineEnd);
+					Input = new StringView(Input, LineEnd + 1);
+				}
+
+				// If this is a parameter, invoke the method to extract the name first
+				int ParamStart = Line.Span.IndexOf(ParamTag, StringComparison.OrdinalIgnoreCase);
+				if (ParamStart >= 0)
+				{
+					AddParameterToolTip(new StringView(Line, ParamStart + ParamTag.Length), ParamMap);
+					continue;
+				}
+
+				// If this is the return value, we have a known name.  Invoke the method where we already know the name
+				ParamStart = Line.Span.IndexOf(ReturnTag, StringComparison.OrdinalIgnoreCase);
+				if (ParamStart >= 0)
+				{
+					AddParameterToolTip(new StringView(UhtNames.ReturnValue), new StringView(Line, ParamStart + ReturnTag.Length), ParamMap);
+				}
+			}
+
+			return ParamMap;
+		}
+		#endregion
+
+		/// <inheritdoc/>
+		public override void CollectReferences(IUhtReferenceCollector Collector)
+		{
+			if (this.FunctionType != UhtFunctionType.Function)
+			{
+				Collector.AddExportType(this);
+				Collector.AddSingleton(this);
+				Collector.AddDeclaration(this, true);
+				Collector.AddCrossModuleReference(this, true);
+				Collector.AddCrossModuleReference(this.Package, true);
+			}
+			foreach (UhtType Child in this.Children)
+			{
+				Child.CollectReferences(Collector);
+			}
+		}
+	}
+}
diff --git a/Engine/Source/Programs/Shared/EpicGames.UHT/Types/UhtHeaderFile.cs b/Engine/Source/Programs/Shared/EpicGames.UHT/Types/UhtHeaderFile.cs
new file mode 100644
--- /dev/null
+++ b/Engine/Source/Programs/Shared/EpicGames.UHT/Types/UhtHeaderFile.cs
@@ -0,0 +1,300 @@
+// Copyright Epic Games, Inc. All Rights Reserved.
+
+using EpicGames.Core;
+using EpicGames.UHT.Utils;
+using System;
+using System.Collections.Generic;
+using System.IO;
+using System.Text;
+using System.Text.Json.Serialization;
+
+namespace EpicGames.UHT.Types
+{
+	public enum UhtHeaderFileType
+	{
+		Classes,
+		Public,
+		Internal,
+		Private,
+	}
+
+	/// <summary>
+	/// Series of flags not part of the engine's class flags that affect code generation or verification
+	/// </summary>
+	[Flags]
+	public enum UhtHeaderFileExportFlags : UInt32
+	{
+		None = 0,
+
+		/// <summary>
+		/// This header is being included by another header
+		/// </summary>
+		Referenced = 0x00000001, 
+	}
+
+	/// <summary>
+	/// Helper methods for testing flags.  These methods perform better than the generic HasFlag which hits
+	/// the GC and stalls.
+	/// </summary>
+	public static class UhtHeaderFileExportFlagsExtensions
+	{
+
+		/// <summary>
+		/// Test to see if any of the specified flags are set
+		/// </summary>
+		/// <param name="InFlags">Current flags</param>
+		/// <param name="TestFlags">Flags to test for</param>
+		/// <returns>True if any of the flags are set</returns>
+		public static bool HasAnyFlags(this UhtHeaderFileExportFlags InFlags, UhtHeaderFileExportFlags TestFlags)
+		{
+			return (InFlags & TestFlags) != 0;
+		}
+
+		/// <summary>
+		/// Test to see if all of the specified flags are set
+		/// </summary>
+		/// <param name="InFlags">Current flags</param>
+		/// <param name="TestFlags">Flags to test for</param>
+		/// <returns>True if all the flags are set</returns>
+		public static bool HasAllFlags(this UhtHeaderFileExportFlags InFlags, UhtHeaderFileExportFlags TestFlags)
+		{
+			return (InFlags & TestFlags) == TestFlags;
+		}
+
+		/// <summary>
+		/// Test to see if a specific set of flags have a specific value.
+		/// </summary>
+		/// <param name="InFlags">Current flags</param>
+		/// <param name="TestFlags">Flags to test for</param>
+		/// <param name="MatchFlags">Expected value of the tested flags</param>
+		/// <returns>True if the given flags have a specific value.</returns>
+		public static bool HasExactFlags(this UhtHeaderFileExportFlags InFlags, UhtHeaderFileExportFlags TestFlags, UhtHeaderFileExportFlags MatchFlags)
+		{
+			return (InFlags & TestFlags) == MatchFlags;
+		}
+	}
+
+	public delegate bool UhtReferencedHeaderDelegate(UhtHeaderFile HeaderFile);
+
+	public class UhtHeaderFile : UhtType
+	{
+		[JsonIgnore]
+		public StringView Data { get => this.SourceFile.Data; }
+
+		[JsonIgnore]
+		public string FilePath { get => this.SourceFile.FilePath; }
+
+		public readonly string FileNameWithoutExtension;
+		public readonly string GeneratedHeaderFileName;
+		public readonly bool bIsNoExportTypes;
+
+		public string ModuleRelativeFilePath = string.Empty;
+		public string IncludeFilePath = string.Empty;
+		public UhtHeaderFileType HeaderFileType = UhtHeaderFileType.Private;
+		public readonly int HeaderFileTypeIndex;
+
+		[JsonConverter(typeof(JsonStringEnumConverter))]
+		public UhtHeaderFileExportFlags HeaderFileExportFlags { get; set; } = UhtHeaderFileExportFlags.None;
+
+		[JsonIgnore]
+		public bool bShouldExport => this.HeaderFileExportFlags.HasAnyFlags(UhtHeaderFileExportFlags.Referenced) || this.Children.Count > 0;
+		public UhtReferenceCollector References = new UhtReferenceCollector();
+
+		[JsonIgnore]
+		public override UhtPackage Package
+		{
+			get
+			{
+				if (this.Outer == null)
+				{
+					throw new UhtIceException("Attempt to fetch header file package but it has no outer");
+				}
+				return (UhtPackage)this.Outer;
+			}
+		}
+
+		[JsonIgnore]
+		public override UhtHeaderFile HeaderFile { get => this; }
+
+		[JsonIgnore]
+		public override UhtEngineType EngineType => UhtEngineType.Header;
+
+		/// <inheritdoc/>
+		public override string EngineClassName { get => "UhtHeaderFile"; }
+
+		private List<UhtHeaderFile> ReferencedHeaders = new List<UhtHeaderFile>();
+
+		[JsonIgnore]
+		public List<UhtHeaderFile> IncludedHeaders { get; set; } = new List<UhtHeaderFile>();
+
+		private UhtSimpleMessageSite MessageSite;
+
+		#region IUHTMessageSite implementation
+		[JsonIgnore]
+		public override IUhtMessageSession MessageSession => this.MessageSite.MessageSession;
+
+		[JsonIgnore]
+		public override IUhtMessageSource? MessageSource => this.MessageSite.MessageSource;
+		#endregion
+
+		private UhtSourceFile SourceFile;
+
+		public UhtHeaderFile(UhtPackage Package, string Path) : base(Package, 1)
+		{
+			this.HeaderFileTypeIndex = this.Session.GetNextHeaderFileTypeIndex();
+			this.MessageSite = new UhtSimpleMessageSite(this.Session);
+			this.SourceFile = new UhtSourceFile(this.Session, Path);
+			this.MessageSite.MessageSource = this.SourceFile;
+			this.FileNameWithoutExtension = System.IO.Path.GetFileNameWithoutExtension(this.SourceFile.FilePath);
+			this.GeneratedHeaderFileName = this.FileNameWithoutExtension + ".generated.h";
+			this.SourceName = System.IO.Path.GetFileName(this.SourceFile.FilePath);
+			this.bIsNoExportTypes = string.Compare(this.SourceFile.FileName, "NoExportTypes", true) == 0;
+		}
+
+		public void Read()
+		{
+			this.SourceFile.Read();
+		}
+
+		public void AddReferencedHeader(string Id, bool bIsIncludedFile)
+		{
+			UhtHeaderFile? HeaderFile = this.Session.FindHeaderFile(Path.GetFileName(Id));
+			if (HeaderFile != null)
+			{
+				AddReferencedHeader(HeaderFile, bIsIncludedFile);
+			}
+		}
+
+		public void AddReferencedHeader(UhtType Type)
+		{
+			AddReferencedHeader(Type.HeaderFile, false);
+		}
+
+		public void AddReferencedHeader(UhtHeaderFile HeaderFile, bool bIsIncludedFile)
+		{
+			lock (this.ReferencedHeaders)
+			{
+				// Check for a duplicate
+				foreach (UhtHeaderFile Ref in this.ReferencedHeaders)
+				{
+					if (Ref == HeaderFile)
+					{
+						return;
+					}
+				}
+
+				// There is questionable compatibility hack where a source file will always be exported
+				// regardless of having types when it is being included by the SAME package.
+				if (HeaderFile.Package == this.Package)
+				{
+					HeaderFile.HeaderFileExportFlags |= UhtHeaderFileExportFlags.Referenced;
+				}
+				this.ReferencedHeaders.Add(HeaderFile);
+				if (bIsIncludedFile)
+				{
+					this.IncludedHeaders.Add(HeaderFile);
+				}
+			}
+		}
+
+		/// <summary>
+		/// Return an enumerator without locking.  This method can only be utilized AFTER all header parsing is complete. 
+		/// </summary>
+		[JsonIgnore]
+		public IEnumerable<UhtHeaderFile> ReferencedHeadersNoLock
+		{
+			get
+			{
+				return this.ReferencedHeaders;
+			}
+		}
+
+		/// <summary>
+		/// Return an enumerator of all current referenced headers under a lock.  This should be used during parsing.
+		/// </summary>
+		[JsonIgnore]
+		public IEnumerable<UhtHeaderFile> ReferencedHeadersLocked
+		{
+			get
+			{
+				lock (this.ReferencedHeaders)
+				{
+					foreach (UhtHeaderFile Ref in this.ReferencedHeaders)
+					{
+						yield return Ref;
+					}
+				}
+			}
+		}
+
+		public override void GetPathName(StringBuilder Builder, UhtType? StopOuter = null)
+		{
+			// Headers do not contribute to path names
+			if (this != StopOuter && this.Outer != null)
+			{
+				this.Outer.GetPathName(Builder, StopOuter);
+			}
+		}
+
+		protected override UhtValidationOptions Validate(UhtValidationOptions Options)
+		{
+			Options = base.Validate(Options | UhtValidationOptions.Shadowing);
+
+			Dictionary<int, UhtFunction> UsedRPCIds = new Dictionary<int, UhtFunction>();
+			Dictionary<int, UhtFunction> RPCsNeedingHookup = new Dictionary<int, UhtFunction>();
+			foreach (UhtType Type in this.Children)
+			{
+				if (Type is UhtClass Class)
+				{
+					foreach (UhtType Child in Class.Children)
+					{
+						if (Child is UhtFunction Function)
+						{
+							if (Function.FunctionType != UhtFunctionType.Function || !Function.FunctionFlags.HasAnyFlags(EFunctionFlags.Net))
+							{
+								continue;
+							}
+
+							if (Function.RPCId > 0)
+							{
+								UhtFunction? ExistingFunc;
+								if (UsedRPCIds.TryGetValue(Function.RPCId, out ExistingFunc))
+								{
+									Function.LogError($"Function {ExistingFunc.SourceName} already uses identifier {Function.RPCId}");
+								}
+								else
+								{
+									UsedRPCIds.Add(Function.RPCId, Function);
+								}
+
+								if (Function.FunctionFlags.HasAnyFlags(EFunctionFlags.NetResponse))
+								{
+									// Look for another function expecting this response
+									RPCsNeedingHookup.Remove(Function.RPCId);
+								}
+							}
+
+							if (Function.RPCResponseId > 0 && Function.EndpointName != "JSBridge")
+							{
+								// Look for an existing response function, if not found then add to the list of ids awaiting hookup
+								if (!UsedRPCIds.ContainsKey(Function.RPCResponseId))
+								{
+									RPCsNeedingHookup.Add(Function.RPCResponseId, Function);
+								}
+							}
+						}
+					}
+				}
+			}
+
+			if (RPCsNeedingHookup.Count > 0)
+			{
+				foreach (KeyValuePair<int, UhtFunction> KVP in RPCsNeedingHookup)
+				{
+					KVP.Value.LogError($"Request function '{KVP.Value.SourceName}' is missing a response function with the id of '{KVP.Key}'");
+				}
+			}
+			return Options;
+		}
+	}
+}
diff --git a/Engine/Source/Programs/Shared/EpicGames.UHT/Types/UhtManifestFile.cs b/Engine/Source/Programs/Shared/EpicGames.UHT/Types/UhtManifestFile.cs
new file mode 100644
--- /dev/null
+++ b/Engine/Source/Programs/Shared/EpicGames.UHT/Types/UhtManifestFile.cs
@@ -0,0 +1,32 @@
+// Copyright Epic Games, Inc. All Rights Reserved.
+
+using EpicGames.Core;
+using EpicGames.UHT.Utils;
+using System.Text.Json;
+
+namespace EpicGames.UHT.Types
+{
+	public class UhtManifestFile : IUhtMessageSite
+	{
+		public UHTManifest? Manifest = null;
+
+		private UhtSourceFile SourceFile;
+
+		#region IUHTMessageSite implementation
+		public IUhtMessageSession MessageSession => this.SourceFile.MessageSession;
+		public IUhtMessageSource? MessageSource => this.SourceFile.MessageSource;
+		public IUhtMessageLineNumber? MessageLineNumber => null;
+		#endregion
+
+		public UhtManifestFile(UhtSession Session, string FilePath)
+		{
+			this.SourceFile = new UhtSourceFile(Session, FilePath);
+		}
+
+		public void Read() 
+		{
+			this.SourceFile.Read();
+			this.Manifest = JsonSerializer.Deserialize<UHTManifest>(this.SourceFile.Data.ToString());
+		}
+	}
+}
