UHT in UBT initial version.  Currently disabled.

diff --git a/Engine/Source/Programs/Shared/EpicGames.UHT/Types/UhtMetaData.cs b/Engine/Source/Programs/Shared/EpicGames.UHT/Types/UhtMetaData.cs
new file mode 100644
--- /dev/null
+++ b/Engine/Source/Programs/Shared/EpicGames.UHT/Types/UhtMetaData.cs
@@ -0,0 +1,402 @@
+// Copyright Epic Games, Inc. All Rights Reserved.
+
+using EpicGames.Core;
+using EpicGames.UHT.Utils;
+using System;
+using System.Collections.Generic;
+using System.Diagnostics.CodeAnalysis;
+
+namespace EpicGames.UHT.Types
+{
+	public interface IUhtMetaDataKeyConversion
+	{
+		string GetMetaDataKey(string Name, int NameIndex);
+	}
+
+	public struct UhtMetaDataKey
+	{
+		public string Name;
+		public int Index;
+
+		public UhtMetaDataKey(string Name, int Index = UhtMetaData.INDEX_NONE)
+		{
+			this.Name = Name;
+			this.Index = Index;
+		}
+
+		public override string ToString()
+		{
+			if (Index == UhtMetaData.INDEX_NONE)
+			{
+				return this.Name.ToString();
+			}
+			else
+			{
+				return $"{this.Name}:{this.Index}";
+			}
+		}
+	}
+
+	class UhtMetaDataKeyComparer : IEqualityComparer<UhtMetaDataKey>, IComparer<UhtMetaDataKey>
+	{
+		public int Compare([AllowNull] UhtMetaDataKey x, [AllowNull] UhtMetaDataKey y)
+		{
+			if (x.Index < y.Index)
+			{
+				return -1;
+			}
+			else if (x.Index > y.Index)
+			{
+				return 1;
+			}
+			return string.Compare(x.Name, y.Name, StringComparison.OrdinalIgnoreCase);
+		}
+
+		public bool Equals([AllowNull] UhtMetaDataKey x, [AllowNull] UhtMetaDataKey y)
+		{
+			return x.Index == y.Index && x.Name.Equals(y.Name, StringComparison.OrdinalIgnoreCase);
+		}
+
+		public int GetHashCode([DisallowNull] UhtMetaDataKey obj)
+		{
+			return obj.Name.GetHashCode() ^ obj.Name.GetHashCode(StringComparison.OrdinalIgnoreCase);
+		}
+	}
+
+	/// <summary>
+	/// Represents a collection of key/value pairs.  Each type has this collection and enumerations also
+	/// have key/value pairs for each enumeration value index.
+	/// </summary>
+	public class UhtMetaData
+	{
+
+		/// <summary>
+		/// Helper comparer for meta data keys
+		/// </summary>
+		private static UhtMetaDataKeyComparer Comparer = new UhtMetaDataKeyComparer();
+
+		/// <summary>
+		/// Empty collection of meta data
+		/// </summary>
+		public static UhtMetaData Empty = new UhtMetaData(null);
+
+		/// <summary>
+		/// The meta data of the outer object for the type that owns this meta data
+		/// </summary>
+		public UhtMetaData? Parent = null;
+
+		/// <summary>
+		/// Message site associated with the meta data.  That in combination with the line number is used to generate errors
+		/// </summary>
+		public IUhtMessageSite? MessageSite = null;
+		
+		/// <summary>
+		/// Source code line number where the meta data is declared
+		/// </summary>
+		public int LineNumber { get; internal set; } = 1;
+
+		/// <summary>
+		/// Contains the meta data entries.  Due to the small size of these dictionaries, a SortedList performs
+		/// better than a Dictionary.
+		/// </summary>
+		public SortedList<UhtMetaDataKey, string>? Dictionary = null;
+
+		/// <summary>
+		/// Enumerations implement this interface so that the index part of the key can be converted to a string
+		/// </summary>
+		internal IUhtMetaDataKeyConversion? KeyConversion = null;
+
+		public const int INDEX_NONE = -1;
+
+		public UhtMetaData(IUhtMessageSite? MessageSite)
+		{
+			this.MessageSite = MessageSite;
+		}
+
+		/// <summary>
+		/// Test to see if the meta data object contains no entries
+		/// </summary>
+		/// <returns>True if the meta data object contains no entries</returns>
+		public bool IsEmpty()
+		{
+			return this.Dictionary == null || this.Dictionary.Count == 0;
+		}
+
+		/// <summary>
+		/// Remove all meta data entries
+		/// </summary>
+		public void Clear()
+		{
+			if (this.Dictionary != null)
+			{
+				this.Dictionary.Clear();
+			}
+		}
+
+		/// <summary>
+		/// Copy the elements of the meta data
+		/// </summary>
+		/// <returns>A deep copy of the meta data</returns>
+		public UhtMetaData Clone()
+		{
+			return (UhtMetaData)MemberwiseClone();
+		}
+
+		/// <summary>
+		/// Test to see if the meta data contains the given key
+		/// </summary>
+		/// <param name="Name">Name of the meta data</param>
+		/// <param name="NameIndex">Enumeration value index or -1 for the type's meta data</param>
+		/// <returns>True if the key is found</returns>
+		public bool ContainsKey(string Name, int NameIndex = INDEX_NONE)
+		{
+			if (this.Dictionary == null)
+			{
+				return false;
+			}
+			return this.Dictionary.ContainsKey(new UhtMetaDataKey(Name, NameIndex));
+		}
+
+		/// <summary>
+		/// Test to see if the meta data or parent meta data contains the given key
+		/// </summary>
+		/// <param name="Name">Name of the meta data</param>
+		/// <param name="NameIndex">Enumeration value index or -1 for the type's meta data</param>
+		/// <returns>True if the key is found</returns>
+		public bool ContainsKeyHierarchical(string Name, int NameIndex = INDEX_NONE)
+		{
+			for (UhtMetaData? Current = this; Current != null; Current = Current.Parent)
+			{
+				if (Current.ContainsKey(Name, NameIndex))
+				{
+					return true;
+				}
+			}
+			return false;
+		}
+
+		/// <summary>
+		/// Attempt to get the value associated with the key
+		/// </summary>
+		/// <param name="Name">Name of the meta data</param>
+		/// <returns>True if the key is found</returns>
+		public bool TryGetValue(string Name, [NotNullWhen(true)] out string? Value)
+		{
+			return TryGetValue(Name, INDEX_NONE, out Value);
+		}
+
+		/// <summary>
+		/// Attempt to get the value associated with the key including parent meta data
+		/// </summary>
+		/// <param name="Name">Name of the meta data</param>
+		/// <returns>True if the key is found</returns>
+		public bool TryGetValueHierarchical(string Name, [NotNullWhen(true)] out string? Value)
+		{
+			return TryGetValueHierarchical(Name, INDEX_NONE, out Value);
+		}
+
+		/// <summary>
+		/// Attempt to get the value associated with the key
+		/// </summary>
+		/// <param name="Name">Name of the meta data</param>
+		/// <param name="NameIndex">Enumeration value index or -1 for the type's meta data</param>
+		/// <returns>True if the key is found</returns>
+		public bool TryGetValue(string Name, int NameIndex, [NotNullWhen(true)] out string? Value)
+		{
+			if (this.Dictionary == null)
+			{
+				Value = String.Empty;
+				return false;
+			}
+			return this.Dictionary.TryGetValue(new UhtMetaDataKey(Name, NameIndex), out Value);
+		}
+
+		/// <summary>
+		/// Attempt to get the value associated with the key including parent meta data
+		/// </summary>
+		/// <param name="Name">Name of the meta data</param>
+		/// <param name="NameIndex">Enumeration value index or -1 for the type's meta data</param>
+		/// <returns>True if the key is found</returns>
+		public bool TryGetValueHierarchical(string Name, int NameIndex, [NotNullWhen(true)] out string? Value)
+		{
+			for (UhtMetaData? Current = this; Current != null; Current = Current.Parent)
+			{
+				if (Current.TryGetValue(Name, NameIndex, out Value))
+				{
+					return true;
+				}
+			}
+			Value = String.Empty;
+			return false;
+		}
+
+		public string GetValueOrDefault(string Name, int NameIndex = INDEX_NONE)
+		{
+			string? Out;
+			if (TryGetValue(Name, NameIndex, out Out))
+			{
+				return Out;
+			}
+			return String.Empty;
+		}
+
+		public string GetValueOrDefaultHierarchical(string Name, int NameIndex = INDEX_NONE)
+		{
+			string? Out;
+			if (TryGetValueHierarchical(Name, NameIndex, out Out))
+			{
+				return Out;
+			}
+			return string.Empty;
+		}
+
+		public bool GetBoolean(string Name, int NameIndex = INDEX_NONE)
+		{
+			return UhtFCString.ToBool(GetValueOrDefault(Name, NameIndex));
+		}
+
+		public bool GetBooleanHierarchical(string Name, int NameIndex = INDEX_NONE)
+		{
+			return UhtFCString.ToBool(GetValueOrDefaultHierarchical(Name, NameIndex));
+		}
+
+		public double GetDouble(string Name, int NameIndex = INDEX_NONE)
+		{
+			string Value = GetValueOrDefault(Name, NameIndex);
+			double Result;
+			if (!double.TryParse(Value, out Result))
+			{
+				Result = 0;
+			}
+			return Result;
+		}
+
+		public string[]? GetStringArray(string Name, int NameIndex = INDEX_NONE)
+		{
+			string? Temp;
+			if (TryGetValue(Name, NameIndex, out Temp))
+			{
+				return Temp.ToString().Split(' ', StringSplitOptions.RemoveEmptyEntries);
+			}
+			return null;
+		}
+
+		public string[]? GetStringArrayHierarchical(string Name, int NameIndex = INDEX_NONE)
+		{
+			string? Temp;
+			if (TryGetValueHierarchical(Name, NameIndex, out Temp))
+			{
+				return Temp.ToString().Split(' ', StringSplitOptions.RemoveEmptyEntries);
+			}
+			return null;
+		}
+
+		public void Add(string Name, string Value)
+		{
+			Add(Name, INDEX_NONE, Value);
+		}
+
+		public void Add(string Name, int NameIndex, string Value)
+		{
+			string RemappedName;
+			if (UhtConfig.Instance.RedirectMetaDataKey(Name, out RemappedName))
+			{
+				if (this.MessageSite != null)
+				{
+					this.MessageSite.LogWarning(this.LineNumber, $"Remapping old metadata key '{Name}' to new key '{RemappedName}', please update the declaration.");
+				}
+			}
+			GetDictionary()[new UhtMetaDataKey(Name, NameIndex)] = Value;
+		}
+
+		public void Add(string Name, bool Value, int NameIndex = INDEX_NONE)
+		{
+			Add(Name, NameIndex, Value ? "true" : "false");
+		}
+
+		public void Add(string Name, List<string> Strings, char Separator = ' ')
+		{
+			Add(Name, string.Join(Separator, Strings));
+		}
+
+		public void AddIfNotEmpty(string Name, List<string> Strings, char Separator = ' ')
+		{
+			if (Strings.Count != 0)
+			{
+				Add(Name, string.Join(Separator, Strings));
+			}
+		}
+
+		public void Add(UhtMetaData MetaData)
+		{
+			if (MetaData.Dictionary != null)
+			{
+				SortedList<UhtMetaDataKey, string> Dictionary = GetDictionary();
+				foreach (KeyValuePair<UhtMetaDataKey, string> Kvp in MetaData.Dictionary)
+				{
+					Dictionary[Kvp.Key] = Kvp.Value;
+				}
+			}
+		}
+
+		public void Remove(string Name, int NameIndex = INDEX_NONE)
+		{
+			if (this.Dictionary != null)
+			{
+				this.Dictionary.Remove(new UhtMetaDataKey(Name, NameIndex));
+			}
+		}
+
+		private SortedList<UhtMetaDataKey, string> GetDictionary()
+		{
+			if (this.Dictionary == null)
+			{
+				this.Dictionary = new SortedList<UhtMetaDataKey, string>(Comparer);
+			}
+			return this.Dictionary;
+		}
+
+		public string GetKeyString(UhtMetaDataKey Key)
+		{
+			return GetKeyString(Key.Name, Key.Index);
+		}
+
+		public string GetKeyString(string Name, int NameIndex)
+		{
+			if (NameIndex == INDEX_NONE)
+			{
+				return Name;
+			}
+			else if (KeyConversion != null)
+			{
+				return KeyConversion.GetMetaDataKey(Name, NameIndex);
+			}
+			else
+			{
+				throw new UhtIceException("Attempt to generate an indexed meta data key name but no key conversion interface was set.");
+			}
+		}
+
+		/// <summary>
+		/// Given a type, return an array of all the meta data formatted and sorted by name
+		/// </summary>
+		/// <returns>List of meta data key and value pairs</returns>
+		public List<KeyValuePair<string, string>> GetSorted()
+		{
+			List<KeyValuePair<string, string>> Out = new List<KeyValuePair<string, string>>(this.Dictionary != null ? this.Dictionary.Count : 0);
+			if (this.Dictionary != null && this.Dictionary.Count > 0)
+			{
+				foreach (KeyValuePair<UhtMetaDataKey, string> Kvp in this.Dictionary)
+				{
+					Out.Add(new KeyValuePair<string, string>(this.GetKeyString(Kvp.Key), Kvp.Value));
+				}
+
+				Out.Sort((KeyValuePair<string, string> Lhs, KeyValuePair<string, string> Rhs) =>
+				{
+					return StringComparerUE.OrdinalIgnoreCase.Compare(Lhs.Key, Rhs.Key);
+				});
+			}
+			return Out;
+		}
+	}
+}
diff --git a/Engine/Source/Programs/Shared/EpicGames.UHT/Types/UhtObject.cs b/Engine/Source/Programs/Shared/EpicGames.UHT/Types/UhtObject.cs
new file mode 100644
--- /dev/null
+++ b/Engine/Source/Programs/Shared/EpicGames.UHT/Types/UhtObject.cs
@@ -0,0 +1,34 @@
+// Copyright Epic Games, Inc. All Rights Reserved.
+
+using EpicGames.Core;
+using EpicGames.UHT.Tables;
+using EpicGames.UHT.Utils;
+
+namespace EpicGames.UHT.Types
+{
+	[UhtEngineClass(Name = "Object")]
+	public abstract class UhtObject : UhtType
+	{
+		public EInternalObjectFlags InternalObjectFlags = EInternalObjectFlags.None;
+		public readonly int ObjectTypeIndex;
+
+		/// <summary>
+		/// The alternate object is used by the interface system where the native interface will
+		/// update this setting to point to the UInterface derived companion object.
+		/// </summary>
+		public UhtObject? AlternateObject = null;
+
+		/// <inheritdoc/>
+		public override string EngineClassName { get => "Object"; }
+
+		protected UhtObject(UhtSession Session) : base(Session)
+		{
+			ObjectTypeIndex = this.Session.GetNextObjectTypeIndex();
+		}
+
+		protected UhtObject(UhtType Outer, int LineNumber) : base(Outer, LineNumber)
+		{
+			ObjectTypeIndex = this.Session.GetNextObjectTypeIndex();
+		}
+	}
+}
diff --git a/Engine/Source/Programs/Shared/EpicGames.UHT/Types/UhtPackage.cs b/Engine/Source/Programs/Shared/EpicGames.UHT/Types/UhtPackage.cs
new file mode 100644
--- /dev/null
+++ b/Engine/Source/Programs/Shared/EpicGames.UHT/Types/UhtPackage.cs
@@ -0,0 +1,83 @@
+// Copyright Epic Games, Inc. All Rights Reserved.
+
+using EpicGames.Core;
+using EpicGames.UHT.Tables;
+using EpicGames.UHT.Utils;
+using System;
+using System.Text.Json.Serialization;
+using UnrealBuildBase;
+
+namespace EpicGames.UHT.Types
+{
+	[UhtEngineClass(Name = "Package")]
+	public class UhtPackage : UhtObject
+	{
+		public readonly int PackageTypeIndex;
+		public EPackageFlags PackageFlags { get; internal set; } = EPackageFlags.None;
+		public UHTManifest.Module Module { get; internal set; }
+
+		[JsonIgnore]
+		public override UhtEngineType EngineType => UhtEngineType.Package;
+
+		[JsonIgnore]
+		public override UhtPackage Package => this;
+
+		[JsonIgnore]
+		public override UhtHeaderFile HeaderFile => throw new NotImplementedException();
+
+		[JsonIgnore]
+		public bool bIsPartOfEngine
+		{
+			get
+			{
+				switch (this.Module.ModuleType)
+				{
+					case UHTModuleType.Program:
+						return this.Module.BaseDirectory.Replace('\\', '/').StartsWith(Unreal.EngineDirectory.FullName.Replace('\\', '/'));
+					case UHTModuleType.EngineRuntime:
+					case UHTModuleType.EngineUncooked:
+					case UHTModuleType.EngineDeveloper:
+					case UHTModuleType.EngineEditor:
+					case UHTModuleType.EngineThirdParty:
+						return true;
+					case UHTModuleType.GameRuntime:
+					case UHTModuleType.GameUncooked:
+					case UHTModuleType.GameDeveloper:
+					case UHTModuleType.GameEditor:
+					case UHTModuleType.GameThirdParty:
+						return false;
+					default:
+						throw new UhtIceException("Invalid module type");
+				}
+			}
+		}
+
+		[JsonIgnore]
+		public bool bIsPlugin => this.Module.BaseDirectory.Replace('\\', '/').Contains("/Plugins/");
+
+
+		public string ShortName;
+
+		/// <inheritdoc/>
+		public override string EngineClassName { get => "Package"; }
+
+		public UhtPackage(UhtSession Session, UHTManifest.Module Module, EPackageFlags PackageFlags) : base(Session)
+		{
+			this.Module = Module;
+			this.PackageFlags = PackageFlags;
+			this.PackageTypeIndex = this.Session.GetNextPackageTypeIndex();
+
+			int LastSlashIndex = Module.Name.LastIndexOf('/');
+			if (LastSlashIndex == -1)
+			{
+				this.SourceName = $"/Script/{Module.Name}";
+				this.ShortName = Module.Name;
+			}
+			else
+			{
+				this.SourceName = Module.Name;
+				this.ShortName = this.SourceName.Substring(LastSlashIndex + 1);
+			}
+		}
+	}
+}
diff --git a/Engine/Source/Programs/Shared/EpicGames.UHT/Types/UhtProperty.cs b/Engine/Source/Programs/Shared/EpicGames.UHT/Types/UhtProperty.cs
new file mode 100644
--- /dev/null
+++ b/Engine/Source/Programs/Shared/EpicGames.UHT/Types/UhtProperty.cs
@@ -0,0 +1,1767 @@
+// Copyright Epic Games, Inc. All Rights Reserved.
+
+using EpicGames.Core;
+using EpicGames.UHT.Parsers;
+using EpicGames.UHT.Tables;
+using EpicGames.UHT.Tokenizer;
+using EpicGames.UHT.Utils;
+using System;
+using System.Collections.Generic;
+using System.Diagnostics.CodeAnalysis;
+using System.Linq;
+using System.Text;
+using System.Text.Json.Serialization;
+
+namespace EpicGames.UHT.Types
+{
+
+	/// <summary>
+	/// Collection of UHT only flags associated with properties
+	/// </summary>
+	[Flags]
+	public enum UhtPropertyExportFlags
+	{
+		/// <summary>
+		/// Property should be exported as public
+		/// </summary>
+		Public = 0x00000001,
+
+		/// <summary>
+		/// Property should be exported as private
+		/// </summary>
+		Private = 0x00000002,
+
+		/// <summary>
+		/// Property should be exported as protected
+		/// </summary>
+		Protected = 0x00000004,
+		
+		/// <summary>
+		/// The BlueprintPure flag was set in software and should not be considered an error
+		/// </summary>
+		ImpliedBlueprintPure = 0x00000008,
+
+		/// <summary>
+		/// If true, the property has a getter function
+		/// </summary>
+		GetterSpecified = 0x00000010,
+
+		/// <summary>
+		/// If true, the property has a setter function
+		/// </summary>
+		SetterSpecified = 0x00000020,
+
+		/// <summary>
+		/// If true, the getter has been disabled in the specifiers
+		/// </summary>
+		GetterSpecifiedNone = 0x00000040,
+
+		/// <summary>
+		/// If true, the setter has been disabled in the specifiers
+		/// </summary>
+		SetterSpecifiedNone = 0x00000080,
+
+		/// <summary>
+		/// If true, the getter function was found
+		/// </summary>
+		GetterFound = 0x00000100,
+
+		/// <summary>
+		/// If true, the property has a setter function
+		/// </summary>
+		SetterFound = 0x00000200,
+	};
+
+	/// <summary>
+	/// Helper methods for testing flags.  These methods perform better than the generic HasFlag which hits
+	/// the GC and stalls.
+	/// </summary>
+	public static class UhtPropertyExportFlagsExtensions
+	{
+
+		/// <summary>
+		/// Test to see if any of the specified flags are set
+		/// </summary>
+		/// <param name="InFlags">Current flags</param>
+		/// <param name="TestFlags">Flags to test for</param>
+		/// <returns>True if any of the flags are set</returns>
+		public static bool HasAnyFlags(this UhtPropertyExportFlags InFlags, UhtPropertyExportFlags TestFlags)
+		{
+			return (InFlags & TestFlags) != 0;
+		}
+
+		/// <summary>
+		/// Test to see if all of the specified flags are set
+		/// </summary>
+		/// <param name="InFlags">Current flags</param>
+		/// <param name="TestFlags">Flags to test for</param>
+		/// <returns>True if all the flags are set</returns>
+		public static bool HasAllFlags(this UhtPropertyExportFlags InFlags, UhtPropertyExportFlags TestFlags)
+		{
+			return (InFlags & TestFlags) == TestFlags;
+		}
+
+		/// <summary>
+		/// Test to see if a specific set of flags have a specific value.
+		/// </summary>
+		/// <param name="InFlags">Current flags</param>
+		/// <param name="TestFlags">Flags to test for</param>
+		/// <param name="MatchFlags">Expected value of the tested flags</param>
+		/// <returns>True if the given flags have a specific value.</returns>
+		public static bool HasExactFlags(this UhtPropertyExportFlags InFlags, UhtPropertyExportFlags TestFlags, UhtPropertyExportFlags MatchFlags)
+		{
+			return (InFlags & TestFlags) == MatchFlags;
+		}
+	}
+
+	/// <summary>
+	/// The context of the property.
+	/// </summary>
+	public enum UhtPropertyCategory
+	{
+		/// <summary>
+		/// Function parameter for a function that isn't marked as NET
+		/// </summary>
+		RegularParameter,
+
+		/// <summary>
+		/// Function parameter for a function that is marked as NET
+		/// </summary>
+		ReplicatedParameter,
+
+		/// <summary>
+		/// Function return value
+		/// </summary>
+		Return,
+
+		/// <summary>
+		/// Class or a script structure member property
+		/// </summary>
+		Member,
+	};
+
+	public static class UhtPropertyCategoryExtensions
+	{
+		public static string GetHintText(this UhtPropertyCategory PropertyCategory)
+		{
+			switch (PropertyCategory)
+			{
+				case UhtPropertyCategory.ReplicatedParameter:
+				case UhtPropertyCategory.RegularParameter:
+					return "Function parameter";
+
+				case UhtPropertyCategory.Return:
+					return "Function return type";
+
+				case UhtPropertyCategory.Member:
+					return "Member variable declaration";
+
+				default:
+					throw new UhtIceException("Unknown variable category");
+			}
+		}
+
+		public static bool IsMemberProperty(this UhtPropertyCategory PropertyCategory)
+		{
+			return PropertyCategory == UhtPropertyCategory.Member;
+		}
+
+		public static bool IsFunctionProperty(this UhtPropertyCategory PropertyCategory)
+		{
+			return PropertyCategory != UhtPropertyCategory.Member;
+		}
+
+		public static bool IsFunctionReturn(this UhtPropertyCategory PropertyCategory)
+		{
+			return PropertyCategory == UhtPropertyCategory.Return;
+		}
+
+		public static bool IsFunctionArgument(this UhtPropertyCategory PropertyCategory)
+		{
+			return PropertyCategory == UhtPropertyCategory.RegularParameter || PropertyCategory == UhtPropertyCategory.ReplicatedParameter;
+		}
+	}
+
+	public enum UhtPropertyAllocator
+	{
+		Default,
+		MemoryImage
+	};
+
+	public enum UhtPointerType
+	{
+		None,
+		Native,
+	}
+
+	public enum UhtPropertyIntType
+	{
+		None,
+		Sized,
+		Unsized,
+	};
+
+	public enum UhtPropertyRefQualifier
+	{
+		None,
+		ConstRef,
+		NonConstRef,
+	};
+
+	/// <summary>
+	/// Options that customize the properties.
+	/// </summary>
+	[Flags]
+	public enum UhtPropertyOptions
+	{
+		None = 0,
+
+		/// <summary>
+		/// Don't automatically mark properties as CPF_Const
+		/// </summary>
+		NoAutoConst = 1 << 0,
+	}
+
+	/// <summary>
+	/// Helper methods for testing flags.  These methods perform better than the generic HasFlag which hits
+	/// the GC and stalls.
+	/// </summary>
+	public static class UhtPropertyOptionsExtensions
+	{
+
+		/// <summary>
+		/// Test to see if any of the specified flags are set
+		/// </summary>
+		/// <param name="InFlags">Current flags</param>
+		/// <param name="TestFlags">Flags to test for</param>
+		/// <returns>True if any of the flags are set</returns>
+		public static bool HasAnyFlags(this UhtPropertyOptions InFlags, UhtPropertyOptions TestFlags)
+		{
+			return (InFlags & TestFlags) != 0;
+		}
+
+		/// <summary>
+		/// Test to see if all of the specified flags are set
+		/// </summary>
+		/// <param name="InFlags">Current flags</param>
+		/// <param name="TestFlags">Flags to test for</param>
+		/// <returns>True if all the flags are set</returns>
+		public static bool HasAllFlags(this UhtPropertyOptions InFlags, UhtPropertyOptions TestFlags)
+		{
+			return (InFlags & TestFlags) == TestFlags;
+		}
+
+		/// <summary>
+		/// Test to see if a specific set of flags have a specific value.
+		/// </summary>
+		/// <param name="InFlags">Current flags</param>
+		/// <param name="TestFlags">Flags to test for</param>
+		/// <param name="MatchFlags">Expected value of the tested flags</param>
+		/// <returns>True if the given flags have a specific value.</returns>
+		public static bool HasExactFlags(this UhtPropertyOptions InFlags, UhtPropertyOptions TestFlags, UhtPropertyOptions MatchFlags)
+		{
+			return (InFlags & TestFlags) == MatchFlags;
+		}
+	}
+
+	[Flags]
+	public enum UhtPropertyCaps
+	{
+		None = 0,
+
+		/// <summary>
+		/// If true, the property will be passed by reference when generating the full type string
+		/// </summary>
+		PassCppArgsByRef = 1 << 0,
+
+		/// <summary>
+		/// If true, the an argument will need to be added to the constructor
+		/// </summary>
+		RequiresNullConstructorArg = 1 << 1,
+
+		/// <summary>
+		/// If true, the property type can be TArray or TMap value
+		/// </summary>
+		CanBeContainerValue = 1 << 2,
+
+		/// <summary>
+		/// If true, the property type can be a TSet or TMap key
+		/// </summary>
+		CanBeContainerKey = 1 << 3,
+
+		/// <summary>
+		/// True if the property can be instanced
+		/// </summary>
+		CanBeInstanced = 1 << 4,
+
+		/// <summary>
+		/// True if the property can be exposed on spawn
+		/// </summary>
+		CanExposeOnSpawn = 1 << 5,
+
+		/// <summary>
+		/// True if the property can have a config setting
+		/// </summary>
+		CanHaveConfig = 1 << 6,
+
+		/// <summary>
+		/// True if the property allows the BlueprintAssignable flag.
+		/// </summary>
+		CanBeBlueprintAssignable = 1 << 7,
+
+		/// <summary>
+		/// True if the property allows the BlueprintCallable flag.
+		/// </summary>
+		CanBeBlueprintCallable = 1 << 8,
+
+		/// <summary>
+		/// True if the property allows the BlueprintAuthorityOnly flag.
+		/// </summary>
+		CanBeBlueprintAuthorityOnly = 1 << 9,
+
+		/// <summary>
+		/// True to see if the function parameter property is supported by blueprint
+		/// </summary>
+		IsParameterSupportedByBlueprint = 1 << 10,
+
+		/// <summary>
+		/// True to see if the member property is supported by blueprint
+		/// </summary>
+		IsMemberSupportedByBlueprint = 1 << 11,
+	}
+
+	/// <summary>
+	/// Helper methods for testing flags.  These methods perform better than the generic HasFlag which hits
+	/// the GC and stalls.
+	/// </summary>
+	public static class UhtPropertyCapsExtensions
+	{
+
+		/// <summary>
+		/// Test to see if any of the specified flags are set
+		/// </summary>
+		/// <param name="InFlags">Current flags</param>
+		/// <param name="TestFlags">Flags to test for</param>
+		/// <returns>True if any of the flags are set</returns>
+		public static bool HasAnyFlags(this UhtPropertyCaps InFlags, UhtPropertyCaps TestFlags)
+		{
+			return (InFlags & TestFlags) != 0;
+		}
+
+		/// <summary>
+		/// Test to see if all of the specified flags are set
+		/// </summary>
+		/// <param name="InFlags">Current flags</param>
+		/// <param name="TestFlags">Flags to test for</param>
+		/// <returns>True if all the flags are set</returns>
+		public static bool HasAllFlags(this UhtPropertyCaps InFlags, UhtPropertyCaps TestFlags)
+		{
+			return (InFlags & TestFlags) == TestFlags;
+		}
+
+		/// <summary>
+		/// Test to see if a specific set of flags have a specific value.
+		/// </summary>
+		/// <param name="InFlags">Current flags</param>
+		/// <param name="TestFlags">Flags to test for</param>
+		/// <param name="MatchFlags">Expected value of the tested flags</param>
+		/// <returns>True if the given flags have a specific value.</returns>
+		public static bool HasExactFlags(this UhtPropertyCaps InFlags, UhtPropertyCaps TestFlags, UhtPropertyCaps MatchFlags)
+		{
+			return (InFlags & TestFlags) == MatchFlags;
+		}
+	}
+
+	/// <summary>
+	/// Text can be formatted in different context with slightly different results
+	/// </summary>
+	public enum UhtPropertyTextType
+	{
+		/// <summary>
+		/// Used when the declaration is to be displayed to the user in cases such as error messages.
+		/// </summary>
+		UserFacing,
+
+		/// <summary>
+		/// Sparse property declaration
+		/// </summary>
+		Sparse,
+
+		/// <summary>
+		/// Sparse property short name
+		/// </summary>
+		SparseShort,
+
+		/// <summary>
+		/// Class function argument or return value
+		/// </summary>
+		ClassFunction,
+
+		/// <summary>
+		/// Event function argument or return value
+		/// </summary>
+		EventFunction,
+
+		/// <summary>
+		/// Interface function argument or return value
+		/// </summary>
+		InterfaceFunction,
+
+		/// <summary>
+		/// Class or structure member
+		/// </summary>
+		ExportMember,
+
+		/// <summary>
+		/// Members of the event parameters structure used by code generation to invoke events
+		/// </summary>
+		EventParameterMember,
+
+		/// <summary>
+		/// Members of the event parameters structure used by code generation to invoke functions
+		/// </summary>
+		EventParameterFunctionMember,
+
+		/// <summary>
+		/// Instance of the property is being constructed
+		/// </summary>
+		Construction,
+
+		/// <summary>
+		/// Used to get the type argument for a function thunk.  This is used for P_GET_ARRAY_*
+		/// </summary>
+		FunctionThunkParameterArrayType,
+
+		/// <summary>
+		/// If the P_GET macro requires an argument, this is used to fetch that argument
+		/// </summary>
+		FunctionThunkParameterArgType,
+
+		/// <summary>
+		/// Used to get the return type for function thunks
+		/// </summary>
+		FunctionThunkReturn,
+
+		/// <summary>
+		/// Used when generating types for RigVM.  Only invoked for template arguments
+		/// </summary>
+		RigVMTemplateArg,
+
+		/// <summary>
+		/// Type expected in a getter/setter argument list
+		/// </summary>
+		GetterSetterArg,
+	}
+
+	public static class UhtPropertyTextTypeExtensions
+	{
+		public static bool IsParameter(this UhtPropertyTextType TextType)
+		{
+			return
+				TextType == UhtPropertyTextType.ClassFunction ||
+				TextType == UhtPropertyTextType.EventFunction ||
+				TextType == UhtPropertyTextType.InterfaceFunction;
+		}
+	}
+
+	public enum UhtPGetArgumentType
+	{
+		None,
+		EngineClass,
+		TypeText,
+	}
+
+	public interface IUhtPropertyMemberContext
+	{
+		public UhtStruct OuterStruct { get; }
+		public string OuterStructSourceName { get; }
+		public string StaticsName { get; }
+		public string NamePrefix { get; }
+		public string MetaDataSuffix { get; }
+		public uint GetTypeHash(UhtObject Object);
+		public string GetSingletonName(UhtObject? Object, bool bRegistered);
+	}
+
+	public class UhtPropertySettings
+	{
+		public string SourceName = String.Empty;
+		public string EngineName = String.Empty;
+		public UhtMetaData MetaData = UhtMetaData.Empty;
+		public UhtType Outer;
+		public int LineNumber;
+		public UhtPropertyCategory PropertyCategory;
+		public EPropertyFlags PropertyFlags;
+		public EPropertyFlags DisallowPropertyFlags;
+		public UhtPropertyExportFlags PropertyExportFlags;
+		public UhtPropertyAllocator Allocator;
+		public UhtPropertyOptions Options;
+		public UhtPointerType PointerType;
+		public string? RepNotifyName;
+		public string? ArrayDimensions;
+		public string? Setter;
+		public string? Getter;
+		public List<UhtToken>? DefaultValueTokens;
+		public bool bIsBitfield;
+
+#pragma warning disable CS8618 // Non-nullable field must contain a non-null value when exiting constructor. Consider declaring as nullable.
+		public UhtPropertySettings()
+#pragma warning restore CS8618 // Non-nullable field must contain a non-null value when exiting constructor. Consider declaring as nullable.
+		{
+		}
+
+#pragma warning disable CS8618 // Non-nullable field must contain a non-null value when exiting constructor. Consider declaring as nullable.
+		public UhtPropertySettings(UhtType Outer, int LineNumber, UhtPropertyCategory PropertyCategory, EPropertyFlags DisallowPropertyFlags)
+#pragma warning restore CS8618 // Non-nullable field must contain a non-null value when exiting constructor. Consider declaring as nullable.
+		{
+			Reset(Outer, LineNumber, PropertyCategory, DisallowPropertyFlags);
+		}
+
+		public void Reset(UhtType Outer, int LineNumber, UhtPropertyCategory PropertyCategory, EPropertyFlags DisallowPropertyFlags)
+		{
+			this.SourceName = string.Empty;
+			this.EngineName = string.Empty;
+			this.MetaData = new UhtMetaData(Outer);
+			this.Outer = Outer;
+			this.LineNumber = LineNumber;
+			this.PropertyCategory = PropertyCategory;
+			this.PropertyFlags = EPropertyFlags.None;
+			this.DisallowPropertyFlags = DisallowPropertyFlags;
+			this.PropertyExportFlags = UhtPropertyExportFlags.Public;
+			this.RepNotifyName = null;
+			this.Allocator = UhtPropertyAllocator.Default;
+			this.PointerType = UhtPointerType.None;
+			this.ArrayDimensions = null;
+			this.DefaultValueTokens = null;
+			this.Setter = null;
+			this.Getter = null;
+			this.bIsBitfield = false;
+		}
+
+#pragma warning disable CS8618 // Non-nullable field must contain a non-null value when exiting constructor. Consider declaring as nullable.
+		public UhtPropertySettings(UhtProperty Property, UhtPropertyOptions Options)
+#pragma warning restore CS8618 // Non-nullable field must contain a non-null value when exiting constructor. Consider declaring as nullable.
+		{
+			Reset(Property, Options);
+		}
+
+		public void Reset(UhtProperty Property, UhtPropertyOptions Options)
+		{
+			if (Property.Outer == null)
+			{
+				throw new UhtIceException("Property must have an outer specified");
+			}
+			this.SourceName = Property.SourceName;
+			this.EngineName = Property.EngineName;
+			this.MetaData = Property.MetaData;
+			this.Outer = Property.Outer;
+			this.LineNumber = Property.LineNumber;
+			this.PropertyCategory = Property.PropertyCategory;
+			this.PropertyFlags = Property.PropertyFlags;
+			this.DisallowPropertyFlags = Property.DisallowPropertyFlags;
+			this.PropertyExportFlags = Property.PropertyExportFlags;
+			this.Allocator = Property.Allocator;
+			this.Options = Options;
+			this.PointerType = Property.PointerType;
+			this.RepNotifyName = Property.RepNotifyName;
+			this.ArrayDimensions = Property.ArrayDimensions;
+			this.DefaultValueTokens = Property.DefaultValueTokens;
+			this.Setter = Property.Setter;
+			this.Getter = Property.Getter;
+			this.bIsBitfield = Property.bIsBitfield;
+		}
+
+		public UhtPropertySettings(UhtPropertySettings ParentPropertySettings, string SourceName, IUhtMessageSite MessageSite)
+		{
+			this.SourceName = SourceName;
+			this.EngineName = SourceName;
+			this.MetaData = new UhtMetaData(MessageSite);
+			this.Outer = ParentPropertySettings.Outer;
+			this.LineNumber = ParentPropertySettings.LineNumber;
+			this.PropertyCategory = ParentPropertySettings.PropertyCategory;
+			this.PropertyFlags = ParentPropertySettings.PropertyFlags;
+			this.DisallowPropertyFlags = ParentPropertySettings.DisallowPropertyFlags;
+			this.PropertyExportFlags = UhtPropertyExportFlags.Public;
+			this.RepNotifyName = null;
+			this.Allocator = UhtPropertyAllocator.Default;
+			this.Options = ParentPropertySettings.Options;
+			this.PointerType = UhtPointerType.None;
+			this.ArrayDimensions = null;
+			this.DefaultValueTokens = null;
+			this.Setter = null;
+			this.Getter = null;
+			this.bIsBitfield = false;
+		}
+	}
+
+	[UhtEngineClass(Name = "Field", IsProperty = true)]
+	[UhtEngineClass(Name = "Property", IsProperty = true)]
+	public abstract class UhtProperty : UhtType
+	{
+		#region Constants
+		/// <summary>
+		/// Collection of recognized casts when parsing array dimensions
+		/// </summary>
+		private static string[] Casts = new string[]
+		{
+			"(uint32)",
+			"(int32)",
+			"(uint16)",
+			"(int16)",
+			"(uint8)",
+			"(int8)",
+			"(int)",
+			"(unsigned)",
+			"(signed)",
+			"(unsigned int)",
+			"(signed int)",
+		};
+
+		/// <summary>
+		/// Collection of invalid names for parameters 
+		/// </summary>
+		private static readonly HashSet<string> InvalidParamNames = new HashSet<string>(StringComparer.OrdinalIgnoreCase) { "self" };
+
+		/// <summary>
+		/// Standard object flags for properties
+		/// </summary>
+		protected const string ObjectFlags = "RF_Public|RF_Transient|RF_MarkAsNative";
+
+		/// <summary>
+		/// Prefix used when declaring P_GET_ parameters
+		/// </summary>
+		protected const string FunctionParameterThunkPrefix = "Z_Param_";
+		#endregion
+
+		#region Protected property configuration used to simplify implementation details
+		/// <summary>
+		/// Simple native CPP type text.  Will not include any template arguments
+		/// </summary>
+		protected abstract string CppTypeText { get; }
+
+		/// <summary>
+		/// P_GET_ macro name
+		/// </summary>
+		protected abstract string PGetMacroText { get; }
+
+		/// <summary>
+		/// If true, then references must be passed without a pointer
+		/// </summary>
+		protected virtual bool bPGetPassAsNoPtr { get => false; }
+
+		/// <summary>
+		/// Type of the PGet argument if one is required
+		/// </summary>
+		protected virtual UhtPGetArgumentType PGetTypeArgument { get => UhtPGetArgumentType.None; }
+		#endregion
+
+		public UhtPropertyCategory PropertyCategory = UhtPropertyCategory.Member;
+
+		[JsonConverter(typeof(JsonStringEnumConverter))]
+		public EPropertyFlags PropertyFlags { get; set; }
+
+		[JsonIgnore]
+		[JsonConverter(typeof(JsonStringEnumConverter))]
+		public UhtPropertyCaps PropertyCaps { get; set; }
+
+		public EPropertyFlags DisallowPropertyFlags = EPropertyFlags.None;
+		public UhtPropertyExportFlags PropertyExportFlags = UhtPropertyExportFlags.Public;
+		public UhtPropertyRefQualifier RefQualifier = UhtPropertyRefQualifier.None;
+		public UhtPointerType PointerType = UhtPointerType.None;
+		public UhtPropertyAllocator Allocator = UhtPropertyAllocator.Default;
+		public string? RepNotifyName = null;
+		public string? ArrayDimensions = null;
+		public string? Setter = null;
+		public string? Getter = null;
+		public List<UhtToken>? DefaultValueTokens = null;
+		public bool bIsBitfield = false;
+
+		[JsonIgnore]
+		public override UhtEngineType EngineType => UhtEngineType.Property;
+
+		///<inheritdoc/>
+		[JsonIgnore]
+		public override bool bDeprecated => this.PropertyFlags.HasAnyFlags(EPropertyFlags.Deprecated);
+
+		///<inheritdoc/>
+		protected override UhtSpecifierValidatorTable? SpecifierValidatorTable
+		{
+			get
+			{
+				switch (this.PropertyCategory)
+				{
+					case UhtPropertyCategory.Member:
+						return UhtPropertyMemberSpecifiers.SpecifierValidatorTable;
+					case UhtPropertyCategory.RegularParameter:
+					case UhtPropertyCategory.ReplicatedParameter:
+					case UhtPropertyCategory.Return:
+						return UhtPropertyArgumentSpecifiers.SpecifierValidatorTable;
+					default:
+						throw new UhtIceException("Unknown property category type");
+				}
+			}
+		}
+
+		// Other attributes
+		[JsonIgnore]
+		public bool bIsStaticArray => !string.IsNullOrEmpty(this.ArrayDimensions);
+
+		[JsonIgnore]
+		public bool bIsEditorOnlyProperty => this.PropertyFlags.HasAnyFlags(EPropertyFlags.DevelopmentAssets);
+
+		public UhtProperty(UhtType Outer, int LineNumber) : base(Outer, LineNumber)
+		{
+			this.PropertyFlags = EPropertyFlags.None;
+			this.PropertyCaps = UhtPropertyCaps.CanBeContainerValue | UhtPropertyCaps.CanBeContainerKey | UhtPropertyCaps.CanHaveConfig;
+		}
+
+		public UhtProperty(UhtPropertySettings PropertySettings) : base(PropertySettings.Outer, PropertySettings.LineNumber, PropertySettings.MetaData)
+		{
+			this.SourceName = PropertySettings.SourceName;
+			// Engine name defaults to source name.  If it doesn't match what is coming in, then set it.
+			if (PropertySettings.EngineName.Length > 0 && PropertySettings.SourceName != PropertySettings.EngineName ) 
+			{
+				this.EngineName = PropertySettings.EngineName;
+			}
+			this.PropertyCategory = PropertySettings.PropertyCategory;
+			this.PropertyFlags = PropertySettings.PropertyFlags;
+			this.DisallowPropertyFlags = PropertySettings.DisallowPropertyFlags;
+			this.PropertyExportFlags = PropertySettings.PropertyExportFlags;
+			this.Allocator = PropertySettings.Allocator;
+			this.PointerType = PropertySettings.PointerType;
+			this.RepNotifyName = PropertySettings.RepNotifyName;
+			this.ArrayDimensions = PropertySettings.ArrayDimensions;
+			this.DefaultValueTokens = PropertySettings.DefaultValueTokens;
+			this.Getter = PropertySettings.Getter;
+			this.Setter = PropertySettings.Setter;
+			this.bIsBitfield = PropertySettings.bIsBitfield;
+			this.PropertyCaps = UhtPropertyCaps.CanBeContainerValue | UhtPropertyCaps.CanBeContainerKey | UhtPropertyCaps.CanHaveConfig;
+		}
+
+		/// <inheritdoc/>
+		protected override string GetDisplayStringFromEngineName()
+		{
+			return UhtFCString.NameToDisplayString(this.EngineName, this is UhtBooleanProperty);
+		}
+
+		#region Text and code generation support
+		/// <summary>
+		/// Internal version of AppendText.  Don't append any text to the builder to get the default behavior
+		/// </summary>
+		/// <param name="Builder">Destination builder</param>
+		/// <param name="TextType">Text type of where the property is being referenced</param>
+		/// <param name="bIsTemplateArgument">If true, this property is a template arguments</param>
+		/// <returns>Destination builder</returns>
+		public virtual StringBuilder AppendText(StringBuilder Builder, UhtPropertyTextType TextType, bool bIsTemplateArgument = false)
+		{
+			// By default, we assume it will be just the simple text
+			return Builder.Append(this.CppTypeText);
+		}
+
+		/// <summary>
+		/// Append the full declaration including such things as property name and const<amp/> requirements
+		/// </summary>
+		/// <param name="Builder">Destination builder</param>
+		/// <param name="TextType">Text type of where the property is being referenced</param>
+		/// <param name="bSkipParameterName">If true, do not include the property name</param>
+		/// <returns>Destination builder</returns>
+		public virtual StringBuilder AppendFullDecl(StringBuilder Builder, UhtPropertyTextType TextType, bool bSkipParameterName = false)
+		{
+			UhtPropertyCaps Caps = this.PropertyCaps;
+
+			bool bIsParameter = TextType.IsParameter();
+			bool bIsInterfaceProp = this is UhtInterfaceProperty;
+
+			// When do we need a leading const:
+			// 1) If this is a object or object ptr property and the referenced class is const
+			// 2) If this is not an out parameter or reference parameter then,
+			//		if this is a parameter
+			//		AND - if this is a const param OR (if this is an interface property and not an out param)
+			// 3) If this is a parameter without array dimensions, must be passed by reference, but not an out parameter or const param
+			bool bPassCppArgsByRef = Caps.HasAnyFlags(UhtPropertyCaps.PassCppArgsByRef);
+			bool bIsConstParam = bIsParameter && (this.PropertyFlags.HasAnyFlags(EPropertyFlags.ConstParm) || (bIsInterfaceProp && !this.PropertyFlags.HasAnyFlags(EPropertyFlags.OutParm)));
+			bool bIsConstArgsByRef = bIsParameter && this.ArrayDimensions == null && bPassCppArgsByRef && !this.PropertyFlags.HasAnyFlags(EPropertyFlags.ConstParm | EPropertyFlags.OutParm);
+			bool bIsOnConstClass = false;
+			if (this is UhtObjectProperty ObjectProperty)
+			{
+				bIsOnConstClass = ObjectProperty.Class.ClassFlags.HasAnyFlags(EClassFlags.Const);
+			}
+			bool bShouldHaveRef = this.PropertyFlags.HasAnyFlags(EPropertyFlags.OutParm | EPropertyFlags.ReferenceParm);
+
+			bool bConstAtTheBeginning = bIsOnConstClass || bIsConstArgsByRef || (bIsConstParam && !bShouldHaveRef);
+			if (bConstAtTheBeginning)
+			{
+				Builder.Append("const ");
+			}
+
+			this.AppendText(Builder, TextType);
+
+			bool bFromConstClass = false;
+			if (TextType == UhtPropertyTextType.ExportMember && this.Outer != null)
+			{
+				if (this.Outer is UhtClass OuterClass)
+				{
+					bFromConstClass = OuterClass.ClassFlags.HasAnyFlags(EClassFlags.Const);
+				}
+			}
+			bool bConstAtTheEnd = bFromConstClass || (bIsConstParam && bShouldHaveRef);
+			if (bConstAtTheEnd)
+			{
+				Builder.Append(" const");
+			}
+
+			if (bIsParameter && this.ArrayDimensions == null && (bPassCppArgsByRef || this.PropertyFlags.HasAnyFlags(EPropertyFlags.OutParm | EPropertyFlags.ReferenceParm)))
+			{
+				Builder.Append('&');
+			}
+
+			Builder.Append(' ');
+			if (!bSkipParameterName)
+			{
+				Builder.Append(this.SourceName);
+			}
+
+			if (this.ArrayDimensions != null)
+			{
+				Builder.Append('[').Append(this.ArrayDimensions).Append(']');
+			}
+			return Builder;
+		}
+
+		/// <summary>
+		/// Append the required code to declare the property as a member
+		/// </summary>
+		/// <param name="Builder">Output builder</param>
+		/// <param name="NamePrefix">Prefix to the property name</param>
+		/// <param name="Name">Name of the property.  This is needed in some cases where the name in the declarations doesn't match the property name.</param>
+		/// <param name="NameSuffix">Suffix to the property name</param>
+		/// <param name="Tabs">Number of tabs prefix the line with</param>
+		/// <returns>Output builder</returns>
+		public abstract StringBuilder AppendMemberDecl(StringBuilder Builder, IUhtPropertyMemberContext Context, string Name, string NameSuffix, int Tabs);
+
+		/// <summary>
+		/// Append the required code to declare the property as a member
+		/// </summary>
+		/// <param name="Builder">Output builder</param>
+		/// <param name="NamePrefix">Prefix to the property name</param>
+		/// <param name="Name">Name of the property.  This is needed in some cases where the name in the declarations doesn't match the property name.</param>
+		/// <param name="NameSuffix">Suffix to the property name</param>
+		/// <param name="Tabs">Number of tabs prefix the line with</param>
+		/// <param name="ParamsStructName">Structure name</param>
+		/// <param name="bAppendMetaDataDecl">If true, add the meta data decl prior to the member decl</param>
+		/// <returns>Output builder</returns>
+		public StringBuilder AppendMemberDecl(StringBuilder Builder, IUhtPropertyMemberContext Context, string Name, string NameSuffix, int Tabs, string ParamsStructName, bool bAppendMetaDataDecl = true)
+		{
+			if (bAppendMetaDataDecl)
+			{
+				Builder.AppendMetaDataDecl(this, Context, Name, NameSuffix, Tabs);
+			}
+			Builder.AppendTabs(Tabs).Append("static const UECodeGen_Private::").Append(ParamsStructName).Append(' ').AppendNameDecl(Context, Name, NameSuffix).Append(";\r\n");
+			return Builder;
+		}
+
+		/// <summary>
+		/// Append the required code to define the property as a member
+		/// </summary>
+		/// <param name="Builder">Output builder</param>
+		/// <param name="Context">Context of the call</param>
+		/// <param name="Name">Name of the property.  This is needed in some cases where the name in the declarations doesn't match the property name.</param>
+		/// <param name="NameSuffix">Suffix to the property name</param>
+		/// <param name="Offset">Offset of the property</param>
+		/// <param name="Tabs">Number of tabs prefix the line with</param>
+		/// <returns>Output builder</returns>
+		public abstract StringBuilder AppendMemberDef(StringBuilder Builder, IUhtPropertyMemberContext Context, string Name, string NameSuffix, string? Offset, int Tabs);
+
+		/// <summary>
+		/// Append the required start of code to define the property as a member
+		/// </summary>
+		/// <param name="Builder">Output builder</param>
+		/// <param name="Context">Context of the call</param>
+		/// <param name="Name">Name of the property.  This is needed in some cases where the name in the declarations doesn't match the property name.</param>
+		/// <param name="NameSuffix">Suffix to the property name</param>
+		/// <param name="Offset">Offset of the property</param>
+		/// <param name="Tabs">Number of tabs prefix the line with</param>
+		/// <param name="ParamsStructName">Structure name</param>
+		/// <param name="ParamsGenFlags">Structure flags</param>
+		/// <param name="bAppendMetaDataDef">If true, add the meta data def prior to the member def</param>
+		/// <param name="bAppendOffset">If true, add the offset parameter</param>
+		/// <returns>Output builder</returns>
+		public StringBuilder AppendMemberDefStart(StringBuilder Builder, IUhtPropertyMemberContext Context, string Name, string NameSuffix, string? Offset, int Tabs,
+			string ParamsStructName, string ParamsGenFlags, bool bAppendMetaDataDef = true, bool bAppendOffset = true)
+		{
+			if (bAppendMetaDataDef)
+			{
+				Builder.AppendMetaDataDef(this, Context, Name, NameSuffix, Tabs);
+			}
+			Builder
+				.AppendTabs(Tabs)
+				.Append("const UECodeGen_Private::").Append(ParamsStructName).Append(' ')
+				.AppendNameDef(Context, Name, NameSuffix).Append(" = { ")
+				.AppendUTF8LiteralString(this.EngineName).Append(", ")
+				.AppendNotifyFunc(this).Append(", ")
+				.AppendFlags(this.PropertyFlags).Append(", ")
+				.Append(ParamsGenFlags).Append(", ")
+				.Append(ObjectFlags).Append(", ")
+				.AppendArrayDim(this, Context).Append(", ");
+
+			if (this.PropertyExportFlags.HasAnyFlags(UhtPropertyExportFlags.SetterFound))
+			{
+				Builder.Append('&').Append(this.Outer!.SourceName).Append("::").AppendPropertySetterWrapperName(this).Append(", ");
+			}
+			else
+			{
+				Builder.Append("nullptr, ");
+			}
+
+			if (this.PropertyExportFlags.HasAnyFlags(UhtPropertyExportFlags.GetterFound))
+			{
+				Builder.Append('&').Append(this.Outer!.SourceName).Append("::").AppendPropertyGetterWrapperName(this).Append(", ");
+			}
+			else
+			{
+				Builder.Append("nullptr, ");
+			}
+
+			if (bAppendOffset)
+			{
+				if (!string.IsNullOrEmpty(Offset))
+				{
+					Builder.Append(Offset).Append(", ");
+				}
+				else
+				{
+					Builder.Append("STRUCT_OFFSET(").Append(Context.OuterStructSourceName).Append(", ").Append(this.SourceName).Append("), ");
+				}
+			}
+			return Builder;
+		}
+
+		/// <summary>
+		/// Append the required end of code to define the property as a member
+		/// </summary>
+		/// <param name="Builder">Output builder</param>
+		/// <param name="Context">Context of the call</param>
+		/// <param name="Name">Name of the property.  This is needed in some cases where the name in the declarations doesn't match the property name.</param>
+		/// <param name="NameSuffix">Suffix to the property name</param>
+		/// <returns>Output builder</returns>
+		protected StringBuilder AppendMemberDefEnd(StringBuilder Builder, IUhtPropertyMemberContext Context, string Name, string NameSuffix)
+		{
+			Builder
+				.AppendMetaDataParams(this, Context, Name, NameSuffix)
+				.Append(" };")
+				.AppendObjectHashes(this, Context)
+				.Append("\r\n");
+			return Builder;
+		}
+
+		/// <summary>
+		/// Append the a type reference to the member definition
+		/// </summary>
+		/// <param name="Builder">Output builder</param>
+		/// <param name="Context">Context of the call</param>
+		/// <param name="Object">Referenced object</param>
+		/// <param name="bRegistered">True if the registered singleton name is to be used</param>
+		/// <param name="bAppendNull">True if a "nullptr" is to be appended if the object is null</param>
+		/// <returns>Output builder</returns>
+		protected StringBuilder AppendMemberDefRef(StringBuilder Builder, IUhtPropertyMemberContext Context, UhtObject? Object, bool bRegistered, bool bAppendNull = false)
+		{
+			if (Object != null)
+			{
+				Builder.AppendSingletonName(Context, Object, bRegistered).Append(", ");
+			}
+			else if (bAppendNull)
+			{
+				Builder.Append("nullptr, ");
+			}
+			return Builder;
+		}
+
+		/// <summary>
+		/// Append the required code to add the properties to a pointer array
+		/// </summary>
+		/// <param name="Builder">Output builder</param>
+		/// <param name="Context">Context of the call</param>
+		/// <param name="Name">Name of the property.  This is needed in some cases where the name in the declarations doesn't match the property name.</param>
+		/// <param name="NameSuffix">Suffix to the property name</param>
+		/// <param name="Tabs">Number of tabs prefix the line with</param>
+		/// <returns>Output builder</returns>
+		public virtual StringBuilder AppendMemberPtr(StringBuilder Builder, IUhtPropertyMemberContext Context, string Name, string NameSuffix, int Tabs)
+		{
+			Builder.AppendTabs(Tabs).Append("(const UECodeGen_Private::FPropertyParamsBase*)&").AppendNameDef(Context, Name, NameSuffix).Append(",\r\n");
+			return Builder;
+		}
+
+		/// <summary>
+		/// Append a P_GET macro
+		/// </summary>
+		/// <param name="Builder">Destination builder</param>
+		/// <returns>Destination builder</returns>
+		public virtual StringBuilder AppendFunctionThunkParameterGet(StringBuilder Builder)
+		{
+			Builder.Append("P_GET_");
+			if (this.ArrayDimensions != null)
+			{
+				Builder.Append("ARRAY");
+			}
+			else
+			{
+				Builder.Append(this.PGetMacroText);
+			}
+			if (this.PropertyFlags.HasAnyFlags(EPropertyFlags.OutParm))
+			{
+				if (this.bPGetPassAsNoPtr)
+				{
+					Builder.Append("_REF_NO_PTR");
+				}
+				else
+				{
+					Builder.Append("_REF");
+				}
+			}
+			Builder.Append("(");
+			if (this.ArrayDimensions != null)
+			{
+				Builder.AppendFunctionThunkParameterArrayType(this).Append(",");
+			}
+			else
+			{
+				switch (this.PGetTypeArgument)
+				{
+					case UhtPGetArgumentType.None:
+						break;
+
+					case UhtPGetArgumentType.EngineClass:
+						Builder.Append('F').Append(this.EngineClassName).Append(',');
+						break;
+
+					case UhtPGetArgumentType.TypeText:
+						Builder.AppendPropertyText(this, UhtPropertyTextType.FunctionThunkParameterArgType).Append(",");
+						break;
+				}
+			}
+			Builder.AppendFunctionThunkParameterName(this).Append(")");
+			return Builder;
+		}
+
+		/// <summary>
+		/// Append the text for a function thunk call argument
+		/// </summary>
+		/// <param name="Builder">Output builder</param>
+		/// <returns>Output builder</returns>
+		public virtual StringBuilder AppendFunctionThunkParameterArg(StringBuilder Builder)
+		{
+			return Builder.AppendFunctionThunkParameterName(this);
+		}
+
+		/// <summary>
+		/// Apppend the name of a function thunk paramter
+		/// </summary>
+		/// <param name="Builder">Output builder</param>
+		/// <returns>Output builder</returns>
+		public StringBuilder AppendFunctionThunkParameterName(StringBuilder Builder)
+		{
+			Builder.Append(FunctionParameterThunkPrefix);
+			if (this.PropertyFlags.HasAnyFlags(EPropertyFlags.OutParm))
+			{
+				Builder.Append("Out_");
+			}
+			Builder.Append(this.EngineName);
+			return Builder;
+		}
+
+		/// <summary>
+		/// Append the appropriate values to initialize the property to a "NULL" value;
+		/// </summary>
+		/// <param name="Builder"></param>
+		/// <param name="bIsInitializer"></param>
+		/// <returns></returns>
+		public abstract StringBuilder AppendNullConstructorArg(StringBuilder Builder, bool bIsInitializer);
+
+		/// <summary>
+		/// Return the basic declaration type text for user facing messages
+		/// </summary>
+		/// <returns></returns>
+		public string GetUserFacingDecl()
+		{
+			StringBuilder Builder = new StringBuilder();
+			AppendText(Builder, UhtPropertyTextType.UserFacing);
+			return Builder.ToString();
+		}
+
+		/// <summary>
+		/// Appends any applicable objects and child properties
+		/// </summary>
+		/// <param name="Builder">Output builder</param>
+		/// <param name="StartingLength">Initial length of the builder prior to appending the hashes</param>
+		/// <param name="Context">Context used to lookup the hashes</param>
+		public virtual void AppendObjectHashes(StringBuilder Builder, int StartingLength, IUhtPropertyMemberContext Context)
+		{
+		}
+#endregion
+
+#region Parsing support
+		/// <summary>
+		/// Parse a default value for the property and return a sanitized string representation.
+		/// 
+		/// All tokens in the token reader must be consumed.  Otherwise the default value will be considered to be invalid.
+		/// </summary>
+		/// <param name="DefaultValueReader">Reader containing the default value</param>
+		/// <param name="InnerDefaultValue">Sanitized representation of default value</param>
+		/// <returns>True if a default value was parsed.</returns>
+		public abstract bool SanitizeDefaultValue(IUhtTokenReader DefaultValueReader, StringBuilder InnerDefaultValue);
+#endregion
+
+#region Resolution support
+		protected override bool ResolveSelf(UhtResolvePhase Phase)
+		{
+			bool bResult = base.ResolveSelf(Phase);
+
+			switch (Phase)
+			{
+				case UhtResolvePhase.Final:
+					if (this.ArrayDimensions != null)
+					{
+						ReadOnlySpan<char> Dim = this.ArrayDimensions.AsSpan();
+
+						bool bAgain = false;
+						do
+						{
+							bAgain = false;
+
+							// Remove any outer brackets
+							if (Dim[0] == '(')
+							{
+								for (int Index = 1, Depth = 1; Index < Dim.Length; ++Index)
+								{
+									if (Dim[Index] == ')')
+									{
+										if (--Depth == 0)
+										{
+											if (Index == Dim.Length - 1)
+											{
+												Dim = Dim.Slice(1, Index - 1);
+												bAgain = true;
+											}
+											break;
+										}
+									}
+									else if (Dim[Index] == '(')
+									{
+										++Depth;
+									}
+								}
+							}
+
+							// Remove any well known casts
+							if (Dim.Length > 0)
+							{
+								foreach (string Cast in Casts)
+								{
+									if (Dim.StartsWith(Cast))
+									{
+										Dim = Dim.Slice(Cast.Length);
+										bAgain = true;
+										break;
+									}
+								}
+							}
+
+						} while (bAgain && Dim.Length > 0);
+
+						//COMPATIBILITY-TODO - This method is more robust, but causes differences.  See UhtSession for future
+						// plans on fix in this.
+						//// Now try to see if this is an enum
+						//UhtEnum? Enum = null;
+						//int Sep = Dim.IndexOf("::");
+						//if (Sep >= 0)
+						//{
+						//	//COMPATIBILITY-TODO "Bob::Type::Value" did not generate a match with the old code
+						//	if (Dim.LastIndexOf("::") != Sep)
+						//	{
+						//		break;
+						//	}
+						//	Dim = Dim.Slice(0, Sep);
+						//}
+						//else
+						//{
+						//	Enum = this.Session.FindRegularEnumValue(Dim.ToString());
+						//}
+						//if (Enum == null)
+						//{
+						//	Enum = this.Session.FindType(this.Outer, UhtFindOptions.Enum | UhtFindOptions.SourceName, Dim.ToString()) as UhtEnum;
+						//}
+						//if (Enum != null)
+						//{
+						//	this.MetaData.Add(UhtNames.ArraySizeEnum, Enum.GetPathName());
+						//}
+
+						if (Dim.Length > 0 && !UhtFCString.IsDigit(Dim[0]))
+						{
+							UhtEnum? Enum = this.Session.FindRegularEnumValue(Dim.ToString());
+							if (Enum == null)
+							{
+								Enum = this.Session.FindType(this.Outer, UhtFindOptions.Enum | UhtFindOptions.SourceName, this.ArrayDimensions) as UhtEnum;
+							}
+							if (Enum != null)
+							{
+								this.MetaData.Add(UhtNames.ArraySizeEnum, Enum.GetPathName());
+							}
+						}
+					}
+					break;
+			}
+			return bResult;
+		}
+
+		/// <summary>
+		/// Check properties to see if any instances are referenced.
+		/// This method does NOT cache the result.
+		/// </summary>
+		/// <param name="bDeepScan">If true, the ScanForInstancedReferenced method on the properties will also be called.</param>
+		/// <returns></returns>
+		public virtual bool ScanForInstancedReferenced(bool bDeepScan)
+		{
+			return false;
+		}
+#endregion
+
+#region Validation support
+		protected override UhtValidationOptions Validate(UhtValidationOptions Options)
+		{
+			Options = base.Validate(Options);
+
+			// The outer should never be null...
+			if (this.Outer == null)
+			{
+				return Options;
+			}
+
+			// Shadowing checks are done at this level, not in the properties themselves
+			// Note: The old UHT would check for all property type where this version only checks for member variables.
+			// In the old code if you defined the property after the function with argument with the same name, UHT would
+			// not issue an error.  However, if the property was defined PRIOR to the function with the matching argument name,
+			// UHT would generate an error.
+			if (Options.HasAnyFlags(UhtValidationOptions.Shadowing) && this.PropertyCategory == UhtPropertyCategory.Member)
+			{
+
+				// First check for duplicate name in self and then duplicate name in parents
+				UhtType? Existing = this.Outer.FindType(UhtFindOptions.PropertyOrFunction | UhtFindOptions.EngineName | UhtFindOptions.SelfOnly, this.EngineName);
+				if (Existing == this)
+				{
+					Existing = this.Outer.FindType(UhtFindOptions.PropertyOrFunction | UhtFindOptions.EngineName | UhtFindOptions.ParentsOnly | UhtFindOptions.NoGlobal | UhtFindOptions.NoIncludes, this.EngineName);
+				}
+
+				if (Existing != null && Existing != this)
+				{
+					if (Existing is UhtProperty ExistingProperty)
+					{
+						//@TODO: This exception does not seem sound either, but there is enough existing code that it will need to be
+						// fixed up first before the exception it is removed.
+						bool bExistingPropDeprecated = ExistingProperty.PropertyFlags.HasAnyFlags(EPropertyFlags.Deprecated);
+						bool bNewPropDeprecated = this.PropertyCategory == UhtPropertyCategory.Member && this.PropertyFlags.HasAnyFlags(EPropertyFlags.Deprecated);
+						if (!bNewPropDeprecated && !bExistingPropDeprecated)
+						{
+							this.LogShadowingError(ExistingProperty);
+						}
+					}
+					else if (Existing is UhtFunction ExistingFunction)
+					{
+						if (this.PropertyCategory == UhtPropertyCategory.Member)
+						{
+							this.LogShadowingError(ExistingFunction);
+						}
+					}
+				}
+			}
+
+			Validate((UhtStruct)this.Outer, this, Options);
+			return Options;
+		}
+
+		private void LogShadowingError(UhtType Shadows)
+		{
+			this.LogError($"{this.PropertyCategory.GetHintText()}: '{this.SourceName}' cannot be defined in '{this.Outer?.SourceName}' as it is already defined in scope '{Shadows.Outer?.SourceName}' (shadowing is not allowed)");
+		}
+
+		// Validate the property settings
+		public virtual void Validate(UhtStruct OuterStruct, UhtProperty OutermostProperty, UhtValidationOptions Options)
+		{
+			// Check for deprecation
+			if (Options.HasAnyFlags(UhtValidationOptions.Deprecated) && !this.bDeprecated)
+			{
+				ValidateDeprecated();
+			}
+
+			// Validate the types allowed with arrays
+			if (this.ArrayDimensions != null)
+			{
+				switch (this.PropertyCategory)
+				{
+					case UhtPropertyCategory.Return:
+						this.LogError("Arrays aren't allowed as return types");
+						break;
+
+					case UhtPropertyCategory.RegularParameter:
+					case UhtPropertyCategory.ReplicatedParameter:
+						this.LogError("Arrays aren't allowed as function parameters");
+						break;
+				}
+
+				if (this is UhtContainerBaseProperty)
+				{
+					this.LogError("Static arrays of containers are not allowed");
+				}
+
+				if (this is UhtBooleanProperty)
+				{
+					this.LogError("Bool arrays are not allowed");
+				}
+			}
+
+			if (!Options.HasAnyFlags(UhtValidationOptions.IsKey) && this.PropertyFlags.HasAnyFlags(EPropertyFlags.PersistentInstance) && !this.PropertyCaps.HasAnyFlags(UhtPropertyCaps.CanBeInstanced))
+			{
+				this.LogError("'Instanced' is only allowed on an object property, an array of objects, a set of objects, or a map with an object value type.");
+			}
+
+			if (this.PropertyFlags.HasAnyFlags(EPropertyFlags.Config) && !this.PropertyCaps.HasAnyFlags(UhtPropertyCaps.CanHaveConfig))
+			{
+				this.LogError("Not allowed to use 'config' with object variables");
+			}
+
+			if (this.PropertyFlags.HasAnyFlags(EPropertyFlags.ExposeOnSpawn))
+			{
+				if (this.PropertyFlags.HasAnyFlags(EPropertyFlags.DisableEditOnInstance))
+				{
+					this.LogWarning("Property cannot have both 'DisableEditOnInstance' and 'ExposeOnSpawn' flags");
+				}
+				if (!this.PropertyFlags.HasAnyFlags(EPropertyFlags.BlueprintVisible))
+				{
+					this.LogWarning("Property cannot have 'ExposeOnSpawn' without 'BlueprintVisible' flag.");
+				}
+			}
+
+			if (this.PropertyFlags.HasAnyFlags(EPropertyFlags.BlueprintAssignable) && !this.PropertyCaps.HasAnyFlags(UhtPropertyCaps.CanBeBlueprintAssignable))
+			{
+				this.LogError("'BlueprintAssignable' is only allowed on multicast delegate properties");
+			}
+
+			if (this.PropertyFlags.HasAnyFlags(EPropertyFlags.BlueprintCallable) && !this.PropertyCaps.HasAnyFlags(UhtPropertyCaps.CanBeBlueprintCallable))
+			{
+				this.LogError("'BlueprintCallable' is only allowed on multicast delegate properties");
+			}
+
+			if (this.PropertyFlags.HasAnyFlags(EPropertyFlags.BlueprintAuthorityOnly) && !this.PropertyCaps.HasAnyFlags(UhtPropertyCaps.CanBeBlueprintAuthorityOnly))
+			{
+				this.LogError("'BlueprintAuthorityOnly' is only allowed on multicast delegate properties");
+			}
+
+			// Check for invalid transients
+			EPropertyFlags Transients = this.PropertyFlags & (EPropertyFlags.DuplicateTransient | EPropertyFlags.TextExportTransient | EPropertyFlags.NonPIEDuplicateTransient);
+			if (Transients != 0 && !(OuterStruct is UhtClass))
+			{
+				this.LogError($"'{string.Join(", ", Transients.ToStringList(false))}' specifier(s) are only allowed on class member variables");
+			}
+
+			if (!Options.HasAnyFlags(UhtValidationOptions.IsKeyOrValue))
+			{
+				switch (this.PropertyCategory)
+				{
+					case UhtPropertyCategory.Member:
+						ValidateMember(OuterStruct, Options);
+						break;
+
+					case UhtPropertyCategory.RegularParameter:
+					case UhtPropertyCategory.ReplicatedParameter:
+						ValidateFunctionArgument((UhtFunction)OuterStruct, Options);
+						break;
+
+					case UhtPropertyCategory.Return:
+						break;
+				}
+			}
+			return;
+		}
+
+		/// <summary>
+		/// Validate that we don't reference any deprecated classes
+		/// </summary>
+		public virtual void ValidateDeprecated()
+		{
+		}
+
+		/// <summary>
+		/// Verify function argument
+		/// </summary>
+		protected virtual void ValidateFunctionArgument(UhtFunction Function, UhtValidationOptions Options)
+		{
+			if (Function.FunctionFlags.HasAnyFlags(EFunctionFlags.Net))
+			{
+				if (this.PropertyFlags.HasExactFlags(EPropertyFlags.ReferenceParm | EPropertyFlags.ConstParm, EPropertyFlags.ReferenceParm))
+				{
+					this.LogError($"Replicated parameters cannot be passed by non-const reference");
+				}
+
+				if (Function.FunctionFlags.HasAnyFlags(EFunctionFlags.NetRequest))
+				{
+					if (this.PropertyFlags.HasExactFlags(EPropertyFlags.OutParm | EPropertyFlags.RepSkip, EPropertyFlags.OutParm))
+					{
+						// This is difficult to trigger since NotReplicated also sets the property category
+						this.LogError("Service request functions cannot contain out parameters, unless marked NotReplicated");
+					}
+				}
+				else
+				{
+					if (this.PropertyFlags.HasAnyFlags(EPropertyFlags.OutParm))
+					{
+						this.LogError("Replicated functions cannot contain out parameters");
+					}
+
+					if (this.PropertyFlags.HasAnyFlags(EPropertyFlags.RepSkip))
+					{
+						// This is difficult to trigger since NotReplicated also sets the property category
+						this.LogError("Only service request functions cannot contain NotReplicated parameters");
+					}
+				}
+			}
+
+			// The following checks are not performed on the value of a container
+			if (Function.FunctionFlags.HasAnyFlags(EFunctionFlags.BlueprintEvent | EFunctionFlags.BlueprintCallable))
+			{
+				// Check that the parameter name is valid and does not conflict with pre-defined types
+				if (InvalidParamNames.Contains(this.SourceName))
+				{
+					this.LogError($"Parameter name '{this.SourceName}' in function is invalid, '{this.SourceName}' is a reserved name.");
+				}
+			}
+		}
+
+		/// <summary>
+		/// Validate member settings
+		/// </summary>
+		/// <param name="Struct">Containing struct.  This is either a UhtScriptStruct or UhtClass</param>
+		/// <param name="Options">Validation options</param>
+		protected virtual void ValidateMember(UhtStruct Struct, UhtValidationOptions Options)
+		{
+			if (!Options.HasAnyFlags(UhtValidationOptions.IsKeyOrValue))
+			{
+				// First check if the category was specified at all and if the property was exposed to the editor.
+				string? Category;
+				if (!this.MetaData.TryGetValue(UhtNames.Category, out Category))
+				{
+					if (this.PropertyFlags.HasAnyFlags(EPropertyFlags.Edit | EPropertyFlags.BlueprintVisible))
+					{
+						if (this.Package.bIsPartOfEngine)
+						{
+							this.LogError("An explicit Category specifier is required for any property exposed to the editor or Blueprints in an Engine module.");
+						}
+					}
+				}
+
+				// If the category was specified explicitly, it wins
+				if (!string.IsNullOrEmpty(Category) && !this.PropertyFlags.HasAnyFlags(EPropertyFlags.Edit | EPropertyFlags.BlueprintVisible |
+					EPropertyFlags.BlueprintAssignable | EPropertyFlags.BlueprintCallable))
+				{
+					this.LogWarning("Property has a Category set but is not exposed to the editor or Blueprints with EditAnywhere, BlueprintReadWrite, " +
+						"VisibleAnywhere, BlueprintReadOnly, BlueprintAssignable, BlueprintCallable keywords.");
+				}
+			}
+
+			// Make sure that editblueprint variables are editable
+			if (!this.PropertyFlags.HasAnyFlags(EPropertyFlags.Edit))
+			{
+				if (this.PropertyFlags.HasAnyFlags(EPropertyFlags.DisableEditOnInstance))
+				{
+					this.LogError("Property cannot have 'DisableEditOnInstance' without being editable");
+				}
+
+				if (this.PropertyFlags.HasAnyFlags(EPropertyFlags.DisableEditOnTemplate))
+				{
+					this.LogError("Property cannot have 'DisableEditOnTemplate' without being editable");
+				}
+			}
+
+			string ExposeOnSpawnValue = this.MetaData.GetValueOrDefault(UhtNames.ExposeOnSpawn);
+			if (ExposeOnSpawnValue.Equals("true", StringComparison.OrdinalIgnoreCase) && !this.PropertyCaps.HasAnyFlags(UhtPropertyCaps.CanExposeOnSpawn))
+			{
+				this.LogError("ExposeOnSpawn - Property cannot be exposed");
+			}
+		}
+
+		/// <summary>
+		/// Generate an error if the class has been deprecated
+		/// </summary>
+		/// <param name="Class">Class to check</param>
+		protected void ValidateDeprecatedClass(UhtClass? Class)
+		{
+			if (Class == null)
+			{
+				return;
+			}
+
+			if (!Class.bDeprecated)
+			{
+				return;
+			}
+
+			if (this.PropertyCategory == UhtPropertyCategory.Member)
+			{
+				this.LogError($"Property is using a deprecated class: '{Class.SourceName}'.  Property should be marked deprecated as well.");
+			}
+			else
+			{
+				this.LogError($"Function is using a deprecated class: '{Class.SourceName}'.  Function should be marked deprecated as well.");
+			}
+		}
+
+		/// <summary>
+		/// Check to see if the property is valid as a member of a networked structure
+		/// </summary>
+		/// <param name="ReferencingProperty">The property referencing the structure property.  All error should be logged on the referencing property.</param>
+		/// <returns>True if the property is valid, false if not.  If the property is not valid, an error should also be generated.</returns>
+		public virtual bool ValidateStructPropertyOkForNet(UhtProperty ReferencingProperty)
+		{
+			return true;
+		}
+
+		/// <summary>
+		/// Test to see if this property references something that requires the argument to be marked as const
+		/// </summary>
+		/// <param name="ErrorType">If const is required, returns the type that is forcing the const</param>
+		/// <returns>True if the argument must be marked as const</returns>
+		public virtual bool MustBeConstArgument([NotNullWhen(true)] out UhtType? ErrorType)
+		{
+			ErrorType = null;
+			return false;
+		}
+#endregion
+
+#region Reference support
+		/// <inheritdoc/>
+		public override void CollectReferences(IUhtReferenceCollector Collector)
+		{
+			CollectReferencesInternal(Collector, false);
+			if (this.PropertyFlags.HasAnyFlags(EPropertyFlags.ParmFlags))
+			{
+				Collector.AddForwardDeclaration(this.GetForwardDeclarations());
+			}
+		}
+
+		public virtual void CollectReferencesInternal(IUhtReferenceCollector Collector, bool bIsTemplateProperty)
+		{
+		}
+
+		public virtual string? GetForwardDeclarations()
+		{
+			return null;
+		}
+
+		public virtual IEnumerable<UhtType> EnumerateReferencedTypes()
+		{
+			return Enumerable.Empty<UhtType>();
+		}
+#endregion
+
+#region RigVM support methods
+		public abstract string? GetRigVMType(ref UhtRigVMParameterFlags ParameterFlags);
+#endregion
+
+#region Helper methods
+		protected string GetNameSuffix(string OuterSuffix, string NewSuffix)
+		{
+			return string.IsNullOrEmpty(OuterSuffix) ? NewSuffix : $"{OuterSuffix}{NewSuffix}";
+		}
+
+		public abstract bool IsSameType(UhtProperty Other);
+
+		public bool MatchesType(UhtProperty Other)
+		{
+			if (this.PropertyFlags.HasAnyFlags(EPropertyFlags.OutParm))
+			{
+				if (!Other.PropertyFlags.HasAnyFlags(EPropertyFlags.OutParm))
+				{
+					return false;
+				}
+
+				if (this.PropertyFlags.HasAnyFlags(EPropertyFlags.ConstParm) && !Other.PropertyFlags.HasAnyFlags(EPropertyFlags.ConstParm))
+				{
+					return false;
+				}
+			}
+			return IsSameType(Other);
+		}
+#endregion
+	}
+
+	public static class UhtPropertyStringBuilderExtensions
+	{
+		public static StringBuilder AppendPropertyText(this StringBuilder Builder, UhtProperty Property, UhtPropertyTextType TextType, bool bIsTemplateArgument = false)
+		{
+			return Property.AppendText(Builder, TextType, bIsTemplateArgument);
+		}
+
+		public static StringBuilder AppendMemberDecl(this StringBuilder Builder, UhtProperty Property, IUhtPropertyMemberContext Context, string Name, string NameSuffix, int Tabs)
+		{
+			return Property.AppendMemberDecl(Builder, Context, Name, NameSuffix, Tabs);
+		}
+
+		public static StringBuilder AppendMemberDef(this StringBuilder Builder, UhtProperty Property, IUhtPropertyMemberContext Context, string Name, string NameSuffix, string? Offset, int Tabs)
+		{
+			return Property.AppendMemberDef(Builder, Context, Name, NameSuffix, Offset, Tabs);
+		}
+
+		public static StringBuilder AppendMemberPtr(this StringBuilder Builder, UhtProperty Property, IUhtPropertyMemberContext Context, string Name, string NameSuffix, int Tabs)
+		{
+			return Property.AppendMemberPtr(Builder, Context, Name, NameSuffix, Tabs);
+		}
+
+		public static StringBuilder AppendFullDecl(this StringBuilder Builder, UhtProperty Property, UhtPropertyTextType TextType, bool bSkipParameterName)
+		{
+			return Property.AppendFullDecl(Builder, TextType, bSkipParameterName);
+		}
+
+		public static StringBuilder AppendUserFacingDecl(this StringBuilder Builder, UhtProperty Property)
+		{
+			return Property.AppendText(Builder, UhtPropertyTextType.UserFacing);
+		}
+
+		public static StringBuilder AppendFunctionThunkParameterGet(this StringBuilder Builder, UhtProperty Property)
+		{
+			return Property.AppendFunctionThunkParameterGet(Builder);
+		}
+
+		public static StringBuilder AppendFunctionThunkParameterArrayType(this StringBuilder Builder, UhtProperty Property)
+		{
+			return Property.AppendText(Builder, UhtPropertyTextType.FunctionThunkParameterArrayType);
+		}
+
+		public static StringBuilder AppendFunctionThunkParameterArg(this StringBuilder Builder, UhtProperty Property)
+		{
+			return Property.AppendFunctionThunkParameterArg(Builder);
+		}
+
+		public static StringBuilder AppendFunctionThunkReturn(this StringBuilder Builder, UhtProperty Property)
+		{
+			return Property.AppendText(Builder, UhtPropertyTextType.FunctionThunkReturn);
+		}
+
+		public static StringBuilder AppendFunctionThunkParameterName(this StringBuilder Builder, UhtProperty Property)
+		{
+			return Property.AppendFunctionThunkParameterName(Builder);
+		}
+
+		public static StringBuilder AppendSparse(this StringBuilder Builder, UhtProperty Property)
+		{
+			return Property.AppendText(Builder, UhtPropertyTextType.Sparse);
+		}
+
+		public static StringBuilder AppendNullConstructorArg(this StringBuilder Builder, UhtProperty Property, bool bIsInitializer)
+		{
+			Property.AppendNullConstructorArg(Builder, bIsInitializer);
+			return Builder;
+		}
+
+		public static StringBuilder AppendNotifyFunc(this StringBuilder Builder, UhtProperty Property)
+		{
+			if (Property.RepNotifyName != null)
+			{
+				Builder.AppendUTF8LiteralString(Property.RepNotifyName);
+			}
+			else
+			{
+				Builder.Append("nullptr");
+			}
+			return Builder;
+		}
+
+		public static StringBuilder AppendFlags(this StringBuilder Builder, EPropertyFlags PropertyFlags)
+		{
+			PropertyFlags &= ~EPropertyFlags.ComputedFlags;
+			return Builder.Append("(EPropertyFlags)0x").AppendFormat("{0:x16}", (ulong)PropertyFlags);
+		}
+
+		public static StringBuilder AppendArrayDim(this StringBuilder Builder, UhtProperty Property, IUhtPropertyMemberContext Context)
+		{
+			if (Property.ArrayDimensions != null)
+			{
+				Builder.Append("CPP_ARRAY_DIM(").Append(Property.SourceName).Append(", ").Append(Context.OuterStruct.SourceName).Append(')');
+			}
+			else
+			{
+				Builder.Append("1");
+			}
+			return Builder;
+		}
+
+		/// <summary>
+		/// Given an object, append the hash (if applicable) to the builder
+		/// </summary>
+		/// <param name="Builder">Output builder</param>
+		/// <param name="ReferingType">Type asking for an object hash</param>
+		/// <param name="StartingLength">Initial length of the builder prior to appending the hashes</param>
+		/// <param name="Context">Context used to lookup the hashes</param>
+		/// <param name="Object">Object being appended</param>
+		public static StringBuilder AppendObjectHash(this StringBuilder Builder, int StartingLength, UhtType ReferingType, IUhtPropertyMemberContext Context, UhtObject? Object)
+		{
+			if (Object == null)
+			{
+				return Builder;
+			}
+			else if (Object is UhtClass Class)
+			{
+				if (Class.ClassFlags.HasAnyFlags(EClassFlags.NoExport))
+				{
+					return Builder;
+				}
+			}
+			else if (Object is UhtScriptStruct ScriptStruct)
+			{
+				if (ScriptStruct.ScriptStructFlags.HasAnyFlags(EStructFlags.NoExport))
+				{
+					return Builder;
+				}
+			}
+
+			Builder.Append(Builder.Length == StartingLength ? " // " : " ");
+			uint Hash = Context.GetTypeHash(Object);
+			if (Hash == 0)
+			{
+				string Type = ReferingType is UhtProperty ? "property" : "object";
+				ReferingType.LogWarning($"The {Type} \"{ReferingType.SourceName}\" references type \"{Object.SourceName}\" but the code generation hash is zero");
+			}
+			Builder.Append(Context.GetTypeHash(Object));
+			return Builder;
+		}
+
+		/// <summary>
+		/// Given an object, append the hash (if applicable) to the builder
+		/// </summary>
+		/// <param name="Builder">Output builder</param>
+		/// <param name="ReferingType">Type asking for an object hash</param>
+		/// <param name="Context">Context used to lookup the hashes</param>
+		/// <param name="Object">Object being appended</param>
+		public static StringBuilder AppendObjectHash(this StringBuilder Builder, UhtType ReferingType, IUhtPropertyMemberContext Context, UhtObject? Object)
+		{
+			return Builder.AppendObjectHash(Builder.Length, ReferingType, Context, Object);
+		}
+
+		public static StringBuilder AppendObjectHashes(this StringBuilder Builder, UhtProperty Property, IUhtPropertyMemberContext Context)
+		{
+			Property.AppendObjectHashes(Builder, Builder.Length, Context);
+			return Builder;
+		}
+
+		public static StringBuilder AppendSingletonName(this StringBuilder Builder, IUhtPropertyMemberContext Context, UhtObject? Type, bool bRegistered)
+		{
+			return Builder.Append(Context.GetSingletonName(Type, bRegistered));
+		}
+
+		public static StringBuilder AppendPropertyGetterWrapperName(this StringBuilder Builder, UhtProperty Property)
+		{
+			return Builder.Append("Get").Append(Property.SourceName).Append("_WrapperImpl");
+		}
+
+		public static StringBuilder AppendPropertySetterWrapperName(this StringBuilder Builder, UhtProperty Property)
+		{
+			return Builder.Append("Set").Append(Property.SourceName).Append("_WrapperImpl");
+		}
+	}
+}
diff --git a/Engine/Source/Programs/Shared/EpicGames.UHT/Types/UhtScriptStruct.cs b/Engine/Source/Programs/Shared/EpicGames.UHT/Types/UhtScriptStruct.cs
new file mode 100644
--- /dev/null
+++ b/Engine/Source/Programs/Shared/EpicGames.UHT/Types/UhtScriptStruct.cs
@@ -0,0 +1,500 @@
+// Copyright Epic Games, Inc. All Rights Reserved.
+
+using EpicGames.Core;
+using EpicGames.UHT.Tables;
+using EpicGames.UHT.Utils;
+using System;
+using System.Collections.Generic;
+using System.Text.Json.Serialization;
+
+namespace EpicGames.UHT.Types
+{
+	/// <summary>
+	/// Flags to represent information about a RigVM parameter
+	/// </summary>
+	[Flags]
+	public enum UhtRigVMParameterFlags : UInt32
+	{
+		None			= 0,
+		Constant		= 0x00000001,
+		Input			= 0x00000002,
+		Output			= 0x00000004,
+		Singleton		= 0x00000008,
+		EditorOnly		= 0x00000010,
+		IsEnum			= 0x00010000,
+		IsArray			= 0x00020000,
+		IsFixedArray	= 0x00040000,
+	}
+
+	/// <summary>
+	/// Helper methods for testing flags.  These methods perform better than the generic HasFlag which hits
+	/// the GC and stalls.
+	/// </summary>
+	public static class UhtRigVMParameterFlagsExtensions
+	{
+
+		/// <summary>
+		/// Test to see if any of the specified flags are set
+		/// </summary>
+		/// <param name="InFlags">Current flags</param>
+		/// <param name="TestFlags">Flags to test for</param>
+		/// <returns>True if any of the flags are set</returns>
+		public static bool HasAnyFlags(this UhtRigVMParameterFlags InFlags, UhtRigVMParameterFlags TestFlags)
+		{
+			return (InFlags & TestFlags) != 0;
+		}
+
+		/// <summary>
+		/// Test to see if all of the specified flags are set
+		/// </summary>
+		/// <param name="InFlags">Current flags</param>
+		/// <param name="TestFlags">Flags to test for</param>
+		/// <returns>True if all the flags are set</returns>
+		public static bool HasAllFlags(this UhtRigVMParameterFlags InFlags, UhtRigVMParameterFlags TestFlags)
+		{
+			return (InFlags & TestFlags) == TestFlags;
+		}
+
+		/// <summary>
+		/// Test to see if a specific set of flags have a specific value.
+		/// </summary>
+		/// <param name="InFlags">Current flags</param>
+		/// <param name="TestFlags">Flags to test for</param>
+		/// <param name="MatchFlags">Expected value of the tested flags</param>
+		/// <returns>True if the given flags have a specific value.</returns>
+		public static bool HasExactFlags(this UhtRigVMParameterFlags InFlags, UhtRigVMParameterFlags TestFlags, UhtRigVMParameterFlags MatchFlags)
+		{
+			return (InFlags & TestFlags) == MatchFlags;
+		}
+	}
+
+	/**
+	 * The FRigVMParameter represents a single parameter of a method
+	 * marked up with RIGVM_METHOD.
+	 * Each parameter can be marked with Constant, Input or Output
+	 * metadata - this struct simplifies access to that information.
+	 */
+
+	public class UhtRigVMParameter
+	{
+		public readonly UhtProperty? Property;
+		public string Name { get; } = String.Empty;
+		public string Type { get; } = String.Empty;
+		public readonly string Getter = String.Empty;
+		public readonly string? CastName = null;
+		public readonly string? CastType = null;
+		public UhtRigVMParameterFlags ParameterFlags = UhtRigVMParameterFlags.None;
+
+		[JsonIgnore]
+		public bool bConstant => this.ParameterFlags.HasAnyFlags(UhtRigVMParameterFlags.Constant);
+
+		[JsonIgnore]
+		public bool bInput => this.ParameterFlags.HasAnyFlags(UhtRigVMParameterFlags.Input);
+
+		[JsonIgnore]
+		public bool bOutput => this.ParameterFlags.HasAnyFlags(UhtRigVMParameterFlags.Output);
+
+		[JsonIgnore]
+		public bool bSingleton => this.ParameterFlags.HasAnyFlags(UhtRigVMParameterFlags.Singleton);
+
+		[JsonIgnore]
+		public bool bEditorOnly => this.ParameterFlags.HasAnyFlags(UhtRigVMParameterFlags.EditorOnly);
+
+		[JsonIgnore]
+		public bool bIsEnum => this.ParameterFlags.HasAnyFlags(UhtRigVMParameterFlags.IsEnum);
+
+		[JsonIgnore]
+		public bool bIsArray => this.ParameterFlags.HasAnyFlags(UhtRigVMParameterFlags.IsArray);
+
+		[JsonIgnore]
+		public bool bIsFixedArray => this.ParameterFlags.HasAnyFlags(UhtRigVMParameterFlags.IsFixedArray);
+
+		public UhtRigVMParameter(UhtProperty Property, int Index)
+		{
+			this.Property = Property;
+
+			this.Name = Property.EngineName;
+			this.ParameterFlags |= Property.MetaData.ContainsKey("Constant") ? UhtRigVMParameterFlags.Constant : UhtRigVMParameterFlags.None;
+			this.ParameterFlags |= Property.MetaData.ContainsKey("Input") ? UhtRigVMParameterFlags.Input : UhtRigVMParameterFlags.None;
+			this.ParameterFlags |= Property.MetaData.ContainsKey("Output") ? UhtRigVMParameterFlags.Output : UhtRigVMParameterFlags.None;
+			this.ParameterFlags |= Property.bIsEditorOnlyProperty ? UhtRigVMParameterFlags.EditorOnly : UhtRigVMParameterFlags.None;
+			this.ParameterFlags |= Property.MetaData.ContainsKey("Singleton") ? UhtRigVMParameterFlags.Singleton : UhtRigVMParameterFlags.None;
+			this.Getter = "GetRef";
+
+			if (Property.MetaData.ContainsKey("Visible"))
+			{
+				this.ParameterFlags |= UhtRigVMParameterFlags.Constant;
+				this.ParameterFlags &= (UhtRigVMParameterFlags.Input | UhtRigVMParameterFlags.Output);
+			}
+
+			if (this.bEditorOnly)
+			{
+				Property.LogError($"RigVM Struct '{this.Property.Outer?.SourceName}' - Member '{this.Property.SourceName}' is editor only - WITH_EDITORONLY_DATA not allowed on structs with RIGVM_METHOD.");
+			}
+
+			string? RigVMType = Property.GetRigVMType(ref this.ParameterFlags);
+			if (RigVMType == null)
+			{
+				Property.LogError($"RigVM Struct '{this.Property.Outer?.SourceName}' - Member '{this.Property.SourceName}' type '{this.Property.GetUserFacingDecl()}' not supported by RigVM.");
+			}
+			else
+			{
+				this.Type = RigVMType;
+			}
+
+			if (this.bIsArray)
+			{
+				if (this.IsConst())
+				{
+					string ExtendedType = this.ExtendedType(false);
+					this.CastName = $"{this.Name}_{Index}_Array";
+					this.CastType = $"TArrayView<const {ExtendedType.AsMemory(1, ExtendedType.Length - 2)}>";
+				}
+			}
+		}
+
+		public UhtRigVMParameter(string Name, string Type)
+		{
+			this.Name = Name;
+			this.Type = Type;
+		}
+
+		public string NameOriginal(bool bCastName = false)
+		{
+			return bCastName && this.CastName != null ? this.CastName : this.Name;
+		}
+
+		public string TypeOriginal(bool bCastType = false)
+		{
+			return bCastType && this.CastType != null ? this.CastType : this.Type;
+		}
+
+		public string Declaration(bool bCastType = false, bool bCastName = false)
+		{
+			return $"{TypeOriginal(bCastType)} {NameOriginal(bCastName)}";
+		}
+
+		public string BaseType(bool bCastType = false)
+		{
+			string String = TypeOriginal(bCastType);
+
+			int LesserPos = String.IndexOf('<');
+			if (LesserPos >= 0)
+			{
+				return String.Substring(0, LesserPos);
+			}
+			else
+			{
+				return String;
+			}
+		}
+
+		public string ExtendedType(bool bCastType = false)
+		{
+			string String = TypeOriginal(bCastType);
+
+			int LesserPos = String.IndexOf('<');
+			if (LesserPos >= 0)
+			{
+				return String.Substring(LesserPos);
+			}
+			else
+			{
+				return String;
+			}
+		}
+
+		public string TypeConstRef(bool bCastType = false)
+		{
+			string String = TypeNoRef(bCastType);
+			if (String.StartsWith("T") || String.StartsWith("F"))
+			{
+				return $"const {String}&";
+			}
+			else
+			{
+				return $"const {String}";
+			}
+		}
+
+		public string TypeRef(bool bCastType = false)
+		{
+			string String = TypeNoRef(bCastType);
+			return $"{String}&";
+		}
+
+		public string TypeNoRef(bool bCastType = false)
+		{
+			string String = TypeOriginal(bCastType);
+			if (String.EndsWith("&"))
+			{
+				return String.Substring(0, String.Length - 1);
+			}
+			else
+			{
+				return String;
+			}
+		}
+
+		public string TypeVariableRef(bool bCastType = false)
+		{
+			return IsConst() ? TypeConstRef(bCastType) : TypeRef(bCastType);
+		}
+
+		public string Variable(bool bCastType = false, bool bCastName = false)
+		{
+			return $"{TypeVariableRef(bCastType)} {NameOriginal(bCastName)}";
+		}
+
+		public bool IsConst()
+		{
+			return this.bConstant || (this.bInput && !this.bOutput);
+		}
+
+		public bool RequiresCast()
+		{
+			return this.CastType != null && this.CastName != null;
+		}
+	}
+
+	/**
+	 * A single info dataset for a function marked with RIGVM_METHOD.
+	 * This struct provides access to its name, the return type and all parameters.
+	 */
+	public class UhtRigVMMethodInfo
+	{
+		private static string NoPrefixInternal = string.Empty;
+		private static string ReturnPrefixInternal = "return ";
+
+		public string ReturnType { get; set; } = string.Empty;
+		public string Name { get; set; } = string.Empty;
+		public List<UhtRigVMParameter> Parameters { get; set; } = new List<UhtRigVMParameter>();
+
+		public string ReturnPrefix()
+		{
+			return (ReturnType.Length == 0 || ReturnType == "void") ? NoPrefixInternal : ReturnPrefixInternal;
+		}
+	}
+
+	/**
+	 * An info dataset providing access to all functions marked with RIGVM_METHOD
+	 * for each struct.
+	 */
+	public class UhtRigVMStructInfo
+	{
+		public string Name { get; set; } = string.Empty;
+		public List<UhtRigVMParameter> Members = new List<UhtRigVMParameter>();
+		public List<UhtRigVMMethodInfo> Methods { get; set; } = new List<UhtRigVMMethodInfo>();
+	};
+
+	/// <summary>
+	/// Series of flags not part of the engine's script struct flags that affect code generation or verification
+	/// </summary>
+	[Flags]
+	public enum UhtScriptStructExportFlags : UInt32
+	{
+		None = 0,
+		HasDefaults = 1 << 0,
+		HasNoOpConstructor = 1 << 1,
+		IsAlwaysAccessible = 1 << 2,
+		IsCoreType = 1 << 3,
+	}
+
+	/// <summary>
+	/// Helper methods for testing flags.  These methods perform better than the generic HasFlag which hits
+	/// the GC and stalls.
+	/// </summary>
+	public static class UhtScriptStructExportFlagsExtensions
+	{
+
+		/// <summary>
+		/// Test to see if any of the specified flags are set
+		/// </summary>
+		/// <param name="InFlags">Current flags</param>
+		/// <param name="TestFlags">Flags to test for</param>
+		/// <returns>True if any of the flags are set</returns>
+		public static bool HasAnyFlags(this UhtScriptStructExportFlags InFlags, UhtScriptStructExportFlags TestFlags)
+		{
+			return (InFlags & TestFlags) != 0;
+		}
+
+		/// <summary>
+		/// Test to see if all of the specified flags are set
+		/// </summary>
+		/// <param name="InFlags">Current flags</param>
+		/// <param name="TestFlags">Flags to test for</param>
+		/// <returns>True if all the flags are set</returns>
+		public static bool HasAllFlags(this UhtScriptStructExportFlags InFlags, UhtScriptStructExportFlags TestFlags)
+		{
+			return (InFlags & TestFlags) == TestFlags;
+		}
+
+		/// <summary>
+		/// Test to see if a specific set of flags have a specific value.
+		/// </summary>
+		/// <param name="InFlags">Current flags</param>
+		/// <param name="TestFlags">Flags to test for</param>
+		/// <param name="MatchFlags">Expected value of the tested flags</param>
+		/// <returns>True if the given flags have a specific value.</returns>
+		public static bool HasExactFlags(this UhtScriptStructExportFlags InFlags, UhtScriptStructExportFlags TestFlags, UhtScriptStructExportFlags MatchFlags)
+		{
+			return (InFlags & TestFlags) == MatchFlags;
+		}
+	}
+
+
+	[UhtEngineClass(Name = "ScriptStruct")]
+	public class UhtScriptStruct : UhtStruct
+	{
+		private static UhtSpecifierValidatorTable ScriptStructSpecifierValidatorTable = UhtSpecifierValidatorTables.Instance.Get(UhtTableNames.ScriptStruct);
+
+		[JsonConverter(typeof(JsonStringEnumConverter))]
+		public EStructFlags ScriptStructFlags { get; set; } = EStructFlags.NoFlags;
+
+		[JsonConverter(typeof(JsonStringEnumConverter))]
+		public UhtScriptStructExportFlags ScriptStructExportFlags { get; set; } = UhtScriptStructExportFlags.None;
+
+		public int MacroDeclaredLineNumber { get; set; } = -1;
+
+		public UhtRigVMStructInfo? RigVMStructInfo { get; set; } = null;
+
+		[JsonIgnore]
+		public override UhtEngineType EngineType => UhtEngineType.ScriptStruct;
+
+		[JsonIgnore]
+		public bool bHasDefaults => this.ScriptStructExportFlags.HasAnyFlags(UhtScriptStructExportFlags.HasDefaults);
+
+		[JsonIgnore]
+		public bool bIsAlwaysAccessible => this.ScriptStructExportFlags.HasAnyFlags(UhtScriptStructExportFlags.IsAlwaysAccessible);
+
+		[JsonIgnore]
+		public bool bIsCoreType => this.ScriptStructExportFlags.HasAnyFlags(UhtScriptStructExportFlags.IsCoreType);
+
+		[JsonIgnore]
+		public bool bHasNoOpConstructor => this.ScriptStructExportFlags.HasAnyFlags(UhtScriptStructExportFlags.HasNoOpConstructor);
+
+		/// <inheritdoc/>
+		public override string EngineClassName { get => "ScriptStruct"; }
+
+		[JsonIgnore]
+		public override string EngineNamePrefix
+		{
+			get
+			{
+				return UhtConfig.Instance.IsStructWithTPrefix(this.EngineName) ? "T" : "F";
+			}
+		}
+
+		///<inheritdoc/>
+		[JsonIgnore]
+		protected override UhtSpecifierValidatorTable? SpecifierValidatorTable { get => UhtScriptStruct.ScriptStructSpecifierValidatorTable; }
+
+		[JsonConverter(typeof(UhtNullableTypeSourceNameJsonConverter<UhtScriptStruct>))]
+		public UhtScriptStruct? SuperScriptStruct => (UhtScriptStruct?)this.Super;
+
+		public UhtScriptStruct(UhtType Outer, int LineNumber) : base(Outer, LineNumber)
+		{
+		}
+
+		/// <inheritdoc/>
+		protected override void ResolveChildren(UhtResolvePhase Phase)
+		{
+			base.ResolveChildren(Phase);
+
+			switch (Phase)
+			{
+				case UhtResolvePhase.Final:
+					if (ScanForInstancedReferenced(false))
+					{
+						this.ScriptStructFlags |= EStructFlags.HasInstancedReference;
+					}
+					CollectRigVMMembers();
+					break;
+			}
+		}
+
+		/// <inheritdoc/>
+		public override bool ScanForInstancedReferenced(bool bDeepScan)
+		{
+			if (this.ScriptStructFlags.HasAnyFlags(EStructFlags.HasInstancedReference))
+			{
+				return true;
+			}
+
+			if (this.SuperScriptStruct != null && this.SuperScriptStruct.ScanForInstancedReferenced(bDeepScan))
+			{
+				return true;
+			}
+
+			return base.ScanForInstancedReferenced(bDeepScan);
+		}
+
+		#region Validation support
+		protected override UhtValidationOptions Validate(UhtValidationOptions Options)
+		{
+			Options = base.Validate(Options);
+
+			if (this.ScriptStructFlags.HasAnyFlags(EStructFlags.Immutable))
+			{
+				if (this.HeaderFile != this.Session.UObject.HeaderFile)
+				{
+					this.LogError("Immutable is being phased out in favor of SerializeNative, and is only legal on the mirror structs declared in UObject");
+				}
+			}
+
+			// Validate the engine name
+			string ExpectedName = $"{this.EngineNamePrefix}{this.EngineName}";
+			if (this.SourceName != ExpectedName)
+			{
+				this.LogError($"Struct '{this.SourceName}' has an invalid Unreal prefix, expecting '{ExpectedName}");
+			}
+
+			return Options |= UhtValidationOptions.Shadowing;
+		}
+		#endregion
+
+		/// <inheritdoc/>
+		public override void CollectReferences(IUhtReferenceCollector Collector)
+		{
+			if (this.ScriptStructFlags.HasAnyFlags(EStructFlags.NoExport))
+			{
+				Collector.AddSingleton(this);
+			}
+			Collector.AddExportType(this);
+			Collector.AddDeclaration(this, true);
+			Collector.AddCrossModuleReference(this, true);
+			if (this.SuperScriptStruct != null)
+			{
+				Collector.AddCrossModuleReference(this.SuperScriptStruct, true);
+			}
+			Collector.AddCrossModuleReference(this.Package, true);
+			foreach (UhtType Child in this.Children)
+			{
+				Child.CollectReferences(Collector);
+			}
+		}
+
+		private void CollectRigVMMembers()
+		{
+			if (this.RigVMStructInfo != null)
+			{
+				for (UhtStruct? Current = this; Current != null; Current = Current.SuperStruct)
+				{
+					foreach (UhtProperty Property in Current.Properties)
+					{
+						this.RigVMStructInfo.Members.Add(new UhtRigVMParameter(Property, this.RigVMStructInfo.Members.Count));
+					}
+				}
+
+				if (this.RigVMStructInfo.Members.Count == 0)
+				{
+					this.LogError($"RigVM Struct '{this.SourceName}' - has zero members - invalid RIGVM_METHOD.");
+				}
+				else if (this.RigVMStructInfo.Members.Count > 64)
+				{
+					this.LogError($"RigVM Struct '{this.SourceName}' - has {this.RigVMStructInfo.Members.Count} members (64 is the limit).");
+				}
+			}
+		}
+	}
+}
\ No newline at end of file
diff --git a/Engine/Source/Programs/Shared/EpicGames.UHT/Types/UhtSourceFile.cs b/Engine/Source/Programs/Shared/EpicGames.UHT/Types/UhtSourceFile.cs
new file mode 100644
--- /dev/null
+++ b/Engine/Source/Programs/Shared/EpicGames.UHT/Types/UhtSourceFile.cs
@@ -0,0 +1,50 @@
+// Copyright Epic Games, Inc. All Rights Reserved.
+
+using EpicGames.Core;
+using EpicGames.UHT.Utils;
+using System.IO;
+
+namespace EpicGames.UHT.Types
+{
+	public class UhtSourceFile : IUhtMessageSite, IUhtMessageSource
+	{
+		public readonly UhtSession Session;
+		public readonly string FilePath;
+		public readonly string FileName;
+		public StringView Data { get => SourceFragment.Data; }
+		public UhtSourceFile? FragmentSourceFile { get => SourceFragment.SourceFile; }
+		public int FragmentLineNumber { get => SourceFragment.LineNumber; }
+		public string FullFilePath { get => SourceFragment.FilePath; }
+
+		private UhtSimpleMessageSite MessageSite;
+		private UhtSourceFragment SourceFragment;
+
+		#region IUHTMessageSite implementation
+		public IUhtMessageSession MessageSession => this.MessageSite.MessageSession;
+		public IUhtMessageSource? MessageSource => this.MessageSite.MessageSource;
+		public IUhtMessageLineNumber? MessageLineNumber => null;
+		#endregion
+
+		#region IUHTMessageSource implementation
+		public string MessageFilePath { get => this.FilePath; }
+		public string MessageFullFilePath { get => this.FullFilePath; }
+		public bool bMessageIsFragment { get => this.FragmentSourceFile != null; }
+		public string MessageFragmentFilePath { get => this.FragmentSourceFile != null ? this.FragmentSourceFile.MessageFilePath : ""; }
+		public string MessageFragmentFullFilePath { get => this.FragmentSourceFile != null ? this.FragmentSourceFile.MessageFullFilePath : ""; }
+		public int MessageFragmentLineNumber { get => this.FragmentLineNumber; }
+		#endregion
+
+		public UhtSourceFile(UhtSession Session, string FilePath)
+		{
+			this.Session = Session;
+			this.FilePath = FilePath;
+			this.FileName = Path.GetFileNameWithoutExtension(this.FilePath);
+			this.MessageSite = new UhtSimpleMessageSite(this.Session, this);
+		}
+
+		public virtual void Read()
+		{
+			this.SourceFragment = Session.ReadSource(this.FilePath);
+		}
+	}
+}
