UHT in UBT initial version.  Currently disabled.

diff --git a/Engine/Source/Programs/Shared/EpicGames.UHT/Types/UhtStruct.cs b/Engine/Source/Programs/Shared/EpicGames.UHT/Types/UhtStruct.cs
new file mode 100644
--- /dev/null
+++ b/Engine/Source/Programs/Shared/EpicGames.UHT/Types/UhtStruct.cs
@@ -0,0 +1,390 @@
+// Copyright Epic Games, Inc. All Rights Reserved.
+
+using EpicGames.Core;
+using EpicGames.UHT.Tables;
+using EpicGames.UHT.Tokenizer;
+using EpicGames.UHT.Utils;
+using System;
+using System.Collections.Generic;
+using System.Text.Json.Serialization;
+
+namespace EpicGames.UHT.Types
+{
+	[UhtEngineClass(Name = "Struct")]
+	public abstract class UhtStruct : UhtField
+	{
+		public EGeneratedCodeVersion GeneratedCodeVersion = EGeneratedCodeVersion.None;
+
+		/// <summary>
+		/// Return a collection of children that are properties
+		/// 
+		/// NOTE: This method allocates memory to construct the enumerator.  In code
+		/// invoked a large number of times, the loop should be written directly into
+		/// the code and not use this method.  Also, the Linq version performs even
+		/// worse.
+		/// </summary>
+		[JsonIgnore]
+		public IEnumerable<UhtProperty> Properties
+		{
+			get
+			{
+				foreach (UhtType Type in this.Children)
+				{
+					if (Type is UhtProperty Property)
+					{
+						yield return Property;
+					}
+				}
+			}
+		}
+
+		/// <summary>
+		/// Return a collection of children that are functions
+		/// 
+		/// NOTE: This method allocates memory to construct the enumerator.  In code
+		/// invoked a large number of times, the loop should be written directly into
+		/// the code and not use this method.  Also, the Linq version performs even
+		/// worse.
+		/// </summary>
+		[JsonIgnore]
+		public IEnumerable<UhtFunction> Functions
+		{
+			get
+			{
+				foreach (UhtType Type in this.Children)
+				{
+					if (Type is UhtFunction Function)
+					{
+						yield return Function;
+					}
+				}
+			}
+		}
+
+		[JsonIgnore]
+		public virtual string EngineNamePrefix { get => "F"; }
+
+		/// <inheritdoc/>
+		public override string EngineClassName { get => "Struct"; }
+
+		public UhtStruct? Super = null;
+
+		[JsonConverter(typeof(UhtNullableTypeListJsonConverter<UhtStruct>))]
+		public List<UhtStruct>? Bases { get; set; } = null;
+
+		[JsonIgnore]
+		public UhtStruct? SuperStruct => Super;
+
+		public UhtStruct(UhtType Outer, int LineNumber) : base(Outer, LineNumber)
+		{
+		}
+
+		/// <summary>
+		/// Test to see if the given struct is derived from the base structure
+		/// </summary>
+		/// <param name="SomeBase">Base structure.</param>
+		/// <returns>True if the given structure is the specified base or derives from the base.  If the base is null, the false is returned.</returns>
+		public bool IsChildOf(UhtStruct? Base)
+		{
+			if (Base == null)
+			{
+				return false;
+			}
+			for (UhtStruct? Current = this; Current != null; Current = Current.Super)
+			{
+				if (Current == Base)
+				{
+					return true;
+				}
+			}
+			return false;
+		}
+
+		#region Resolution support
+		protected override void ResolveSuper(UhtResolvePhase ResolvePhase)
+		{
+			if (this.Super != null)
+			{
+				this.Super.Resolve(ResolvePhase);
+			}
+
+			if (this.Bases != null)
+			{
+				foreach (UhtStruct Base in this.Bases)
+				{
+					Base.Resolve(ResolvePhase);
+				}
+			}
+
+			base.ResolveSuper(ResolvePhase);
+		}
+
+		/// <summary>
+		/// Check properties to see if any instances are referenced.
+		/// This method does NOT cache the result.
+		/// </summary>
+		/// <param name="bDeepScan">If true, the ScanForInstancedReferenced method on the properties will also be called.</param>
+		/// <returns></returns>
+		public virtual bool ScanForInstancedReferenced(bool bDeepScan)
+		{
+			foreach (UhtType Type in this.Children)
+			{
+				if (Type is UhtProperty Property)
+				{
+					if (Property.PropertyFlags.HasAnyFlags(EPropertyFlags.ContainsInstancedReference | EPropertyFlags.InstancedReference))
+					{
+						return true;
+					}
+					if (bDeepScan && Property.ScanForInstancedReferenced(bDeepScan))
+					{
+						return true;
+					}
+				}
+			}
+			return false;
+		}
+		#endregion
+
+		#region Super and base binding helper methods
+		public void BindAndResolveSuper(ref UhtToken SuperIdentifier, UhtFindOptions FindOptions)
+		{
+			if (SuperIdentifier)
+			{
+				this.Super = (UhtStruct?)this.FindType(FindOptions | UhtFindOptions.SourceName | UhtFindOptions.NoSelf, ref SuperIdentifier);
+				if (this.Super == null)
+				{
+					throw new UhtException(this, $"Unable to find parent {this.EngineType.ShortLowercaseText()} type for '{this.SourceName}' named '{SuperIdentifier.Value}'");
+				}
+				this.HeaderFile.AddReferencedHeader(this.Super);
+				this.MetaData.Parent = this.Super.MetaData;
+				this.Super.Resolve(UhtResolvePhase.Bases);
+			}
+		}
+
+		public void BindAndResolveBases(List<UhtToken[]>? BaseIdentifiers, UhtFindOptions FindOptions)
+		{
+			if (BaseIdentifiers != null)
+			{
+				foreach (UhtToken[] BaseIdentifier in BaseIdentifiers)
+				{
+					// We really only case about interfaces, but we can also handle structs
+					UhtStruct? Base = (UhtStruct?)this.FindType(FindOptions | UhtFindOptions.Class | UhtFindOptions.ScriptStruct | UhtFindOptions.SourceName | UhtFindOptions.NoSelf, BaseIdentifier);
+					if (Base != null)
+					{
+						if (this.Bases == null)
+						{
+							this.Bases = new List<UhtStruct>();
+						}
+						this.Bases.Add(Base);
+						this.HeaderFile.AddReferencedHeader(Base);
+						Base.Resolve(UhtResolvePhase.Bases);
+					}
+				}
+			}
+		}
+		#endregion
+
+		#region Validation support
+		protected override UhtValidationOptions Validate(UhtValidationOptions Options)
+		{
+			Options = base.Validate(Options);
+			ValidateSparseClassData();
+			return Options;
+		}
+
+		private static bool CheckUIMinMaxRangeFromMetaData(UhtType Child)
+		{
+			string UIMin = Child.MetaData.GetValueOrDefault(UhtNames.UIMin);
+			string UIMax = Child.MetaData.GetValueOrDefault(UhtNames.UIMax);
+			if (UIMin.Length == 0 || UIMax.Length == 0)
+			{
+				return false;
+			}
+
+			// NOTE: Old UHT didn't handle parse errors
+			double MinValue;
+			if (!double.TryParse(UIMin, out MinValue))
+			{
+				MinValue = 0;
+			}
+
+			double MaxValue;
+			if (!double.TryParse(UIMax, out MaxValue))
+			{
+				MaxValue = 0;
+			}
+
+			// NOTE: that we actually allow UIMin == UIMax to disable the range manually.
+			return MinValue <= MaxValue;
+		}
+
+		protected override void ValidateDocumentationPolicy(UhtDocumentationPolicy Policy)
+		{
+			if (Policy.bClassOrStructCommentRequired)
+			{
+				string ClassTooltip = this.MetaData.GetValueOrDefault(UhtNames.ToolTip);
+				if (ClassTooltip.Length == 0 || ClassTooltip.Equals(this.EngineName, StringComparison.OrdinalIgnoreCase))
+				{
+					this.LogError($"{this.EngineType.CapitalizedText()} '{this.SourceName}' does not provide a tooltip / comment (DocumentationPolicy).");
+				}
+			}
+
+			if (Policy.bMemberToolTipsRequired)
+			{
+				Dictionary<string, UhtProperty> ToolTipToType = new Dictionary<string, UhtProperty>();
+				foreach (UhtProperty Property in this.Properties)
+				{
+					string ToolTip = Property.GetToolTipText();
+					if (ToolTip.Length == 0 || ToolTip == Property.GetDisplayNameText())
+					{
+						Property.LogError($"Property '{this.SourceName}::{Property.SourceName}' does not provide a tooltip / comment (DocumentationPolicy).");
+						continue;
+					}
+
+					UhtProperty? Existing;
+					if (ToolTipToType.TryGetValue(ToolTip, out Existing))
+					{
+						Property.LogError($"Property '{this.SourceName}::{Existing.SourceName}' and '{this.SourceName}::{Property.SourceName}' are using identical tooltips (DocumentationPolicy).");
+					}
+					else
+					{
+						ToolTipToType.Add(ToolTip, Property);
+					}
+				}
+			}
+
+			if (Policy.bFloatRangesRequired)
+			{
+				foreach (UhtType Child in this.Children)
+				{
+					if (Child is UhtDoubleProperty || Child is UhtFloatProperty)
+					{
+						if (!CheckUIMinMaxRangeFromMetaData(Child))
+						{
+							Child.LogError($"Property '{this.SourceName}::{Child.SourceName}' does not provide a valid UIMin / UIMax (DocumentationPolicy).");
+						}
+					}
+				}
+			}
+
+			// also compare all tooltips to see if they are unique
+			if (Policy.bFunctionToolTipsRequired)
+			{
+				if (this is UhtClass)
+				{
+					Dictionary<string, UhtType> ToolTipToType = new Dictionary<string, UhtType>();
+					foreach (UhtType Child in this.Children)
+					{
+						if (Child is UhtFunction Function)
+						{
+							string ToolTip = Function.GetToolTipText();
+							if (ToolTip.Length == 0)
+							{
+								// NOTE: This does not fire because it doesn't check to see if it matches the display name as above.
+								Child.LogError($"Function '{this.SourceName}::{Function.SourceName}' does not provide a tooltip / comment (DocumentationPolicy).");
+								continue;
+							}
+
+							UhtType? Existing;
+							if (ToolTipToType.TryGetValue(ToolTip, out Existing))
+							{
+								Child.LogError($"Functions '{this.SourceName}::{Existing.SourceName}' and '{this.SourceName}::{Function.SourceName}' are using identical tooltips / comments (DocumentationPolicy).");
+							}
+							else
+							{
+								ToolTipToType.Add(ToolTip, Function);
+							}
+						}
+					}
+				}
+			}
+		}
+
+		protected void ValidateSparseClassData()
+		{
+			// Fetch the data types
+			string[]? SparseClassDataTypes = this.MetaData.GetStringArray(UhtNames.SparseClassDataTypes);
+			if (SparseClassDataTypes == null)
+			{
+				return;
+			}
+
+			// Make sure we don't try to have sparse class data inside of a struct instead of a class
+			UhtClass? Class = this as UhtClass;
+			if (Class == null)
+			{
+				this.LogError($"{this.EngineType.CapitalizedText()} '{this.SourceName}' contains sparse class data but is not a class.");
+				return;
+			}
+
+			// for now we only support one sparse class data structure per class
+			if (SparseClassDataTypes.Length > 1)
+			{
+				this.LogError($"Class '{this.SourceName}' contains multiple sparse class data types");
+				return;
+			}
+			if (SparseClassDataTypes.Length == 0)
+			{
+				this.LogError($"Class '{this.SourceName}' has sparse class metadata but does not specify a type");
+				return;
+			}
+
+			foreach (string SparseClassDataTypeName in SparseClassDataTypes)
+			{
+				UhtScriptStruct? SparseScriptStruct = this.Session.FindType(null, UhtFindOptions.EngineName | UhtFindOptions.ScriptStruct, SparseClassDataTypeName) as UhtScriptStruct;
+
+				// make sure the sparse class data struct actually exists
+				if (SparseScriptStruct == null)
+				{
+					this.LogError($"Unable to find sparse data type '{SparseClassDataTypeName}' for class '{this.SourceName}'");
+					continue;
+				}
+
+				// check the data struct for invalid properties
+				foreach (UhtProperty Property in SparseScriptStruct.Properties)
+				{
+					if (Property.PropertyFlags.HasAnyFlags(EPropertyFlags.BlueprintAssignable))
+					{
+						Property.LogError($"Sparse class data types can not contain blueprint assignable delegates. Type '{SparseScriptStruct.EngineName}' Delegate '{Property.SourceName}'");
+					}
+
+					// all sparse properties should have EditDefaultsOnly
+					if (!Property.PropertyFlags.HasAllFlags(EPropertyFlags.Edit | EPropertyFlags.DisableEditOnInstance))
+					{
+						Property.LogError($"Sparse class data types must be VisibleDefaultsOnly or EditDefaultsOnly. Type '{SparseScriptStruct.EngineName}' Property '{Property.SourceName}'");
+					}
+
+					// no sparse properties should have BlueprintReadWrite
+					if (Property.PropertyFlags.HasAnyFlags(EPropertyFlags.BlueprintVisible) && !Property.PropertyFlags.HasAnyFlags(EPropertyFlags.BlueprintReadOnly))
+					{
+						Property.LogError($"Sparse class data types must not be BlueprintReadWrite. Type '{SparseScriptStruct.EngineName}' Property '{Property.SourceName}'");
+					}
+				}
+
+				// if the class's parent has a sparse class data struct then the current class must also use the same struct or one that inherits from it
+				UhtClass? SuperClass = Class.SuperClass;
+				if (SuperClass != null)
+				{
+					string[]? SuperSparseClassDataTypes = SuperClass.MetaData.GetStringArray(UhtNames.SparseClassDataTypes);
+					if (SuperSparseClassDataTypes != null)
+					{
+						foreach (string SuperSparseClassDataTypeName in SuperSparseClassDataTypes)
+						{
+							UhtScriptStruct? SuperSparseScriptStruct = this.Session.FindType(null, UhtFindOptions.EngineName | UhtFindOptions.ScriptStruct, SuperSparseClassDataTypeName) as UhtScriptStruct;
+							if (SuperSparseScriptStruct != null)
+							{
+								if (!SparseScriptStruct.IsChildOf(SuperSparseScriptStruct))
+								{
+									this.LogError(
+										$"Class '{this.SourceName}' is a child of '{SuperClass.SourceName}' but its sparse class data struct " +
+										$"'{SparseScriptStruct.EngineName}', does not inherit from '{SuperSparseScriptStruct.EngineName}'.");
+								}
+							}
+						}
+					}
+				}
+			}
+		}
+		#endregion
+	}
+}
diff --git a/Engine/Source/Programs/Shared/EpicGames.UHT/Types/UhtType.cs b/Engine/Source/Programs/Shared/EpicGames.UHT/Types/UhtType.cs
new file mode 100644
--- /dev/null
+++ b/Engine/Source/Programs/Shared/EpicGames.UHT/Types/UhtType.cs
@@ -0,0 +1,1095 @@
+// Copyright Epic Games, Inc. All Rights Reserved.
+
+using EpicGames.UHT.Tables;
+using EpicGames.UHT.Tokenizer;
+using EpicGames.UHT.Utils;
+using System;
+using System.Collections.Generic;
+using System.Text;
+using System.Text.Json.Serialization;
+using System.Threading;
+
+namespace EpicGames.UHT.Types
+{
+
+	/// <summary>
+	/// Enumeration for all the different engine types
+	/// </summary>
+	public enum UhtEngineType
+	{
+
+		/// <summary>
+		/// Type is a header file (UhtHeader)
+		/// </summary>
+		Header,
+
+		/// <summary>
+		/// Type is a package (UhtPackage)
+		/// </summary>
+		Package,
+
+		/// <summary>
+		/// Type is a class (UhtClass)
+		/// </summary>
+		Class,
+
+		/// <summary>
+		/// Type is an interface (UhtClass)
+		/// </summary>
+		Interface,
+
+		/// <summary>
+		/// Type is a native interface (UhtClass)
+		/// </summary>
+		NativeInterface,
+
+		/// <summary>
+		/// Type is a script struct (UhtScriptStruct)
+		/// </summary>
+		ScriptStruct,
+
+		/// <summary>
+		/// Type is an enumeration (UhtEnum)
+		/// </summary>
+		Enum,
+
+		/// <summary>
+		/// Type is a function (UhtFunction)
+		/// </summary>
+		Function,
+
+		/// <summary>
+		/// Type is a delegate (UhtFunction)
+		/// </summary>
+		Delegate,
+
+		/// <summary>
+		/// Type is a property (UhtProperty)
+		/// </summary>
+		Property,
+	}
+
+	/// <summary>
+	/// Collection of extension methods used to query information about the type
+	/// </summary>
+	public static class UhtEngineTypeExtensions
+	{
+
+		/// <summary>
+		/// Get the name of the type where shared distinctions aren't made between different class and function types.
+		/// </summary>
+		/// <param name="EngineType">The engine type in question</param>
+		/// <returns>Lowercase name</returns>
+		/// <exception cref="UhtIceException">If the requested type isn't supported</exception>
+		public static string ShortLowercaseText(this UhtEngineType EngineType)
+		{
+			switch (EngineType)
+			{
+				case UhtEngineType.Header: return "header";
+				case UhtEngineType.Package: return "package";
+				case UhtEngineType.Class: 
+				case UhtEngineType.Interface:
+				case UhtEngineType.NativeInterface: return "class";
+				case UhtEngineType.ScriptStruct: return "struct";
+				case UhtEngineType.Enum: return "enum";
+				case UhtEngineType.Function:
+				case UhtEngineType.Delegate: return "function";
+				case UhtEngineType.Property: return "property";
+				default:
+					throw new UhtIceException("Invalid engine extended type");
+			}
+		}
+
+		/// <summary>
+		/// Get the name of the type where shared distinctions aren't made between different class and function types.
+		/// </summary>
+		/// <param name="EngineType">The engine type in question</param>
+		/// <returns>Capitalized name</returns>
+		/// <exception cref="UhtIceException">If the requested type isn't supported</exception>
+		public static string ShortCapitalizedText(this UhtEngineType EngineType)
+		{
+			switch (EngineType)
+			{
+				case UhtEngineType.Header: return "Header";
+				case UhtEngineType.Package: return "Package";
+				case UhtEngineType.Class:
+				case UhtEngineType.Interface:
+				case UhtEngineType.NativeInterface: return "Class";
+				case UhtEngineType.ScriptStruct: return "Struct";
+				case UhtEngineType.Enum: return "Enum";
+				case UhtEngineType.Function:
+				case UhtEngineType.Delegate: return "Function";
+				case UhtEngineType.Property: return "Property";
+				default:
+					throw new UhtIceException("Invalid engine extended type");
+			}
+		}
+
+		/// <summary>
+		/// Get the name of the type
+		/// </summary>
+		/// <param name="EngineType">The engine type in question</param>
+		/// <returns>Lowercase name</returns>
+		/// <exception cref="UhtIceException">If the requested type isn't supported</exception>
+		public static string LowercaseText(this UhtEngineType EngineType)
+		{
+			switch (EngineType)
+			{
+				case UhtEngineType.Header: return "header";
+				case UhtEngineType.Package: return "package";
+				case UhtEngineType.Class: return "class";
+				case UhtEngineType.Interface: return "interface";
+				case UhtEngineType.NativeInterface: return "native interface";
+				case UhtEngineType.ScriptStruct: return "struct";
+				case UhtEngineType.Enum: return "enum";
+				case UhtEngineType.Function: return "function";
+				case UhtEngineType.Delegate: return "delegate";
+				case UhtEngineType.Property: return "property";
+				default:
+					throw new UhtIceException("Invalid engine extended type");
+			}
+		}
+
+
+		/// <summary>
+		/// Get the name of the type
+		/// </summary>
+		/// <param name="EngineType">The engine type in question</param>
+		/// <returns>Capitalized name</returns>
+		/// <exception cref="UhtIceException">If the requested type isn't supported</exception>
+		public static string CapitalizedText(this UhtEngineType EngineType)
+		{
+			switch (EngineType)
+			{
+				case UhtEngineType.Header: return "Header";
+				case UhtEngineType.Package: return "Package";
+				case UhtEngineType.Class: return "Class";
+				case UhtEngineType.Interface: return "Interface";
+				case UhtEngineType.NativeInterface: return "Native interface";
+				case UhtEngineType.ScriptStruct: return "Struct";
+				case UhtEngineType.Enum: return "Enum";
+				case UhtEngineType.Function: return "Function";
+				case UhtEngineType.Delegate: return "Delegate";
+				case UhtEngineType.Property: return "Property";
+				default:
+					throw new UhtIceException("Invalid engine extended type");
+			}
+		}
+
+		/// <summary>
+		/// Return true if the type must be unique in the symbol table
+		/// </summary>
+		/// <param name="EngineType">Type in question</param>
+		/// <returns>True if it must be unique</returns>
+		public static bool MustBeUnique(this UhtEngineType EngineType)
+		{
+			switch (EngineType)
+			{
+				case UhtEngineType.Class:
+				case UhtEngineType.Interface:
+				case UhtEngineType.NativeInterface:
+				case UhtEngineType.ScriptStruct:
+				case UhtEngineType.Enum:
+					return true;
+
+				default:
+					return false;
+			}
+		}
+
+		/// <summary>
+		/// Find options to be used when attempting to find a duplicate name
+		/// </summary>
+		/// <param name="EngineType">Type in question</param>
+		/// <returns>Find options to be used for search</returns>
+		public static UhtFindOptions MustBeUniqueFindOptions(this UhtEngineType EngineType)
+		{
+			switch (EngineType)
+			{
+				case UhtEngineType.Class:
+				case UhtEngineType.Interface:
+				case UhtEngineType.ScriptStruct:
+				case UhtEngineType.Enum:
+					return UhtFindOptions.Class | UhtFindOptions.ScriptStruct | UhtFindOptions.Enum;
+
+				case UhtEngineType.NativeInterface:
+					return UhtFindOptions.Class;
+
+				default:
+					return (UhtFindOptions)0;
+			}
+		}
+
+		/// <summary>
+		/// Return true if the type must not be a reserved name.
+		/// </summary>
+		/// <param name="EngineType">Type in question</param>
+		/// <returns>True if the type must not have a reserved name.</returns>
+		public static bool MustNotBeReserved(this UhtEngineType EngineType)
+		{
+			switch (EngineType)
+			{
+				case UhtEngineType.Class:
+				case UhtEngineType.Interface:
+				case UhtEngineType.NativeInterface:
+				case UhtEngineType.ScriptStruct:
+				case UhtEngineType.Enum:
+					return true;
+
+				default:
+					return false;
+			}
+		}
+
+		/// <summary>
+		/// Return true if the type should be added to the engine symbol table
+		/// </summary>
+		/// <param name="EngineType">Type in question</param>
+		/// <returns>True if the type should be added to the symbol table</returns>
+		public static bool HasEngineName(this UhtEngineType EngineType)
+		{
+			return EngineType != UhtEngineType.NativeInterface;
+		}
+
+		/// <summary>
+		/// Return true if the children of the type should be added to the symbol table
+		/// </summary>
+		/// <param name="EngineType">Type in question</param>
+		/// <returns>True if children should be added</returns>
+		public static bool AddChildrenToSymbolTable(this UhtEngineType EngineType)
+		{
+			switch (EngineType)
+			{
+				case UhtEngineType.Property:
+					return false;
+
+				default:
+					return true;
+			}
+		}
+
+		/// <summary>
+		/// Convert the engine type to find options
+		/// </summary>
+		/// <param name="EngineType">Engine type</param>
+		/// <returns>Find options</returns>
+		/// <exception cref="UhtIceException">Engine type is invalid</exception>
+		public static UhtFindOptions FindOptions(this UhtEngineType EngineType)
+		{
+			switch (EngineType)
+			{
+				case UhtEngineType.Header: return 0;
+				case UhtEngineType.Package: return 0;
+				case UhtEngineType.Class: return UhtFindOptions.Class;
+				case UhtEngineType.Interface: return UhtFindOptions.Class;
+				case UhtEngineType.NativeInterface: return UhtFindOptions.Class;
+				case UhtEngineType.ScriptStruct: return UhtFindOptions.ScriptStruct;
+				case UhtEngineType.Enum: return UhtFindOptions.Enum;
+				case UhtEngineType.Function: return UhtFindOptions.Function;
+				case UhtEngineType.Delegate: return UhtFindOptions.DelegateFunction;
+				case UhtEngineType.Property: return UhtFindOptions.Property;
+				default:
+					throw new UhtIceException("Invalid engine extended type");
+			}
+		}
+	}
+
+	[Flags]
+	public enum UhtFindOptions
+	{
+		EngineName = 1 << 0,        // This is the default
+		SourceName = 1 << 1,
+
+		Enum = 1 << 8,
+		ScriptStruct = 1 << 9,
+		Class = 1 << 10,
+		DelegateFunction = 1 << 11,
+		Function = 1 << 12,
+		Property = 1 << 13,
+
+		NoParents = 1 << 16,
+		NoOuter = 1 << 17,
+		NoIncludes = 1 << 18,
+		NoGlobal = 1 << 19,
+		NoSelf = 1 << 20,
+		CaseCompare = 1 << 21,
+		CaselessCompare = 1 << 22,
+		ParentsOnly = NoOuter | NoIncludes | NoGlobal | NoSelf,
+
+		NamesMask = EngineName | SourceName,
+		TypesMask = Enum | ScriptStruct | Class | DelegateFunction | Function | Property,
+		ObjectTypesMask = Enum | ScriptStruct | Class,
+		DelgeateTypesMask = DelegateFunction,
+		PropertyOrFunction = Property | Function | DelegateFunction,
+		SelfOnly = NoParents | NoOuter | NoIncludes | NoGlobal,
+	}
+
+	/// <summary>
+	/// Helper methods for testing flags.  These methods perform better than the generic HasFlag which hits
+	/// the GC and stalls.
+	/// </summary>
+	public static class UhtFindOptionsExtensions
+	{
+
+		/// <summary>
+		/// Test to see if any of the specified flags are set
+		/// </summary>
+		/// <param name="InFlags">Current flags</param>
+		/// <param name="TestFlags">Flags to test for</param>
+		/// <returns>True if any of the flags are set</returns>
+		public static bool HasAnyFlags(this UhtFindOptions InFlags, UhtFindOptions TestFlags)
+		{
+			return (InFlags & TestFlags) != 0;
+		}
+
+		/// <summary>
+		/// Test to see if all of the specified flags are set
+		/// </summary>
+		/// <param name="InFlags">Current flags</param>
+		/// <param name="TestFlags">Flags to test for</param>
+		/// <returns>True if all the flags are set</returns>
+		public static bool HasAllFlags(this UhtFindOptions InFlags, UhtFindOptions TestFlags)
+		{
+			return (InFlags & TestFlags) == TestFlags;
+		}
+
+		/// <summary>
+		/// Test to see if a specific set of flags have a specific value.
+		/// </summary>
+		/// <param name="InFlags">Current flags</param>
+		/// <param name="TestFlags">Flags to test for</param>
+		/// <param name="MatchFlags">Expected value of the tested flags</param>
+		/// <returns>True if the given flags have a specific value.</returns>
+		public static bool HasExactFlags(this UhtFindOptions InFlags, UhtFindOptions TestFlags, UhtFindOptions MatchFlags)
+		{
+			return (InFlags & TestFlags) == MatchFlags;
+		}
+	}
+
+	[Flags]
+	public enum UhtValidationOptions
+	{
+		None,
+
+		/// <summary>
+		/// Test to see if any referenced types have been deprecated.  This should not
+		/// be set when the parent object or property has already been marked as deprecated.
+		/// </summary>
+		Deprecated = 1 << 0,
+
+		/// <summary>
+		/// Test to see if the name of the property conflicts with another property in the super chain
+		/// </summary>
+		Shadowing = 1 << 1,
+
+		/// <summary>
+		/// The property is part of a key in a TMap
+		/// </summary>
+		IsKey = 1 << 2,
+
+		/// <summary>
+		/// The property is part of a value in a TArray, TSet, or TMap
+		/// </summary>
+		IsValue = 1 << 3,
+
+		/// <summary>
+		/// The property is either a key or a value part of a TArray, TSet, or TMap
+		/// </summary>
+		IsKeyOrValue = IsKey | IsValue,
+	}
+
+	/// <summary>
+	/// Helper methods for testing flags.  These methods perform better than the generic HasFlag which hits
+	/// the GC and stalls.
+	/// </summary>
+	public static class UhtValidationOptionsExtensions
+	{
+
+		/// <summary>
+		/// Test to see if any of the specified flags are set
+		/// </summary>
+		/// <param name="InFlags">Current flags</param>
+		/// <param name="TestFlags">Flags to test for</param>
+		/// <returns>True if any of the flags are set</returns>
+		public static bool HasAnyFlags(this UhtValidationOptions InFlags, UhtValidationOptions TestFlags)
+		{
+			return (InFlags & TestFlags) != 0;
+		}
+
+		/// <summary>
+		/// Test to see if all of the specified flags are set
+		/// </summary>
+		/// <param name="InFlags">Current flags</param>
+		/// <param name="TestFlags">Flags to test for</param>
+		/// <returns>True if all the flags are set</returns>
+		public static bool HasAllFlags(this UhtValidationOptions InFlags, UhtValidationOptions TestFlags)
+		{
+			return (InFlags & TestFlags) == TestFlags;
+		}
+
+		/// <summary>
+		/// Test to see if a specific set of flags have a specific value.
+		/// </summary>
+		/// <param name="InFlags">Current flags</param>
+		/// <param name="TestFlags">Flags to test for</param>
+		/// <param name="MatchFlags">Expected value of the tested flags</param>
+		/// <returns>True if the given flags have a specific value.</returns>
+		public static bool HasExactFlags(this UhtValidationOptions InFlags, UhtValidationOptions TestFlags, UhtValidationOptions MatchFlags)
+		{
+			return (InFlags & TestFlags) == MatchFlags;
+		}
+	}
+
+	public enum UhtAccessSpecifier
+	{
+		None,
+		Public,
+		Private,
+		Protected,
+	}
+
+	public enum UhtResolvePhase : Int32
+	{
+		InvalidCheck = 1,
+		Bases,
+		Properties,
+		Final,
+		Count,
+	}
+
+	public class UhtDocumentationPolicy
+	{
+		private static UhtDocumentationPolicy StrictDocumentationPolicy = new UhtDocumentationPolicy
+		{
+			bPolicySet = true,
+			bClassOrStructCommentRequired = true,
+			bFunctionToolTipsRequired = true,
+			bMemberToolTipsRequired = true,
+			bParameterToolTipsRequired = true,
+			bFloatRangesRequired = true,
+		};
+
+		private static UhtDocumentationPolicy NoDocumentationPolicy = new UhtDocumentationPolicy
+		{
+			bPolicySet = false,
+			bClassOrStructCommentRequired = false,
+			bFunctionToolTipsRequired = false,
+			bMemberToolTipsRequired = false,
+			bParameterToolTipsRequired = false,
+			bFloatRangesRequired = false,
+		};
+
+		public bool bPolicySet;
+		public bool bClassOrStructCommentRequired;
+		public bool bFunctionToolTipsRequired;
+		public bool bMemberToolTipsRequired;
+		public bool bParameterToolTipsRequired;
+		public bool bFloatRangesRequired;
+
+		public static bool TryGet(UhtType Type, string DocumentationPolicyName, out UhtDocumentationPolicy Policy)
+		{
+			if (DocumentationPolicyName == "Strict")
+			{
+				Policy = StrictDocumentationPolicy;
+				return true;
+			}
+
+			Policy = NoDocumentationPolicy;
+			return DocumentationPolicyName.Length == 0;
+		}
+	};
+
+	public abstract class UhtType : IUhtMessageSite, IUhtMessageLineNumber
+	{
+		/// <summary>
+		/// Empty list of type used when children are requested but no children have been added.
+		/// </summary>
+		private static List<UhtType> EmptyTypeList = new List<UhtType>();
+
+		/// <summary>
+		/// All UHT runs are associated with a given session.  The session holds all the global information for a run.
+		/// </summary>
+		public UhtSession Session;
+
+		/// <summary>
+		/// Every type in a session is assigned a unique, non-zero id.
+		/// </summary>
+		public readonly int TypeIndex;
+
+		/// <summary>
+		/// The outer object containing this type.  For example, the header file would be the outer for a class defined in the global scope.
+		/// </summary>
+		public UhtType? Outer = null;
+
+		/// <summary>
+		/// The owning package of the type
+		/// </summary>
+		[JsonIgnore]
+		public virtual UhtPackage Package
+		{
+			get
+			{
+				UhtType Type = this;
+				while (Type.Outer != null)
+				{
+					Type = Type.Outer;
+				}
+				return (UhtPackage)Type;
+			}
+		}
+
+		/// <summary>
+		/// The header file containing the type
+		/// </summary>
+		[JsonIgnore]
+		public virtual UhtHeaderFile HeaderFile
+		{
+			get
+			{
+				UhtType Type = this;
+				while (Type.Outer != null && Type.Outer.Outer != null)
+				{
+					Type = Type.Outer;
+				}
+				return (UhtHeaderFile)Type;
+			}
+		}
+
+		/// <summary>
+		/// The name of the type as it appears in the source file
+		/// </summary>
+		public string SourceName { get; set; }
+
+		/// <summary>
+		/// The name of the type as used by the engines type system.  If not set, will default to the source name
+		/// </summary>
+		public string EngineName { get => this.EngineNameInternal ?? this.SourceName; set => this.EngineNameInternal = value; }
+
+		/// <summary>
+		/// Simple enumeration that can be used to detect the different types of objects
+		/// </summary>
+		[JsonIgnore]
+		public abstract UhtEngineType EngineType { get; }
+
+		/// <summary>
+		/// The name of the engine's class for this type
+		/// </summary>
+		public abstract string EngineClassName { get; }
+
+		/// <summary>
+		/// If true, then this type is visible to the symbol table.
+		/// </summary>
+		public bool bVisibleType = true;
+
+		/// <summary>
+		/// Return true if the type has been deprecated
+		/// </summary>
+		[JsonIgnore]
+		public virtual bool bDeprecated => false;
+
+		/// <summary>
+		/// The line number of where the definition began
+		/// </summary>
+		public int LineNumber { get; set; }
+
+		/// <summary>
+		/// Return a combination of the engine type name followed by the path name of the type
+		/// </summary>
+		[JsonIgnore]
+		public virtual string FullName { get => GetFullName(); }
+
+		/// <summary>
+		/// Return the path name of the type which includes all parent outer objects excluding the header file
+		/// </summary>
+		[JsonIgnore]
+		public string PathName { get => GetPathName(); }
+
+		/// <summary>
+		/// Return the name of the documentation policy to be used
+		/// </summary>
+		[JsonIgnore]
+		public string DocumentationPolicyName { get => this.MetaData.GetValueOrDefaultHierarchical(UhtNames.DocumentationPolicy); }
+
+		/// <summary>
+		/// Return the specifier validation table to be used to validate the meta data on this type
+		/// </summary>
+		protected virtual UhtSpecifierValidatorTable? SpecifierValidatorTable { get => null; }
+
+		/// <summary>
+		/// Meta data associated with the type
+		/// </summary>
+		public UhtMetaData MetaData;
+
+		/// <summary>
+		/// Children types of this type
+		/// </summary>
+		[JsonConverter(typeof(UhtNullableTypeListJsonConverter<UhtType>))]
+		public List<UhtType> Children { get => ChildrenInternal == null ? UhtType.EmptyTypeList : ChildrenInternal; }
+
+		/// <summary>
+		/// Helper to string method
+		/// </summary>
+		/// <returns></returns>
+		public override string ToString() { return this.SourceName; }
+
+		private string? EngineNameInternal = null;
+		private List<UhtType>? ChildrenInternal = null;
+		private int ResolveState = 0;
+
+		#region IUHTMessageSite implementation
+		[JsonIgnore]
+		public virtual IUhtMessageSession MessageSession => this.HeaderFile.MessageSession;
+
+		[JsonIgnore]
+		public virtual IUhtMessageSource? MessageSource => this.HeaderFile.MessageSource;
+
+		[JsonIgnore]
+		public IUhtMessageLineNumber? MessageLineNumber => this;
+		#endregion
+
+		#region IUHTMessageListNumber implementation
+		int IUhtMessageLineNumber.MessageLineNumber => LineNumber;
+		#endregion
+
+		/// <summary>
+		/// Construct the base type information.  This constructor is used exclusively by UHTPackage which is at 
+		/// the root of all type hierarchies.
+		/// </summary>
+		/// <param name="Session"></param>
+		protected UhtType(UhtSession Session)
+		{
+			this.Session = Session;
+			this.Outer = null;
+			this.SourceName = string.Empty;
+			this.LineNumber = 1;
+			this.TypeIndex = Session.GetNextTypeIndex();
+			this.MetaData = new UhtMetaData(this);
+		}
+
+		/// <summary>
+		/// Construct instance of the type given a parent type.
+		/// </summary>
+		/// <param name="Outer">The outer type of the type being constructed.  For example, a class defined in a given header will have the header as the outer.</param>
+		/// <param name="LineNumber">The line number in the source file where the type was defined.</param>
+		/// <param name="MetaData">Optional meta data to be associated with the type.</param>
+		protected UhtType(UhtType Outer, int LineNumber, UhtMetaData? MetaData = null)
+		{
+			this.Session = Outer.Session;
+			this.Outer = Outer;
+			this.SourceName = string.Empty;
+			this.LineNumber = LineNumber;
+			this.TypeIndex = this.Session.GetNextTypeIndex();
+			this.MetaData = MetaData ?? new UhtMetaData(this);
+			this.MetaData.MessageSite = this; // Make sure the site is correct
+		}
+
+		/// <summary>
+		/// Add a type as a child
+		/// </summary>
+		/// <param name="Child">The child to be added.</param>
+		public virtual void AddChild(UhtType Child)
+		{
+			if (this.ChildrenInternal == null)
+			{
+				this.ChildrenInternal = new List<UhtType>();
+			}
+			this.ChildrenInternal.Add(Child);
+		}
+
+		/// <summary>
+		/// Return the list of children attached to the type.  The 
+		/// list associated with the type is cleared.
+		/// </summary>
+		/// <returns>The current list of children</returns>
+		public List<UhtType>? DetchChildren()
+		{
+			List<UhtType>? Out = this.ChildrenInternal;
+			this.ChildrenInternal = null;
+			return Out;
+		}
+
+		/// <summary>
+		/// Merge the list children into the type.  The child's outer is set.
+		/// </summary>
+		/// <param name="Children">List of children to add.</param>
+		public void AddChildren(List<UhtType>? Children)
+		{
+			if (Children != null)
+			{
+				foreach (UhtType Child in Children)
+				{
+					Child.Outer = this;
+					AddChild(Child);
+				}
+			}
+		}
+
+		#region Type resolution
+		public bool Resolve(UhtResolvePhase Phase)
+		{
+			int InitState = (int)Phase << 1;
+
+			int CurrentState = Interlocked.Add(ref this.ResolveState, 0);
+			
+			// If we are already passed the init phase for this state, then we are good
+			if (CurrentState > InitState)
+			{
+				return true;
+			}
+
+			// Attempt to make myself the resolve thread for this type.  If we aren't the ones
+			// to get the init state, then we have to wait until the thread that did completes
+			CurrentState = Interlocked.CompareExchange(ref this.ResolveState, InitState, CurrentState);
+			if (CurrentState > InitState)
+			{
+				return true;
+			}
+			else if (CurrentState == InitState)
+			{
+				do
+				{
+					Thread.Yield();
+					CurrentState = Interlocked.Add(ref this.ResolveState, 0);
+				} while (CurrentState == InitState);
+				return true;
+			}
+			else
+			{
+				bool bResult = true;
+				try
+				{
+					ResolveSuper(Phase);
+					bResult = ResolveSelf(Phase);
+					ResolveChildren(Phase);
+				}
+				finally
+				{
+					Interlocked.Increment(ref this.ResolveState);
+				}
+				return bResult;
+			}
+		}
+
+		protected virtual void ResolveSuper(UhtResolvePhase Phase)
+		{
+		}
+
+		protected virtual bool ResolveSelf(UhtResolvePhase Phase)
+		{
+			return true;
+		}
+
+		protected virtual void ResolveChildren(UhtResolvePhase Phase)
+		{
+			if (this.ChildrenInternal != null)
+			{
+				if (Phase == UhtResolvePhase.InvalidCheck)
+				{
+					int OutIndex = 0;
+					for (int Index = 0; Index < this.ChildrenInternal.Count; ++Index)
+					{
+						UhtType Child = this.ChildrenInternal[Index];
+						if (Child.Resolve(Phase))
+						{
+							this.ChildrenInternal[OutIndex++] = Child;
+						}
+					}
+					if (OutIndex < this.ChildrenInternal.Count)
+					{
+						this.ChildrenInternal.RemoveRange(OutIndex, this.ChildrenInternal.Count - OutIndex);
+					}
+				}
+				else
+				{
+					foreach (UhtType Child in this.ChildrenInternal)
+					{
+						Child.Resolve(Phase);
+					}
+				}
+			}
+		}
+		#endregion
+
+		#region FindType support
+		/// <summary>
+		/// Find the given type in the type hierarchy
+		/// </summary>
+		/// <param name="Options">Options controlling what is searched</param>
+		/// <param name="Name">Name of the type.</param>
+		/// <param name="MessageSite">If supplied, then a error message will be generated if the type can not be found</param>
+		/// <param name="LineNumber">Source code line number requesting the lookup.</param>
+		/// <returns>The located type of null if not found</returns>
+		public UhtType? FindType(UhtFindOptions Options, string Name, IUhtMessageSite? MessageSite = null, int LineNumber = -1)
+		{
+			return this.Session.FindType(this, Options, Name, MessageSite, LineNumber);
+		}
+
+		/// <summary>
+		/// Find the given type in the type hierarchy
+		/// </summary>
+		/// <param name="Options">Options controlling what is searched</param>
+		/// <param name="Name">Name of the type.</param>
+		/// <param name="MessageSite">If supplied, then a error message will be generated if the type can not be found</param>
+		/// <returns>The located type of null if not found</returns>
+		public UhtType? FindType(UhtFindOptions Options, ref UhtToken Name, IUhtMessageSite? MessageSite = null)
+		{
+			return this.Session.FindType(this, Options, ref Name, MessageSite);
+		}
+
+		/// <summary>
+		/// Find the given type in the type hierarchy
+		/// </summary>
+		/// <param name="Options">Options controlling what is searched</param>
+		/// <param name="Identifiers">Enumeration of identifiers.</param>
+		/// <param name="MessageSite">If supplied, then a error message will be generated if the type can not be found</param>
+		/// <param name="LineNumber">Source code line number requesting the lookup.</param>
+		/// <returns>The located type of null if not found</returns>
+		public UhtType? FindType(UhtFindOptions Options, UhtTokenList Identifiers, IUhtMessageSite? MessageSite = null, int LineNumber = -1)
+		{
+			return this.Session.FindType(this, Options, Identifiers, MessageSite, LineNumber);
+		}
+
+		/// <summary>
+		/// Find the given type in the type hierarchy
+		/// </summary>
+		/// <param name="Options">Options controlling what is searched</param>
+		/// <param name="Identifiers">Enumeration of identifiers.</param>
+		/// <param name="MessageSite">If supplied, then a error message will be generated if the type can not be found</param>
+		/// <param name="LineNumber">Source code line number requesting the lookup.</param>
+		/// <returns>The located type of null if not found</returns>
+		public UhtType? FindType(UhtFindOptions Options, UhtToken[] Identifiers, IUhtMessageSite? MessageSite = null, int LineNumber = -1)
+		{
+			return this.Session.FindType(this, Options, Identifiers, MessageSite, LineNumber);
+		}
+		#endregion
+
+		#region Validation support
+		// Validate the property settings
+		protected virtual UhtValidationOptions Validate(UhtValidationOptions Options)
+		{
+			ValidateMetaData();
+			ValidateDocumentationPolicy();
+			return Options;
+		}
+
+		private void ValidateMetaData()
+		{
+			UhtSpecifierValidatorTable? Table = this.SpecifierValidatorTable;
+			if (Table != null)
+			{
+				if (this.MetaData.Dictionary != null && this.MetaData.Dictionary.Count > 0)
+				{
+					foreach (var KVP in this.MetaData.Dictionary)
+					{
+						UhtSpecifierValidator? Specifier;
+						if (Table.TryGetValue(KVP.Key.Name, out Specifier))
+						{
+							Specifier.Delegate(this, MetaData, KVP.Key, KVP.Value);
+						}
+					}
+				}
+			}
+		}
+
+		private void ValidateDocumentationPolicy()
+		{
+			if (UhtDocumentationPolicy.TryGet(this, this.DocumentationPolicyName, out UhtDocumentationPolicy Policy))
+			{
+				if (Policy.bPolicySet)
+				{
+					this.ValidateDocumentationPolicy(Policy);
+				}
+			}
+			else
+			{
+				this.LogError(this.MetaData.LineNumber, $"Documentation policy '{this.DocumentationPolicyName}' is not known");
+			}
+		}
+
+		protected virtual void ValidateDocumentationPolicy(UhtDocumentationPolicy Policy)
+		{
+		}
+
+		public static void ValidateType(UhtType Type, UhtValidationOptions Options)
+		{
+			// Invoke the new method
+			Options = Type.Validate(Options);
+
+			// Invoke validate on the children
+			foreach (UhtType Child in Type.Children)
+			{
+				ValidateType(Child, Options);
+			}
+		}
+		#endregion
+
+		#region Name and user facing text support
+		public string GetPathName(UhtType? StopOuter = null)
+		{
+			StringBuilder Builder = new StringBuilder();
+			GetPathName(Builder, StopOuter);
+			return Builder.ToString();
+		}
+
+		public virtual void GetPathName(StringBuilder Builder, UhtType? StopOuter = null)
+		{
+			GetPathNameInternal(Builder, StopOuter);
+			if (Builder.Length == 0)
+			{
+				Builder.Append("None");
+			}
+		}
+
+		public virtual void GetPathNameInternal(StringBuilder Builder, UhtType? StopOuter = null)
+		{
+			if (this != StopOuter)
+			{
+				UhtType? Outer = this.Outer;
+				if (Outer != null && Outer != StopOuter)
+				{
+					Outer.GetPathName(Builder, StopOuter);
+
+					// SubObjectDelimiter is used to indicate that this object's outer is not a UPackage
+					// In the C# version of UHT, we key off the header file
+					if (!(Outer is UhtHeaderFile) && Outer.Outer is UhtHeaderFile)
+					{
+						Builder.Append(UhtFCString.SubObjectDelimiter);
+					}
+
+					// SubObjectDelimiter is also used to bridge between the object:properties in the path
+					else if (!(Outer is UhtProperty) && this is UhtProperty)
+					{
+						Builder.Append(UhtFCString.SubObjectDelimiter);
+					}
+					else
+					{
+						Builder.Append('.');
+					}
+				}
+				Builder.Append(this.EngineName);
+			}
+		}
+
+		public string GetFullName()
+		{
+			StringBuilder Builder = new StringBuilder();
+			GetFullName(Builder);
+			return Builder.ToString();
+		}
+
+		public virtual void GetFullName(StringBuilder Builder)
+		{
+			Builder.Append(this.EngineClassName);
+			Builder.Append(" ");
+			GetPathName(Builder);
+		}
+
+		/// <summary>
+		/// Finds the localized display name or native display name as a fallback.
+		/// </summary>
+		/// <returns>The display name for this object.</returns>
+		public string GetDisplayNameText()
+		{
+			string? NativeDisplayName;
+			if (this.MetaData.TryGetValue(UhtNames.DisplayName, out NativeDisplayName))
+			{
+				return NativeDisplayName;
+			}
+
+			return GetDisplayStringFromEngineName();
+		}
+
+		/// <summary>
+		/// Finds the localized tooltip or native tooltip as a fallback.
+		/// </summary>
+		/// <param name="bShortToolTip">Look for a shorter version of the tooltip (falls back to the long tooltip if none was specified)</param>
+		/// <returns>The tooltip for this object.</returns>
+		public virtual string GetToolTipText(bool bShortToolTip = false)
+		{
+			string NativeToolTip = string.Empty;
+			if (bShortToolTip)
+			{
+				NativeToolTip = this.MetaData.GetValueOrDefault(UhtNames.ShortToolTip);
+			}
+
+			if (NativeToolTip.Length == 0)
+			{
+				NativeToolTip = this.MetaData.GetValueOrDefault(UhtNames.ToolTip);
+			}
+
+			if (NativeToolTip.Length == 0)
+			{
+				return GetDisplayStringFromEngineName();
+			}
+			else
+			{
+				return NativeToolTip.TrimEnd().Replace("@see", "See:");
+			}
+		}
+
+		/// <summary>
+		/// Return a generated display string from the engine name
+		/// </summary>
+		/// <returns>The generated string</returns>
+		protected virtual string GetDisplayStringFromEngineName()
+		{
+			return UhtFCString.NameToDisplayString(this.EngineName, false);
+		}
+		#endregion
+
+		public virtual void CollectReferences(IUhtReferenceCollector Collector)
+		{
+		}
+	}
+
+	public static class UhtTypeStringBuilderExtensions
+	{
+		/// <summary>
+		/// Append all of the outer names associated with the type
+		/// </summary>
+		/// <param name="Builder">Output string builder</param>
+		/// <param name="Outer">Outer to be appended</param>
+		/// <returns>Output string builder</returns>
+		public static StringBuilder AppendOuterNames(this StringBuilder Builder, UhtType? Outer)
+		{
+			if (Outer == null)
+			{
+				return Builder;
+			}
+
+			if (Outer is UhtClass Class)
+			{
+				Builder.Append('_');
+				Builder.Append(Outer.SourceName);
+			}
+			else if (Outer is UhtScriptStruct)
+			{
+				// Structs can also have UPackage outer.
+				if (!(Outer.Outer is UhtHeaderFile))
+				{
+					AppendOuterNames(Builder, Outer.Outer);
+				}
+				Builder.Append('_');
+				Builder.Append(Outer.SourceName);
+			}
+			else if (Outer is UhtPackage Package)
+			{
+				Builder.Append('_');
+				Builder.Append(Package.ShortName);
+			}
+			else if (Outer is UhtHeaderFile)
+			{
+				// Pickup the package
+				AppendOuterNames(Builder, Outer.Outer);
+			}
+			else
+			{
+				AppendOuterNames(Builder, Outer.Outer);
+				Builder.Append('_');
+				Builder.Append(Outer.EngineName);
+			}
+
+			return Builder;
+		}
+	}
+}
diff --git a/Engine/Source/Programs/Shared/EpicGames.UHT/Utils/UhtAttributeScanner.cs b/Engine/Source/Programs/Shared/EpicGames.UHT/Utils/UhtAttributeScanner.cs
new file mode 100644
--- /dev/null
+++ b/Engine/Source/Programs/Shared/EpicGames.UHT/Utils/UhtAttributeScanner.cs
@@ -0,0 +1,202 @@
+// Copyright Epic Games, Inc. All Rights Reserved.
+
+using EpicGames.Core;
+using EpicGames.UHT.Tables;
+using System;
+using System.Collections.Generic;
+using System.Linq;
+using System.Reflection;
+
+namespace EpicGames.UHT.Utils
+{
+	/// <summary>
+	/// UnrealHeaderTool avoids hard coding any table contents by using attributes to add entries to the tables.
+	/// 
+	/// There are two basic styles of tables in UHT.  
+	/// 
+	/// The first style is just a simple association of a string and an attribute.  For example, 
+	/// the engine class table is just a collection of all the engine class names supported by the engine.
+	/// 
+	/// The second style is a table of tables.  Depending on the context (i.e. is a "class" or "function" 
+	/// being parsed), attributes will "extend" a given table adding an entry to that table and every
+	/// table the derives from that table.  For example, the Keywords table will add "private" to the 
+	/// "ClassBase" table.  Since the "Class", "Interface", and "NativeInterface" tables derive from
+	/// "ClassBase", all of those tables will contain the keyword "private".
+	/// 
+	/// See UhtTables.cs for a list of table names and how they relate to each other.
+	/// 
+	/// Tables are implemented in the following source files:
+	/// 
+	///		UhtEngineClassTable.cs - Collection of all the engine class names.
+	///		UhtKeywordTable.cs - Collection of the C++ keywords that UHT understands.
+	///		UhtLocTextDefaultValueTable.cs - Collection of loctext default value parsing
+	///		UhtPropertyTypeTable.cs - Collection of the property type keywords.
+	///		UhtSpecifierTable.cs - Collection of the known specifiers
+	///		UhtSpecifierValidatorTable.cs - Collection of the specifier validators
+	///		UhtStructDefaultValueTable.cs - Collection of structure default value parsing
+	/// </summary>
+	public class UhtAttributeScanner
+	{
+		private static readonly Lazy<UhtAttributeScanner> LazyInit = new Lazy<UhtAttributeScanner>(() =>
+		{
+			UhtAttributeScanner Scanner = new UhtAttributeScanner();
+			Scanner.Initialize();
+			return Scanner;
+		});
+
+		private readonly UhtSpecifierTables SpecifierTables = UhtSpecifierTables.Instance;
+		private readonly UhtSpecifierValidatorTables SpecifierValidatorTables = UhtSpecifierValidatorTables.Instance;
+		private readonly UhtKeywordTables KeywordTables = UhtKeywordTables.Instance;
+		private readonly UhtPropertyTypeTable PropertyTypeTable = UhtPropertyTypeTable.Instance;
+		private readonly UhtStructDefaultValueTable StructDefaultValueTable = UhtStructDefaultValueTable.Instance;
+		private readonly UhtEngineClassTable EngineClassTable = UhtEngineClassTable.Instance;
+		private readonly UhtExporterTable GeneraterTable = UhtExporterTable.Instance;
+		private readonly UhtLocTextDefaultValueTable LocTextDefaultValueTable = UhtLocTextDefaultValueTable.Instance;
+
+		/// <summary>
+		/// Initialize the attribute scanner and scan all assemblies for attributes
+		/// </summary>
+		public static void Scan()
+		{
+			_ = LazyInit.Value;
+		}
+
+		/// <summary>
+		/// Perform one time initialization of the system
+		/// </summary>
+		private void Initialize()
+		{
+			IEnumerable<Assembly> LoadedAssemblies = AppDomain.CurrentDomain.GetAssemblies()
+				// Exclude the AutomationTool driver assembly - it contains no tasks, and trying to load it in the context of
+				// BuildGraph can result in an exception if Microsoft.Build.Framework is not able to be loaded
+				.Where(A => !String.Equals("AutomationTool", A.GetName().Name));
+
+			// Collect all of the types from the loaded assemblies
+			foreach (Assembly LoadedAssembly in LoadedAssemblies)
+			{
+				Type[] Types;
+				try
+				{
+					Types = LoadedAssembly.GetTypes();
+				}
+				catch (ReflectionTypeLoadException ex)
+				{
+					Log.TraceWarning("Exception {0} while trying to get types from assembly {1}. LoaderExceptions: {2}", 
+						ex, LoadedAssembly, string.Join("\n", ex.LoaderExceptions.Select(x => x?.Message)));
+					continue;
+				}
+
+				foreach (Type Type in Types)
+				{
+					CheckForAttributes(Type);
+				}
+			}
+
+			PerformPostInitialization();
+		}
+
+		/// <summary>
+		/// For the given type, look for any table related attributes 
+		/// </summary>
+		/// <param name="Type">Type in question</param>
+		private void CheckForAttributes(Type Type)
+		{
+			if (Type.IsClass)
+			{
+
+				// Loop through the attributes
+				foreach (Attribute ClassAttribute in Type.GetCustomAttributes(false))
+				{
+					if (ClassAttribute is UnrealHeaderToolAttribute ParserAttribute)
+					{
+						HandleUnrealHeaderToolAttribute(Type, ParserAttribute);
+					}
+					else if (ClassAttribute is UhtEngineClassAttribute EngineClassAttribute)
+					{
+						this.EngineClassTable.OnEngineClassAttribute(Type, EngineClassAttribute);
+					}
+				}
+			}
+		}
+
+		private void PerformPostInitialization()
+		{
+			this.KeywordTables.Merge();
+			this.SpecifierTables.Merge();
+			this.SpecifierValidatorTables.Merge();
+
+			// Invoke this to generate an exception if there is no default
+			_ = this.PropertyTypeTable.Default;
+			_ = this.StructDefaultValueTable.Default;
+		}
+
+		private void HandleUnrealHeaderToolAttribute(Type Type, UnrealHeaderToolAttribute ParserAttribute)
+		{
+			if (!string.IsNullOrEmpty(ParserAttribute.InitMethod))
+			{
+				MethodInfo? InitInfo = Type.GetMethod(ParserAttribute.InitMethod, BindingFlags.Static | BindingFlags.Public | BindingFlags.NonPublic);
+				if (InitInfo == null)
+				{
+					throw new Exception($"Unable to find UnrealHeaderTool attribute InitMethod {ParserAttribute.InitMethod}");
+				}
+				InitInfo.Invoke(null, new Object[] { });
+			}
+
+			// Scan the methods for things we are interested in
+			foreach (MethodInfo MethodInfo in Type.GetMethods(BindingFlags.Static | BindingFlags.Public | BindingFlags.NonPublic))
+			{
+
+				// Scan for attributes we care about
+				foreach (Attribute MethodAttribute in MethodInfo.GetCustomAttributes())
+				{
+					if (MethodAttribute is UhtSpecifierAttribute SpecifierAttribute)
+					{
+						this.SpecifierTables.OnSpecifierAttribute(Type, MethodInfo, SpecifierAttribute);
+					}
+					else if (MethodAttribute is UhtSpecifierValidatorAttribute SpecifierValidatorAttribute)
+					{
+						this.SpecifierValidatorTables.OnSpecifierValidatorAttribute(Type, MethodInfo, SpecifierValidatorAttribute);
+					}
+					else if (MethodAttribute is UhtKeywordCatchAllAttribute KeywordCatchAllAttribute)
+					{
+						this.KeywordTables.OnKeywordCatchAllAttribute(Type, MethodInfo, KeywordCatchAllAttribute);
+					}
+					else if (MethodAttribute is UhtKeywordAttribute KeywordAttribute)
+					{
+						this.KeywordTables.OnKeywordAttribute(Type, MethodInfo, KeywordAttribute);
+					}
+					else if (MethodAttribute is UhtPropertyTypeAttribute PropertyTypeAttribute)
+					{
+						this.PropertyTypeTable.OnPropertyTypeAttribute(Type, MethodInfo, PropertyTypeAttribute);
+					}
+					else if (MethodAttribute is UhtStructDefaultValueAttribute StructDefaultValueAttribute)
+					{
+						this.StructDefaultValueTable.OnStructDefaultValueAttribute(Type, MethodInfo, StructDefaultValueAttribute);
+					}
+					else if (MethodAttribute is UhtExporterAttribute ExporterAttribute)
+					{
+						this.GeneraterTable.OnExporterAttribute(Type, MethodInfo, ExporterAttribute);
+					}
+					else if (MethodAttribute is UhtLocTextDefaultValueAttribute LocTextDefaultValueAttribute)
+					{
+						this.LocTextDefaultValueTable.OnLocTextDefaultValueAttribute(Type, MethodInfo, LocTextDefaultValueAttribute);
+					}
+				}
+			}
+		}
+	}
+
+	/// <summary>
+	/// This attribute must be applied to any class that contains other UHT attributes.
+	/// </summary>
+	[AttributeUsage(AttributeTargets.Class)]
+	public class UnrealHeaderToolAttribute : Attribute
+	{
+
+		/// <summary>
+		/// If specified, this method will be invoked once during the scan for attributes.
+		/// It can be used to perform some one time initialization.
+		/// </summary>
+		public string InitMethod = string.Empty;
+	}
+}
diff --git a/Engine/Source/Programs/Shared/EpicGames.UHT/Utils/UhtBuffer.cs b/Engine/Source/Programs/Shared/EpicGames.UHT/Utils/UhtBuffer.cs
new file mode 100644
--- /dev/null
+++ b/Engine/Source/Programs/Shared/EpicGames.UHT/Utils/UhtBuffer.cs
@@ -0,0 +1,220 @@
+// Copyright Epic Games, Inc. All Rights Reserved.
+
+using System;
+using System.Numerics;
+using System.Text;
+
+namespace EpicGames.UHT.Utils
+{
+
+	/// <summary>
+	/// Cached character buffer system.
+	/// 
+	/// Invoke UhtBuffer.Borrow method to get a buffer of the given size.
+	/// Invoke UhtBuffer.Return to return the buffer to the cache.
+	/// </summary>
+	public class UhtBuffer
+	{
+		/// <summary>
+		/// Any requests of the given size or smaller will be placed in bucket zero with the given size.
+		/// </summary>
+		private static int MinSize = 1024 * 16;
+
+		/// <summary>
+		/// Adjustment to the bucket index to account for the minimum bucket size
+		/// </summary>
+		private static int BuckedAdjustment = BitOperations.Log2((uint)UhtBuffer.MinSize);
+		
+		/// <summary>
+		/// Total number of supported buckets
+		/// </summary>
+		private static int BucketCount = 32 - UhtBuffer.BuckedAdjustment;
+
+		/// <summary>
+		/// Bucket lookaside list
+		/// </summary>
+		private static UhtBuffer?[] LookAsideArray = new UhtBuffer?[UhtBuffer.BucketCount];
+
+		/// <summary>
+		/// The bucket index associated with the buffer
+		/// </summary>
+		private int Bucket;
+
+		/// <summary>
+		/// Single list link to the next cached buffer
+		/// </summary>
+		private UhtBuffer? NextBuffer = null;
+
+		/// <summary>
+		/// The backing character block.  The size of the array will normally be larger than the 
+		/// requested size.
+		/// </summary>
+		public char[] Block;
+
+		/// <summary>
+		/// Memory region sized to the requested size
+		/// </summary>
+		public Memory<char> Memory;
+
+		/// <summary>
+		/// Construct a new buffer
+		/// </summary>
+		/// <param name="Size">The initial size of the buffer</param>
+		/// <param name="Bucket">The bucket associated with the buffer</param>
+		/// <param name="BucketSize">The size all blocks in this bucket</param>
+		private UhtBuffer(int Size, int Bucket, int BucketSize)
+		{
+			this.Block = new char[BucketSize];
+			this.Bucket = Bucket;
+			Reset(Size);
+		}
+
+		/// <summary>
+		/// Reset the memory region to the given size
+		/// </summary>
+		/// <param name="Size"></param>
+		public void Reset(int Size)
+		{
+			this.Memory = new Memory<char>(this.Block, 0, Size);
+		}
+
+		/// <summary>
+		/// Borrow a new buffer of the given size
+		/// </summary>
+		/// <param name="Size">Size of the buffer</param>
+		/// <returns>Buffer that should be returned with a call to Return</returns>
+		public static UhtBuffer Borrow(int Size)
+		{
+			if (Size <= UhtBuffer.MinSize)
+			{
+				return BorrowInternal(Size, 0, UhtBuffer.MinSize);
+			}
+			else
+			{
+
+				// Round up the size to the next larger power of two if it isn't a power of two
+				uint USize = (uint)Size;
+				--USize;
+				USize |= USize >> 1;
+				USize |= USize >> 2;
+				USize |= USize >> 4;
+				USize |= USize >> 8;
+				USize |= USize >> 16;
+				++USize;
+				int Bucket = BitOperations.Log2(USize) - UhtBuffer.BuckedAdjustment;
+				return BorrowInternal(Size, Bucket, (int)USize);
+			}
+		}
+
+		/// <summary>
+		/// Return a buffer initialized with the string builder.
+		/// </summary>
+		/// <param name="Builder">Source builder content</param>
+		/// <returns>Buffer that should be returned with a call to Return</returns>
+		public static UhtBuffer Borrow(StringBuilder Builder)
+		{
+			int Length = Builder.Length;
+			UhtBuffer Buffer = Borrow(Length);
+			Builder.CopyTo(0, Buffer.Memory.Span, Length);
+			return Buffer;
+		}
+
+		/// <summary>
+		/// Return a buffer initialized with the string builder sub string.
+		/// </summary>
+		/// <param name="Builder">Source builder content</param>
+		/// <param name="StartIndex">Starting index in the builder</param>
+		/// <param name="Length">Length of the content</param>
+		/// <returns>Buffer that should be returned with a call to Return</returns>
+		public static UhtBuffer Borrow(StringBuilder Builder, int StartIndex, int Length)
+		{
+			UhtBuffer Buffer = Borrow(Length);
+			Builder.CopyTo(StartIndex, Buffer.Memory.Span, Length);
+			return Buffer;
+		}
+
+		/// <summary>
+		/// Return the buffer to the cache.  The buffer should no longer be accessed.
+		/// </summary>
+		/// <param name="Buffer">The buffer to be returned.</param>
+		public static void Return(UhtBuffer Buffer)
+		{
+			lock (UhtBuffer.LookAsideArray)
+			{
+				Buffer.NextBuffer = UhtBuffer.LookAsideArray[Buffer.Bucket];
+				UhtBuffer.LookAsideArray[Buffer.Bucket] = Buffer;
+			}
+		}
+
+		/// <summary>
+		/// Internal helper to allocate a buffer
+		/// </summary>
+		/// <param name="Size">The initial size of the buffer</param>
+		/// <param name="Bucket">The bucket associated with the buffer</param>
+		/// <param name="BucketSize">The size all blocks in this bucket</param>
+		/// <returns>The allocated buffer</returns>
+		private static UhtBuffer BorrowInternal(int Size, int Bucket, int BucketSize)
+		{
+			lock (UhtBuffer.LookAsideArray)
+			{
+				if (UhtBuffer.LookAsideArray[Bucket] != null)
+				{
+					UhtBuffer Buffer = UhtBuffer.LookAsideArray[Bucket]!;
+					UhtBuffer.LookAsideArray[Bucket] = Buffer.NextBuffer;
+					Buffer.Reset(Size);
+					return Buffer;
+				}
+			}
+			return new UhtBuffer(Size, Bucket, BucketSize);
+		}
+	}
+
+	/// <summary>
+	/// Helper class for using pattern to borrow and return a buffer.
+	/// </summary>
+	public struct UhtBorrowBuffer : IDisposable
+	{
+
+		/// <summary>
+		/// The borrowed buffer
+		/// </summary>
+		public UhtBuffer Buffer;
+
+		/// <summary>
+		/// Borrow a buffer with the given size
+		/// </summary>
+		/// <param name="Size">The size to borrow</param>
+		public UhtBorrowBuffer(int Size)
+		{
+			this.Buffer = UhtBuffer.Borrow(Size);
+		}
+
+		/// <summary>
+		/// Borrow a buffer populated with the builder contents
+		/// </summary>
+		/// <param name="Builder">Initial contents of the buffer</param>
+		public UhtBorrowBuffer(StringBuilder Builder)
+		{
+			this.Buffer = UhtBuffer.Borrow(Builder);
+		}
+
+		/// <summary>
+		/// Borrow a buffer populated with the builder contents
+		/// </summary>
+		/// <param name="Builder">Initial contents of the buffer</param>
+		/// <param name="StartIndex">Starting index into the builder</param>
+		/// <param name="Length">Length of the data in the builder</param>
+		public UhtBorrowBuffer(StringBuilder Builder, int StartIndex, int Length)
+		{
+			this.Buffer = UhtBuffer.Borrow(Builder, StartIndex, Length);
+		}
+
+		/// <summary>
+		/// Return the borrowed buffer to the cache
+		/// </summary>
+		public void Dispose()
+		{
+			UhtBuffer.Return(this.Buffer);
+		}
+	}
+}
diff --git a/Engine/Source/Programs/Shared/EpicGames.UHT/Utils/UhtConfig.cs b/Engine/Source/Programs/Shared/EpicGames.UHT/Utils/UhtConfig.cs
new file mode 100644
--- /dev/null
+++ b/Engine/Source/Programs/Shared/EpicGames.UHT/Utils/UhtConfig.cs
@@ -0,0 +1,144 @@
+// Copyright Epic Games, Inc. All Rights Reserved.
+
+using EpicGames.Core;
+using EpicGames.UHT.Tokenizer;
+
+namespace EpicGames.UHT.Utils
+{
+
+	/// <summary>
+	/// Describes how pointers will generate warnings or errors
+	/// </summary>
+	public enum UhtPointerMemberBehavior
+	{
+
+		/// <summary>
+		/// An error will be generated.
+		/// </summary>
+		Disallow,
+
+		/// <summary>
+		/// Ignore the pointer
+		/// </summary>
+		AllowSilently,
+
+		/// <summary>
+		/// Log a warning about the pointer
+		/// </summary>
+		AllowAndLog,
+	};
+
+	/// <summary>
+	/// Interface for accessing configuration data.  Since UnrealBuildTool depends on EpicGames.UHT and all 
+	/// of the configuration support exists in UBT, configuration data must be accessed through an interface.
+	/// </summary>
+	public interface IUhtConfig
+	{
+		/// <summary>
+		/// Default version of generated code. Defaults to oldest possible, unless specified otherwise in config.
+		/// </summary>
+		public EGeneratedCodeVersion DefaultGeneratedCodeVersion { get; }
+
+		/// <summary>
+		/// Pointer warning for native pointers in the engine
+		/// </summary>
+		public UhtPointerMemberBehavior EngineNativePointerMemberBehavior { get; }
+
+		/// <summary>
+		/// Pointer warning for object pointers in the engine
+		/// </summary>
+		public UhtPointerMemberBehavior EngineObjectPtrMemberBehavior { get; }
+
+		/// <summary>
+		/// Pointer warning for native pointers outside the engine
+		/// </summary>
+		public UhtPointerMemberBehavior NonEngineNativePointerMemberBehavior { get; }
+
+		/// <summary>
+		/// Pointer warning for object pointers outside the engine
+		/// </summary>
+		public UhtPointerMemberBehavior NonEngineObjectPtrMemberBehavior { get; }
+
+		/// <summary>
+		/// If true setters and getters will be automatically (without specifying their function names on a property) parsed and generated if a function with matching signature is found
+		/// </summary>
+		public bool bAllowAutomaticSettersAndGetters { get; }
+
+		/// <summary>
+		/// If the token references a remapped identifier, update the value in the token 
+		/// </summary>
+		/// <param name="Token">Token to be remapped</param>
+		public void RedirectTypeIdentifier(ref UhtToken Token);
+
+		/// <summary>
+		/// Return the remapped key or the existing key
+		/// </summary>
+		/// <param name="Key">Key to be remapped.</param>
+		/// <param name="NewKey">Resulting key name</param>
+		/// <returns>True if the key has been remapped and has changed.</returns>
+		public bool RedirectMetaDataKey(string Key, out string NewKey);
+
+		/// <summary>
+		/// Test to see if the given units are valid.
+		/// </summary>
+		/// <param name="Units">Units to test</param>
+		/// <returns>True if the units are valid, false if not</returns>
+		public bool IsValidUnits(StringView Units);
+
+		/// <summary>
+		/// Test to see if the structure name should be using a "T" prefix.
+		/// </summary>
+		/// <param name="Name">Name of the structure to test without any prefix.</param>
+		/// <returns>True if the structure should have a "T" prefix.</returns>
+		public bool IsStructWithTPrefix(StringView Name);
+
+		/// <summary>
+		/// Test to see if the given macro has a parameter count as part of the name.
+		/// </summary>
+		/// <param name="DelegateMacro">Macro to test</param>
+		/// <returns>-1 if the macro does not contain a parameter count.  The number of parameters minus one.</returns>
+		public int FindDelegateParameterCount(StringView DelegateMacro);
+
+		/// <summary>
+		/// Get the parameter count string associated with the given index
+		/// </summary>
+		/// <param name="Index">Index from a prior call to FindDelegateParameterCount or -1.</param>
+		/// <returns>Parameter count string or an empty string if Index is -1.</returns>
+		public StringView GetDelegateParameterCountString(int Index);
+
+		/// <summary>
+		/// Test to see if the exporter is enabled
+		/// </summary>
+		/// <param name="Name">Name of the exporter</param>
+		/// <returns>True if the exporter is enabled, false if not</returns>
+		public bool IsExporterEnabled(string Name);
+	}
+
+	/// <summary>
+	/// Helper class to house the IUhtConfig instance
+	/// </summary>
+	public static class UhtConfig
+	{
+		private static IUhtConfig? InstanceInternal = null;
+
+		/// <summary>
+		/// Get/Set the instance of IUhtConfig
+		/// </summary>
+		public static IUhtConfig Instance 
+		{ 
+			get
+			{
+				if (UhtConfig.InstanceInternal == null)
+				{
+					throw new UhtIceException("UhtConfig instance has not been set");
+				}
+				return UhtConfig.InstanceInternal;
+			}
+
+			set
+			{
+				UhtConfig.InstanceInternal = value;
+			}
+		}
+	}
+}
diff --git a/Engine/Source/Programs/Shared/EpicGames.UHT/Utils/UhtException.cs b/Engine/Source/Programs/Shared/EpicGames.UHT/Utils/UhtException.cs
new file mode 100644
--- /dev/null
+++ b/Engine/Source/Programs/Shared/EpicGames.UHT/Utils/UhtException.cs
@@ -0,0 +1,123 @@
+// Copyright Epic Games, Inc. All Rights Reserved.
+
+using EpicGames.UHT.Tokenizer;
+using System;
+using System.Text;
+
+namespace EpicGames.UHT.Utils
+{
+	/// <summary>
+	/// Compiler exception that immediately stops the processing of the current header file.
+	/// </summary>
+	public class UhtException : Exception
+	{
+		/// <summary>
+		/// The generated message
+		/// </summary>
+		public UhtMessage UhtMessage;
+
+		/// <summary>
+		/// Internal do nothing constructor
+		/// </summary>
+		protected UhtException()
+		{
+		}
+
+		/// <summary>
+		/// Exception with a simple message.  Context will be the current header file.
+		/// </summary>
+		/// <param name="Message">Text of the error</param>
+		public UhtException(string Message)
+		{
+			this.UhtMessage = UhtMessage.MakeMessage(UhtMessageType.Error, null, null, 1, Message);
+		}
+
+		/// <summary>
+		/// Exception with a simple message.  Context from the given message site.
+		/// </summary>
+		/// <param name="MessageSite">Site generating the exception</param>
+		/// <param name="Message">Text of the error</param>
+		/// <param name="ExtraContext">Addition context to be appended to the error message</param>
+		public UhtException(IUhtMessageSite MessageSite, string Message, object? ExtraContext = null)
+		{
+			if (ExtraContext != null)
+			{
+				Message = $"{Message} while parsing {UhtMessage.FormatContext(ExtraContext)}";
+			}
+			this.UhtMessage = UhtMessage.MakeMessage(UhtMessageType.Error, MessageSite.MessageSource, null, MessageSite.GetLineNumber(), Message);
+		}
+
+		/// <summary>
+		/// Make an exception to be thrown
+		/// </summary>
+		/// <param name="MessageSite">Message site to be associated with the exception</param>
+		/// <param name="LineNumber">Line number of the error</param>
+		/// <param name="Message">Text of the error</param>
+		/// <param name="ExtraContext">Addition context to be appended to the error message</param>
+		public UhtException(IUhtMessageSite MessageSite, int LineNumber, string Message, object? ExtraContext = null)
+		{
+			if (ExtraContext != null)
+			{
+				Message = $"{Message} while parsing {UhtMessage.FormatContext(ExtraContext)}";
+			}
+			this.UhtMessage = UhtMessage.MakeMessage(UhtMessageType.Error, MessageSite.MessageSource, null, MessageSite.GetLineNumber(LineNumber), Message);
+		}
+	}
+
+	/// <summary>
+	/// Exception where the current token isn't what was expected
+	/// </summary>
+	public class UhtTokenException : UhtException
+	{
+
+		/// <summary>
+		/// Make a parsing error for when there is a mismatch between the expected token and what was parsed.
+		/// </summary>
+		/// <param name="MessageSite">Message site to be associated with the exception</param>
+		/// <param name="Got">The parsed token.  Support for EOF also provided.</param>
+		/// <param name="Expected">What was expected.</param>
+		/// <param name="ExtraContext">Extra context to be appended to the error message</param>
+		/// <returns>The exception object to throw</returns>
+		public UhtTokenException(IUhtMessageSite MessageSite, UhtToken Got, object? Expected, object? ExtraContext = null)
+		{
+			string Message = Expected != null 
+				? $"Found {UhtMessage.FormatContext(Got)} when expecting {UhtMessage.FormatContext(Expected)}{FormatExtraContext(MessageSite, ExtraContext)}"
+				: $"Found {UhtMessage.FormatContext(Got)}{FormatExtraContext(MessageSite, ExtraContext)}";
+			this.UhtMessage = UhtMessage.MakeMessage(UhtMessageType.Error, MessageSite.MessageSource, null, Got.InputLine, Message);
+		}
+
+		/// <summary>
+		/// Format any extra context supplied by the caller or the message site
+		/// </summary>
+		/// <param name="MessageSite">Message site that may contain extra context</param>
+		/// <param name="ExtraContext">Additional caller supplied context</param>
+		/// <returns></returns>
+		private string FormatExtraContext(IUhtMessageSite MessageSite, object? ExtraContext = null)
+		{
+			StringBuilder Builder = new StringBuilder(" while parsing ");
+			int StartingLength = Builder.Length;
+			if (ExtraContext != null)
+			{
+				Builder.Append(UhtMessage.FormatContext(ExtraContext));
+			}
+			UhtMessage.Append(Builder, UhtTlsMessageExtraContext.GetMessageExtraContext(), StartingLength != Builder.Length);
+
+			return Builder.Length != StartingLength ? Builder.ToString() : string.Empty;
+		}
+	}
+
+	/// <summary>
+	/// Internal compiler error exception
+	/// </summary>
+	public class UhtIceException : UhtException
+	{
+		/// <summary>
+		/// Exception with a simple message.  Context will be the current header file.
+		/// </summary>
+		/// <param name="Message">Text of the error</param>
+		public UhtIceException(string Message)
+		{
+			this.UhtMessage = UhtMessage.MakeMessage(UhtMessageType.Ice, null, null, 1, Message);
+		}
+	}
+}
diff --git a/Engine/Source/Programs/Shared/EpicGames.UHT/Utils/UhtExport.cs b/Engine/Source/Programs/Shared/EpicGames.UHT/Utils/UhtExport.cs
new file mode 100644
--- /dev/null
+++ b/Engine/Source/Programs/Shared/EpicGames.UHT/Utils/UhtExport.cs
@@ -0,0 +1,192 @@
+// Copyright Epic Games, Inc. All Rights Reserved.
+
+using EpicGames.Core;
+using EpicGames.UHT.Types;
+using System;
+using System.Collections.Generic;
+using System.Text;
+using System.Threading.Tasks;
+
+namespace EpicGames.UHT.Utils
+{
+
+	/// <summary>
+	/// Export options.
+	/// </summary>
+	[Flags]
+	public enum UhtExportOptions
+	{
+
+		/// <summary>
+		/// No export options specified
+		/// </summary>
+		None = 0,
+
+		/// <summary>
+		/// If set, output is to be written
+		/// </summary>
+		WriteOutput = 1 << 0,
+	}
+
+	/// <summary>
+	/// Helper methods for testing flags.  These methods perform better than the generic HasFlag which hits
+	/// the GC and stalls.
+	/// </summary>
+	public static class UhtExportOptionsExtensions
+	{
+
+		/// <summary>
+		/// Test to see if any of the specified flags are set
+		/// </summary>
+		/// <param name="InFlags">Current flags</param>
+		/// <param name="TestFlags">Flags to test for</param>
+		/// <returns>True if any of the flags are set</returns>
+		public static bool HasAnyFlags(this UhtExportOptions InFlags, UhtExportOptions TestFlags)
+		{
+			return (InFlags & TestFlags) != 0;
+		}
+
+		/// <summary>
+		/// Test to see if all of the specified flags are set
+		/// </summary>
+		/// <param name="InFlags">Current flags</param>
+		/// <param name="TestFlags">Flags to test for</param>
+		/// <returns>True if all the flags are set</returns>
+		public static bool HasAllFlags(this UhtExportOptions InFlags, UhtExportOptions TestFlags)
+		{
+			return (InFlags & TestFlags) == TestFlags;
+		}
+
+		/// <summary>
+		/// Test to see if a specific set of flags have a specific value.
+		/// </summary>
+		/// <param name="InFlags">Current flags</param>
+		/// <param name="TestFlags">Flags to test for</param>
+		/// <param name="MatchFlags">Expected value of the tested flags</param>
+		/// <returns>True if the given flags have a specific value.</returns>
+		public static bool HasExactFlags(this UhtExportOptions InFlags, UhtExportOptions TestFlags, UhtExportOptions MatchFlags)
+		{
+			return (InFlags & TestFlags) == MatchFlags;
+		}
+	}
+
+	/// <summary>
+	/// Represents an output file
+	/// </summary>
+	public interface IUhtExportOutput
+	{
+
+		/// <summary>
+		/// Session associated with the output.
+		/// </summary>
+		public UhtSession Session { get; }
+
+		/// <summary>
+		/// Export options
+		/// </summary>
+		public UhtExportOptions Options { get; set; }
+
+		/// <summary>
+		/// Commit the contents of the string builder as the output.
+		/// </summary>
+		/// <param name="Builder">Source for the content</param>
+		/// <returns>A string view of the committed output.  Use this instead of the 
+		/// StringBuilder so the builder can be recycled.</returns>
+		public StringView CommitOutput(StringBuilder Builder);
+
+		/// <summary>
+		/// Commit the value of the string as the output
+		/// </summary>
+		/// <param name="Output">Output to commit</param>
+		public void CommitOutput(string Output);
+	}
+
+	/// <summary>
+	/// Represents a task to export outputs
+	/// </summary>
+	public interface IUhtExportTask
+	{
+
+		/// <summary>
+		/// Session being run
+		/// </summary>
+		public UhtSession Session { get; }
+
+		/// <summary>
+		/// Export options
+		/// </summary>
+		public UhtExportOptions Options { get; }
+
+		/// <summary>
+		/// Task handle associated with the task for it can be waited on
+		/// </summary>
+		public Task? ActionTask { get; }
+	}
+
+	/// <summary>
+	/// Delegate used to generate a single output from a task
+	/// </summary>
+	/// <param name="Output">The single output</param>
+	public delegate void UhtExportSingleDelegate(IUhtExportOutput Output);
+
+	/// <summary>
+	/// Delegate used to generate two outputs from a task.
+	/// </summary>
+	/// <param name="HeaderOutput">Header output</param>
+	/// <param name="CppOutput">CPP output</param>
+	public delegate void UhtExportPairDelegate(IUhtExportOutput HeaderOutput, IUhtExportOutput CppOutput);
+
+	/// <summary>
+	/// Factory object used to generate export tasks
+	/// </summary>
+	public interface IUhtExportFactory
+	{
+		
+		/// <summary>
+		/// Session being run
+		/// </summary>
+		public UhtSession Session { get; }
+
+		/// <summary>
+		/// Export options
+		/// </summary>
+		public UhtExportOptions Options { get; }
+
+		/// <summary>
+		/// Create a task to export a single file
+		/// </summary>
+		/// <param name="Path">Destination file path</param>
+		/// <param name="Prereqs">Tasks that must be completed prior to this task running</param>
+		/// <param name="AdditionalOptions">Additional options</param>
+		/// <param name="Action">Action to be invoked to generate the output</param>
+		/// <returns>Task interface.</returns>
+		public IUhtExportTask CreateTask(string Path, List<IUhtExportTask> Prereqs, UhtExportOptions AdditionalOptions, UhtExportSingleDelegate Action);
+
+		/// <summary>
+		/// Create a task to export two files
+		/// </summary>
+		/// <param name="HeaderPath">Header destination file path</param>
+		/// <param name="CppPath">Cpp destination file path</param>
+		/// <param name="Prereqs">Tasks that must be completed prior to this task running</param>
+		/// <param name="AdditionalOptions">Additional options</param>
+		/// <param name="Action">Action to be invoked to generate the output</param>
+		/// <returns>Task interface.</returns>
+		public IUhtExportTask CreateTask(string HeaderPath, string CppPath, List<IUhtExportTask> Prereqs, UhtExportOptions AdditionalOptions, UhtExportPairDelegate Action);
+
+		/// <summary>
+		/// Make a path for an output based on the header file name.
+		/// </summary>
+		/// <param name="HeaderFile">Header file being exported.</param>
+		/// <param name="Suffix">Suffix to be added to the file name.</param>
+		/// <returns>Output file path</returns>
+		public string MakePath(UhtHeaderFile HeaderFile, string Suffix);
+
+		/// <summary>
+		/// Make a path for an output based on the package name.
+		/// </summary>
+		/// <param name="Package">Package being exported</param>
+		/// <param name="Suffix">Suffix to be added to the file name.</param>
+		/// <returns>Output file path</returns>
+		public string MakePath(UhtPackage Package, string Suffix);
+	}
+}
diff --git a/Engine/Source/Programs/Shared/EpicGames.UHT/Utils/UhtFCString.cs b/Engine/Source/Programs/Shared/EpicGames.UHT/Utils/UhtFCString.cs
new file mode 100644
--- /dev/null
+++ b/Engine/Source/Programs/Shared/EpicGames.UHT/Utils/UhtFCString.cs
@@ -0,0 +1,533 @@
+// Copyright Epic Games, Inc. All Rights Reserved.
+
+using EpicGames.Core;
+using System;
+using System.Text;
+
+namespace EpicGames.UHT.Utils
+{
+
+	/// <summary>
+	/// Collection of helper methods to emulate UHT string functions
+	/// </summary>
+	public static class UhtFCString
+	{
+
+		/// <summary>
+		/// Delimiter between the subobjects
+		/// </summary>
+		public static char SubObjectDelimiter = ':';
+
+		/// <summary>
+		/// Test to see if the string is a boolean
+		/// </summary>
+		/// <param name="Value">Boolean to test</param>
+		/// <returns>True if the value is true</returns>
+		public static bool ToBool(StringView Value)
+		{
+			ReadOnlySpan<char> Span = Value.Span;
+			if (Span.Length == 0)
+			{
+				return false;
+			}
+			else if (
+				Span.Equals("true", StringComparison.OrdinalIgnoreCase) ||
+				Span.Equals("yes", StringComparison.OrdinalIgnoreCase) ||
+				Span.Equals("on", StringComparison.OrdinalIgnoreCase))
+			{
+				return true;
+			}
+			else if (
+				Span.Equals("false", StringComparison.OrdinalIgnoreCase) ||
+				Span.Equals("no", StringComparison.OrdinalIgnoreCase) ||
+				Span.Equals("off", StringComparison.OrdinalIgnoreCase))
+			{
+				return false;
+			}
+			else
+			{
+				int IntValue;
+				return int.TryParse(Span, out IntValue) && IntValue != 0;
+			}
+		}
+
+		/// <summary>
+		/// Test to see if the string is a boolean
+		/// </summary>
+		/// <param name="Value">Boolean to test</param>
+		/// <returns>True if the value is true</returns>
+		public static bool ToBool(string? Value)
+		{
+			if (string.IsNullOrEmpty(Value))
+			{
+				return false;
+			}
+			else if (string.Compare(Value, "true", true) == 0 ||
+				string.Compare(Value, "yes", true) == 0 ||
+				string.Compare(Value, "on", true) == 0)
+			{
+				return true;
+			}
+			else if (string.Compare(Value, "false", true) == 0 ||
+				string.Compare(Value, "no", true) == 0 ||
+				string.Compare(Value, "off", true) == 0)
+			{
+				return false;
+			}
+			else
+			{
+				int IntValue;
+				return int.TryParse(Value, out IntValue) && IntValue != 0;
+			}
+		}
+
+		/// <summary>
+		/// Test to see if the character is a digit
+		/// </summary>
+		/// <param name="C">Character to test</param>
+		/// <returns>True if it is</returns>
+		public static bool IsDigit(char C)
+		{
+			return C >= '0' && C <= '9';
+		}
+
+		/// <summary>
+		/// Test to see if the character is an alphabet character
+		/// </summary>
+		/// <param name="C">Character to test</param>
+		/// <returns>True if it is</returns>
+		public static bool IsAlpha(char C)
+		{
+			return (C >= 'A' && C <= 'Z') || (C >= 'a' && C <= 'z');
+		}
+
+		/// <summary>
+		/// Test to see if the character is a digit or alphabet character
+		/// </summary>
+		/// <param name="C">Character to test</param>
+		/// <returns>True if it is</returns>
+		public static bool IsAlnum(char C)
+		{
+			return IsDigit(C) || IsAlpha(C);
+		}
+
+		/// <summary>
+		/// Test to see if the character is a hex digit (A-F)
+		/// </summary>
+		/// <param name="C">Character to test</param>
+		/// <returns>True if it is</returns>
+		public static bool IsHexAlphaDigit(char C)
+		{
+			return (C >= 'A' && C <= 'F') || (C >= 'a' && C <= 'f');
+		}
+
+		/// <summary>
+		/// Test to see if the character is a hex digit (0-9a-f)
+		/// </summary>
+		/// <param name="C">Character to test</param>
+		/// <returns>True if it is</returns>
+		public static bool IsHexDigit(char C)
+		{
+			return IsDigit(C) || IsHexAlphaDigit(C);
+		}
+
+		/// <summary>
+		/// Test to see if the character is whitespace
+		/// </summary>
+		/// <param name="C">Character to test</param>
+		/// <returns>True if it is</returns>
+		public static bool IsWhitespace(char C)
+		{
+			return Char.IsWhiteSpace(C);
+		}
+
+		/// <summary>
+		/// Test to see if the character is a sign
+		/// </summary>
+		/// <param name="C">Character to test</param>
+		/// <returns>True if it is</returns>
+		public static bool IsSign(char C)
+		{
+			return C == '+' || C == '-';
+		}
+
+		/// <summary>
+		/// Test to see if the character is a hex marker (xX)
+		/// </summary>
+		/// <param name="C">Character to test</param>
+		/// <returns>True if it is</returns>
+		public static bool IsHexMarker(char C)
+		{
+			return C == 'x' || C == 'X';
+		}
+
+		/// <summary>
+		/// Test to see if the character is a float marker (fF)
+		/// </summary>
+		/// <param name="C">Character to test</param>
+		/// <returns>True if it is</returns>
+		public static bool IsFloatMarker(char C)
+		{
+			return C == 'f' || C == 'F';
+		}
+
+		/// <summary>
+		/// Test to see if the character is an exponent marker (eE)
+		/// </summary>
+		/// <param name="C">Character to test</param>
+		/// <returns>True if it is</returns>
+		public static bool IsExponentMarker(char C)
+		{
+			return C == 'e' || C == 'E';
+		}
+
+		/// <summary>
+		/// Test to see if the character is an unsigned marker (uU)
+		/// </summary>
+		/// <param name="C">Character to test</param>
+		/// <returns>True if it is</returns>
+		public static bool IsUnsignedMarker(char C)
+		{
+			return C == 'u' || C == 'U';
+		}
+
+		/// <summary>
+		/// Test to see if the character is a long marker (lL)
+		/// </summary>
+		/// <param name="C">Character to test</param>
+		/// <returns>True if it is</returns>
+		public static bool IsLongMarker(char C)
+		{
+			return C == 'l' || C == 'L';
+		}
+
+		/// <summary>
+		/// Test to see if the span is a numeric value
+		/// </summary>
+		/// <param name="Span">Span to test</param>
+		/// <returns>True if it is</returns>
+		public static bool IsNumeric(ReadOnlySpan<char> Span)
+		{
+			int Index = 0;
+			char C = Span[Index];
+			if (C == '-' || C == '+')
+			{
+				++Index;
+			}
+
+			bool bHasDot = false;
+			for (; Index < Span.Length; ++Index)
+			{
+				C = Span[Index];
+				if (C == '.')
+				{
+					if (bHasDot)
+					{
+						return false;
+					}
+					bHasDot = true;
+				}
+				else if (!IsDigit(C))
+				{
+					return false;
+				}
+			}
+			return true;
+		}
+
+		/// <summary>
+		/// Test to see if the character is a linebreak character
+		/// </summary>
+		/// <param name="C">Character to test</param>
+		/// <returns>True if it is</returns>
+		public static bool IsLinebreak(char C)
+		{
+			return (C >= 0x0a && C <= 0x0d) || C == 0x85 || C == 0x2028 || C == 0x2029;
+		}
+
+		/// <summary>
+		/// Return an unescaped string
+		/// </summary>
+		/// <param name="Text">Text to unescape</param>
+		/// <returns>Resulting string</returns>
+		public static string UnescapeText(string Text)
+		{
+			StringBuilder Result = new StringBuilder();
+			for (int Idx = 0; Idx < Text.Length; Idx++)
+			{
+				if (Text[Idx] == '\\' && Idx + 1 < Text.Length)
+				{
+					switch (Text[++Idx])
+					{
+						case 't':
+							Result.Append('\t');
+							break;
+						case 'r':
+							Result.Append('\r');
+							break;
+						case 'n':
+							Result.Append('\n');
+							break;
+						case '\'':
+							Result.Append('\'');
+							break;
+						case '\"':
+							Result.Append('\"');
+							break;
+						case 'u':
+							bool bParsed = false;
+							if (Idx + 4 < Text.Length)
+							{
+								int Value;
+								if (int.TryParse(Text.Substring(Idx + 1, 4), System.Globalization.NumberStyles.HexNumber, null, out Value))
+								{
+									bParsed = true;
+									Result.Append((char)Value);
+									Idx += 4;
+								}
+							}
+							if (!bParsed)
+							{
+								Result.Append('\\');
+								Result.Append('u');
+							}
+							break;
+						default:
+							Result.Append(Text[Idx]);
+							break;
+					}
+				}
+				else
+				{
+					Result.Append(Text[Idx]);
+				}
+			}
+			return Result.ToString();
+		}
+
+		/// <summary>
+		/// Words that are considered articles when parsing comment strings
+		/// Some words are always forced lowercase
+		/// </summary>
+		private static string[] Articles = new string[]
+		{
+			"In",
+			"As",
+			"To",
+			"Or",
+			"At",
+			"On",
+			"If",
+			"Be",
+			"By",
+			"The",
+			"For",
+			"And",
+			"With",
+			"When",
+			"From",
+		};
+
+		/// <summary>
+		/// Convert a name to a display string
+		/// </summary>
+		/// <param name="Name">Name to convert</param>
+		/// <param name="bIsBool">True if the name is from a boolean property</param>
+		/// <returns>Display string</returns>
+		public static string NameToDisplayString(StringView Name, bool bIsBool)
+		{
+			ReadOnlySpan<char> Chars = Name.Span;
+
+			// This is used to indicate that we are in a run of uppercase letter and/or digits.  The code attempts to keep
+			// these characters together as breaking them up often looks silly (i.e. "Draw Scale 3 D" as opposed to "Draw Scale 3D"
+			bool bInARun = false;
+			bool bWasSpace = false;
+			bool bWasOpenParen = false;
+			bool bWasNumber = false;
+			bool bWasMinusSign = false;
+
+			StringBuilder Builder = new StringBuilder();
+			for (int CharIndex = 0; CharIndex < Chars.Length; ++CharIndex)
+			{
+				char ch = Chars[CharIndex];
+
+				bool bLowerCase = Char.IsLower(ch);
+				bool bUpperCase = Char.IsUpper(ch);
+				bool bIsDigit = Char.IsDigit(ch);
+				bool bIsUnderscore = ch == '_';
+
+				// Skip the first character if the property is a bool (they should all start with a lowercase 'b', which we don't want to keep)
+				if (CharIndex == 0 && bIsBool && ch == 'b')
+				{
+					// Check if next character is uppercase as it may be a user created string that doesn't follow the rules of Unreal variables
+					if (Chars.Length > 1 && Char.IsUpper(Chars[1]))
+					{
+						continue;
+					}
+				}
+
+				// If the current character is upper case or a digit, and the previous character wasn't, then we need to insert a space if there wasn't one previously
+				// We don't do this for numerical expressions, for example "-1.2" should not be formatted as "- 1. 2"
+				if ((bUpperCase || (bIsDigit && !bWasMinusSign)) && !bInARun && !bWasOpenParen && !bWasNumber)
+				{
+					if (!bWasSpace && Builder.Length > 0)
+					{
+						Builder.Append(' ');
+						bWasSpace = true;
+					}
+					bInARun = true;
+				}
+
+				// A lower case character will break a run of upper case letters and/or digits
+				if (bLowerCase)
+				{
+					bInARun = false;
+				}
+
+				// An underscore denotes a space, so replace it and continue the run
+				if (bIsUnderscore)
+				{
+					ch = ' ';
+					bInARun = true;
+				}
+
+				// If this is the first character in the string, then it will always be upper-case
+				if (Builder.Length == 0)
+				{
+					ch = Char.ToUpper(ch);
+				}
+				else if (!bIsDigit && (bWasSpace || bWasOpenParen)) // If this is first character after a space, then make sure it is case-correct
+				{
+
+					// Search for a word that needs case repaired
+					bool bIsArticle = false;
+					foreach (string Article in Articles)
+					{
+						// Make sure the character following the string we're testing is not lowercase (we don't want to match "in" with "instance")
+						int ArticleLength = Article.Length;
+						if (CharIndex + ArticleLength < Chars.Length && !Char.IsLower(Chars[CharIndex + ArticleLength]))
+						{
+							// Does this match the current article?
+							if (Chars.Slice(CharIndex, ArticleLength).Equals(Article.AsSpan(), StringComparison.OrdinalIgnoreCase))
+							{
+								bIsArticle = true;
+								break;
+							}
+						}
+					}
+
+					// Start of a keyword, force to lowercase
+					if (bIsArticle)
+					{
+						ch = Char.ToLower(ch);
+					}
+					else    // First character after a space that's not a reserved keyword, make sure it's uppercase
+					{
+						ch = Char.ToUpper(ch);
+					}
+				}
+
+				bWasSpace = ch == ' ';
+				bWasOpenParen = ch == '(';
+
+				// What could be included as part of a numerical representation.
+				// For example -1.2
+				bWasMinusSign = ch == '-';
+				bool bPotentialNumericalChar = bWasMinusSign || ch == '.';
+				bWasNumber = bIsDigit || (bWasNumber && bPotentialNumericalChar);
+				Builder.Append(ch);
+			}
+			return Builder.ToString();
+		}
+
+		/// <summary>
+		/// Replace tabs to spaces in a string containing only a single line.
+		/// </summary>
+		/// <param name="Input">Input string</param>
+		/// <param name="TabSpacing">Number of spaces to exchange for tabs</param>
+		/// <param name="bEmulateCrBug">Due to a bug in UE ConvertTabsToSpacesInline, any \n is considered part of the line length.</param>
+		/// <returns>Resulting string or the original string if the string didn't contain any spaces.</returns>
+		public static ReadOnlyMemory<char> TabsToSpaces(ReadOnlyMemory<char> Input, int TabSpacing, bool bEmulateCrBug)
+		{
+			// If we have any tab characters, then we need to convert them to spaces
+			int TabIndex = Input.Span.IndexOf('\t');
+			if (TabIndex == -1)
+			{
+				return Input;
+			}
+			return TabsToSpacesInternal(Input.Span, TabSpacing, bEmulateCrBug, TabIndex).AsMemory();
+		}
+
+		/// <summary>
+		/// Replace tabs to spaces in a string containing zero or more lines.
+		/// </summary>
+		/// <param name="Input">Input string</param>
+		/// <param name="TabSpacing">Number of spaces to exchange for tabs</param>
+		/// <param name="bEmulateCrBug">Due to a bug in UE ConvertTabsToSpacesInline, any \n is considered part of the line length.</param>
+		/// <returns>Resulting string or the original string if the string didn't contain any spaces.</returns>
+		public static string TabsToSpaces(String Input, int TabSpacing, bool bEmulateCrBug)
+		{
+			// If we have any tab characters, then we need to convert them to spaces
+			int TabIndex = Input.IndexOf('\t');
+			if (TabIndex == -1)
+			{
+				return Input;
+			}
+			return TabsToSpacesInternal(Input.AsSpan(), TabSpacing, bEmulateCrBug, TabIndex);
+		}
+
+		/// <summary>
+		/// Replace tabs to spaces in a string containing zero or more lines.
+		/// </summary>
+		/// <param name="Span">Input string to convert</param>
+		/// <param name="TabSpacing">Number of spaces to exchange for tabs</param>
+		/// <param name="bEmulateCrBug">Due to a bug in UE ConvertTabsToSpacesInline, any \n is considered part of the line length.</param>
+		/// <param name="TabIndex">Initial tab index</param>
+		/// <returns>Resulting string or the original string if the string didn't contain any spaces.</returns>
+		private static string TabsToSpacesInternal(ReadOnlySpan<char> Span, int TabSpacing, bool bEmulateCrBug, int TabIndex)
+		{
+			// Locate the last \n since all tabs have to be computed relative to this.
+			int CrPos = Span.Slice(0, TabIndex).LastIndexOf('\n') + 1;
+
+			using (BorrowStringBuilder Borrower = new BorrowStringBuilder(StringBuilderCache.Small))
+			{
+				StringBuilder Builder = Borrower.StringBuilder;
+				int Committed = 0;
+				do
+				{
+					// Commit everything prior to the tab to the builder
+					Builder.Append(Span.Slice(0, TabIndex));
+					Span = Span.Slice(TabIndex + 1);
+					Committed += TabIndex;
+
+					// Add the appropriate number of spaces
+					int AdjustedCrPos = CrPos;
+					if (bEmulateCrBug && AdjustedCrPos > 0)
+					{
+						--AdjustedCrPos;
+					}
+					int SpacesToInsert = TabSpacing - (Committed - AdjustedCrPos) % TabSpacing;
+					Builder.AppendSpaces(SpacesToInsert);
+					Committed += SpacesToInsert;
+
+					// Search for the next \t or \n
+					for (TabIndex = 0; TabIndex < Span.Length; ++TabIndex)
+					{
+						if (Span[TabIndex] == '\n')
+						{
+							CrPos = Committed + TabIndex + 1;
+						}
+						else if (Span[TabIndex] == '\t')
+						{
+							break;
+						}
+					}
+				} while (TabIndex < Span.Length);
+
+				// Commit the remaining data
+				Builder.Append(Span);
+				return Builder.ToString();
+			}
+		}
+	}
+}
