UHT in UBT initial version.  Currently disabled.

diff --git a/Engine/Source/Programs/Shared/EpicGames.UHT/Utils/UhtFileManager.cs b/Engine/Source/Programs/Shared/EpicGames.UHT/Utils/UhtFileManager.cs
new file mode 100644
--- /dev/null
+++ b/Engine/Source/Programs/Shared/EpicGames.UHT/Utils/UhtFileManager.cs
@@ -0,0 +1,219 @@
+// Copyright Epic Games, Inc. All Rights Reserved.
+
+using EpicGames.Core;
+using System;
+using System.IO;
+using System.Linq;
+using System.Text;
+
+namespace EpicGames.UHT.Utils
+{
+
+	/// <summary>
+	/// Interface used to read/write files
+	/// </summary>
+	public interface IUhtFileManager
+	{
+		/// <summary>
+		/// Return the full file path for a partial path
+		/// </summary>
+		/// <param name="FilePath">The partial file path</param>
+		/// <returns>The full file path</returns>
+		public string GetFullFilePath(string FilePath);
+
+		/// <summary>
+		/// Read the given source file
+		/// </summary>
+		/// <param name="FilePath">File path</param>
+		/// <param name="Fragment">Read fragment information</param>
+		/// <returns>True if the file was read</returns>
+		public bool ReadSource(string FilePath, out UhtSourceFragment Fragment);
+
+		/// <summary>
+		/// Read the given source file
+		/// </summary>
+		/// <param name="FilePath">File path</param>
+		/// <returns>Buffer containing the read data or null if not found.  The returned buffer must be returned to the cache via a call to UhtBuffer.Return</returns>
+		public UhtBuffer? ReadOutput(string FilePath);
+
+		/// <summary>
+		/// Write the given contents to the file
+		/// </summary>
+		/// <param name="FilePath">Path to write to</param>
+		/// <param name="Contents">Contents to write</param>
+		/// <returns>True if the source was written</returns>
+		public bool WriteOutput(string FilePath, ReadOnlySpan<char> Contents);
+
+		/// <summary>
+		/// Rename the given file
+		/// </summary>
+		/// <param name="OldFilePath">Old file path name</param>
+		/// <param name="NewFilePath">New file path name</param>
+		/// <returns>True if the file was renamed</returns>
+		public bool RenameOutput(string OldFilePath, string NewFilePath);
+	}
+
+	/// <summary>
+	/// Implementation of a file manager that reads/writes from disk
+	/// </summary>
+	public class UhtStdFileManager : IUhtFileManager
+	{
+
+		/// <summary>
+		/// Construct a new file manager
+		/// </summary>
+		public UhtStdFileManager()
+		{
+		}
+
+		/// <inheritdoc/>
+		public string GetFullFilePath(string FilePath)
+		{
+			return FilePath;
+		}
+
+		/// <inheritdoc/>
+		public bool ReadSource(string FilePath, out UhtSourceFragment Fragment)
+		{
+			if (ReadFile(FilePath, out StringView Data))
+			{
+				Fragment = new UhtSourceFragment { SourceFile = null, FilePath = FilePath, LineNumber = 0, Data = Data };
+				return true;
+			}
+			Fragment = new UhtSourceFragment { SourceFile = null, FilePath = String.Empty, LineNumber = 0, Data = new StringView() };
+			return false;
+		}
+
+		/// <inheritdoc/>
+		public UhtBuffer? ReadOutput(string FilePath)
+		{
+			// Exceptions are very expensive.  Don't bother trying to open the file if it doesn't exist
+			if (!File.Exists(FilePath))
+			{
+				return null;
+			}
+
+			try
+			{
+				using (FileStream fs = new FileStream(FilePath, FileMode.Open, FileAccess.Read, FileShare.Read, 4 * 1024, FileOptions.SequentialScan))
+				{
+					using (StreamReader sr = new StreamReader(fs, Encoding.UTF8, true))
+					{
+
+						// Try to read the whole file into a buffer created by hand.  This avoids a LOT of memory allocations which in turn reduces the
+						// GC stress on the system.  Removing the StreamReader would be nice in the future.
+						long RawFileLength = fs.Length + 32;
+						UhtBuffer InitialBuffer = UhtBuffer.Borrow((int)RawFileLength);
+						int ReadLength = sr.Read(InitialBuffer.Memory.Span);
+						if (sr.EndOfStream)
+						{
+							InitialBuffer.Reset(ReadLength);
+							return InitialBuffer;
+						}
+						else
+						{
+							string Remaining = sr.ReadToEnd();
+							long TotalSize = ReadLength + Remaining.Length;
+							UhtBuffer Combined = UhtBuffer.Borrow((int)TotalSize);
+							Buffer.BlockCopy(InitialBuffer.Block, 0, Combined.Block, 0, ReadLength * sizeof(char));
+							Buffer.BlockCopy(Remaining.ToArray(), 0, Combined.Block, ReadLength * sizeof(char), Remaining.Length * sizeof(char));
+							UhtBuffer.Return(InitialBuffer);
+							return Combined;
+						}
+					}
+				}
+			}
+			catch (IOException)
+			{
+				return null;
+			}
+		}
+
+		/// <inheritdoc/>
+		public bool WriteOutput(string FilePath, ReadOnlySpan<char> Contents)
+		{
+			try
+			{
+				string? FileDirectory = Path.GetDirectoryName(FilePath);
+				if (!string.IsNullOrEmpty(FileDirectory))
+				{
+					Directory.CreateDirectory(FileDirectory);
+				}
+				using (StreamWriter Writer = new StreamWriter(FilePath, false, new UTF8Encoding(false, true), 16 * 1024))
+				{
+					Writer.Write(Contents);
+				}
+				return true;
+			}
+			catch (Exception)
+			{
+				return false;
+			}
+		}
+
+		/// <inheritdoc/>
+		public bool RenameOutput(string OldFilePath, string NewFilePath)
+		{
+			try
+			{
+				File.Move(OldFilePath, OldFilePath, true);
+				return true;
+			}
+			catch (Exception)
+			{
+				return false;
+			}
+		}
+
+		///// <summary>
+		///// Read the given source file
+		///// </summary>
+		///// <param name="FilePath">Full file path</param>
+		///// <param name="Contents">Contents of the file</param>
+		///// <returns>True if the file was read, false if not</returns>
+		private bool ReadFile(string FilePath, out StringView Contents)
+		{
+			// Exceptions are very expensive.  Don't bother trying to open the file if it doesn't exist
+			if (!File.Exists(FilePath))
+			{
+				Contents = new StringView();
+				return false;
+			}
+
+			try
+			{
+				using (FileStream fs = new FileStream(FilePath, FileMode.Open, FileAccess.Read, FileShare.Read, 4 * 1024, FileOptions.SequentialScan))
+				{
+					using (StreamReader sr = new StreamReader(fs, Encoding.UTF8, true))
+					{
+
+						// Try to read the whole file into a buffer created by hand.  This avoids a LOT of memory allocations which in turn reduces the
+						// GC stress on the system.  Removing the StreamReader would be nice in the future.
+						long RawFileLength = fs.Length + 32;
+						Memory<char> InitialBuffer = new Memory<char>(new char[RawFileLength]);
+						int ReadLength = sr.Read(InitialBuffer.Span);
+						if (sr.EndOfStream)
+						{
+							Contents = new StringView(InitialBuffer, 0, ReadLength);
+						}
+						else
+						{
+							string Remaining = sr.ReadToEnd();
+							long TotalSize = ReadLength + Remaining.Length;
+							Memory<char> Combined = new Memory<char>(new char[TotalSize]);
+							Buffer.BlockCopy(InitialBuffer.Span.ToArray(), 0, Combined.Span.ToArray(), 0, ReadLength * sizeof(char));
+							Buffer.BlockCopy(Remaining.ToArray(), 0, Combined.Span.ToArray(), ReadLength * sizeof(char), Remaining.Length * sizeof(char));
+							Contents = new StringView(Combined);
+						}
+						return true;
+					}
+				}
+			}
+			catch (IOException)
+			{
+				Contents = new StringView();
+				return false;
+			}
+		}
+	}
+}
diff --git a/Engine/Source/Programs/Shared/EpicGames.UHT/Utils/UhtHash.cs b/Engine/Source/Programs/Shared/EpicGames.UHT/Utils/UhtHash.cs
new file mode 100644
--- /dev/null
+++ b/Engine/Source/Programs/Shared/EpicGames.UHT/Utils/UhtHash.cs
@@ -0,0 +1,82 @@
+// Copyright Epic Games, Inc. All Rights Reserved.
+
+using EpicGames.Core;
+using System;
+
+namespace EpicGames.UHT.Utils
+{
+
+	/// <summary>
+	/// Support for generating hash values using the existing UHT algorithms
+	/// </summary>
+	public struct UhtHash
+	{
+		private ulong HashInternal;
+
+		/// <summary>
+		/// Start the hash computation
+		/// </summary>
+		public void Begin()
+		{
+			this.HashInternal = 0;
+		}
+
+		/// <summary>
+		/// Add a span of text to the current hash value
+		/// </summary>
+		/// <param name="Text"></param>
+		public void Add(ReadOnlySpan<char> Text)
+		{
+			while (Text.Length != 0)
+			{
+				int CrPos = Text.IndexOf('\r');
+				if (CrPos == -1)
+				{
+					HashText(Text.Slice(0, Text.Length));
+					break;
+				}
+				else
+				{
+					HashText(Text.Slice(0, CrPos));
+					Text = Text.Slice(CrPos + 1);
+				}
+			}
+		}
+
+		/// <summary>
+		/// Return the final hash value
+		/// </summary>
+		/// <returns>Final hash value</returns>
+		public uint End()
+		{
+			return (uint)(this.HashInternal + (this.HashInternal >> 32));
+		}
+
+		/// <summary>
+		/// Generate the hash value for a block of text
+		/// </summary>
+		/// <param name="Text">Text to hash</param>
+		/// <returns>Hash value</returns>
+		public static uint GenenerateTextHash(ReadOnlySpan<char> Text)
+		{
+			UhtHash Hash = new UhtHash();
+			Hash.Begin();
+			Hash.Add(Text);
+			return Hash.End();
+		}
+
+		private void HashText(ReadOnlySpan<char> Text)
+		{
+			if (Text.Length > 0)
+			{
+				unsafe
+				{
+					fixed (char* TextPtr = Text)
+					{
+						this.HashInternal = CityHash.CityHash64WithSeed((byte*)TextPtr, (uint)(Text.Length * sizeof(char)), this.HashInternal);
+					}
+				}
+			}
+		}
+	}
+}
diff --git a/Engine/Source/Programs/Shared/EpicGames.UHT/Utils/UhtJsonConverters.cs b/Engine/Source/Programs/Shared/EpicGames.UHT/Utils/UhtJsonConverters.cs
new file mode 100644
--- /dev/null
+++ b/Engine/Source/Programs/Shared/EpicGames.UHT/Utils/UhtJsonConverters.cs
@@ -0,0 +1,151 @@
+// Copyright Epic Games, Inc. All Rights Reserved.
+
+using EpicGames.UHT.Types;
+using System;
+using System.Collections.Generic;
+using System.Text.Json;
+using System.Text.Json.Serialization;
+
+namespace EpicGames.UHT.Utils
+{
+
+	/// <summary>
+	/// JSON converter to output the source name of a type
+	/// </summary>
+	/// <typeparam name="T"></typeparam>
+	public class UhtTypeSourceNameJsonConverter<T> : JsonConverter<T> where T : UhtType
+	{
+		/// <summary>
+		/// Read the JSON value
+		/// </summary>
+		/// <param name="Reader">Source reader</param>
+		/// <param name="TypeToConvert">Type to convert</param>
+		/// <param name="Options">Serialization options</param>
+		/// <returns>Read value</returns>
+		/// <exception cref="NotImplementedException"></exception>
+		public override T Read(ref Utf8JsonReader Reader, Type TypeToConvert, JsonSerializerOptions Options)
+		{
+			throw new NotImplementedException();
+		}
+
+		/// <summary>
+		/// Write the JSON value
+		/// </summary>
+		/// <param name="Writer">Destination writer</param>
+		/// <param name="Type">Value being written</param>
+		/// <param name="Options">Serialization options</param>
+		public override void Write(Utf8JsonWriter Writer, T Type, JsonSerializerOptions Options)
+		{
+			Writer.WriteStringValue(Type.SourceName);
+		}
+	}
+
+	/// <summary>
+	/// Read/Write type source name for an optional type value
+	/// </summary>
+	/// <typeparam name="T"></typeparam>
+	public class UhtNullableTypeSourceNameJsonConverter<T> : JsonConverter<T> where T : UhtType
+	{
+		/// <summary>
+		/// Read the JSON value
+		/// </summary>
+		/// <param name="Reader">Source reader</param>
+		/// <param name="TypeToConvert">Type to convert</param>
+		/// <param name="Options">Serialization options</param>
+		/// <returns>Read value</returns>
+		/// <exception cref="NotImplementedException"></exception>
+		public override T Read(ref Utf8JsonReader Reader, Type TypeToConvert, JsonSerializerOptions Options)
+		{
+			throw new NotImplementedException();
+		}
+
+		/// <summary>
+		/// Write the JSON value
+		/// </summary>
+		/// <param name="Writer">Destination writer</param>
+		/// <param name="Type">Value being written</param>
+		/// <param name="Options">Serialization options</param>
+		public override void Write(Utf8JsonWriter Writer, T? Type, JsonSerializerOptions Options)
+		{
+			if (Type != null)
+			{
+				Writer.WriteStringValue(Type.SourceName);
+			}
+		}
+	}
+
+	/// <summary>
+	/// Serialize a list of types
+	/// </summary>
+	/// <typeparam name="T"></typeparam>
+	public class UhtTypeListJsonConverter<T> : JsonConverter<List<T>> where T : UhtType
+	{
+		/// <summary>
+		/// Read the JSON value
+		/// </summary>
+		/// <param name="Reader">Source reader</param>
+		/// <param name="TypeToConvert">Type to convert</param>
+		/// <param name="Options">Serialization options</param>
+		/// <returns>Read value</returns>
+		/// <exception cref="NotImplementedException"></exception>
+		public override List<T> Read(ref Utf8JsonReader Reader, Type TypeToConvert, JsonSerializerOptions Options)
+		{
+			throw new NotImplementedException();
+		}
+
+		/// <summary>
+		/// Write the JSON value
+		/// </summary>
+		/// <param name="Writer">Destination writer</param>
+		/// <param name="Collection">Value being written</param>
+		/// <param name="Options">Serialization options</param>
+		public override void Write(Utf8JsonWriter Writer, List<T> Collection, JsonSerializerOptions Options)
+		{
+			Writer.WriteStartArray();
+			foreach (UhtType Type in Collection)
+			{
+				JsonSerializer.Serialize(Writer, Type, Type.GetType(), Options);
+			}
+			Writer.WriteEndArray();
+		}
+	}
+
+	/// <summary>
+	/// Serialize a nullable list of types
+	/// </summary>
+	/// <typeparam name="T"></typeparam>
+	public class UhtNullableTypeListJsonConverter<T> : JsonConverter<List<T>> where T : UhtType
+	{
+		/// <summary>
+		/// Read the JSON value
+		/// </summary>
+		/// <param name="Reader">Source reader</param>
+		/// <param name="TypeToConvert">Type to convert</param>
+		/// <param name="Options">Serialization options</param>
+		/// <returns>Read value</returns>
+		/// <exception cref="NotImplementedException"></exception>
+		public override List<T> Read(ref Utf8JsonReader Reader, Type TypeToConvert, JsonSerializerOptions Options)
+		{
+			throw new NotImplementedException();
+		}
+
+		/// <summary>
+		/// Write the JSON value
+		/// </summary>
+		/// <param name="Writer">Destination writer</param>
+		/// <param name="Collection">Value being written</param>
+		/// <param name="Options">Serialization options</param>
+		public override void Write(Utf8JsonWriter Writer, List<T>? Collection, JsonSerializerOptions Options)
+		{
+			Writer.WriteStartArray();
+			if (Collection != null)
+			{
+				foreach (UhtType Type in Collection)
+				{
+					JsonSerializer.Serialize(Writer, Type, Type.GetType(), Options);
+				}
+			}
+			Writer.WriteEndArray();
+		}
+	}
+}
diff --git a/Engine/Source/Programs/Shared/EpicGames.UHT/Utils/UhtMessage.cs b/Engine/Source/Programs/Shared/EpicGames.UHT/Utils/UhtMessage.cs
new file mode 100644
--- /dev/null
+++ b/Engine/Source/Programs/Shared/EpicGames.UHT/Utils/UhtMessage.cs
@@ -0,0 +1,579 @@
+// Copyright Epic Games, Inc. All Rights Reserved.
+
+using EpicGames.UHT.Tokenizer;
+using System;
+using System.Collections.Generic;
+using System.Text;
+using System.Threading;
+
+namespace EpicGames.UHT.Utils
+{
+
+	/// <summary>
+	/// Type of message
+	/// </summary>
+	public enum UhtMessageType
+	{
+		/// <summary>
+		/// The message is an error and goes to the log and console
+		/// </summary>
+		Error,
+
+		/// <summary>
+		/// The message is a warning and goes to the log and console
+		/// </summary>
+		Warning,
+
+		/// <summary>
+		/// The message is for information only and goes to the log and console
+		/// </summary>
+		Info,
+
+		/// <summary>
+		/// The message is for debugging and goes to the log only
+		/// </summary>
+		Trace,
+
+		/// <summary>
+		/// The message is an internal error and goes to the log and console
+		/// </summary>
+		Ice,
+	}
+
+	/// <summary>
+	/// A message session is the destination object for all generated messages
+	/// </summary>
+	public interface IUhtMessageSession
+	{
+
+		/// <summary>
+		/// Add the given message
+		/// </summary>
+		/// <param name="Message">The message to be added</param>
+		void AddMessage(UhtMessage Message);
+	}
+
+	/// <summary>
+	/// A message source represents the source file where the message occurred.
+	/// </summary>
+	public interface IUhtMessageSource
+	{
+		/// <summary>
+		/// File path of the file being parsed
+		/// </summary>
+		string MessageFilePath { get; }
+
+		/// <summary>
+		/// The full file path of being parsed
+		/// </summary>
+		string MessageFullFilePath { get; }
+
+		/// <summary>
+		/// If true, the source is a source fragment from the testing harness
+		/// </summary>
+		bool bMessageIsFragment { get; }
+
+		/// <summary>
+		/// If this is a fragment, this is the container file path of the fragment
+		/// </summary>
+		string MessageFragmentFilePath { get; }
+
+		/// <summary>
+		/// If this is a fragment, this is the container full file path of the fragment
+		/// </summary>
+		string MessageFragmentFullFilePath { get; }
+
+		/// <summary>
+		/// If this is a fragment, this is the line number in the container file where the fragment is defined.
+		/// </summary>
+		int MessageFragmentLineNumber { get; }
+	}
+
+	/// <summary>
+	/// A message site can automatically provide a line number where the site was defined
+	/// in the source.  If no line number is provided when the message is created or if the
+	/// site doesn't support this interface, the line number of '1' will be used.
+	/// </summary>
+	public interface IUhtMessageLineNumber
+	{
+
+		/// <summary>
+		/// Line number where the type was defined.
+		/// </summary>
+		int MessageLineNumber { get; }
+	}
+
+	/// <summary>
+	/// This interface provides a mechanism for things to provide more context for an error
+	/// </summary>
+	public interface IUhtMessageExtraContext
+	{
+		/// <summary>
+		/// Enumeration of objects to add as extra context.
+		/// </summary>
+		IEnumerable<object?>? MessageExtraContext { get; }
+	}
+
+	/// <summary>
+	/// A message site is any object that can generate a message.  In general, all 
+	/// types are also message sites. This provides a convenient method to log messages
+	/// where the type was defined.
+	/// </summary>
+	public interface IUhtMessageSite
+	{
+
+		/// <summary>
+		/// Destination message session for the messages
+		/// </summary>
+		public IUhtMessageSession MessageSession { get; }
+		
+		/// <summary>
+		/// Source file generating messages
+		/// </summary>
+		public IUhtMessageSource? MessageSource { get; }
+
+		/// <summary>
+		/// Optional line number where type was defined
+		/// </summary>
+		public IUhtMessageLineNumber? MessageLineNumber { get; }
+	}
+
+	/// <summary>
+	/// Represents a UHT message
+	/// </summary>
+	public struct UhtMessage
+	{
+		/// <summary>
+		/// The type of message
+		/// </summary>
+		public UhtMessageType MessageType;
+
+		/// <summary>
+		/// Optional message source for the message.  Either the MessageSource or FilePath must be set.
+		/// </summary>
+		public IUhtMessageSource? MessageSource;
+
+		/// <summary>
+		/// Optional file path for the message.  Either the MessageSource or FilePath must be set.
+		/// </summary>
+		public string? FilePath;
+
+		/// <summary>
+		/// Line number where error occurred.
+		/// </summary>
+		public int LineNumber;
+
+		/// <summary>
+		/// Text of the message
+		/// </summary>
+		public string Message;
+
+		/// <summary>
+		/// Make a new message with the given settings
+		/// </summary>
+		/// <param name="MessageType">Type of message</param>
+		/// <param name="MessageSource">Source of the message</param>
+		/// <param name="FilePath">File path of the message</param>
+		/// <param name="LineNumber">Line number where message occurred</param>
+		/// <param name="Message">Text of the message</param>
+		/// <returns>Created message</returns>
+		public static UhtMessage MakeMessage(UhtMessageType MessageType, IUhtMessageSource? MessageSource, string? FilePath, int LineNumber, string Message)
+		{
+			return new UhtMessage
+			{
+				MessageType = MessageType,
+				MessageSource = MessageSource,
+				FilePath = FilePath,
+				LineNumber = LineNumber,
+				Message = Message
+			};
+		}
+
+		/// <summary>
+		/// Format an object to be included in a message
+		/// </summary>
+		/// <param name="Context">Contextual object</param>
+		/// <returns>Formatted context</returns>
+		public static string FormatContext(object Context)
+		{
+			if (Context is IUhtMessageExtraContext ExtraContextInterface)
+			{
+				StringBuilder Builder = new StringBuilder();
+				Append(Builder, ExtraContextInterface, false);
+				return Builder.ToString();
+			}
+			else if (Context is UhtToken Token)
+			{
+				switch (Token.TokenType)
+				{
+					case UhtTokenType.EndOfFile:
+						return "EOF";
+					case UhtTokenType.EndOfDefault:
+						return "'end of default value'";
+					case UhtTokenType.EndOfType:
+						return "'end of type'";
+					case UhtTokenType.EndOfDeclaration:
+						return "'end of declaration'";
+					case UhtTokenType.StringConst:
+						return "string constant";
+					default:
+						return $"'{Token.Value}'";
+				}
+			}
+			else if (Context is char C)
+			{
+				return $"'{C}'";
+			}
+			else if (Context is string[] StringArray)
+			{
+				return UhtUtilities.MergeTypeNames(StringArray, "or", true);
+			}
+			else
+			{
+				return Context.ToString() ?? string.Empty;
+			}
+		}
+
+		/// <summary>
+		/// Append message context
+		/// </summary>
+		/// <param name="Builder">Destination builder</param>
+		/// <param name="MessageExtraContextInterface">Extra context to append</param>
+		/// <param name="bAlwaysIncludeSeparator">If true, always include the separator</param>
+		public static void Append(StringBuilder Builder, IUhtMessageExtraContext? MessageExtraContextInterface, bool bAlwaysIncludeSeparator)
+		{
+			if (MessageExtraContextInterface == null)
+			{
+				return;
+			}
+
+			IEnumerable<object?>? ExtraContextList = MessageExtraContextInterface.MessageExtraContext;
+			if (ExtraContextList != null)
+			{
+				int StartingLength = Builder.Length;
+				foreach (object? EC in ExtraContextList)
+				{
+					if (EC != null)
+					{
+						if (Builder.Length != StartingLength || bAlwaysIncludeSeparator)
+						{
+							Builder.Append(" in ");
+						}
+						Builder.Append(UhtMessage.FormatContext(EC));
+					}
+				}
+			}
+		}
+	}
+
+	/// <summary>
+	/// A placeholder message site
+	/// </summary>
+	public class UhtEmptyMessageSite : IUhtMessageSite
+	{
+		/// <inheritdoc/>
+		public IUhtMessageSession MessageSession => throw new NotImplementedException();
+
+		/// <inheritdoc/>
+		public IUhtMessageSource? MessageSource => throw new NotImplementedException();
+
+		/// <inheritdoc/>
+		public IUhtMessageLineNumber? MessageLineNumber => throw new NotImplementedException();
+	}
+
+	/// <summary>
+	/// Creates a message site from a message session interface and a message source interface
+	/// </summary>
+	public class UhtSimpleMessageSite : IUhtMessageSite
+	{
+		private readonly IUhtMessageSession MessageSessionInternal;
+		private IUhtMessageSource? MessageSourceInternal;
+
+		#region IUHTMessageSite implementation
+		/// <inheritdoc/>
+		public IUhtMessageSession MessageSession => this.MessageSessionInternal;
+		/// <inheritdoc/>
+		public IUhtMessageSource? MessageSource { get => this.MessageSourceInternal; set => this.MessageSourceInternal = value; }
+		/// <inheritdoc/>
+		public IUhtMessageLineNumber? MessageLineNumber => null;
+		#endregion
+
+		/// <summary>
+		/// Create a simple message site for the given session and source
+		/// </summary>
+		/// <param name="MessageSession">Associated message session</param>
+		/// <param name="MessageSource">Source for the messages</param>
+		public UhtSimpleMessageSite(IUhtMessageSession MessageSession, IUhtMessageSource? MessageSource = null)
+		{
+			this.MessageSessionInternal = MessageSession;
+			this.MessageSourceInternal = MessageSource;
+		}
+	}
+
+	/// <summary>
+	/// Simple message site for the given file.
+	/// </summary>
+	public class UhtSimpleFileMessageSite : UhtSimpleMessageSite, IUhtMessageSource
+	{
+		#region IUHTMessageSource implementation
+		/// <inheritdoc/>
+		public string MessageFilePath { get => this.FilePath; }
+		/// <inheritdoc/>
+		public string MessageFullFilePath { get => this.FilePath; }
+		/// <inheritdoc/>
+		public bool bMessageIsFragment { get => false; }
+		/// <inheritdoc/>
+		public string MessageFragmentFilePath { get => ""; }
+		/// <inheritdoc/>
+		public string MessageFragmentFullFilePath { get => ""; }
+		/// <inheritdoc/>
+		public int MessageFragmentLineNumber { get => -1; }
+		#endregion
+
+		private string FilePath;
+
+		/// <summary>
+		/// Create a simple file site
+		/// </summary>
+		/// <param name="MessageSession">Associated message session</param>
+		/// <param name="FilePath">File associated with the site</param>
+		public UhtSimpleFileMessageSite(IUhtMessageSession MessageSession, string FilePath) : base(MessageSession, null)
+		{
+			this.FilePath = FilePath;
+			this.MessageSource = this;
+		}
+	}
+
+	/// <summary>
+	/// Series of extensions for message sites.
+	/// </summary>
+	public static class UhtMessageSiteExtensions
+	{
+		/// <summary>
+		/// Get the line number generating the error
+		/// </summary>
+		/// <param name="MessageSite">The message site generating the error.</param>
+		/// <param name="LineNumber">An override line number</param>
+		/// <returns>Either the overriding line number or the line number from the message site.</returns>
+		public static int GetLineNumber(this IUhtMessageSite MessageSite, int LineNumber = -1)
+		{
+			if (LineNumber != -1)
+			{
+				return LineNumber;
+			}
+			else if (MessageSite.MessageLineNumber != null)
+			{
+				return MessageSite.MessageLineNumber.MessageLineNumber;
+			}
+			else
+			{
+				return 1;
+			}
+		}
+
+		/// <summary>
+		/// Log an error
+		/// </summary>
+		/// <param name="MessageSite">Message site associated with the message</param>
+		/// <param name="LineNumber">Line number of the error</param>
+		/// <param name="Message">Text of the error</param>
+		/// <param name="ExtraContext">Addition context to be appended to the error message</param>
+		public static void LogError(this IUhtMessageSite MessageSite, int LineNumber, string Message, object? ExtraContext = null)
+		{
+			LogMessage(UhtMessageType.Error, MessageSite, LineNumber, Message, ExtraContext);
+		}
+
+		/// <summary>
+		/// Log an error
+		/// </summary>
+		/// <param name="MessageSite">Message site associated with the message</param>
+		/// <param name="Message">Text of the error</param>
+		/// <param name="ExtraContext">Addition context to be appended to the error message</param>
+		public static void LogError(this IUhtMessageSite MessageSite, string Message, object? ExtraContext = null)
+		{
+			LogMessage(UhtMessageType.Error, MessageSite, -1, Message, ExtraContext);
+		}
+
+		/// <summary>
+		/// Log a warning
+		/// </summary>
+		/// <param name="MessageSite">Message site associated with the message</param>
+		/// <param name="LineNumber">Line number of the warning</param>
+		/// <param name="Message">Text of the warning</param>
+		/// <param name="ExtraContext">Addition context to be appended to the error message</param>
+		public static void LogWarning(this IUhtMessageSite MessageSite, int LineNumber, string Message, object? ExtraContext = null)
+		{
+			LogMessage(UhtMessageType.Warning, MessageSite, LineNumber, Message, ExtraContext);
+		}
+
+		/// <summary>
+		/// Log a warning
+		/// </summary>
+		/// <param name="MessageSite">Message site associated with the message</param>
+		/// <param name="Message">Text of the warning</param>
+		/// <param name="ExtraContext">Addition context to be appended to the error message</param>
+		public static void LogWarning(this IUhtMessageSite MessageSite, string Message, object? ExtraContext = null)
+		{
+			LogMessage(UhtMessageType.Warning, MessageSite, -1, Message, ExtraContext);
+		}
+
+		/// <summary>
+		/// Log information
+		/// </summary>
+		/// <param name="MessageSite">Message site associated with the message</param>
+		/// <param name="LineNumber">Line number of the information</param>
+		/// <param name="Message">Text of the information</param>
+		/// <param name="ExtraContext">Addition context to be appended to the error message</param>
+		public static void LogInfo(this IUhtMessageSite MessageSite, int LineNumber, string Message, object? ExtraContext = null)
+		{
+			LogMessage(UhtMessageType.Info, MessageSite, LineNumber, Message, ExtraContext);
+		}
+
+		/// <summary>
+		/// Log a message directly to the log
+		/// </summary>
+		/// <param name="MessageSite">Message site associated with the message</param>
+		/// <param name="Message">Text of the information</param>
+		/// <param name="ExtraContext">Addition context to be appended to the error message</param>
+		public static void LogTrace(this IUhtMessageSite MessageSite, string Message, object? ExtraContext = null)
+		{
+			LogMessage(UhtMessageType.Trace, MessageSite, -1, Message, ExtraContext);
+		}
+
+		/// <summary>
+		/// Log a message directly to the log
+		/// </summary>
+		/// <param name="MessageSite">Message site associated with the message</param>
+		/// <param name="LineNumber">Line number of the information</param>
+		/// <param name="Message">Text of the information</param>
+		/// <param name="ExtraContext">Addition context to be appended to the error message</param>
+		public static void LogLog(this IUhtMessageSite MessageSite, int LineNumber, string Message, object? ExtraContext = null)
+		{
+			LogMessage(UhtMessageType.Trace, MessageSite, LineNumber, Message, ExtraContext);
+		}
+
+		/// <summary>
+		/// Log a information
+		/// </summary>
+		/// <param name="MessageSite">Message site associated with the message</param>
+		/// <param name="Message">Text of the information</param>
+		/// <param name="ExtraContext">Addition context to be appended to the error message</param>
+		public static void LogInfo(this IUhtMessageSite MessageSite, string Message, object? ExtraContext = null)
+		{
+			LogMessage(UhtMessageType.Info, MessageSite, -1, Message, ExtraContext);
+		}
+
+		/// <summary>
+		/// Log a message
+		/// </summary>
+		/// <param name="MessageType">Type of the message being generated</param>
+		/// <param name="MessageSite">Message site associated with the message</param>
+		/// <param name="LineNumber">Line number of the information</param>
+		/// <param name="Message">Text of the information</param>
+		/// <param name="ExtraContext">Addition context to be appended to the error message</param>
+		private static void LogMessage(UhtMessageType MessageType, IUhtMessageSite MessageSite, int LineNumber, string Message, object? ExtraContext)
+		{
+			if (ExtraContext != null)
+			{
+				Message = $"{Message} in {UhtMessage.FormatContext(ExtraContext)}";
+			}
+			MessageSite.MessageSession.AddMessage(UhtMessage.MakeMessage(MessageType, MessageSite.MessageSource, null, MessageSite.GetLineNumber(LineNumber), Message));
+		}
+	}
+
+	/// <summary>
+	/// Thread based message context.  Used to improve performance by avoiding allocations.
+	/// </summary>
+	public class UhtTlsMessageExtraContext : IUhtMessageExtraContext
+	{
+		private Stack<object?>? ExtraContexts;
+		private static ThreadLocal<UhtTlsMessageExtraContext> Tls = new ThreadLocal<UhtTlsMessageExtraContext>(() => new UhtTlsMessageExtraContext());
+
+		#region IUHTMessageExtraContext implementation
+		IEnumerable<object?>? IUhtMessageExtraContext.MessageExtraContext => this.ExtraContexts;
+		#endregion
+
+		/// <summary>
+		/// Add an extra context
+		/// </summary>
+		/// <param name="ExceptionContext"></param>
+		public void PushExtraContext(object? ExceptionContext)
+		{
+			if (this.ExtraContexts == null)
+			{
+				this.ExtraContexts = new Stack<object?>(8);
+			}
+			this.ExtraContexts.Push(ExceptionContext);
+		}
+
+		/// <summary>
+		/// Pop the top most extra context
+		/// </summary>
+		public void PopExtraContext()
+		{
+			if (this.ExtraContexts != null)
+			{
+				this.ExtraContexts.Pop();
+			}
+		}
+
+		/// <summary>
+		/// Get the extra context associated with this thread
+		/// </summary>
+		/// <returns>Extra context</returns>
+		public static UhtTlsMessageExtraContext? GetTls() { return UhtTlsMessageExtraContext.Tls.Value; }
+
+		/// <summary>
+		/// Get the extra context interface
+		/// </summary>
+		/// <returns>Extra context interface</returns>
+		public static IUhtMessageExtraContext? GetMessageExtraContext() { return UhtTlsMessageExtraContext.Tls.Value; }
+	}
+
+	/// <summary>
+	/// A "using" object to automate the push/pop of extra context to the thread's current extra context
+	/// </summary>
+	public struct UhtMessageContext : IDisposable
+	{
+		private readonly UhtTlsMessageExtraContext? Stack;
+
+		/// <summary>
+		/// Construct a new entry
+		/// </summary>
+		/// <param name="MessageSite">Destination message site</param>
+		/// <param name="ExtraContext">Extra context to be added</param>
+		public UhtMessageContext(IUhtMessageSite MessageSite, object? ExtraContext)
+		{
+			this.Stack = UhtTlsMessageExtraContext.GetTls();
+			if (this.Stack != null)
+			{
+				this.Stack.PushExtraContext(ExtraContext);
+			}
+		}
+
+		/// <summary>
+		/// Replace the extra context.  This replaces the top level context and thus 
+		/// can have unexpected results if done when a more deeper context has been added
+		/// </summary>
+		/// <param name="ExtraContext">New extra context</param>
+		public void Reset(object? ExtraContext)
+		{
+			if (this.Stack != null)
+			{
+				this.Stack.PopExtraContext();
+				this.Stack.PushExtraContext(ExtraContext);
+			}
+		}
+
+		/// <summary>
+		/// Dispose the object and auto-remove the added context
+		/// </summary>
+		public void Dispose()
+		{
+			if (this.Stack != null)
+			{
+				this.Stack.PopExtraContext();
+			}
+		}
+	}
+}
diff --git a/Engine/Source/Programs/Shared/EpicGames.UHT/Utils/UhtReferenceCollector.cs b/Engine/Source/Programs/Shared/EpicGames.UHT/Utils/UhtReferenceCollector.cs
new file mode 100644
--- /dev/null
+++ b/Engine/Source/Programs/Shared/EpicGames.UHT/Utils/UhtReferenceCollector.cs
@@ -0,0 +1,249 @@
+// Copyright Epic Games, Inc. All Rights Reserved.
+
+using EpicGames.Core;
+using EpicGames.UHT.Types;
+using System;
+using System.Collections.Generic;
+
+namespace EpicGames.UHT.Utils
+{
+	/// <summary>
+	/// Interface used to collect all the objects referenced by a given type.
+	/// Not all types such as UhtPackage and UhtHeaderFile support collecting
+	/// references due to assorted reasons.
+	/// </summary>
+	public interface IUhtReferenceCollector
+	{
+
+		/// <summary>
+		/// Add a cross module reference to a given object type.
+		/// </summary>
+		/// <param name="Object">Object type being referenced</param>
+		/// <param name="bRegistered">If true, the method being invoked must return the registered object.  This only applies to classes.</param>
+		void AddCrossModuleReference(UhtObject? Object, bool bRegistered);
+
+		/// <summary>
+		/// Add an object declaration
+		/// </summary>
+		/// <param name="Object">Object in question</param>
+		/// <param name="bRegistered">If true, the method being invoked must return the registered object.  This only applies to classes.</param>
+		void AddDeclaration(UhtObject Object, bool bRegistered);
+
+		/// <summary>
+		/// Add a field as a singleton for exporting
+		/// </summary>
+		/// <param name="Field">Field to be added</param>
+		void AddSingleton(UhtField Field);
+
+		/// <summary>
+		/// Add a field as a type being exported
+		/// </summary>
+		/// <param name="Field">Field to be added</param>
+		void AddExportType(UhtField Field);
+
+		/// <summary>
+		/// Add a forward declaration.  The string can contain multiple declarations but must only exist on one line.
+		/// </summary>
+		/// <param name="Declaration">The declarations to add</param>
+		void AddForwardDeclaration(string? Declaration);
+	}
+
+	/// <summary>
+	/// Delegate used to fetch the string associated with a reference
+	/// </summary>
+	/// <param name="ObjectIndex">Index of the referenced object</param>
+	/// <param name="bRegistered">If true return the registered string, otherwise the unregistered string.  Classes have an unregistered version.</param>
+	/// <returns>The requested string</returns>
+	public delegate string GetReferenceStringDelegate(int ObjectIndex, bool bRegistered);
+
+	/// <summary>
+	/// Maintains a list of referenced object indices.
+	/// </summary>
+	public class UhtUniqueReferenceCollection
+	{
+		/// <summary>
+		/// Collection use to quickly detect if a reference is already in the collection
+		/// </summary>
+		private HashSet<int> Uniques = new HashSet<int>();
+
+		/// <summary>
+		/// List of all unique reference keys.  Use UngetKey to get the object index and the flag.
+		/// </summary>
+		public List<int> References = new List<int>();
+
+		/// <summary>
+		/// Return an encoded key that represents the object and registered flag. 
+		/// If the object has the alternate object set (i.e. native interfaces), then
+		/// that object's index is used to generate the key.
+		/// </summary>
+		/// <param name="Object">Object being referenced</param>
+		/// <param name="bRegistered">If true, then the API that ensures the object is registered is returned.</param>
+		/// <returns>Integer key value.</returns>
+		public static int GetKey(UhtObject Object, bool bRegistered)
+		{
+			return Object.AlternateObject != null ? GetKey(Object.AlternateObject, bRegistered) : (Object.ObjectTypeIndex << 1) + (bRegistered ? 1 : 0);
+		}
+
+		/// <summary>
+		/// Given a key, return the object index and registered flag
+		/// </summary>
+		/// <param name="Key">The key in question</param>
+		/// <param name="ObjectIndex">Index of the referenced object</param>
+		/// <param name="bRegistered">True if referencing the registered API.</param>
+		public static void UngetKey(int Key, out int ObjectIndex, out bool bRegistered)
+		{
+			ObjectIndex = Key >> 1;
+			bRegistered = (Key & 1) != 0;
+		}
+
+		/// <summary>
+		/// Add the given object to the references
+		/// </summary>
+		/// <param name="Object">Object to be added</param>
+		/// <param name="bRegistered">True if the registered API is being returned.</param>
+		public void Add(UhtObject? Object, bool bRegistered)
+		{
+			if (Object != null)
+			{
+				int Key = GetKey(Object, bRegistered);
+				if (this.Uniques.Add(Key))
+				{
+					this.References.Add(Key);
+				}
+			}
+		}
+
+		/// <summary>
+		/// Return the collection of references sorted by the API string returned by the delegate.
+		/// </summary>
+		/// <param name="Delegate">Delegate to invoke to return the requested object API string</param>
+		/// <returns>Read only memory region of all the string.</returns>
+		public ReadOnlyMemory<string> GetSortedReferences(GetReferenceStringDelegate Delegate)
+		{
+			// Collect the unsorted array
+			string[] Sorted = new string[this.References.Count];
+			for (int Index = 0; Index < this.References.Count; ++Index)
+			{
+				int Key = this.References[Index];
+				int ObjectIndex;
+				bool bRegistered;
+				UngetKey(Key, out ObjectIndex, out bRegistered);
+				Sorted[Index] = Delegate(ObjectIndex, bRegistered);
+			}
+
+			// Sort the array
+			Array.Sort(Sorted, StringComparerUE.OrdinalIgnoreCase);
+
+			// Remove duplicates.  In some instances the different keys might return the same string.
+			// This removes those duplicates
+			if (this.References.Count > 1)
+			{
+				int PriorOut = 0;
+				for (int Index = 1; Index < Sorted.Length; ++Index)
+				{
+					if (Sorted[Index] != Sorted[PriorOut])
+					{
+						++PriorOut;
+						Sorted[PriorOut] = Sorted[Index];
+					}
+				}
+				return Sorted.AsMemory(0, PriorOut + 1);
+			}
+			else
+			{
+				return Sorted.AsMemory();
+			}
+		}
+	}
+
+	/// <summary>
+	/// Standard implementation of the reference collector interface
+	/// </summary>
+	public class UhtReferenceCollector : IUhtReferenceCollector
+	{
+
+		/// <summary>
+		/// Collection of unique cross module references
+		/// </summary>
+		public UhtUniqueReferenceCollection CrossModule = new UhtUniqueReferenceCollection();
+
+		/// <summary>
+		/// Collection of unique declarations
+		/// </summary>
+		public UhtUniqueReferenceCollection Declaration = new UhtUniqueReferenceCollection();
+
+		/// <summary>
+		/// Collection of singletons
+		/// </summary>
+		public List<UhtField> Singletons = new List<UhtField>();
+
+		/// <summary>
+		/// Collection of types to export
+		/// </summary>
+		public List<UhtField> ExportTypes = new List<UhtField>();
+
+		/// <summary>
+		/// Collection of forward declarations
+		/// </summary>
+		public HashSet<string> ForwardDeclarations = new HashSet<string>();
+
+		/// <summary>
+		/// Collection of referenced headers
+		/// </summary>
+		public HashSet<UhtHeaderFile> ReferencedHeaders = new HashSet<UhtHeaderFile>();
+
+		/// <summary>
+		/// Add a cross module reference
+		/// </summary>
+		/// <param name="Object">Object being referenced</param>
+		/// <param name="bRegistered">True if the object being referenced must be registered</param>
+		public void AddCrossModuleReference(UhtObject? Object, bool bRegistered)
+		{
+			CrossModule.Add(Object, bRegistered);
+			if (Object != null && !(Object is UhtPackage) && bRegistered)
+			{
+				this.ReferencedHeaders.Add(Object.HeaderFile);
+			}
+		}
+
+		/// <summary>
+		/// Add a declaration
+		/// </summary>
+		/// <param name="Object">Object being declared</param>
+		/// <param name="bRegistered">True if the object being declared must be registered</param>
+		public void AddDeclaration(UhtObject Object, bool bRegistered)
+		{
+			this.Declaration.Add(Object, bRegistered);
+		}
+
+		/// <summary>
+		/// Add a singleton.  These are added as forward declared functions in the package file.
+		/// </summary>
+		/// <param name="Field">Field being added</param>
+		public void AddSingleton(UhtField Field)
+		{
+			this.Singletons.Add(Field);
+		}
+
+		/// <summary>
+		/// Add a type to be exported.
+		/// </summary>
+		/// <param name="Field">Type to be exported</param>
+		public void AddExportType(UhtField Field)
+		{
+			this.ExportTypes.Add(Field);
+		}
+
+		/// <summary>
+		/// Add a symbol that must be forward declared
+		/// </summary>
+		/// <param name="Declaration">Symbol to be forward declared.</param>
+		public void AddForwardDeclaration(string? Declaration)
+		{
+			if (!string.IsNullOrEmpty(Declaration))
+			{
+				this.ForwardDeclarations.Add(Declaration);
+			}
+		}
+	}
+}
