UHT in UBT initial version.  Currently disabled.

diff --git a/Engine/Source/Programs/Shared/EpicGames.UHT/Utils/UhtSession.cs b/Engine/Source/Programs/Shared/EpicGames.UHT/Utils/UhtSession.cs
new file mode 100644
--- /dev/null
+++ b/Engine/Source/Programs/Shared/EpicGames.UHT/Utils/UhtSession.cs
@@ -0,0 +1,2383 @@
+// Copyright Epic Games, Inc. All Rights Reserved.
+
+using EpicGames.Core;
+using EpicGames.UHT.Parsers;
+using EpicGames.UHT.Tokenizer;
+using EpicGames.UHT.Tables;
+using Microsoft.Extensions.Logging;
+using System;
+using System.Collections.Generic;
+using System.IO;
+using System.Linq;
+using System.Text;
+using System.Text.Json;
+using System.Threading;
+using System.Threading.Tasks;
+using EpicGames.UHT.Types;
+using System.IO.Enumeration;
+
+namespace EpicGames.UHT.Utils
+{
+
+	/// <summary>
+	/// To support the testing framework, source files can be containing in other source files.  
+	/// A source fragment represents this possibility.
+	/// </summary>
+	public struct UhtSourceFragment
+	{
+
+		/// <summary>
+		/// When not null, this source comes from another source file
+		/// </summary>
+		public UhtSourceFile? SourceFile;
+
+		/// <summary>
+		/// The file path of the source
+		/// </summary>
+		public string FilePath;
+
+		/// <summary>
+		/// The line number of the fragment in the containing source file.
+		/// </summary>
+		public int LineNumber;
+
+		/// <summary>
+		/// Data of the source file
+		/// </summary>
+		public StringView Data;
+	}
+
+	/// <summary>
+	/// Implementation of the export output interface
+	/// </summary>
+	class UhtExportOutput : IUhtExportOutput
+	{
+		/// <summary>
+		/// Export task that created the export output
+		/// </summary>
+		private readonly IUhtExportTask ExportTask;
+
+		/// <summary>
+		/// Export options for the output
+		/// </summary>
+		private UhtExportOptions OptionsInternal;
+
+		/// <summary>
+		/// Output buffer associated with the exported output
+		/// </summary>
+		private UhtBuffer? ExportedBuffer = null;
+
+		/// <summary>
+		/// Destination file path of the output
+		/// </summary>
+		public readonly string FilePath;
+
+		/// <summary>
+		/// Destination file path of the temporary copy of the output
+		/// </summary>
+		public readonly string TempFilePath;
+
+		/// <summary>
+		/// Exported data
+		/// </summary>
+		public StringView Exported;
+
+		/// <summary>
+		/// True if the exported data was committed
+		/// </summary>
+		public bool bCommitted = false;
+
+		/// <summary>
+		/// True if the exported data saves (i.e. it doesn't match the existing file)
+		/// </summary>
+		public bool bSaved = false;
+
+		/// <summary>
+		/// UHT session
+		/// </summary>
+		public UhtSession Session => ExportTask.Session;
+
+		/// <summary>
+		/// Options associated with the output
+		/// </summary>
+		public UhtExportOptions Options
+		{
+			get => this.OptionsInternal;
+			set => this.OptionsInternal = value;
+		}
+
+		/// <summary>
+		/// Construct a new instance of an export output
+		/// </summary>
+		/// <param name="ExportTask">Task associated with the output</param>
+		/// <param name="FilePath">Destination file path</param>
+		public UhtExportOutput(IUhtExportTask ExportTask, string FilePath)
+		{
+			this.ExportTask = ExportTask;
+			this.Options = ExportTask.Options;
+			this.FilePath = FilePath;
+			this.TempFilePath = FilePath + ".tmp";
+		}
+
+		/// <summary>
+		/// Commit the output using the given string builder
+		/// </summary>
+		/// <param name="Builder">Builder containing the output.</param>
+		/// <returns>String view containing the committed data in the output buffer. 
+		/// While supported, this view should not be saved for any significant length of time.</returns>
+
+		public StringView CommitOutput(StringBuilder Builder)
+		{
+			this.bCommitted = true;
+			this.ExportedBuffer = UhtBuffer.Borrow(Builder);
+			this.Exported = new StringView(this.ExportedBuffer.Memory);
+			return this.Exported;
+		}
+
+		/// <summary>
+		/// Save the given committed output
+		/// </summary>
+		/// <param name="Output">The output to save</param>
+		public void CommitOutput(string Output)
+		{
+			this.bCommitted = true;
+			this.Exported = Output;
+		}
+
+		/// <summary>
+		/// Reset the contents of the output task so the memory can be freed and output buffers returned to the cache.
+		/// </summary>
+		public void Reset()
+		{
+			if (this.ExportedBuffer != null)
+			{
+				UhtBuffer.Return(this.ExportedBuffer);
+				this.ExportedBuffer = null;
+			}
+			this.Exported = new StringView();
+		}
+	}
+
+	/// <summary>
+	/// Base implementation of an export task
+	/// </summary>
+	abstract class UhtExportBaseTask : IUhtExportTask
+	{
+
+		/// <summary>
+		/// Factory requesting the export
+		/// </summary>
+		protected readonly UhtExportFactory ExportFactory;
+
+		/// <summary>
+		/// Action task created to invoke the task
+		/// </summary>
+		private Task? ActionTaskInternal = null;
+
+		/// <summary>
+		/// Export options associated with the task
+		/// </summary>
+		private readonly UhtExportOptions OptionsInternal;
+
+		/// <summary>
+		/// UHT session
+		/// </summary>
+		public UhtSession Session => this.ExportFactory.Session;
+
+		/// <summary>
+		/// Export options associated with the task
+		/// </summary>
+		public UhtExportOptions Options => this.OptionsInternal;
+
+		/// <summary>
+		/// Export options associated with the task
+		/// </summary>
+		public Task? ActionTask => this.ActionTaskInternal;
+
+		/// <summary>
+		/// Construct a new instance of an export task
+		/// </summary>
+		/// <param name="ExportFactory">Factory requesting the task</param>
+		/// <param name="AdditionalOptions">Addition options to added to the existing factory options</param>
+		public UhtExportBaseTask(UhtExportFactory ExportFactory, UhtExportOptions AdditionalOptions)
+		{
+			this.ExportFactory = ExportFactory;
+			this.OptionsInternal = this.ExportFactory.Options | AdditionalOptions;
+		}
+
+		/// <summary>
+		/// Queue the task
+		/// </summary>
+		/// <param name="Prereqs">List of prerequisite tasks that must complete prior to this task being invoked.</param>
+		public void Queue(List<IUhtExportTask> Prereqs)
+		{
+			if (this.Session.bGoWide)
+			{
+				List<Task> PrereqTasks = new List<Task>();
+				foreach (IUhtExportTask Prereq in Prereqs)
+				{
+					if (Prereq.ActionTask != null)
+					{
+						PrereqTasks.Add(Prereq.ActionTask);
+					}
+				}
+				if (PrereqTasks.Count > 0)
+				{
+					this.ActionTaskInternal = Task.Factory.ContinueWhenAll(PrereqTasks.ToArray(), (Task[] Tasks) => { Export(); });
+				}
+				else
+				{
+					this.ActionTaskInternal = Task.Factory.StartNew(() => { Export(); });
+				}
+			}
+			else
+			{
+				Export();
+			}
+		}
+
+		/// <summary>
+		/// Invoked to perform the actual export
+		/// </summary>
+		public abstract void Export();
+	}
+
+	/// <summary>
+	/// Export task that generates a single output
+	/// </summary>
+	class UhtExportSingleTask : UhtExportBaseTask
+	{
+		/// <summary>
+		/// Single output
+		/// </summary>
+		public readonly UhtExportOutput File;
+
+		/// <summary>
+		/// Action to invoke to generate the output
+		/// </summary>
+		private readonly UhtExportSingleDelegate Action;
+
+		/// <summary>
+		/// Create a new instance of a single output task
+		/// </summary>
+		/// <param name="ExportFactory">Requesting factory</param>
+		/// <param name="Path">Destination path of the output</param>
+		/// <param name="AdditionalOptions">Additional options for the output</param>
+		/// <param name="Action">Action to be invoked to generate the output</param>
+		public UhtExportSingleTask(UhtExportFactory ExportFactory, string Path, UhtExportOptions AdditionalOptions, UhtExportSingleDelegate Action) : base(ExportFactory, AdditionalOptions)
+		{
+			this.File = new UhtExportOutput(this, Path);
+			this.Action = Action;
+		}
+
+		/// <summary>
+		/// Invoke the actions and save the output
+		/// </summary>
+		public override void Export()
+		{
+			this.Action(this.File);
+			this.ExportFactory.SaveIfChanged(this.File);
+		}
+	}
+
+	/// <summary>
+	/// Export task that generates a pair of outputs
+	/// </summary>
+	class UhtExportPairTask : UhtExportBaseTask
+	{
+
+		/// <summary>
+		/// First/Header file being generated
+		/// </summary>
+		private readonly UhtExportOutput HeaderFile;
+
+		/// <summary>
+		/// Second/Cpp file being generated
+		/// </summary>
+		private readonly UhtExportOutput CppFile;
+
+		/// <summary>
+		/// Action to be invoked to generate the output
+		/// </summary>
+		private readonly UhtExportPairDelegate Action;
+
+		/// <summary>
+		/// Construct a new instance of the pair output task
+		/// </summary>
+		/// <param name="ExportFactory">Requesting factory</param>
+		/// <param name="HeaderPath">Path of the header file</param>
+		/// <param name="CppPath">Path of the cpp file</param>
+		/// <param name="AdditionalOptions">Additional options to be applied to the outputs</param>
+		/// <param name="Action">Action to invoke to generate the output</param>
+		public UhtExportPairTask(UhtExportFactory ExportFactory, string HeaderPath, string CppPath, UhtExportOptions AdditionalOptions, UhtExportPairDelegate Action) : base(ExportFactory, AdditionalOptions)
+		{
+			this.HeaderFile = new UhtExportOutput(this, HeaderPath);
+			this.CppFile = new UhtExportOutput(this, CppPath);
+			this.Action = Action;
+		}
+
+		/// <summary>
+		/// Invoked to export and save the outputs
+		/// </summary>
+		public override void Export()
+		{
+			this.Action(this.HeaderFile, this.CppFile);
+			this.ExportFactory.SaveIfChanged(this.HeaderFile);
+			this.ExportFactory.SaveIfChanged(this.CppFile);
+		}
+	}
+
+	/// <summary>
+	/// Implementation of the export factory
+	/// </summary>
+	class UhtExportFactory : IUhtExportFactory
+	{
+
+		/// <summary>
+		/// UHT session
+		/// </summary>
+		private readonly UhtSession SessionInternal;
+
+		/// <summary>
+		/// Export options 
+		/// </summary>
+		private readonly UhtExportOptions OptionsInternal;
+
+		/// <summary>
+		/// Limiter for the number of files being saved to the reference directory.
+		/// The OS can get swamped on high core systems
+		/// </summary>
+		private Semaphore WriteRefSemaphore = new Semaphore(32, 32);
+
+		/// <summary>
+		/// Requesting exporter
+		/// </summary>
+		public readonly UhtExporter Exporter;
+
+		/// <summary>
+		/// UHT Session
+		/// </summary>
+		public UhtSession Session => this.SessionInternal;
+
+		/// <summary>
+		/// Export options 
+		/// </summary>
+		public UhtExportOptions Options => this.OptionsInternal;
+		
+		/// <summary>
+		/// Collection of error from mismatches with the reference files
+		/// </summary>
+		public Dictionary<string, bool> ReferenceErrorMessages = new Dictionary<string, bool>();
+
+		/// <summary>
+		/// List of export outputs
+		/// </summary>
+		public List<UhtExportOutput> Outputs = new List<UhtExportOutput>();
+
+		/// <summary>
+		/// Directory for the reference output
+		/// </summary>
+		public string ReferenceDirectory = string.Empty;
+
+		/// <summary>
+		/// Directory for the verify output
+		/// </summary>
+		public string VerifyDirectory = string.Empty;
+
+		/// <summary>
+		/// Create a new instance of the export factory
+		/// </summary>
+		/// <param name="Session">UHT session</param>
+		/// <param name="Exporter">Exporter being run</param>
+		/// <param name="Options">Export options</param>
+		public UhtExportFactory(UhtSession Session, UhtExporter Exporter, UhtExportOptions Options)
+		{
+			this.Exporter = Exporter;
+			this.SessionInternal = Session;
+			this.OptionsInternal = Options;
+			if (this.Session.ReferenceMode != UhtReferenceMode.None)
+			{
+				this.ReferenceDirectory = Path.Combine(this.Session.ReferenceDirectory, this.Exporter.Name);
+				this.VerifyDirectory = Path.Combine(this.Session.VerifyDirectory, this.Exporter.Name);
+				Directory.CreateDirectory(this.Session.ReferenceMode == UhtReferenceMode.Reference ? this.ReferenceDirectory : this.VerifyDirectory);
+			}
+		}
+
+		/// <summary>
+		/// Create an export task that write one output
+		/// </summary>
+		/// <param name="Path">Destination file path</param>
+		/// <param name="Prereqs">List of required export tasks that must be completed prior to this export</param>
+		/// <param name="AdditionalOptions">Additional export options for the output</param>
+		/// <param name="Action">Action to be invoked to generated the output</param>
+		/// <returns>Created task</returns>
+		public IUhtExportTask CreateTask(string Path, List<IUhtExportTask> Prereqs, UhtExportOptions AdditionalOptions, UhtExportSingleDelegate Action)
+		{
+			UhtExportBaseTask Task = new UhtExportSingleTask(this, Path, AdditionalOptions, Action);
+			Task.Queue(Prereqs);
+			return Task;
+		}
+
+		/// <summary>
+		/// Create an export task that write one output
+		/// </summary>
+		/// <param name="HeaderPath">Destination header file path</param>
+		/// <param name="CppPath">Destination cpp file path</param>
+		/// <param name="Prereqs">List of required export tasks that must be completed prior to this export</param>
+		/// <param name="AdditionalOptions">Additional export options for the output</param>
+		/// <param name="Action">Action to be invoked to generated the output</param>
+		/// <returns>Created task</returns>
+		public IUhtExportTask CreateTask(string HeaderPath, string CppPath, List<IUhtExportTask> Prereqs, UhtExportOptions AdditionalOptions, UhtExportPairDelegate Action)
+		{
+			UhtExportBaseTask Task = new UhtExportPairTask(this, HeaderPath, CppPath, AdditionalOptions, Action);
+			Task.Queue(Prereqs);
+			return Task;
+		}
+
+		/// <summary>
+		/// Given a header file, generate the output file name.
+		/// </summary>
+		/// <param name="HeaderFile">Header file</param>
+		/// <param name="Suffix">Suffix/extension to be added to the file name.</param>
+		/// <returns>Resulting file name</returns>
+		public string MakePath(UhtHeaderFile HeaderFile, string Suffix)
+		{
+			UHTManifest.Module Module = HeaderFile.Package.Module;
+			return Path.Combine(Module.OutputDirectory, HeaderFile.FileNameWithoutExtension) + Suffix;
+		}
+
+		/// <summary>
+		/// Given a package file, generate the output file name
+		/// </summary>
+		/// <param name="Package">Package file</param>
+		/// <param name="Suffix">Suffix/extension to be added to the file name.</param>
+		/// <returns>Resulting file name</returns>
+		public string MakePath(UhtPackage Package, string Suffix)
+		{
+			UHTManifest.Module Module = Package.Module;
+			return Path.Combine(Module.OutputDirectory, Package.ShortName.ToString()) + Suffix;
+		}
+
+		/// <summary>
+		/// Helper method to test to see if the output has changed.
+		/// </summary>
+		/// <param name="Output">The output file</param>
+		internal void SaveIfChanged(UhtExportOutput Output)
+		{
+
+			if (!Output.Options.HasAnyFlags(UhtExportOptions.WriteOutput) || !Output.bCommitted)
+			{
+				return;
+			}
+
+			// Add this to the list of outputs
+			lock (this.Outputs)
+			{
+				this.Outputs.Add(Output);
+			}
+
+			StringView Exported = Output.Exported;
+			ReadOnlySpan<char> ExportedSpan = Exported.Span;
+
+			if (this.Session.ReferenceMode != UhtReferenceMode.None)
+			{
+				string FileName = Path.GetFileName(Output.FilePath);
+
+				// Writing billions of files to the same directory causes issues.  Use ourselves to throttle reference writes
+				try
+				{
+					this.WriteRefSemaphore.WaitOne();
+					{
+						string OutPath = Path.Combine(this.Session.ReferenceMode == UhtReferenceMode.Reference ? this.ReferenceDirectory : this.VerifyDirectory, FileName);
+						if (!this.Session.WriteSource(OutPath, Exported.Span))
+						{
+							new UhtSimpleFileMessageSite(this.Session, OutPath).LogWarning($"Unable to write reference file {OutPath}");
+						}
+					}
+				}
+				finally
+				{
+					this.WriteRefSemaphore.Release();
+				}
+
+				// If we are verifying, read the existing file and check the contents
+				if (this.Session.ReferenceMode == UhtReferenceMode.Verify)
+				{
+					string Message = String.Empty;
+					string RefPath = Path.Combine(this.ReferenceDirectory, FileName);
+					UhtBuffer? ExistingRef = this.Session.ReadSourceToBuffer(RefPath);
+					if (ExistingRef != null)
+					{
+						ReadOnlySpan<char> ExistingSpan = ExistingRef.Memory.Span;
+						if (ExistingSpan.CompareTo(ExportedSpan, StringComparison.Ordinal) != 0)
+						{
+							Message = $"********************************* {FileName} has changed.";
+						}
+						UhtBuffer.Return(ExistingRef);
+					}
+					else
+					{
+						Message = $"********************************* {FileName} appears to be a new generated file.";
+					}
+
+					if (Message != String.Empty)
+					{
+						Log.Logger.LogInformation(Message);
+						lock (this.ReferenceErrorMessages)
+						{
+							this.ReferenceErrorMessages.Add(Message, true);
+						}
+					}
+				}
+			}
+
+			// Check to see if the contents have changed
+			UhtBuffer? Original = this.Session.ReadSourceToBuffer(Output.FilePath);
+			bool bSave = Original == null;
+			if (Original != null)
+			{
+				ReadOnlySpan<char> OriginalSpan = Original.Memory.Span;
+				if (OriginalSpan.CompareTo(ExportedSpan, StringComparison.Ordinal) != 0)
+				{
+					if (this.Session.bFailIfGeneratedCodeChanges)
+					{
+						string ConflictPath = Output.FilePath + ".conflict";
+						if (!this.Session.WriteSource(ConflictPath, Exported.Span))
+						{
+							new UhtSimpleFileMessageSite(this.Session, Output.FilePath).LogError($"Changes to generated code are not allowed - conflicts written to '{ConflictPath}'");
+						}
+					}
+					bSave = true;
+				}
+				UhtBuffer.Return(Original);
+			}
+
+			// If changed of the original didn't exist, then save the new version
+			if (bSave && !this.Session.bNoOutput)
+			{
+				Output.bSaved = true;
+				if (!this.Session.WriteSource(Output.TempFilePath, Exported.Span))
+				{
+					new UhtSimpleFileMessageSite(this.Session, Output.FilePath).LogWarning($"Failed to save export file: '{Output.TempFilePath}'");
+				}
+			}
+
+			// Reset the output to try and free up some of the temporary buffers
+			Output.Reset();
+		}
+
+		/// <summary>
+		/// Run the output exporter
+		/// </summary>
+		public void Run()
+		{
+
+			// Invoke the exported via the delegate
+			this.Exporter.Delegate(this);
+
+			// If outputs were exported
+			if (this.Outputs.Count > 0)
+			{
+
+				// These outputs are used to cull old outputs from the directories
+				Dictionary<string, HashSet<string>> OutputsByDirectory = new Dictionary<string, HashSet<string>>(StringComparer.OrdinalIgnoreCase);
+				List<UhtExportOutput> Saves = new List<UhtExportOutput>();
+
+				// Collect information about the outputs
+				foreach (UhtExportOutput Output in this.Outputs)
+				{
+
+					// Add this output to the list of expected outputs by directory
+					string? FileDirectory = Path.GetDirectoryName(Output.FilePath);
+					if (FileDirectory != null)
+					{
+						HashSet<string>? Files;
+						if (!OutputsByDirectory.TryGetValue(FileDirectory, out Files))
+						{
+							Files = new HashSet<string>(StringComparer.OrdinalIgnoreCase);
+							OutputsByDirectory.Add(FileDirectory, Files);
+						}
+						Files.Add(Path.GetFileName(Output.FilePath));
+					}
+
+					// Add the save task
+					if (Output.bSaved)
+					{
+						Saves.Add(Output);
+					}
+				}
+
+				// Perform the renames
+				if (this.Session.bGoWide)
+				{
+					Parallel.ForEach(Saves, (UhtExportOutput Output) =>
+					{
+						RenameSource(Output);
+					});
+				}
+				else
+				{ 
+					foreach (UhtExportOutput Output in Saves)
+					{
+						RenameSource(Output);
+					}
+				}
+
+				// Perform the culling of the output directories
+				if (this.Session.bCullOutput && !this.Session.bNoOutput && 
+					(this.Exporter.CppFilters != null || this.Exporter.HeaderFilters != null || this.Exporter.OtherFilters != null))
+				{
+					if (this.Session.bGoWide)
+					{
+						Parallel.ForEach(OutputsByDirectory, (KeyValuePair<string, HashSet<string>> Kvp) =>
+						{
+							CullOutputDirectory(Kvp.Key, Kvp.Value);
+						});
+					}
+					else
+					{
+						foreach (KeyValuePair<string, HashSet<string>> Kvp in OutputsByDirectory)
+						{
+							CullOutputDirectory(Kvp.Key, Kvp.Value);
+						}
+					}
+				}
+			}
+		}
+
+		/// <summary>
+		/// Given an output, rename the output file from the temporary file name to the final file name.
+		/// If there exists a current final file, it will be replaced.
+		/// </summary>
+		/// <param name="Output">The output file to rename</param>
+		private void RenameSource(UhtExportOutput Output)
+		{
+			this.Session.RenameSource(Output.TempFilePath, Output.FilePath);
+		}
+
+		/// <summary>
+		/// Given a directory and a list of known files, delete any unknown file that matches of the supplied filters
+		/// </summary>
+		/// <param name="OutputDirectory">Output directory to scan</param>
+		/// <param name="KnownOutputs">Collection of known output files not to be deleted</param>
+		private void CullOutputDirectory(string OutputDirectory, HashSet<string> KnownOutputs)
+		{
+			foreach (string FilePath in Directory.EnumerateFiles(OutputDirectory))
+			{
+				string FileName = Path.GetFileName(FilePath);
+				if (KnownOutputs.Contains(FileName))
+				{
+					continue;
+				}
+
+				if (IsFilterMatch(FileName, this.Exporter.CppFilters) ||
+					IsFilterMatch(FileName, this.Exporter.HeaderFilters) ||
+					IsFilterMatch(FileName, this.Exporter.OtherFilters))
+				{
+					try
+					{
+						File.Delete(Path.Combine(OutputDirectory, FilePath));
+					}
+					catch (Exception)
+					{
+					}
+				}
+			}
+		}
+
+		/// <summary>
+		/// Test to see if the given filename (without directory), matches one of the given filters
+		/// </summary>
+		/// <param name="FileName">File name to test</param>
+		/// <param name="Filters">List of wildcard filters</param>
+		/// <returns>True if there is a match</returns>
+		private static bool IsFilterMatch(string FileName, string[]? Filters)
+		{
+			if (Filters != null)
+			{
+				foreach (string Filter in Filters)
+				{
+					if (FileSystemName.MatchesSimpleExpression(Filter, FileName, true))
+					{
+						return true;
+					}
+				}
+			}
+			return false;
+		}
+	}
+
+	/// <summary>
+	/// UHT supports the exporting of two reference output directories for testing.  The reference version can be used to test
+	/// modification to UHT and verify there are no output changes or just expected changes.
+	/// </summary>
+	public enum UhtReferenceMode
+	{
+		/// <summary>
+		/// Do not export any reference output files
+		/// </summary>
+		None,
+
+		/// <summary>
+		/// Export the reference copy
+		/// </summary>
+		Reference,
+
+		/// <summary>
+		/// Export the verify copy and compare to the reference copy
+		/// </summary>
+		Verify,
+	};
+
+	/// <summary>
+	/// Session object that represents a UHT run
+	/// </summary>
+	public class UhtSession : IUhtMessageSite, IUhtMessageSession
+	{
+
+		/// <summary>
+		/// Helper class for returning a sequence of auto-incrementing indices
+		/// </summary>
+		private class TypeCounter
+		{
+
+			/// <summary>
+			/// Current number of types
+			/// </summary>
+			private int CountInternal = 0;
+
+			/// <summary>
+			/// Get the next type index
+			/// </summary>
+			/// <returns>Index starting at zero</returns>
+			public int GetNext()
+			{
+				return Interlocked.Increment(ref this.CountInternal) - 1;
+			}
+
+			/// <summary>
+			/// The number of times GetNext was called.
+			/// </summary>
+			public int Count => Interlocked.Add(ref this.CountInternal, 0) + 1;
+		}
+
+		/// <summary>
+		/// Pair that represents a specific value for an enumeration
+		/// </summary>
+		private struct EnumAndValue
+		{
+			public UhtEnum Enum;
+			public long Value;
+		}
+
+		/// <summary>
+		/// Collection of reserved names
+		/// </summary>
+		private static HashSet<string> ReservedNames = new HashSet<string> { "none" };
+
+		#region Configurable settings
+
+		/// <summary>
+		/// Interface used to read/write files
+		/// </summary>
+		public IUhtFileManager? FileManager;
+
+		/// <summary>
+		/// Location of the engine code
+		/// </summary>
+		public string? EngineDirectory;
+
+		/// <summary>
+		/// Optional location of the project
+		/// </summary>
+		public string? ProjectDirectory;
+
+		/// <summary>
+		/// Root directory for the engine.  This is usually just EngineDirectory without the Engine directory.
+		/// </summary>
+		public string? RootDirectory;
+
+		/// <summary>
+		/// Directory to store the reference output
+		/// </summary>
+		public string ReferenceDirectory = string.Empty;
+
+		/// <summary>
+		/// Directory to store the verification output
+		/// </summary>
+		public string VerifyDirectory = string.Empty;
+
+		/// <summary>
+		/// Mode for generating and/or testing reference output
+		/// </summary>
+		public UhtReferenceMode ReferenceMode = UhtReferenceMode.None;
+
+		/// <summary>
+		/// If true, warnings are considered to be errors
+		/// </summary>
+		public bool bWarningsAsErrors = false;
+
+		/// <summary>
+		/// If true, include relative file paths in the log file
+		/// </summary>
+		public bool bRelativePathInLog = false;
+
+		/// <summary>
+		/// If true, use concurrent tasks to run UHT
+		/// </summary>
+		public bool bGoWide = true;
+
+		/// <summary>
+		/// If any output file mismatches existing outputs, an error will be generated
+		/// </summary>
+		public bool bFailIfGeneratedCodeChanges = false;
+
+		/// <summary>
+		/// If true, no output files will be saved
+		/// </summary>
+		public bool bNoOutput = false;
+
+		/// <summary>
+		/// If true, cull the output for any extra files
+		/// </summary>
+		public bool bCullOutput = true;
+
+		/// <summary>
+		/// If true, include extra output in code generation
+		/// </summary>
+		public bool bIncludeDebugOutput = false;
+
+		/// <summary>
+		/// If true, cache any error messages until the end of processing.  This is used by the testing
+		/// harness to generate more stable console output.
+		/// </summary>
+		public bool bCacheMessages = false;
+		#endregion
+
+		/// <summary>
+		/// Manifest file
+		/// </summary>
+		public UhtManifestFile? ManifestFile { get; set; } = null;
+
+		/// <summary>
+		/// Manifest data from the manifest file
+		/// </summary>
+		public UHTManifest? Manifest { get => this.ManifestFile != null ? this.ManifestFile.Manifest : null; }
+
+		/// <summary>
+		/// Collection of packages from the manifest
+		/// </summary>
+		public IReadOnlyList<UhtPackage> Packages { get => this.PackagesInternal; }
+
+		/// <summary>
+		/// Collection of header files from the manifest.  The header files will also appear as the children 
+		/// of the packages
+		/// </summary>
+		public IReadOnlyList<UhtHeaderFile> HeaderFiles { get => this.HeaderFilesInternal; }
+
+		/// <summary>
+		/// Collection of header files topologically sorted.  This will not be populated until after header files
+		/// are parsed and resolved.
+		/// </summary>
+		public IReadOnlyList<UhtHeaderFile> SortedHeaderFiles { get => this.SortedHeaderFilesInternal; }
+
+		/// <summary>
+		/// Dictionary of stripped file name to the header file
+		/// </summary>
+		public IReadOnlyDictionary<string, UhtHeaderFile> HeaderFileDictionary { get => this.HeaderFileDictionaryInternal; }
+
+		/// <summary>
+		/// After headers are parsed, returns the UObject class.
+		/// </summary>
+		public UhtClass UObject
+		{
+			get
+			{
+				if (this.UObjectInternal == null)
+				{
+					throw new UhtIceException("UObject was not defined.");
+				}
+				return this.UObjectInternal;
+			}
+		}
+
+		/// <summary>
+		/// After headers are parsed, returns the UClass class.
+		/// </summary>
+		public UhtClass UClass
+		{
+			get
+			{
+				if (this.UClassInternal == null)
+				{
+					throw new UhtIceException("UClass was not defined.");
+				}
+				return this.UClassInternal;
+			}
+		}
+
+		/// <summary>
+		/// After headers are parsed, returns the UInterface class.
+		/// </summary>
+		public UhtClass UInterface
+		{
+			get
+			{
+				if (this.UInterfaceInternal == null)
+				{
+					throw new UhtIceException("UInterface was not defined.");
+				}
+				return this.UInterfaceInternal;
+			}
+		}
+
+		/// <summary>
+		/// After headers are parsed, returns the IInterface class.
+		/// </summary>
+		public UhtClass IInterface
+		{
+			get
+			{
+				if (this.IInterfaceInternal == null)
+				{
+					throw new UhtIceException("IInterface was not defined.");
+				}
+				return this.IInterfaceInternal;
+			}
+		}
+
+		/// <summary>
+		/// After headers are parsed, returns the AActor class.  Unlike such properties as "UObject", there
+		/// is no requirement for AActor to be defined.  May be null.
+		/// </summary>
+		public UhtClass? AActor = null;
+
+		private List<UhtPackage> PackagesInternal = new List<UhtPackage>();
+		private List<UhtHeaderFile> HeaderFilesInternal = new List<UhtHeaderFile>();
+		private List<UhtHeaderFile> SortedHeaderFilesInternal = new List<UhtHeaderFile>();
+		private Dictionary<string, UhtHeaderFile> HeaderFileDictionaryInternal = new Dictionary<string, UhtHeaderFile>(StringComparer.OrdinalIgnoreCase);
+		private long ErrorCountInternal = 0;
+		private long WarningCountInternal = 0;
+		private List<UhtMessage> Messages = new List<UhtMessage>();
+		private Task? MessageTask = null;
+		private Dictionary<string, UhtSourceFragment> SourceFragments = new Dictionary<string, UhtSourceFragment>();
+		private UhtClass? UObjectInternal = null;
+		private UhtClass? UClassInternal = null;
+		private UhtClass? UInterfaceInternal = null;
+		private UhtClass? IInterfaceInternal = null;
+		private TypeCounter TypeCounterInternal = new TypeCounter();
+		private TypeCounter PackageTypeCountInternal = new TypeCounter();
+		private TypeCounter HeaderFileTypeCountInternal = new TypeCounter();
+		private TypeCounter ObjectTypeCountInternal = new TypeCounter();
+		private UhtSymbolTable SourceNameSymbolTable = new UhtSymbolTable(0);
+		private UhtSymbolTable EngineNameSymbolTable = new UhtSymbolTable(0);
+		private bool bSymbolTablePopulated = false;
+		private Task? ReferenceDeleteTask = null;
+		private Dictionary<string, bool> ExporterStates = new Dictionary<string, bool>(StringComparer.OrdinalIgnoreCase);
+		private Dictionary<string, EnumAndValue> FullEnumValueLookup = new Dictionary<string, EnumAndValue>();
+		private Dictionary<string, UhtEnum> ShortEnumValueLookup = new Dictionary<string, UhtEnum>();
+
+		/// <summary>
+		/// The number of errors
+		/// </summary>
+		public long ErrorCount
+		{
+			get => Interlocked.Read(ref this.ErrorCountInternal);
+		}
+
+		/// <summary>
+		/// The number of warnings
+		/// </summary>
+		public long WarningCount
+		{
+			get => Interlocked.Read(ref this.WarningCountInternal);
+		}
+
+		/// <summary>
+		/// True if any errors have occurred or warnings if warnings are to be treated as errors 
+		/// </summary>
+		public bool bHasErrors
+		{
+			get => this.ErrorCount > 0 || (this.bWarningsAsErrors && this.WarningCount > 0);
+		}
+
+		#region IUHTMessageSession implementation
+		/// <inheritdoc/>
+		public IUhtMessageSession MessageSession => this;
+		/// <inheritdoc/>
+		public IUhtMessageSource? MessageSource => null;
+		/// <inheritdoc/>
+		public IUhtMessageLineNumber? MessageLineNumber => null;
+		/// <inheritdoc/>
+		public IUhtMessageExtraContext? MessageExtraContext => null;
+		#endregion
+
+		/// <summary>
+		/// Return the index for a newly defined type
+		/// </summary>
+		/// <returns>New index</returns>
+		public int GetNextTypeIndex()
+		{
+			return this.TypeCounterInternal.GetNext();
+		}
+
+		/// <summary>
+		/// Return the number of types that have been defined.  This includes all types.
+		/// </summary>
+		public int TypeCount => this.TypeCounterInternal.Count;
+
+		/// <summary>
+		/// Return the index for a newly defined packaging
+		/// </summary>
+		/// <returns>New index</returns>
+		public int GetNextHeaderFileTypeIndex()
+		{
+			return this.HeaderFileTypeCountInternal.GetNext();
+		}
+
+		/// <summary>
+		/// Return the number of headers that have been defined
+		/// </summary>
+		public int HeaderFileTypeCount => this.HeaderFileTypeCountInternal.Count;
+
+		/// <summary>
+		/// Return the index for a newly defined package
+		/// </summary>
+		/// <returns>New index</returns>
+		public int GetNextPackageTypeIndex()
+		{
+			return this.PackageTypeCountInternal.GetNext();
+		}
+
+		/// <summary>
+		/// Return the number of UPackage types that have been defined
+		/// </summary>
+		public int PackageTypeCount => this.PackageTypeCountInternal.Count;
+
+		/// <summary>
+		/// Return the index for a newly defined object
+		/// </summary>
+		/// <returns>New index</returns>
+		public int GetNextObjectTypeIndex()
+		{
+			return this.ObjectTypeCountInternal.GetNext();
+		}
+
+		/// <summary>
+		/// Return the total number of UObject types that have been defined
+		/// </summary>
+		public int ObjectTypeCount => this.ObjectTypeCountInternal.Count;
+
+		/// <summary>
+		/// Enable/Disable an exporter.  This overrides the default state of the exporter.
+		/// </summary>
+		/// <param name="Name">Name of the exporter</param>
+		/// <param name="Enabled">If true, the exporter is to be enabled</param>
+		public void SetExporterStatus(string Name, bool Enabled)
+		{
+			this.ExporterStates[Name] = Enabled;
+		}
+
+		/// <summary>
+		/// Run UHT on the given manifest.  Use the bHasError property to see if process was successful.
+		/// </summary>
+		/// <param name="ManifestFilePath">Path to the manifest file</param>
+		public void Run(string ManifestFilePath)
+		{
+			if (this.FileManager == null)
+			{
+				Interlocked.Increment(ref this.ErrorCountInternal);
+				Log.Logger.LogError("No file manager supplied, aborting.");
+				return;
+			}
+
+			// Initialize
+			Try(null, () => {
+				_ = UhtConfig.Instance;
+				UhtAttributeScanner.Scan();
+			});
+
+			switch (this.ReferenceMode)
+			{
+				case UhtReferenceMode.None:
+					break;
+
+				case UhtReferenceMode.Reference:
+					if (string.IsNullOrEmpty(this.ReferenceDirectory))
+					{
+						Log.Logger.LogError("WRITEREF requested but directory not set, ignoring");
+						this.ReferenceMode = UhtReferenceMode.None;
+					}
+					break;
+
+				case UhtReferenceMode.Verify:
+					if (string.IsNullOrEmpty(this.ReferenceDirectory) || string.IsNullOrEmpty(this.VerifyDirectory))
+					{
+						Log.Logger.LogError("VERIFYREF requested but directories not set, ignoring");
+						this.ReferenceMode = UhtReferenceMode.None;
+					}
+					break;
+			}
+
+			if (this.ReferenceMode != UhtReferenceMode.None)
+			{
+				string DirectoryToDelete = this.ReferenceMode == UhtReferenceMode.Reference ? this.ReferenceDirectory : this.VerifyDirectory;
+				this.ReferenceDeleteTask = Task.Factory.StartNew(() =>
+				{
+					try
+					{
+						Directory.Delete(DirectoryToDelete, true);
+					}
+					catch (Exception)
+					{ }
+				});
+			}
+
+			StepReadManifestFile(ManifestFilePath);
+			StepPrepareModules();
+			StepPrepareHeaders();
+			StepParseHeaders();
+			StepPopulateTypeTable();
+			StepResolveInvalidCheck();
+			StepResolveBases();
+			StepResolveProperties();
+			StepResolveFinal();
+			StepResolveValidate();
+			StepCollectReferences();
+			TopologicalSortHeaderFiles();
+
+			// If we are deleting the reference directory, then wait for that task to complete
+			if (this.ReferenceDeleteTask != null)
+			{
+				Log.Logger.LogTrace("Step - Waiting for reference output to be cleared.");
+				this.ReferenceDeleteTask.Wait();
+			}
+
+			StepExport();
+		}
+
+		/// <summary>
+		/// Try the given action regardless of any prior errors.  If an exception occurs that doesn't have the required
+		/// context, use the supplied context to generate the message.
+		/// </summary>
+		/// <param name="MessageSource">Message context for when the exception doesn't contain a context.</param>
+		/// <param name="Action">The lambda to be invoked</param>
+		public void TryAlways(IUhtMessageSource? MessageSource, Action Action)
+		{
+			try
+			{
+				Action();
+			}
+			catch (Exception E)
+			{
+				HandleException(MessageSource, E);
+			}
+		}
+
+		/// <summary>
+		/// Try the given action.  If an exception occurs that doesn't have the required
+		/// context, use the supplied context to generate the message.
+		/// </summary>
+		/// <param name="MessageSource">Message context for when the exception doesn't contain a context.</param>
+		/// <param name="Action">The lambda to be invoked</param>
+		public void Try(IUhtMessageSource? MessageSource, Action Action)
+		{
+			if (!this.bHasErrors)
+			{
+				try
+				{
+					Action();
+				}
+				catch (Exception E)
+				{
+					HandleException(MessageSource, E);
+				}
+			}
+		}
+
+		/// <summary>
+		/// Read the given source file
+		/// </summary>
+		/// <param name="FilePath">Full or relative file path</param>
+		/// <returns>Information about the read source</returns>
+		public UhtSourceFragment ReadSource(string FilePath)
+		{
+			if (this.FileManager!.ReadSource(FilePath, out UhtSourceFragment Fragment))
+			{
+				return Fragment;
+			}
+			throw new UhtException($"File not found '{FilePath}'");
+		}
+
+		/// <summary>
+		/// Read the given source file
+		/// </summary>
+		/// <param name="FilePath">Full or relative file path</param>
+		/// <returns>Buffer containing the read data or null if not found.  The returned buffer must be returned to the cache via a call to UhtBuffer.Return</returns>
+		public UhtBuffer? ReadSourceToBuffer(string FilePath)
+		{
+			return this.FileManager!.ReadOutput(FilePath);
+		}
+
+		/// <summary>
+		/// Write the given contents to the file
+		/// </summary>
+		/// <param name="FilePath">Path to write to</param>
+		/// <param name="Contents">Contents to write</param>
+		/// <returns>True if the source was written</returns>
+		internal bool WriteSource(string FilePath, ReadOnlySpan<char> Contents)
+		{
+			return this.FileManager!.WriteOutput(FilePath, Contents);
+		}
+
+		/// <summary>
+		/// Rename the given file
+		/// </summary>
+		/// <param name="OldFilePath">Old file path name</param>
+		/// <param name="NewFilePath">New file path name</param>
+		public void RenameSource(string OldFilePath, string NewFilePath)
+		{
+			if (!this.FileManager!.RenameOutput(OldFilePath, NewFilePath))
+			{
+				new UhtSimpleFileMessageSite(this, NewFilePath).LogError($"Failed to rename export file: '{OldFilePath}'");
+			}
+		}
+
+		/// <summary>
+		/// Given the name of a regular enum value, return the enum type
+		/// </summary>
+		/// <param name="Name">Enum value</param>
+		/// <returns>Associated regular enum type or null if not found or enum isn't a regular enum.</returns>
+		public UhtEnum? FindRegularEnumValue(string Name)
+		{
+			//COMPATIBILITY-TODO - See comment below on a more rebust version of the enum lookup
+			//if (this.RegularEnumValueLookup.TryGetValue(Name, out UhtEnum? Enum))
+			//{
+			//	return Enum;
+			//}
+			if (this.FullEnumValueLookup.TryGetValue(Name, out EnumAndValue Value))
+			{
+				if (Value.Value != -1)
+				{
+					return Value.Enum;
+				}
+			}
+
+			if (!Name.Contains("::") && this.ShortEnumValueLookup.TryGetValue(Name, out UhtEnum? Enum))
+			{
+				return Enum;
+			}
+
+			return null;
+		}
+
+		/// <summary>
+		/// Find the given type in the type hierarchy
+		/// </summary>
+		/// <param name="StartingType">Starting point for searches</param>
+		/// <param name="Options">Options controlling what is searched</param>
+		/// <param name="Name">Name of the type.</param>
+		/// <param name="MessageSite">If supplied, then a error message will be generated if the type can not be found</param>
+		/// <param name="LineNumber">Source code line number requesting the lookup.</param>
+		/// <returns>The located type of null if not found</returns>
+		public UhtType? FindType(UhtType? StartingType, UhtFindOptions Options, string Name, IUhtMessageSite? MessageSite = null, int LineNumber = -1)
+		{
+			ValidateFindOptions(Options);
+
+			UhtType? Type = FindTypeInternal(StartingType, Options, Name);
+			if (Type == null && MessageSite != null)
+			{
+				FindTypeError(MessageSite, LineNumber, Options, Name);
+			}
+			return Type;
+		}
+
+		/// <summary>
+		/// Find the given type in the type hierarchy
+		/// </summary>
+		/// <param name="StartingType">Starting point for searches</param>
+		/// <param name="Options">Options controlling what is searched</param>
+		/// <param name="Name">Name of the type.</param>
+		/// <param name="MessageSite">If supplied, then a error message will be generated if the type can not be found</param>
+		/// <returns>The located type of null if not found</returns>
+		public UhtType? FindType(UhtType? StartingType, UhtFindOptions Options, ref UhtToken Name, IUhtMessageSite? MessageSite = null)
+		{
+			ValidateFindOptions(Options);
+
+			UhtType? Type = FindTypeInternal(StartingType, Options, Name.Value.ToString());
+			if (Type == null && MessageSite != null)
+			{
+				FindTypeError(MessageSite, Name.InputLine, Options, Name.Value.ToString());
+			}
+			return Type;
+		}
+
+		/// <summary>
+		/// Find the given type in the type hierarchy
+		/// </summary>
+		/// <param name="StartingType">If specified, this represents the starting type to use when searching base/owner chain for a match</param>
+		/// <param name="Options">Options controlling what is searched</param>
+		/// <param name="Identifiers">Enumeration of identifiers.</param>
+		/// <param name="MessageSite">If supplied, then a error message will be generated if the type can not be found</param>
+		/// <param name="LineNumber">Source code line number requesting the lookup.</param>
+		/// <returns>The located type of null if not found</returns>
+		public UhtType? FindType(UhtType? StartingType, UhtFindOptions Options, UhtTokenList Identifiers, IUhtMessageSite? MessageSite = null, int LineNumber = -1)
+		{
+			ValidateFindOptions(Options);
+
+			if (Identifiers.Next != null && Identifiers.Next.Next != null)
+			{
+				if (MessageSite != null)
+				{
+					MessageSite.LogError(LineNumber, "UnrealHeaderTool only supports C++ identifiers of two or less identifiers");
+					return null;
+				}
+			}
+
+			UhtType? Type = null;
+			if (Identifiers.Next != null)
+			{
+				Type = FindTypeTwoNamesInternal(StartingType, Options, Identifiers.Token.Value.ToString(), Identifiers.Next.Token.Value.ToString());
+			}
+			else
+			{
+				Type = FindTypeInternal(StartingType, Options, Identifiers.Token.Value.ToString());
+			}
+
+			if (Type == null && MessageSite != null)
+			{
+				string FullIdentifier = Identifiers.Join("::");
+				FindTypeError(MessageSite, LineNumber, Options, FullIdentifier);
+			}
+			return Type;
+		}
+
+		/// <summary>
+		/// Find the given type in the type hierarchy
+		/// </summary>
+		/// <param name="StartingType">If specified, this represents the starting type to use when searching base/owner chain for a match</param>
+		/// <param name="Options">Options controlling what is searched</param>
+		/// <param name="Identifiers">Enumeration of identifiers.</param>
+		/// <param name="MessageSite">If supplied, then a error message will be generated if the type can not be found</param>
+		/// <param name="LineNumber">Source code line number requesting the lookup.</param>
+		/// <returns>The located type of null if not found</returns>
+		public UhtType? FindType(UhtType? StartingType, UhtFindOptions Options, UhtToken[] Identifiers, IUhtMessageSite? MessageSite = null, int LineNumber = -1)
+		{
+			ValidateFindOptions(Options);
+
+			if (Identifiers.Length == 0)
+			{
+				throw new UhtIceException("Empty identifier array");
+			}
+			if (Identifiers.Length > 2)
+			{
+				if (MessageSite != null)
+				{
+					MessageSite.LogError(LineNumber, "UnrealHeaderTool only supports C++ identifiers of two or less identifiers");
+					return null;
+				}
+			}
+
+			UhtType? Type = null;
+			if (Identifiers.Length == 0)
+			{
+				Type = FindTypeTwoNamesInternal(StartingType, Options, Identifiers[0].Value.ToString(), Identifiers[1].Value.ToString());
+			}
+			else
+			{
+				Type = FindTypeInternal(StartingType, Options, Identifiers[0].Value.ToString());
+			}
+
+			if (Type == null && MessageSite != null)
+			{
+				string FullIdentifier = string.Join("::", Identifiers);
+				FindTypeError(MessageSite, LineNumber, Options, FullIdentifier);
+			}
+			return Type;
+		}
+
+		/// <summary>
+		/// Find the given type in the type hierarchy
+		/// </summary>
+		/// <param name="StartingType">Starting point for searches</param>
+		/// <param name="Options">Options controlling what is searched</param>
+		/// <param name="FirstName">First name of the type.</param>
+		/// <param name="SecondName">Second name used by delegates in classes and namespace enumerations</param>
+		/// <returns>The located type of null if not found</returns>
+		private UhtType? FindTypeTwoNamesInternal(UhtType? StartingType, UhtFindOptions Options, string FirstName, string SecondName)
+		{
+			// If we have two names
+			if (SecondName.Length > 0)
+			{
+				if (Options.HasAnyFlags(UhtFindOptions.DelegateFunction | UhtFindOptions.Enum))
+				{
+					UhtFindOptions SubOptions = UhtFindOptions.NoParents | (Options & ~UhtFindOptions.TypesMask) | (Options & UhtFindOptions.Enum);
+					if (Options.HasAnyFlags(UhtFindOptions.DelegateFunction))
+					{
+						SubOptions |= UhtFindOptions.Class;
+					}
+					UhtType? Type = FindTypeInternal(StartingType, SubOptions, FirstName);
+					if (Type == null)
+					{
+						return null;
+					}
+					if (Type is UhtEnum)
+					{
+						return Type;
+					}
+					if (Type is UhtClass)
+					{
+						return FindTypeInternal(StartingType, UhtFindOptions.DelegateFunction | UhtFindOptions.NoParents | (Options & ~UhtFindOptions.TypesMask), SecondName);
+					}
+				}
+
+				// We can't match anything at this point
+				return null;
+			}
+
+			// Perform the lookup for just a single name
+			return FindTypeInternal(StartingType, Options, FirstName);
+		}
+
+		/// <summary>
+		/// Find the given type in the type hierarchy
+		/// </summary>
+		/// <param name="StartingType">Starting point for searches</param>
+		/// <param name="Options">Options controlling what is searched</param>
+		/// <param name="Name">Name of the type.</param>
+		/// <returns>The located type of null if not found</returns>
+		public UhtType? FindTypeInternal(UhtType? StartingType, UhtFindOptions Options, string Name)
+		{
+			UhtType? Type = null;
+			if (Options.HasAnyFlags(UhtFindOptions.EngineName))
+			{
+				if (Options.HasAnyFlags(UhtFindOptions.CaseCompare))
+				{
+					Type = this.EngineNameSymbolTable.FindCasedType(StartingType, Options, Name);
+				}
+				else
+				{
+					Type = this.EngineNameSymbolTable.FindCaselessType(StartingType, Options, Name);
+				}
+			}
+			else if (Options.HasAnyFlags(UhtFindOptions.SourceName))
+			{
+				if (Options.HasAnyFlags(UhtFindOptions.CaselessCompare))
+				{
+					Type = this.SourceNameSymbolTable.FindCaselessType(StartingType, Options, Name);
+				}
+				else
+				{
+					Type = this.SourceNameSymbolTable.FindCasedType(StartingType, Options, Name);
+				}
+			}
+			else
+			{
+				throw new UhtIceException("Either EngineName or SourceName must be specified in the options");
+			}
+			return Type;
+		}
+
+		/// <summary>
+		/// Verify that the options are valid.  Will also check to make sure the symbol table has been populated.
+		/// </summary>
+		/// <param name="Options">Find options</param>
+		private void ValidateFindOptions(UhtFindOptions Options)
+		{
+			if (!Options.HasAnyFlags(UhtFindOptions.EngineName | UhtFindOptions.SourceName))
+			{
+				throw new UhtIceException("Either EngineName or SourceName must be specified in the options");
+			}
+
+			if (Options.HasAnyFlags(UhtFindOptions.CaseCompare) && Options.HasAnyFlags(UhtFindOptions.CaselessCompare))
+			{
+				throw new UhtIceException("Both CaseCompare and CaselessCompare can't be specified as FindType options");
+			}
+
+			UhtFindOptions TypeOptions = Options & UhtFindOptions.TypesMask;
+			if (TypeOptions == 0)
+			{
+				throw new UhtIceException("No type options specified");
+			}
+
+			if (!this.bSymbolTablePopulated)
+			{
+				throw new UhtIceException("Symbol table has not been populated, don't call FindType until headers are parsed.");
+			}
+		}
+
+		/// <summary>
+		/// Generate an error message for when a given symbol wasn't found.  The text will contain the list of types that the symbol must be
+		/// </summary>
+		/// <param name="MessageSite">Destination for the message</param>
+		/// <param name="LineNumber">Line number generating the error</param>
+		/// <param name="Options">Collection of required types</param>
+		/// <param name="Name">The name of the symbol</param>
+		private static void FindTypeError(IUhtMessageSite MessageSite, int LineNumber, UhtFindOptions Options, string Name)
+		{
+			List<string> Types = new List<string>();
+			if (Options.HasAnyFlags(UhtFindOptions.Enum))
+			{
+				Types.Add("'enum'");
+			}
+			if (Options.HasAnyFlags(UhtFindOptions.ScriptStruct))
+			{
+				Types.Add("'struct'");
+			}
+			if (Options.HasAnyFlags(UhtFindOptions.Class))
+			{
+				Types.Add("'class'");
+			}
+			if (Options.HasAnyFlags(UhtFindOptions.DelegateFunction))
+			{
+				Types.Add("'delegate'");
+			}
+			if (Options.HasAnyFlags(UhtFindOptions.Function))
+			{
+				Types.Add("'function'");
+			}
+			if (Options.HasAnyFlags(UhtFindOptions.Property))
+			{
+				Types.Add("'property'");
+			}
+
+			MessageSite.LogError(LineNumber, $"Unable to find {UhtUtilities.MergeTypeNames(Types, "or")} with name '{Name}'");
+		}
+
+		/// <summary>
+		/// Search for the given header file by just the file name
+		/// </summary>
+		/// <param name="Name">Name to be found</param>
+		/// <returns></returns>
+		public UhtHeaderFile? FindHeaderFile(string Name)
+		{
+			UhtHeaderFile? HeaderFile;
+			if (this.HeaderFileDictionaryInternal.TryGetValue(Name, out HeaderFile))
+			{
+				return HeaderFile;
+			}
+			return null;
+		}
+
+		#region IUHTMessageSource implementation
+		/// <summary>
+		/// Add a message to the collection of output messages
+		/// </summary>
+		/// <param name="Message">Message being added</param>
+		public void AddMessage(UhtMessage Message)
+		{
+			lock (this.Messages)
+			{
+				this.Messages.Add(Message);
+				
+				// If we aren't caching messages and this is the first message,
+				// start a task to flush the messages.
+				if (!this.bCacheMessages && this.Messages.Count == 1)
+				{
+					this.MessageTask = Task.Factory.StartNew(() => FlushMessages());
+				}
+			}
+
+			switch (Message.MessageType)
+			{
+				case UhtMessageType.Error:
+				case UhtMessageType.Ice:
+					Interlocked.Increment(ref this.ErrorCountInternal);
+					break;
+
+				case UhtMessageType.Warning:
+					Interlocked.Increment(ref this.WarningCountInternal);
+					break;
+
+				case UhtMessageType.Info:
+				case UhtMessageType.Trace:
+					break;
+			}
+		}
+		#endregion
+
+		/// <summary>
+		/// Log all the collected messages to the log/console.  If messages aren't being
+		/// cached, then this just waits until the flush task has completed.  If messages
+		/// are being cached, they are sorted by file name and line number to ensure the 
+		/// output is stable.
+		/// </summary>
+		public void LogMessages()
+		{
+			if (this.MessageTask != null)
+			{
+				this.MessageTask.Wait();
+			}
+
+			foreach (UhtMessage Message in FetchOrderedMessages())
+			{
+				LogMessage(Message);
+			}
+		}
+
+		/// <summary>
+		/// Flush all pending messages to the logger
+		/// </summary>
+		private void FlushMessages()
+		{
+			UhtMessage[]? MessageArray = null;
+			lock (this.Messages)
+			{
+				MessageArray = this.Messages.ToArray();
+				this.Messages.Clear();
+			}
+
+			foreach (UhtMessage Message in MessageArray)
+			{
+				LogMessage(Message);
+			}
+		}
+
+		/// <summary>
+		/// Log the given message
+		/// </summary>
+		/// <param name="Message">The message to be logged</param>
+		private void LogMessage(UhtMessage Message)
+		{
+			string FormattedMessage = FormatMessage(Message);
+			LogLevel LogLevel = LogLevel.Information;
+			switch (Message.MessageType)
+			{
+				default:
+				case UhtMessageType.Error:
+				case UhtMessageType.Ice:
+					LogLevel = LogLevel.Error;
+					break;
+
+				case UhtMessageType.Warning:
+					LogLevel = LogLevel.Warning;
+					break;
+
+				case UhtMessageType.Info:
+					LogLevel = LogLevel.Information;
+					break;
+
+				case UhtMessageType.Trace:
+					LogLevel = LogLevel.Trace;
+					break;
+			}
+
+			Log.Logger.Log(LogLevel, "{0}", FormattedMessage);
+		}
+
+		/// <summary>
+		/// Return all of the messages into a list
+		/// </summary>
+		/// <returns>List of all the messages</returns>
+		public List<string> CollectMessages()
+		{
+			List<string> Out = new List<string>();
+			foreach (UhtMessage Message in FetchOrderedMessages())
+			{
+				Out.Add(FormatMessage(Message));
+			}
+			return Out;
+		}
+
+		/// <summary>
+		/// Given an existing and a new instance, replace the given type in the symbol table.
+		/// This is used by the property resolution system to replace properties created during
+		/// the parsing phase that couldn't be resoled until after all headers are parsed.
+		/// </summary>
+		/// <param name="OldType"></param>
+		/// <param name="NewType"></param>
+		public void ReplaceTypeInSymbolTable(UhtType OldType, UhtType NewType)
+		{
+			this.SourceNameSymbolTable.Replace(OldType, NewType, OldType.SourceName);
+			if (OldType.EngineType.HasEngineName())
+			{
+				this.EngineNameSymbolTable.Replace(OldType, NewType, OldType.EngineName);
+			}
+		}
+
+		/// <summary>
+		/// Return an ordered enumeration of all messages.
+		/// </summary>
+		/// <returns>Enumerator</returns>
+		private IOrderedEnumerable<UhtMessage> FetchOrderedMessages()
+		{
+			List<UhtMessage> Messages = new List<UhtMessage>();
+			lock (this.Messages)
+			{
+				Messages.AddRange(this.Messages);
+				this.Messages.Clear();
+			}
+			return Messages.OrderBy(Context => Context.FilePath).ThenBy(Context => Context.LineNumber + Context.MessageSource?.MessageFragmentLineNumber);
+		}
+
+		/// <summary>
+		/// Format the given message
+		/// </summary>
+		/// <param name="Message">Message to be formatted</param>
+		/// <returns>Text of the formatted message</returns>
+		private string FormatMessage(UhtMessage Message)
+		{
+			string FilePath;
+			string FragmentPath = "";
+			int LineNumber = Message.LineNumber;
+			if (Message.FilePath != null)
+			{
+				FilePath = Message.FilePath;
+			}
+			else if (Message.MessageSource != null)
+			{
+				if (Message.MessageSource.bMessageIsFragment)
+				{
+					if (this.bRelativePathInLog)
+					{
+						FilePath = Message.MessageSource.MessageFragmentFilePath;
+					}
+					else
+					{
+						FilePath = Message.MessageSource.MessageFragmentFullFilePath;
+					}
+					FragmentPath = $"[{Message.MessageSource.MessageFilePath}]";
+					LineNumber += Message.MessageSource.MessageFragmentLineNumber;
+				}
+				else
+				{
+					if (this.bRelativePathInLog)
+					{
+						FilePath = Message.MessageSource.MessageFilePath;
+					}
+					else
+					{
+						FilePath = Message.MessageSource.MessageFullFilePath;
+					}
+				}
+			}
+			else
+			{
+				FilePath = "UnknownSource";
+			}
+
+			switch (Message.MessageType)
+			{
+				case UhtMessageType.Error:
+					return $"{FilePath}({LineNumber}){FragmentPath}: Error: {Message.Message}";
+				case UhtMessageType.Warning:
+					return $"{FilePath}({LineNumber}){FragmentPath}: Warning: {Message.Message}";
+				case UhtMessageType.Info:
+					return $"{FilePath}({LineNumber}){FragmentPath}: Info: {Message.Message}";
+				case UhtMessageType.Trace:
+					return $"{FilePath}({LineNumber}){FragmentPath}: Trace: {Message.Message}";
+				default:
+				case UhtMessageType.Ice:
+					return $"{FilePath}({LineNumber}){FragmentPath}:  Error: Internal Compiler Error - {Message.Message}";
+			}
+		}
+
+		/// <summary>
+		/// Handle the given exception with the provided message context
+		/// </summary>
+		/// <param name="MessageSource">Context for the exception.  Required to handled all exceptions other than UHTException</param>
+		/// <param name="E">Exception being handled</param>
+		private void HandleException(IUhtMessageSource? MessageSource, Exception E)
+		{
+			switch (E)
+			{
+				case UhtException UHTException:
+					UhtMessage Message = UHTException.UhtMessage;
+					if (Message.MessageSource == null)
+					{
+						Message.MessageSource = MessageSource;
+					}
+					AddMessage(Message);
+					break;
+
+				case JsonException JsonException:
+					AddMessage(UhtMessage.MakeMessage(UhtMessageType.Error, MessageSource, null, (int)(JsonException.LineNumber + 1 ?? 1), JsonException.Message));
+					break;
+
+				default:
+					//Log.TraceInformation("{0}", E.StackTrace);
+					AddMessage(UhtMessage.MakeMessage(UhtMessageType.Ice, MessageSource, null, 1, $"{E.GetType().ToString()} - {E.Message}"));
+					break;
+			}
+		}
+
+		/// <summary>
+		/// Return the normalized path converted to a full path if possible. 
+		/// Code should NOT depend on a full path being returned.
+		/// 
+		/// In general, it is assumed that during normal UHT, all paths are already full paths.
+		/// Only the test harness deals in relative paths.
+		/// </summary>
+		/// <param name="FilePath">Path to normalize</param>
+		/// <returns>Normalized path possibly converted to a full path.</returns>
+		private string GetNormalizedFullFilePath(string FilePath)
+		{
+			return NormalizePath(this.FileManager!.GetFullFilePath(FilePath));
+		}
+
+		private string NormalizePath(string FilePath)
+		{
+			return FilePath.Replace('\\', '/');
+		}
+
+		private UhtHeaderFileParser? ParseHeaderFile(UhtHeaderFile HeaderFile)
+		{
+			UhtHeaderFileParser? Parser = null;
+			TryAlways(HeaderFile.MessageSource, () =>
+			{
+				HeaderFile.Read();
+				Parser = UhtHeaderFileParser.Parse(HeaderFile);
+			});
+			return Parser;
+		}
+
+		#region Run steps
+		private void StepReadManifestFile(string ManifestFilePath)
+		{
+			this.ManifestFile = new UhtManifestFile(this, ManifestFilePath);
+
+			Try(this.ManifestFile.MessageSource, () =>
+			{
+				Log.Logger.LogTrace("Step - Read Manifest File");
+
+				this.ManifestFile.Read();
+			});
+		}
+
+		private void StepPrepareModules()
+		{
+			if (this.ManifestFile == null || this.ManifestFile.Manifest == null)
+			{
+				return;
+			}
+
+			Try(this.ManifestFile.MessageSource, () =>
+			{
+				Log.Logger.LogTrace("Step - Prepare Modules");
+
+				foreach (UHTManifest.Module Module in this.ManifestFile.Manifest.Modules)
+				{
+					EPackageFlags PackageFlags = EPackageFlags.ContainsScript | EPackageFlags.Compiling;
+
+					switch (Module.OverrideModuleType)
+					{
+						case EPackageOverrideType.None:
+							switch (Module.ModuleType)
+							{
+								case UHTModuleType.GameEditor:
+								case UHTModuleType.EngineEditor:
+									PackageFlags |= EPackageFlags.EditorOnly;
+									break;
+
+								case UHTModuleType.GameDeveloper:
+								case UHTModuleType.EngineDeveloper:
+									PackageFlags |= EPackageFlags.Developer;
+									break;
+
+								case UHTModuleType.GameUncooked:
+								case UHTModuleType.EngineUncooked:
+									PackageFlags |= EPackageFlags.UncookedOnly;
+									break;
+							}
+							break;
+
+						case EPackageOverrideType.EditorOnly:
+							PackageFlags |= EPackageFlags.EditorOnly;
+							break;
+
+						case EPackageOverrideType.EngineDeveloper:
+						case EPackageOverrideType.GameDeveloper:
+							PackageFlags |= EPackageFlags.Developer;
+							break;
+
+						case EPackageOverrideType.EngineUncookedOnly:
+						case EPackageOverrideType.GameUncookedOnly:
+							PackageFlags |= EPackageFlags.UncookedOnly;
+							break;
+					}
+
+					UhtPackage Package = new UhtPackage(this, Module, PackageFlags);
+					this.PackagesInternal.Add(Package);
+				}
+			});
+		}
+
+		private void StepPrepareHeaders(UhtPackage Package, IEnumerable<string> HeaderFiles, UhtHeaderFileType HeaderFileType)
+		{
+			if (Package.Module == null)
+			{
+				return;
+			}
+
+			string TypeDirectory = HeaderFileType.ToString() + '/';
+			string NormalizedModuleBaseFullFilePath = GetNormalizedFullFilePath(Package.Module.BaseDirectory);
+			foreach (string HeaderFilePath in HeaderFiles)
+			{
+
+				// Make sure this isn't a duplicate
+				string NormalizedFullFilePath = GetNormalizedFullFilePath(HeaderFilePath);
+				string FileName = Path.GetFileName(NormalizedFullFilePath);
+				UhtHeaderFile? ExistingHeaderFile;
+				if (HeaderFileDictionaryInternal.TryGetValue(FileName, out ExistingHeaderFile) && ExistingHeaderFile != null)
+				{
+					string NormalizedExistingFullFilePath = GetNormalizedFullFilePath(ExistingHeaderFile.FilePath);
+					if (string.Compare(NormalizedFullFilePath, NormalizedExistingFullFilePath, true) != 0)
+					{
+						IUhtMessageSite Site = (IUhtMessageSite?)this.ManifestFile ?? this;
+						Site.LogError($"Two headers with the same name is not allowed. '{HeaderFilePath}' conflicts with '{ExistingHeaderFile.FilePath}'");
+						continue;
+					}
+				}
+
+				// Create the header file and add to the collections
+				UhtHeaderFile HeaderFile = new UhtHeaderFile(Package, HeaderFilePath);
+				HeaderFile.HeaderFileType = HeaderFileType;
+				HeaderFilesInternal.Add(HeaderFile);
+				HeaderFileDictionaryInternal.Add(FileName, HeaderFile);
+				Package.AddChild(HeaderFile);
+
+				// Save metadata for the class path, both for it's include path and relative to the module base directory
+				if (NormalizedFullFilePath.StartsWith(NormalizedModuleBaseFullFilePath, true, null))
+				{
+					int StripLength = NormalizedModuleBaseFullFilePath.Length;
+					if (StripLength < NormalizedFullFilePath.Length && NormalizedFullFilePath[StripLength] == '/')
+					{
+						++StripLength;
+					}
+
+					HeaderFile.ModuleRelativeFilePath = NormalizedFullFilePath.Substring(StripLength);
+
+					if (NormalizedFullFilePath.Substring(StripLength).StartsWith(TypeDirectory, true, null))
+					{
+						StripLength += TypeDirectory.Length;
+					}
+
+					HeaderFile.IncludeFilePath = NormalizedFullFilePath.Substring(StripLength);
+				}
+			}
+		}
+
+		private void StepPrepareHeaders()
+		{
+			if (this.ManifestFile == null)
+			{
+				return;
+			}
+
+			Try(this.ManifestFile.MessageSource, () =>
+			{
+				Log.Logger.LogTrace("Step - Prepare Headers");
+
+				foreach (UhtPackage Package in this.PackagesInternal)
+				{
+					if (Package.Module != null)
+					{
+						StepPrepareHeaders(Package, Package.Module.ClassesHeaders, UhtHeaderFileType.Classes);
+						StepPrepareHeaders(Package, Package.Module.PublicHeaders, UhtHeaderFileType.Public);
+						StepPrepareHeaders(Package, Package.Module.InternalHeaders, UhtHeaderFileType.Internal);
+						StepPrepareHeaders(Package, Package.Module.PrivateHeaders, UhtHeaderFileType.Private);
+					}
+				}
+
+				// Locate the NoExportTypes.h file and add it to every other header file
+				if (this.HeaderFileDictionaryInternal.TryGetValue("NoExportTypes.h", out UhtHeaderFile? NoExportTypes))
+				{
+					foreach (UhtPackage Package in this.PackagesInternal)
+					{
+						foreach (UhtHeaderFile HeaderFile in Package.Children)
+						{
+							if (HeaderFile != NoExportTypes)
+							{
+								HeaderFile.AddReferencedHeader(NoExportTypes);
+							}
+						}
+					}
+				}
+
+			});
+		}
+
+		private void StepParseHeaders()
+		{
+			if (this.bHasErrors)
+			{
+				return;
+			}
+
+			Log.Logger.LogTrace("Step - Parse Headers");
+
+			if (this.bGoWide)
+			{
+				Parallel.ForEach(this.HeaderFilesInternal, HeaderFile =>
+				{
+					ParseHeaderFile(HeaderFile);
+				});
+			}
+			else
+			{
+				foreach (UhtHeaderFile HeaderFile in this.HeaderFilesInternal)
+				{
+					ParseHeaderFile(HeaderFile);
+				}
+			}
+		}
+
+		private void StepPopulateTypeTable()
+		{
+			Try(null, () =>
+			{
+				Log.Logger.LogTrace("Step - Populate symbol table");
+
+				this.SourceNameSymbolTable = new UhtSymbolTable(this.TypeCount);
+				this.EngineNameSymbolTable = new UhtSymbolTable(this.TypeCount);
+
+				PopulateSymbolTable();
+
+				this.UObjectInternal = (UhtClass?)FindType(null, UhtFindOptions.SourceName | UhtFindOptions.Class, "UObject");
+				this.UClassInternal = (UhtClass?)FindType(null, UhtFindOptions.SourceName | UhtFindOptions.Class, "UClass");
+				this.UInterfaceInternal = (UhtClass?)FindType(null, UhtFindOptions.SourceName | UhtFindOptions.Class, "UInterface");
+				this.IInterfaceInternal = (UhtClass?)FindType(null, UhtFindOptions.SourceName | UhtFindOptions.Class, "IInterface");
+				this.AActor = (UhtClass?)FindType(null, UhtFindOptions.SourceName | UhtFindOptions.Class, "AActor");
+			});
+		}
+
+		private void StepResolveBases()
+		{
+			StepForAllHeaders(HeaderFile => Resolve(HeaderFile, UhtResolvePhase.Bases));
+		}
+
+		private void StepResolveInvalidCheck()
+		{
+			StepForAllHeaders(HeaderFile => Resolve(HeaderFile, UhtResolvePhase.InvalidCheck));
+		}
+
+		private void StepResolveProperties()
+		{
+			StepForAllHeaders(HeaderFile => Resolve(HeaderFile, UhtResolvePhase.Properties));
+		}
+
+		private void StepResolveFinal()
+		{
+			StepForAllHeaders(HeaderFile => Resolve(HeaderFile, UhtResolvePhase.Final));
+		}
+
+		private void StepResolveValidate()
+		{
+			StepForAllHeaders(HeaderFile => UhtType.ValidateType(HeaderFile, UhtValidationOptions.None));
+		}
+
+		private void StepCollectReferences()
+		{
+			StepForAllHeaders(HeaderFile =>
+			{
+				foreach (UhtType Child in HeaderFile.Children)
+				{
+					Child.CollectReferences(HeaderFile.References);
+				}
+				foreach (UhtHeaderFile RefHeaderFile in HeaderFile.References.ReferencedHeaders)
+				{
+					HeaderFile.AddReferencedHeader(RefHeaderFile);
+				}
+			});
+		}
+
+		private void Resolve(UhtHeaderFile HeaderFile, UhtResolvePhase ResolvePhase)
+		{
+			TryAlways(HeaderFile.MessageSource, () =>
+			{
+				HeaderFile.Resolve(ResolvePhase);
+			});
+		}
+
+		private delegate void StepDelegate(UhtHeaderFile HeaderFile);
+
+		private void StepForAllHeaders(StepDelegate Delegate)
+		{
+			if (this.bHasErrors)
+			{
+				return;
+			}
+
+			if (this.bGoWide)
+			{
+				Parallel.ForEach(this.HeaderFilesInternal, HeaderFile =>
+				{
+					Delegate(HeaderFile);
+				});
+			}
+			else
+			{
+				foreach (UhtHeaderFile HeaderFile in this.HeaderFilesInternal)
+				{
+					Delegate(HeaderFile);
+				}
+			}
+		}
+		#endregion
+
+		#region Symbol table initialization
+		private void PopulateSymbolTable()
+		{
+			foreach (UhtHeaderFile HeaderFile in this.HeaderFilesInternal)
+			{
+				AddTypeToSymbolTable(HeaderFile);
+			}
+			this.bSymbolTablePopulated = true;
+		}
+
+		private void AddTypeToSymbolTable(UhtType Type)
+		{
+			UhtEngineType EngineExtendedType = Type.EngineType;
+
+			if (Type is UhtEnum Enum)
+			{
+				//COMPATIBILITY-TODO: We can get more reliable results by only adding regular enums to the table
+				// and then in the lookup code in the property system to look for the '::' and just lookup
+				// the raw enum name.  In UHT we only care about the enum and not the value.
+				//
+				// The current algorithm has issues with two cases:
+				//
+				//		EnumNamespaceName::EnumTypeName::Value - Where the enum type name is included with a namespace enum
+				//		EnumName::Value - Where the value is defined in terms that can't be parsed.  The -1 check causes it
+				//			to be kicked out.
+				//if (Enum.CppForm == UhtEnumCppForm.Regular)
+				//{
+				//	foreach (UhtEnumValue Value in Enum.EnumValues)
+				//	{
+				//		this.RegularEnumValueLookup.Add(Value.Name.ToString(), Enum);
+				//	}
+				//}
+				bool bAddShortNames = Enum.CppForm == UhtEnumCppForm.Namespaced || Enum.CppForm == UhtEnumCppForm.EnumClass;
+				string CheckName = $"{Enum.SourceName}::";
+				foreach (UhtEnumValue Value in Enum.EnumValues)
+				{
+					this.FullEnumValueLookup.Add(Value.Name, new EnumAndValue { Enum = Enum, Value = Value.Value });
+					if (bAddShortNames)
+					{
+						if (Value.Name.StartsWith(CheckName))
+						{
+							this.ShortEnumValueLookup.TryAdd(Value.Name.Substring(CheckName.Length), Enum);
+						}
+					}
+				}
+			}
+
+			if (EngineExtendedType.FindOptions() != 0)
+			{
+				if (EngineExtendedType.MustNotBeReserved())
+				{
+					if (ReservedNames.Contains(Type.EngineName))
+					{
+						Type.HeaderFile.LogError(Type.LineNumber, $"{EngineExtendedType.CapitalizedText()} '{Type.EngineName}' uses a reserved type name.");
+					}
+				}
+
+				if (EngineExtendedType.HasEngineName() && EngineExtendedType.MustBeUnique())
+				{
+					UhtType? ExistingType = this.EngineNameSymbolTable.FindCaselessType(null, EngineExtendedType.MustBeUniqueFindOptions(), Type.EngineName);
+					if (ExistingType != null)
+					{
+						Type.HeaderFile.LogError(Type.LineNumber, string.Format("{0} '{1}' shares engine name '{2}' with {6} '{3}' in {4}(5)",
+							EngineExtendedType.CapitalizedText(), Type.SourceName, Type.EngineName, ExistingType.SourceName,
+							ExistingType.HeaderFile.FilePath, ExistingType.LineNumber, ExistingType.EngineType.LowercaseText()));
+					}
+				}
+
+				if (Type.bVisibleType)
+				{
+					this.SourceNameSymbolTable.Add(Type, Type.SourceName);
+					if (EngineExtendedType.HasEngineName())
+					{
+						this.EngineNameSymbolTable.Add(Type, Type.EngineName);
+					}
+				}
+			}
+
+			if (EngineExtendedType.AddChildrenToSymbolTable())
+			{
+				foreach (UhtType Child in Type.Children)
+				{
+					AddTypeToSymbolTable(Child);
+				}
+			}
+		}
+		#endregion
+
+		#region Topological sort of the header files
+		private enum TopologicalState
+		{
+			Unmarked,
+			Temporary,
+			Permanent,
+		}
+
+		private void TopologicalRecursion(List<TopologicalState> States, UhtHeaderFile First, UhtHeaderFile Visit)
+		{
+			foreach (UhtHeaderFile Referenced in Visit.ReferencedHeadersNoLock)
+			{
+				if (States[Referenced.HeaderFileTypeIndex] == TopologicalState.Temporary)
+				{
+					First.LogError($"'{Visit.FilePath}' includes/requires '{Referenced.FilePath}'");
+					if (First != Referenced)
+					{
+						TopologicalRecursion(States, First, Referenced);
+					}
+					break;
+				}
+			}
+		}
+
+		private UhtHeaderFile? TopologicalVisit(List<TopologicalState> States, UhtHeaderFile Visit)
+		{
+			switch (States[Visit.HeaderFileTypeIndex])
+			{
+				case TopologicalState.Unmarked:
+					UhtHeaderFile? Recursion = null;
+					States[Visit.HeaderFileTypeIndex] = TopologicalState.Temporary;
+					foreach (UhtHeaderFile Referenced in Visit.ReferencedHeadersNoLock)
+					{
+						if (Visit != Referenced)
+						{
+							UhtHeaderFile? Out = TopologicalVisit(States, Referenced);
+							if (Out != null)
+							{
+								Recursion = Out;
+								break;
+							}
+						}
+					}
+					States[Visit.HeaderFileTypeIndex] = TopologicalState.Permanent;
+					this.SortedHeaderFilesInternal.Add(Visit);
+					return null;
+
+				case TopologicalState.Temporary:
+					return Visit;
+
+				case TopologicalState.Permanent:
+					return null;
+
+				default:
+					throw new UhtIceException("Unknown topological state");
+			}
+		}
+
+		private void TopologicalSortHeaderFiles()
+		{
+			Try(null, () =>
+			{
+				Log.Logger.LogTrace("Step - Topological Sort Header Files");
+
+				// Initialize a scratch table for topological states
+				this.SortedHeaderFilesInternal.Capacity = this.HeaderFileTypeCount;
+				List<TopologicalState> States = new List<TopologicalState>(this.HeaderFileTypeCount);
+				for (int Index = 0; Index < this.HeaderFileTypeCount; ++Index)
+				{
+					States.Add(TopologicalState.Unmarked);
+				}
+
+				foreach (UhtHeaderFile HeaderFile in this.HeaderFiles)
+				{
+					if (States[HeaderFile.HeaderFileTypeIndex] == TopologicalState.Unmarked)
+					{
+						UhtHeaderFile? Recursion = TopologicalVisit(States, HeaderFile);
+						if (Recursion != null)
+						{
+							HeaderFile.LogError("Circular dependency detected:");
+							TopologicalRecursion(States, Recursion, Recursion);
+							return;
+						}
+					}
+				}
+			});
+		}
+		#endregion
+
+		#region Validation helpers
+		private HashSet<UhtScriptStruct> ScriptStructsValidForNet = new HashSet<UhtScriptStruct>();
+
+		/// <summary>
+		/// Validate that the given referenced script structure is valid for network operations.  If the structure
+		/// is valid, then the result will be cached.  It not valid, errors will be generated each time the structure
+		/// is referenced.
+		/// </summary>
+		/// <param name="ReferencingProperty">The property referencing a structure</param>
+		/// <param name="ReferencedScriptStruct">The script structure being referenced</param>
+		/// <returns></returns>
+		public bool ValidateScriptStructOkForNet(UhtProperty ReferencingProperty, UhtScriptStruct ReferencedScriptStruct)
+		{
+
+			// Check for existing value
+			lock (this.ScriptStructsValidForNet)
+			{
+				if (this.ScriptStructsValidForNet.Contains(ReferencedScriptStruct))
+				{
+					return true;
+				}
+			}
+
+			bool bIsStructValid = true;
+
+			// Check the super chain structure
+			UhtScriptStruct? SuperScriptStruct = ReferencedScriptStruct.SuperScriptStruct;
+			if (SuperScriptStruct != null)
+			{
+				if (!ValidateScriptStructOkForNet(ReferencingProperty, SuperScriptStruct))
+				{
+					bIsStructValid = false;
+				}
+			}
+
+			// Check the structure properties
+			foreach (UhtProperty Property in ReferencedScriptStruct.Properties)
+			{
+				if (!Property.ValidateStructPropertyOkForNet(ReferencingProperty))
+				{
+					bIsStructValid = false;
+					break;
+				}
+			}
+
+			// Save the results
+			if (bIsStructValid)
+			{
+				lock (this.ScriptStructsValidForNet)
+				{
+					this.ScriptStructsValidForNet.Add(ReferencedScriptStruct);
+				}
+			}
+			return bIsStructValid;
+		}
+		#endregion
+
+		#region Exporting
+		private void StepExport()
+		{
+			Try(null, () =>
+			{
+				Log.Logger.LogTrace("Step - Exports");
+				UhtExportOptions Options = UhtExportOptions.None;
+
+				foreach (UhtExporter Exporter in UhtExporterTable.Instance)
+				{
+					bool Run = false;
+					if (!this.ExporterStates.TryGetValue(Exporter.Name, out Run))
+					{
+						Run = UhtConfig.Instance.IsExporterEnabled(Exporter.Name) || Exporter.Options.HasAnyFlags(UhtExporterOptions.Default);
+					}
+
+					if (Run)
+					{
+						Log.Logger.LogTrace($"       Running exporter {Exporter.Name}");
+						UhtExportFactory Factory = new UhtExportFactory(this, Exporter, Options);
+						Factory.Run();
+					}
+					else
+					{
+						Log.Logger.LogTrace($"       Exporter {Exporter.Name} skipped");
+					}
+				}
+			});
+		}
+		#endregion
+	}
+}
diff --git a/Engine/Source/Programs/Shared/EpicGames.UHT/Utils/UhtStringBuilderExtensions.cs b/Engine/Source/Programs/Shared/EpicGames.UHT/Utils/UhtStringBuilderExtensions.cs
new file mode 100644
--- /dev/null
+++ b/Engine/Source/Programs/Shared/EpicGames.UHT/Utils/UhtStringBuilderExtensions.cs
@@ -0,0 +1,600 @@
+﻿// Copyright Epic Games, Inc. All Rights Reserved.
+
+using EpicGames.Core;
+using EpicGames.UHT.Types;
+using System;
+using System.Collections.Generic;
+using System.Text;
+
+namespace EpicGames.UHT.Utils
+{
+	static class UhtStringBuilderExtensions
+	{
+		/// <summary>
+		/// String of tabs used to generate code with proper indentation
+		/// </summary>
+		public static StringView TabsString = new StringView(new string('\t', 128));
+
+		/// <summary>
+		/// String of spaces used to generate code with proper indentation
+		/// </summary>
+		public static StringView SpacesString = new StringView(new string(' ', 128));
+
+		/// <summary>
+		/// Append tabs to the builder
+		/// </summary>
+		/// <param name="Builder">Destination builder</param>
+		/// <param name="Tabs">Number of tabs to insert</param>
+		/// <returns>Destination builder</returns>
+		/// <exception cref="ArgumentOutOfRangeException">Thrown if the number of tabs is out of range</exception>
+		public static StringBuilder AppendTabs(this StringBuilder Builder, int Tabs)
+		{
+			if (Tabs < 0 || Tabs > TabsString.Length)
+			{
+				throw new ArgumentOutOfRangeException();
+			}
+			else if (Tabs > 0)
+			{
+				Builder.Append(TabsString.Span.Slice(0, Tabs));
+			}
+			return Builder;
+		}
+
+		/// <summary>
+		/// Append spaces to the builder
+		/// </summary>
+		/// <param name="Builder">Destination builder</param>
+		/// <param name="Spaces">Number of spaces to insert</param>
+		/// <returns>Destination builder</returns>
+		/// <exception cref="ArgumentOutOfRangeException">Thrown if the number of spaces is out of range</exception>
+		public static StringBuilder AppendSpaces(this StringBuilder Builder, int Spaces)
+		{
+			if (Spaces < 0 || Spaces > SpacesString.Length)
+			{
+				throw new ArgumentOutOfRangeException();
+			}
+			else if (Spaces > 0)
+			{
+				Builder.Append(SpacesString.Span.Slice(0, Spaces));
+			}
+			return Builder;
+		}
+
+		/// <summary>
+		/// Append a name declaration to the builder
+		/// </summary>
+		/// <param name="Builder">Destination builder</param>
+		/// <param name="NamePrefix">Name prefix</param>
+		/// <param name="Name">Name</param>
+		/// <param name="NameSuffix">Optional name suffix</param>
+		/// <returns>Destination builder</returns>
+		public static StringBuilder AppendNameDecl(this StringBuilder Builder, string? NamePrefix, string Name, string? NameSuffix)
+		{
+			return Builder.Append(NamePrefix).Append(Name).Append(NameSuffix);
+		}
+
+		/// <summary>
+		/// Append a name declaration to the builder
+		/// </summary>
+		/// <param name="Builder">Destination builder</param>
+		/// <param name="Context">Property context used to get the name prefix</param>
+		/// <param name="Name">Name</param>
+		/// <param name="NameSuffix">Optional name suffix</param>
+		/// <returns>Destination builder</returns>
+		public static StringBuilder AppendNameDecl(this StringBuilder Builder, IUhtPropertyMemberContext Context, string Name, string NameSuffix)
+		{
+			return Builder.AppendNameDecl(Context.NamePrefix, Name, NameSuffix);
+		}
+
+		/// <summary>
+		/// Append a name definition to the builder
+		/// </summary>
+		/// <param name="Builder">Destination builder</param>
+		/// <param name="StaticsName">Optional name of the statics block which will be output in the form of "StaticsName::"</param>
+		/// <param name="NamePrefix">Name prefix</param>
+		/// <param name="Name">Name</param>
+		/// <param name="NameSuffix">Optional name suffix</param>
+		/// <returns>Destination builder</returns>
+		public static StringBuilder AppendNameDef(this StringBuilder Builder, string? StaticsName, string? NamePrefix, string Name, string? NameSuffix)
+		{
+			if (!string.IsNullOrEmpty(StaticsName))
+			{
+				Builder.Append(StaticsName).Append("::");
+			}
+			return Builder.AppendNameDecl(NamePrefix, Name, NameSuffix);
+		}
+
+		/// <summary>
+		/// Append a name definition to the builder
+		/// </summary>
+		/// <param name="Builder">Destination builder</param>
+		/// <param name="Context">Property context used to get the statics name and name prefix</param>
+		/// <param name="Name">Name</param>
+		/// <param name="NameSuffix">Optional name suffix</param>
+		/// <returns>Destination builder</returns>
+		public static StringBuilder AppendNameDef(this StringBuilder Builder, IUhtPropertyMemberContext Context, string Name, string NameSuffix)
+		{
+			return Builder.AppendNameDef(Context.StaticsName, Context.NamePrefix, Name, NameSuffix);
+		}
+
+		/// <summary>
+		/// Append the meta data parameters.  This is intended to be used as arguments to a function call.
+		/// </summary>
+		/// <param name="Builder">Destination builder</param>
+		/// <param name="Type">Source type containing the meta data</param>
+		/// <param name="StaticsName">Optional name of the statics block which will be output in the form of "StaticsName::"</param>
+		/// <param name="NamePrefix">Name prefix</param>
+		/// <param name="Name">Name</param>
+		/// <param name="NameSuffix">Optional name suffix</param>
+		/// <param name="MetaNameSuffix">Suffix to be added to the meta data name</param>
+		/// <returns>Destination builder</returns>
+		private static StringBuilder AppendMetaDataParams(this StringBuilder Builder, UhtType Type, string? StaticsName, string? NamePrefix, string Name, string? NameSuffix, string? MetaNameSuffix)
+		{
+			if (!Type.MetaData.IsEmpty())
+			{
+				return Builder
+					.Append("METADATA_PARAMS(")
+					.AppendNameDef(StaticsName, NamePrefix, Name, NameSuffix).Append(MetaNameSuffix)
+					.Append(", UE_ARRAY_COUNT(")
+					.AppendNameDef(StaticsName, NamePrefix, Name, NameSuffix).Append(MetaNameSuffix)
+					.Append("))");
+			}
+			else
+			{
+				return Builder.Append("METADATA_PARAMS(nullptr, 0)");
+			}
+		}
+
+		/// <summary>
+		/// Append the meta data parameters.  This is intended to be used as arguments to a function call.
+		/// </summary>
+		/// <param name="Builder">Destination builder</param>
+		/// <param name="Type">Source type containing the meta data</param>
+		/// <param name="StaticsName">Optional name of the statics block which will be output in the form of "StaticsName::"</param>
+		/// <param name="Name">Name</param>
+		/// <returns>Destination builder</returns>
+		public static StringBuilder AppendMetaDataParams(this StringBuilder Builder, UhtType Type, string? StaticsName, string Name)
+		{
+			return Builder.AppendMetaDataParams(Type, StaticsName, null, Name, null, null);
+		}
+
+		/// <summary>
+		/// Append the meta data parameters.  This is intended to be used as arguments to a function call.
+		/// </summary>
+		/// <param name="Builder">Destination builder</param>
+		/// <param name="Property">Source type containing the meta data</param>
+		/// <param name="Context">Property context used to get the statics name and name prefix</param>
+		/// <param name="Name">Name</param>
+		/// <param name="NameSuffix">Optional name suffix</param>
+		/// <returns>Destination builder</returns>
+		public static StringBuilder AppendMetaDataParams(this StringBuilder Builder, UhtProperty Property, IUhtPropertyMemberContext Context, string Name, string NameSuffix)
+		{
+			return Builder.AppendMetaDataParams(Property, Context.StaticsName, Context.NamePrefix, Name, NameSuffix, Context.MetaDataSuffix);
+		}
+
+		/// <summary>
+		/// Append the meta data declaration.
+		/// </summary>
+		/// <param name="Builder">Destination builder</param>
+		/// <param name="Type">Source type containing the meta data</param>
+		/// <param name="NamePrefix">Name prefix</param>
+		/// <param name="Name">Name</param>
+		/// <param name="NameSuffix">Optional name suffix</param>
+		/// <param name="MetaNameSuffix">Optional meta data name suffix</param>
+		/// <param name="Tabs">Number of tabs to indent</param>
+		/// <returns>Destination builder</returns>
+		private static StringBuilder AppendMetaDataDecl(this StringBuilder Builder, UhtType Type, string? NamePrefix, string Name, string? NameSuffix, string? MetaNameSuffix, int Tabs)
+		{
+			if (!Type.MetaData.IsEmpty())
+			{
+				Builder.Append("#if WITH_METADATA\r\n");
+				Builder.AppendTabs(Tabs).Append("static const UECodeGen_Private::FMetaDataPairParam ").AppendNameDecl(NamePrefix, Name, NameSuffix).Append(MetaNameSuffix).Append("[];\r\n");
+				Builder.Append("#endif\r\n");
+			}
+			return Builder;
+		}
+
+		/// <summary>
+		/// Append the meta data declaration
+		/// </summary>
+		/// <param name="Builder">Destination builder</param>
+		/// <param name="Type">Source type containing the meta data</param>
+		/// <param name="Name">Name</param>
+		/// <param name="Tabs">Number of tabs to indent</param>
+		/// <returns>Destination builder</returns>
+		public static StringBuilder AppendMetaDataDecl(this StringBuilder Builder, UhtType Type, string Name, int Tabs)
+		{
+			return Builder.AppendMetaDataDecl(Type, null, Name, null, null, Tabs);
+		}
+
+		/// <summary>
+		/// Append the meta data declaration
+		/// </summary>
+		/// <param name="Builder">Destination builder</param>
+		/// <param name="Property">Source type containing the meta data</param>
+		/// <param name="Context">Property context used to get the statics name and name prefix</param>
+		/// <param name="Name">Name</param>
+		/// <param name="NameSuffix">Optional name suffix</param>
+		/// <param name="Tabs">Number of tabs to indent</param>
+		/// <returns>Destination builder</returns>
+		public static StringBuilder AppendMetaDataDecl(this StringBuilder Builder, UhtProperty Property, IUhtPropertyMemberContext Context, string Name, string NameSuffix, int Tabs)
+		{
+			return Builder.AppendMetaDataDecl(Property, Context.NamePrefix, Name, NameSuffix, Context.MetaDataSuffix, Tabs);
+		}
+
+		/// <summary>
+		/// Append the meta data definition.
+		/// </summary>
+		/// <param name="Builder">Destination builder</param>
+		/// <param name="Type">Source type containing the meta data</param>
+		/// <param name="StaticsName">Optional name of the statics block which will be output in the form of "StaticsName::"</param>
+		/// <param name="NamePrefix">Name prefix</param>
+		/// <param name="Name">Name</param>
+		/// <param name="NameSuffix">Optional name suffix</param>
+		/// <param name="MetaNameSuffix">Optional meta data name suffix</param>
+		/// <param name="Tabs">Number of tabs to indent</param>
+		/// <returns>Destination builder</returns>
+		private static StringBuilder AppendMetaDataDef(this StringBuilder Builder, UhtType Type, string? StaticsName, string? NamePrefix, string Name, string? NameSuffix, string? MetaNameSuffix, int Tabs)
+		{
+			if (!Type.MetaData.IsEmpty())
+			{
+				List<KeyValuePair<string, string>> SortedMetaData = Type.MetaData.GetSorted();
+
+				Builder.Append("#if WITH_METADATA\r\n");
+				Builder.AppendTabs(Tabs).Append("const UECodeGen_Private::FMetaDataPairParam ").AppendNameDef(StaticsName, NamePrefix, Name, NameSuffix).Append(MetaNameSuffix).Append("[] = {\r\n");
+
+				foreach (KeyValuePair<string, string> Kvp in SortedMetaData)
+				{
+					Builder.AppendTabs(Tabs + 1).Append("{ ").AppendUTF8LiteralString(Kvp.Key).Append(", ").AppendUTF8LiteralString(Kvp.Value).Append(" },\r\n");
+				}
+
+				Builder.AppendTabs(Tabs).Append("};\r\n");
+				Builder.Append("#endif\r\n");
+			}
+			return Builder;
+		}
+
+		/// <summary>
+		/// Append the meta data definition
+		/// </summary>
+		/// <param name="Builder">Destination builder</param>
+		/// <param name="Type">Source type containing the meta data</param>
+		/// <param name="StaticsName">Optional name of the statics block which will be output in the form of "StaticsName::"</param>
+		/// <param name="Name">Name</param>
+		/// <param name="Tabs">Number of tabs to indent</param>
+		/// <returns>Destination builder</returns>
+		public static StringBuilder AppendMetaDataDef(this StringBuilder Builder, UhtType Type, string? StaticsName, string Name, int Tabs)
+		{
+			return Builder.AppendMetaDataDef(Type, StaticsName, null, Name, null, null, Tabs);
+		}
+
+		/// <summary>
+		/// Append the meta data definition
+		/// </summary>
+		/// <param name="Builder">Destination builder</param>
+		/// <param name="Property">Source type containing the meta data</param>
+		/// <param name="Context">Property context used to get the statics name and name prefix</param>
+		/// <param name="Name">Name</param>
+		/// <param name="NameSuffix">Optional name suffix</param>
+		/// <param name="Tabs">Number of tabs to indent</param>
+		/// <returns>Destination builder</returns>
+		public static StringBuilder AppendMetaDataDef(this StringBuilder Builder, UhtProperty Property, IUhtPropertyMemberContext Context, string Name, string NameSuffix, int Tabs)
+		{
+			return Builder.AppendMetaDataDef(Property, Context.StaticsName, Context.NamePrefix, Name, NameSuffix, Context.MetaDataSuffix, Tabs);
+		}
+
+		/// <summary>
+		/// Append the given text as a UTF8 encoded string
+		/// </summary>
+		/// <param name="Builder">Destination builder</param>
+		/// <param name="bUseText">If false, don't encode the text but include a nullptr</param>
+		/// <param name="Text">Text to include or an empty string if null.</param>
+		/// <returns>Destination builder</returns>
+		public static StringBuilder AppendUTF8LiteralString(this StringBuilder Builder, bool bUseText, string? Text)
+		{
+			if (!bUseText)
+			{
+				Builder.Append("nullptr");
+			}
+			else if (Text == null)
+			{
+				Builder.Append("");
+			}
+			else
+			{
+				Builder.AppendUTF8LiteralString(Text);
+			}
+			return Builder;
+		}
+
+		/// <summary>
+		/// Append the given text as a UTF8 encoded string
+		/// </summary>
+		/// <param name="Builder">Destination builder</param>
+		/// <param name="Text">Text to be encoded</param>
+		/// <returns>Destination builder</returns>
+		public static StringBuilder AppendUTF8LiteralString(this StringBuilder Builder, StringView Text)
+		{
+			Builder.Append("\"");
+
+			ReadOnlySpan<char> Span = Text.Span;
+			int Length = Span.Length;
+
+			if (Length > 0)
+			{
+
+				bool bTrailingHex = false;
+				int Index = 0;
+				while (true)
+				{
+					// Scan forward looking for anything that can just be blindly copied
+					int StartIndex = Index;
+					while (Index < Length)
+					{
+						char cskip = Span[Index];
+						if (cskip < 31 || cskip > 127 || cskip == '"' || cskip == '\\')
+						{
+							break;
+						}
+						++Index;
+					}
+
+					// If we found anything
+					if (StartIndex < Index)
+					{
+						// We close and open the literal here in order to ensure that successive hex characters aren't appended to the hex sequence, causing a different number
+						if (bTrailingHex && UhtFCString.IsHexDigit(Span[StartIndex]))
+						{
+							Builder.Append("\"\"");
+						}
+						Builder.Append(Span.Slice(StartIndex, Index - StartIndex));
+					}
+
+					// We have either reached the end of the string, break
+					if (Index == Length)
+					{
+						break;
+					}
+
+					// This character requires special processing
+					char c = Span[Index++];
+					switch (c)
+					{
+						case '\r':
+							bTrailingHex = false;
+							break;
+						case '\n':
+							bTrailingHex = false;
+							Builder.Append("\\n");
+							break;
+						case '\\':
+							bTrailingHex = false;
+							Builder.Append("\\\\");
+							break;
+						case '\"':
+							bTrailingHex = false;
+							Builder.Append("\\\"");
+							break;
+						default:
+							if (c < 31)
+							{
+								bTrailingHex = true;
+								//Builder.Append($"\\x{(uint)c:x2}");
+								Builder.Append("\\x").AppendFormat("{0:x2}", (uint)c);
+							}
+							else
+							{
+								bTrailingHex = false;
+								if (char.IsHighSurrogate(c))
+								{
+									if (Index == Length)
+									{
+										Builder.Append('?');
+										break;
+									}
+
+									char clow = Span[Index];
+									if (char.IsLowSurrogate(clow))
+									{
+										++Index;
+										Builder.AppendEscapedUtf32((ulong)char.ConvertToUtf32(c, clow));
+										bTrailingHex = true;
+									}
+									else
+									{
+										Builder.Append('?');
+									}
+								}
+								else if (char.IsLowSurrogate(c))
+								{
+									Builder.Append('?');
+								}
+								else
+								{
+									Builder.AppendEscapedUtf32(c);
+									bTrailingHex = true;
+								}
+							}
+							break;
+					}
+				} while (Index < Length);
+			}
+
+			Builder.Append("\"");
+			return Builder;
+		}
+
+		/// <summary>
+		/// Encode a single UTF32 value as UTF8 characters
+		/// </summary>
+		/// <param name="Builder">Destination builder</param>
+		/// <param name="c">Character to encode</param>
+		/// <returns>Destination builder</returns>
+		public static StringBuilder AppendEscapedUtf32(this StringBuilder Builder, ulong c)
+		{
+			if (c < 0x80)
+			{
+				Builder
+					.Append("\\x").AppendFormat("{0:x2}", c);
+			}
+			else if (c < 0x800)
+			{
+				Builder
+					.Append("\\x").AppendFormat("{0:x2}", 0xC0 + (c >> 6))
+					.Append("\\x").AppendFormat("{0:x2}", 0x80 + (c & 0x3f));
+			}
+			else if (c < 0x10000)
+			{
+				Builder
+					.Append("\\x").AppendFormat("{0:x2}", 0xE0 + (c >> 12))
+					.Append("\\x").AppendFormat("{0:x2}", 0x80 + ((c >> 6) & 0x3f))
+					.Append("\\x").AppendFormat("{0:x2}", 0x80 + (c & 0x3f));
+			}
+			else
+			{
+				Builder
+					.Append("\\x").AppendFormat("{0:x2}", 0xF0 + (c >> 18))
+					.Append("\\x").AppendFormat("{0:x2}", 0x80 + ((c >> 12) & 0x3f))
+					.Append("\\x").AppendFormat("{0:x2}", 0x80 + ((c >> 6) & 0x3f))
+					.Append("\\x").AppendFormat("{0:x2}", 0x80 + (c & 0x3f));
+			}
+			return Builder;
+		}
+
+		/// <summary>
+		/// Append the given name of the class but always encode interfaces as the native interface name (i.e. "I...")
+		/// </summary>
+		/// <param name="Builder">Destination builder</param>
+		/// <param name="Class">Class to append name</param>
+		/// <returns>Destination builder</returns>
+		/// <exception cref="NotImplementedException">Class has an unexpected class type</exception>
+		public static StringBuilder AppendClassSourceNameOrInterfaceName(this StringBuilder Builder, UhtClass Class)
+		{
+			switch (Class.ClassType)
+			{
+				case UhtClassType.Class:
+				case UhtClassType.NativeInterface:
+					return Builder.Append(Class.SourceName);
+				case UhtClassType.Interface:
+					return Builder.Append('I').Append(Class.EngineName);
+				default:
+					throw new NotImplementedException();
+			}
+		}
+	}
+
+	/// <summary>
+	/// Provides a cache of StringBuilders
+	/// </summary>
+	public class StringBuilderCache
+	{
+
+		/// <summary>
+		/// Cache of StringBuilders with large initial buffer sizes
+		/// </summary>
+		public static StringBuilderCache Big = new StringBuilderCache(256, 256 * 1024);
+
+		/// <summary>
+		/// Cache of StringBuilders with small initial buffer sizes
+		/// </summary>
+		public static StringBuilderCache Small = new StringBuilderCache(256, 1 * 1024);
+
+		/// <summary>
+		/// Capacity of the cache
+		/// </summary>
+		private int Capacity;
+
+		/// <summary>
+		/// Initial buffer size for new StringBuilders.  Resetting StringBuilders might result
+		/// in the initial chunk size being smaller.
+		/// </summary>
+		private int InitialBufferSize;
+
+		/// <summary>
+		/// Stack of cached StringBuilders
+		/// </summary>
+		private Stack<StringBuilder> Stack;
+
+		/// <summary>
+		/// Create a new StringBuilder cache
+		/// </summary>
+		/// <param name="Capacity">Maximum number of StringBuilders to cache</param>
+		/// <param name="InitialBufferSize">Initial buffer size for newly created StringBuilders</param>
+		public StringBuilderCache(int Capacity, int InitialBufferSize)
+		{
+			this.Capacity = Capacity;
+			this.InitialBufferSize = InitialBufferSize;
+			this.Stack = new Stack<StringBuilder>(this.Capacity);
+		}
+
+		/// <summary>
+		/// Borrow a StringBuilder from the cache.
+		/// </summary>
+		/// <returns></returns>
+		public StringBuilder Borrow()
+		{
+			lock (this.Stack)
+			{
+				if (this.Stack.Count > 0)
+				{
+					return this.Stack.Pop();
+				}
+			}
+
+			return new StringBuilder(this.InitialBufferSize);
+		}
+
+		/// <summary>
+		/// Return a StringBuilder to the cache
+		/// </summary>
+		/// <param name="Builder">The builder being returned</param>
+		public void Return(StringBuilder Builder)
+		{
+			// Sadly, clearing the builder (sets length to 0) will reallocate chunks.
+			Builder.Clear();
+			lock (this.Stack)
+			{
+				if (this.Stack.Count < this.Capacity)
+				{
+					this.Stack.Push(Builder);
+				}
+			}
+		}
+	}
+
+	/// <summary>
+	/// Structure to automate the borrowing and returning of a StringBuilder.
+	/// Use some form of a "using" pattern.
+	/// </summary>
+	public struct BorrowStringBuilder : IDisposable
+	{
+
+		/// <summary>
+		/// Owning cache
+		/// </summary>
+		private readonly StringBuilderCache Cache;
+
+		/// <summary>
+		/// Borrowed string builder
+		/// </summary>
+		public readonly StringBuilder StringBuilder;
+
+		/// <summary>
+		/// Borrow a string builder from the given cache
+		/// </summary>
+		/// <param name="Cache">String builder cache</param>
+		public BorrowStringBuilder(StringBuilderCache Cache)
+		{
+			this.Cache = Cache;
+			this.StringBuilder = this.Cache.Borrow();
+		}
+
+		/// <summary>
+		/// Return the string builder to the cache
+		/// </summary>
+		public void Dispose()
+		{
+			this.Cache.Return(this.StringBuilder);
+		}
+	}
+}
