Fixing an inappropriate word which got copied from the old UHT prior to it getting removed.

#rnx

diff --git a/Engine/Source/Programs/Shared/EpicGames.UHT/Parsers/UhtHeaderFileParser.cs b/Engine/Source/Programs/Shared/EpicGames.UHT/Parsers/UhtHeaderFileParser.cs
--- a/Engine/Source/Programs/Shared/EpicGames.UHT/Parsers/UhtHeaderFileParser.cs
+++ b/Engine/Source/Programs/Shared/EpicGames.UHT/Parsers/UhtHeaderFileParser.cs
@@ -493,7 +493,7 @@ namespace EpicGames.UHT.Parsers
 				}
 			}
 
-			// We look for at least one underscore as a convenient way of whitelisting many known macros
+			// We look for at least one underscore as a convenient way of allowing many known macros
 			// like FORCEINLINE and CONSTEXPR, and non-macros like FPOV and TCHAR.
 			if (UnderscoreCount == 0)
 			{
