Horde: Support for incrementally flushing bundles as they are built up in memory.

- References to BundleNode objects in memory are held through BundleNodeRef objects, which abstracts whether a node is hydrated in memory or referenced via an external hash.
- Deserialized BundleNode objects keep a reference to their incoming reference.
- Tracking of whether nodes have been modified is done via the BundleNodeRef. Only a weak reference to the node is kept for unmodified nodes, and a strong reference is kept for modified nodes.
- When nodes are modified, a timestamp is captured indicating the current time. The Bundle.TrimAsync() method can be used to flush the oldest modified nodes to storage periodically.
- Interface to create/manipulate nodes is now much simpler, and does not require so much explicit state to be passed into node constructor.
- File and directory entries now include a field for the length of the data underneath them.
- Added BuzHash test case.

diff --git a/Engine/Source/Programs/Horde/Horde.Agent/Commands/Bundles/CreateCommand.cs b/Engine/Source/Programs/Horde/Horde.Agent/Commands/Bundles/CreateCommand.cs
--- a/Engine/Source/Programs/Horde/Horde.Agent/Commands/Bundles/CreateCommand.cs
+++ b/Engine/Source/Programs/Horde/Horde.Agent/Commands/Bundles/CreateCommand.cs
@@ -70,9 +70,9 @@ namespace HordeAgent.Commands.Bundles
 		{
 			IStorageClient StorageClient = base.CreateStorageClient(Logger);
 
-			Bundle<DirectoryNode> Bundle = new Bundle<DirectoryNode>(StorageClient, NamespaceId, new BundleOptions(), null);
-			await Bundle.Root.CopyFromDirectoryAsync(InputDir.ToDirectoryInfo(), new ChunkingOptions(), Logger);
-			await Bundle.WriteAsync(BucketId, RefId, false, DateTime.UtcNow);
+			Bundle<DirectoryNode> NewBundle = Bundle.Create<DirectoryNode>(StorageClient, NamespaceId, new BundleOptions(), null);
+			await NewBundle.Root.CopyFromDirectoryAsync(InputDir.ToDirectoryInfo(), new ChunkingOptions(), Logger);
+			await NewBundle.WriteAsync(BucketId, RefId, false, DateTime.UtcNow);
 
 			return 0;
 		}
diff --git a/Engine/Source/Programs/Horde/Horde.Agent/Commands/Bundles/ExtractCommand.cs b/Engine/Source/Programs/Horde/Horde.Agent/Commands/Bundles/ExtractCommand.cs
--- a/Engine/Source/Programs/Horde/Horde.Agent/Commands/Bundles/ExtractCommand.cs
+++ b/Engine/Source/Programs/Horde/Horde.Agent/Commands/Bundles/ExtractCommand.cs
@@ -38,9 +38,8 @@ namespace HordeAgent.Commands.Bundles
 			IStorageClient StorageClient = CreateStorageClient(Logger);
 			using (IMemoryCache Cache = new MemoryCache(new MemoryCacheOptions { SizeLimit = 1024 * 1024 * 1000 }))
 			{
-				Bundle<DirectoryNode> Bundle = new Bundle<DirectoryNode>(StorageClient, NamespaceId, new BundleOptions(), Cache);
-				await Bundle.ReadAsync(BucketId, RefId);
-				await Bundle.Root.CopyToDirectoryAsync(OutputDir.ToDirectoryInfo(), Logger);
+				Bundle<DirectoryNode> NewBundle = await Bundle.ReadAsync<DirectoryNode>(StorageClient, NamespaceId, BucketId, RefId, new BundleOptions(), Cache);
+				await NewBundle.Root.CopyToDirectoryAsync(OutputDir.ToDirectoryInfo(), Logger);
 			}
 			return 0;
 		}
diff --git a/Engine/Source/Programs/Horde/Horde.Build/Commands/Bundles/CreateCommand.cs b/Engine/Source/Programs/Horde/Horde.Build/Commands/Bundles/CreateCommand.cs
--- a/Engine/Source/Programs/Horde/Horde.Build/Commands/Bundles/CreateCommand.cs
+++ b/Engine/Source/Programs/Horde/Horde.Build/Commands/Bundles/CreateCommand.cs
@@ -48,9 +48,9 @@ namespace HordeServer.Commands.Bundles
 		{
 			IStorageClient StorageClient = base.CreateStorageClient(Logger);
 
-			Bundle<DirectoryNode> Bundle = new Bundle<DirectoryNode>(StorageClient, NamespaceId, new BundleOptions(), null);
-			await Bundle.Root.CopyFromDirectoryAsync(InputDir.ToDirectoryInfo(), new ChunkingOptions(), Logger);
-			await Bundle.WriteAsync(BucketId, RefId, false, DateTime.UtcNow);
+			Bundle<DirectoryNode> NewBundle = Bundle.Create<DirectoryNode>(StorageClient, NamespaceId, new BundleOptions(), null);
+			await NewBundle.Root.CopyFromDirectoryAsync(InputDir.ToDirectoryInfo(), new ChunkingOptions(), Logger);
+			await NewBundle.WriteAsync(BucketId, RefId, false, DateTime.UtcNow);
 
 			return 0;
 		}
diff --git a/Engine/Source/Programs/Horde/Horde.Build/Commands/Bundles/ExtractCommand.cs b/Engine/Source/Programs/Horde/Horde.Build/Commands/Bundles/ExtractCommand.cs
--- a/Engine/Source/Programs/Horde/Horde.Build/Commands/Bundles/ExtractCommand.cs
+++ b/Engine/Source/Programs/Horde/Horde.Build/Commands/Bundles/ExtractCommand.cs
@@ -38,9 +38,8 @@ namespace HordeServer.Commands.Bundles
 			IStorageClient StorageClient = CreateStorageClient(Logger);
 			using (IMemoryCache Cache = new MemoryCache(new MemoryCacheOptions { SizeLimit = 1024 * 1024 * 1000 }))
 			{
-				Bundle<DirectoryNode> Bundle = new Bundle<DirectoryNode>(StorageClient, NamespaceId, new BundleOptions(), Cache);
-				await Bundle.ReadAsync(BucketId, RefId);
-				await Bundle.Root.CopyToDirectoryAsync(OutputDir.ToDirectoryInfo(), Logger);
+				Bundle<DirectoryNode> NewBundle = await Bundle.ReadAsync<DirectoryNode>(StorageClient, NamespaceId, BucketId, RefId, new BundleOptions(), Cache);
+				await NewBundle.Root.CopyToDirectoryAsync(OutputDir.ToDirectoryInfo(), Logger);
 			}
 			return 0;
 		}
diff --git a/Engine/Source/Programs/Horde/Horde.Build/Perforce/Impl/CommitService.cs b/Engine/Source/Programs/Horde/Horde.Build/Perforce/Impl/CommitService.cs
--- a/Engine/Source/Programs/Horde/Horde.Build/Perforce/Impl/CommitService.cs
+++ b/Engine/Source/Programs/Horde/Horde.Build/Perforce/Impl/CommitService.cs
@@ -894,22 +894,23 @@ namespace HordeServer.Commits.Impl
 			using IPerforceConnection Perforce = await PerforceConnection.CreateAsync(ClientInfo.Settings, Logger);
 
 			// Get the initial directory state
+			Bundle<DirectoryNode> CommitBundle;
 			if (BaseTree == null)
 			{
 				await FlushWorkspaceAsync(ClientInfo, Perforce, 0);
-				ClientInfo.Contents.Root.Clear();
+				CommitBundle = Bundle.Create<DirectoryNode>(StorageClient, Options.NamespaceId, Options.Bundle, MemoryCache);
 			}
 			else
 			{
 				await FlushWorkspaceAsync(ClientInfo, Perforce, BaseTree.Change);
-				await ClientInfo.Contents.ReadAsync(GetStreamBucketId(BaseTree.StreamId), BaseTree.RefId);
+				CommitBundle = await Bundle.ReadAsync<DirectoryNode>(StorageClient, Options.NamespaceId, GetStreamBucketId(BaseTree.StreamId), BaseTree.RefId, Options.Bundle, MemoryCache);
 			}
 
 			// Apply all the updates
 			Logger.LogInformation("Updating client {Client} to changelist {Change}", ClientInfo.Client.Name, Change);
 			ClientInfo.Change = -1;
 
-			DirectoryNode Contents = ClientInfo.Contents.Root;
+			DirectoryNode Contents = CommitBundle.Root;
 			if (Options.WriteStubFiles)
 			{
 				await foreach (PerforceResponse<SyncRecord> Record in Perforce.StreamCommandAsync<SyncRecord>("sync", new[] { "-k", $"{QueryPath}@{Change}" }, null, default))
@@ -921,7 +922,7 @@ namespace HordeServer.Commits.Impl
 					}
 
 					string Path = SyncRecord.Path.Substring(ClientInfo.Client.Root.Length).Replace('\\', '/');
-					FileNode File = await Contents.CreateFileByPathAsync(Path, DirectoryEntryFlags.PerforceDepotPathAndRevision);
+					FileNode File = await Contents.CreateFileByPathAsync(Path, FileEntryFlags.PerforceDepotPathAndRevision);
 
 					byte[] Data = Encoding.UTF8.GetBytes($"{SyncRecord.DepotFile}#{SyncRecord.Revision}");
 					File.Append(Data, Options.Chunking);
@@ -938,7 +939,7 @@ namespace HordeServer.Commits.Impl
 						if (Io.Command == PerforceIoCommand.Open)
 						{
 							string Path = GetClientRelativePath(Io.Payload, ClientInfo.Client.Root);
-							Files[Io.File] = await Contents.CreateFileByPathAsync(Path, DirectoryEntryFlags.File);
+							Files[Io.File] = await Contents.CreateFileByPathAsync(Path, 0);
 						}
 						else if (Io.Command == PerforceIoCommand.Write)
 						{
@@ -963,7 +964,8 @@ namespace HordeServer.Commits.Impl
 			// Return the new root object
 			RefId RefId = CommitTree.GetDefaultRefId(Change);
 			Logger.LogInformation("Writing ref {RefId} for {StreamId} change {Change}", RefId, Stream.Id, Change);
-			await ClientInfo.Contents.WriteAsync(GetStreamBucketId(Stream.Id), RefId, true, DateTime.UtcNow);
+			await CommitBundle.WriteAsync(GetStreamBucketId(Stream.Id), RefId, true, DateTime.UtcNow);
+
 			return new CommitTree(Stream.Id, Change, RefId);
 		}
 
@@ -1006,16 +1008,14 @@ namespace HordeServer.Commits.Impl
 			public InfoRecord ServerInfo { get; }
 			public ClientRecord Client { get; }
 			public int Change { get; set; }
-			public Bundle<DirectoryNode> Contents { get; }
 
-			public ReplicationClient(PerforceSettings Settings, string ClusterName, InfoRecord ServerInfo, ClientRecord Client, int Change, Bundle<DirectoryNode> Contents)
+			public ReplicationClient(PerforceSettings Settings, string ClusterName, InfoRecord ServerInfo, ClientRecord Client, int Change)
 			{
 				this.Settings = Settings;
 				this.ClusterName = ClusterName;
 				this.ServerInfo = ServerInfo;
 				this.Client = Client;
 				this.Change = Change;
-				this.Contents = Contents;
 			}
 		}
 
@@ -1066,8 +1066,7 @@ namespace HordeServer.Commits.Impl
 				Settings.ClientName = NewClient.Name;
 				Settings.PreferNativeClient = true;
 
-				Bundle<DirectoryNode> Bundle = new Bundle<DirectoryNode>(StorageClient, Options.NamespaceId, Options.Bundle, MemoryCache);
-				ClientInfo = new ReplicationClient(Settings, Stream.ClusterName, ServerInfo, NewClient, -1, Bundle);
+				ClientInfo = new ReplicationClient(Settings, Stream.ClusterName, ServerInfo, NewClient, -1);
 				CachedPerforceClients.Add(Stream.Id, ClientInfo);
 			}
 			return ClientInfo;
diff --git a/Engine/Source/Programs/Shared/EpicGames.Core/BuzHash.cs b/Engine/Source/Programs/Shared/EpicGames.Core/BuzHash.cs
--- a/Engine/Source/Programs/Shared/EpicGames.Core/BuzHash.cs
+++ b/Engine/Source/Programs/Shared/EpicGames.Core/BuzHash.cs
@@ -126,6 +126,15 @@ namespace EpicGames.Core
 			Count--;
 		}
 
+		/// <summary>
+		/// Removes a byte from the start of the hash window
+		/// </summary>
+		/// <param name="Value">Value at the start of the window</param>
+		public static uint Sub(uint State, byte Value, int Count)
+		{
+			return State ^ Rol32(Table[Value], (int)(Count - 1));
+		}
+
 		/// <summary>
 		/// Update the hash with the given data
 		/// </summary>
diff --git a/Engine/Source/Programs/Shared/EpicGames.Horde.Tests/BundleTests.cs b/Engine/Source/Programs/Shared/EpicGames.Horde.Tests/BundleTests.cs
--- a/Engine/Source/Programs/Shared/EpicGames.Horde.Tests/BundleTests.cs
+++ b/Engine/Source/Programs/Shared/EpicGames.Horde.Tests/BundleTests.cs
@@ -10,6 +10,7 @@ using Microsoft.VisualStudio.TestTools.UnitTesting;
 using System;
 using System.Buffers.Binary;
 using System.Collections.Generic;
+using System.Linq;
 using System.Security.Cryptography;
 using System.Text;
 using System.Threading.Tasks;
@@ -23,6 +24,67 @@ namespace HordeServerTests
 		NamespaceId NamespaceId = new NamespaceId("namespace");
 		BucketId BucketId = new BucketId("bucket");
 
+		[TestMethod]
+		public void BuzHashTests()
+		{
+			byte[] Data = new byte[4096];
+			new Random(0).NextBytes(Data);
+
+			const int WindowSize = 128;
+
+			uint RollingHash = 0;
+			for (int MaxIdx = 0; MaxIdx < Data.Length + WindowSize; MaxIdx++)
+			{
+				int MinIdx = MaxIdx - WindowSize;
+
+				if (MaxIdx < Data.Length)
+				{
+					RollingHash = BuzHash.Add(RollingHash, Data[MaxIdx]);
+				}
+
+				int Length = Math.Min(MaxIdx + 1, Data.Length) - Math.Max(MinIdx, 0);
+				uint CleanHash = BuzHash.Add(0, Data.AsSpan(Math.Max(MinIdx, 0), Length));
+				Assert.AreEqual(RollingHash, CleanHash);
+
+				if (MinIdx >= 0)
+				{
+					RollingHash = BuzHash.Sub(RollingHash, Data[MinIdx], Length);
+				}
+			}
+		}
+
+		[TestMethod]
+		public void BasicChunkingTests()
+		{
+			ChunkingOptions Options = new ChunkingOptions();
+			Options.LeafOptions = new ChunkingOptionsForNodeType(8, 8, 8);
+
+			FileNode Node = new FileNode();
+			Node.Append(new byte[7], Options);
+			Assert.AreEqual(0, Node.Depth);
+			Assert.AreEqual(7, Node.Payload.Length);
+
+			Node = new FileNode();
+			Node.Append(new byte[8], Options);
+			Assert.AreEqual(0, Node.Depth);
+			Assert.AreEqual(8, Node.Payload.Length);
+
+			Node = new FileNode();
+			Node.Append(new byte[9], Options);
+			Assert.AreEqual(1, Node.Depth);
+			Assert.AreEqual(2, Node.Children.Count);
+
+			FileNode? ChildNode1 = Node.Children[0].Node;
+			Assert.IsNotNull(ChildNode1);
+			Assert.AreEqual(0, ChildNode1!.Depth);
+			Assert.AreEqual(8, ChildNode1!.Payload.Length);
+
+			FileNode? ChildNode2 = Node.Children[1].Node;
+			Assert.IsNotNull(ChildNode2);
+			Assert.AreEqual(0, ChildNode2!.Depth);
+			Assert.AreEqual(1, ChildNode2!.Payload.Length);
+		}
+
 		[TestMethod]
 		public void SerializationTests()
 		{
@@ -88,37 +150,39 @@ namespace HordeServerTests
 		[TestMethod]
 		public async Task BasicTestDirectory()
 		{
-			Bundle<DirectoryNode> Bundle = new Bundle<DirectoryNode>(StorageClient, NamespaceId, new BundleOptions(), null);
-			DirectoryNode Node = Bundle.Root.AddDirectory("hello");
+			Bundle<DirectoryNode> NewBundle = new Bundle<DirectoryNode>(StorageClient, NamespaceId, new DirectoryNode(), new BundleOptions(), null);
+			DirectoryNode Node = NewBundle.Root.AddDirectory("hello");
 			DirectoryNode Node2 = Node.AddDirectory("world");
 
 			RefId RefId = new RefId("testref");
-			await Bundle.WriteAsync(BucketId, RefId, false, DateTime.MinValue);
+			await NewBundle.WriteAsync(BucketId, RefId, false, DateTime.MinValue);
 
 			// Should be stored inline
 			Assert.AreEqual(1, StorageClient.Refs.Count);
 			Assert.AreEqual(0, StorageClient.Blobs.Count);
 
 			// Check the ref
-			BundleObject Root = await StorageClient.GetRefAsync<BundleObject>(NamespaceId, BucketId, RefId);
-			Assert.AreEqual(3, Root.Exports.Count);
-			Assert.AreEqual(0, Root.Exports[0].Rank);
-			Assert.AreEqual(1, Root.Exports[1].Rank);
-			Assert.AreEqual(2, Root.Exports[2].Rank);
+			BundleObject RootObject = await StorageClient.GetRefAsync<BundleObject>(NamespaceId, BucketId, RefId);
+			Assert.AreEqual(3, RootObject.Exports.Count);
+			Assert.AreEqual(0, RootObject.Exports[0].Rank);
+			Assert.AreEqual(1, RootObject.Exports[1].Rank);
+			Assert.AreEqual(2, RootObject.Exports[2].Rank);
 
 			// Create a new bundle and read it back in again
-			Bundle = new Bundle<DirectoryNode>(StorageClient, NamespaceId, new BundleOptions(), null);
-			await Bundle.ReadAsync(BucketId, RefId);
+			NewBundle = await Bundle.ReadAsync<DirectoryNode>(StorageClient, NamespaceId, BucketId, RefId, new BundleOptions(), null);
 
-			Assert.AreEqual(1, Bundle.Root.Entries.Count);
-			DirectoryNode? OutputNode = await Bundle.Root.FindDirectoryAsync("hello");
+			Assert.AreEqual(0, NewBundle.Root.Files.Count);
+			Assert.AreEqual(1, NewBundle.Root.Directories.Count);
+			DirectoryNode? OutputNode = await NewBundle.Root.FindDirectoryAsync("hello");
 			Assert.IsNotNull(OutputNode);
 
-			Assert.AreEqual(1, OutputNode!.Entries.Count);
+			Assert.AreEqual(0, OutputNode!.Files.Count);
+			Assert.AreEqual(1, OutputNode!.Directories.Count);
 			DirectoryNode? OutputNode2 = await OutputNode!.FindDirectoryAsync("world");
 			Assert.IsNotNull(OutputNode2);
 
-			Assert.AreEqual(0, OutputNode2!.Entries.Count);
+			Assert.AreEqual(0, OutputNode2!.Files.Count);
+			Assert.AreEqual(0, OutputNode2!.Directories.Count);
 		}
 
 		[TestMethod]
@@ -128,12 +192,15 @@ namespace HordeServerTests
 			Options.MaxBlobSize = 1;
 			Options.MaxInlineBlobSize = 1;
 
-			Bundle<DirectoryNode> Bundle = new Bundle<DirectoryNode>(StorageClient, NamespaceId, Options, null);
+			Bundle<DirectoryNode> Bundle = new Bundle<DirectoryNode>(StorageClient, NamespaceId, new DirectoryNode(), Options, null);
 
 			Bundle.Root.AddDirectory("node1");
 			Bundle.Root.AddDirectory("node2");
 			Bundle.Root.AddDirectory("node3");
 
+			Assert.AreEqual(0, StorageClient.Refs.Count);
+			Assert.AreEqual(0, StorageClient.Blobs.Count);
+
 			RefId RefId = new RefId("ref");
 			await Bundle.WriteAsync(BucketId, RefId, false, DateTime.UtcNow);
 
@@ -148,21 +215,20 @@ namespace HordeServerTests
 			Options.MaxBlobSize = 1;
 			Options.MaxInlineBlobSize = 1;
 
-			Bundle<DirectoryNode> Bundle = new Bundle<DirectoryNode>(StorageClient, NamespaceId, Options, null);
+			Bundle<DirectoryNode> InitialBundle = new Bundle<DirectoryNode>(StorageClient, NamespaceId, new DirectoryNode(), Options, null);
 
-			DirectoryNode Node1 = Bundle.Root.AddDirectory("node1");
+			DirectoryNode Node1 = InitialBundle.Root.AddDirectory("node1");
 			DirectoryNode Node2 = Node1.AddDirectory("node2");
 			DirectoryNode Node3 = Node2.AddDirectory("node3");
 			DirectoryNode Node4 = Node3.AddDirectory("node4");
 
 			RefId RefId = new RefId("ref");
-			await Bundle.WriteAsync(BucketId, RefId, false, DateTime.UtcNow);
+			await InitialBundle.WriteAsync(BucketId, RefId, false, DateTime.UtcNow);
 
 			Assert.AreEqual(1, StorageClient.Refs.Count);
 			Assert.AreEqual(4, StorageClient.Blobs.Count);
 
-			Bundle<DirectoryNode> NewBundle = new Bundle<DirectoryNode>(StorageClient, NamespaceId, Options, null);
-			await NewBundle.ReadAsync(BucketId, RefId);
+			Bundle<DirectoryNode> NewBundle = await Bundle.ReadAsync<DirectoryNode>(StorageClient, NamespaceId, BucketId, RefId, Options, null);
 
 			DirectoryNode? NewNode1 = await NewBundle.Root.FindDirectoryAsync("node1");
 			Assert.IsNotNull(NewNode1);
@@ -186,7 +252,7 @@ namespace HordeServerTests
 			Options.MaxBlobSize = 1024 * 1024;
 			Options.MaxInlineBlobSize = 1;
 
-			Bundle<DirectoryNode> Bundle = new Bundle<DirectoryNode>(StorageClient, NamespaceId, Options, null);
+			Bundle<DirectoryNode> Bundle = new Bundle<DirectoryNode>(StorageClient, NamespaceId, new DirectoryNode(), Options, null);
 
 			DirectoryNode Node1 = Bundle.Root.AddDirectory("node1");
 			DirectoryNode Node2 = Node1.AddDirectory("node2");
@@ -201,33 +267,31 @@ namespace HordeServerTests
 
 			IRef Ref1 = StorageClient.Refs[(NamespaceId, BucketId, RefId1)];
 
-			IoHash RootHash1 = Bundle.RootHash;
 			BundleObject RootObject1 = CbSerializer.Deserialize<BundleObject>(Ref1.Value);
 			Assert.AreEqual(1, RootObject1.Exports.Count);
 			Assert.AreEqual(1, RootObject1.ImportObjects.Count);
-			Assert.AreEqual(RootHash1, RootObject1.Exports[0].Hash);
 
 			IoHash LeafHash1 = RootObject1.ImportObjects[0].Object.Hash;
 			BundleObject LeafObject1 = CbSerializer.Deserialize<BundleObject>(StorageClient.Blobs[(NamespaceId, LeafHash1)]);
-			Assert.AreEqual(3, LeafObject1.Exports.Count); // node4 == node3
+			Assert.AreEqual(3, LeafObject1.Exports.Count); // node1 + node2 + node3 (== node4)
 			Assert.AreEqual(0, LeafObject1.ImportObjects.Count);
 
 			// Remove one of the nodes from the root without compacting. the existing blob should be reused.
-			Bundle.Root.Delete("node1");
+			Bundle.Root.DeleteDirectory("node1");
 
 			RefId RefId2 = new RefId("ref2");
 			await Bundle.WriteAsync(BucketId, RefId2, false, Time);
 
 			IRef Ref2 = StorageClient.Refs[(NamespaceId, BucketId, RefId2)];
 
-			IoHash RootHash2 = Bundle.RootHash;
 			BundleObject RootObject2 = CbSerializer.Deserialize<BundleObject>(Ref2.Value);
 			Assert.AreEqual(1, RootObject2.Exports.Count);
 			Assert.AreEqual(1, RootObject2.ImportObjects.Count);
-			Assert.AreEqual(RootHash2, RootObject2.Exports[0].Hash);
 
 			IoHash LeafHash2 = RootObject2.ImportObjects[0].Object.Hash;
 			Assert.AreEqual(LeafHash1, LeafHash2);
+			Assert.AreEqual(3, LeafObject1.Exports.Count); // unused: node1 + node2 + node3, used: node4 (== node3)
+			Assert.AreEqual(0, LeafObject1.ImportObjects.Count);
 
 			// Repack it and check that we make a new object
 			RefId RefId3 = new RefId("ref3");
@@ -235,11 +299,9 @@ namespace HordeServerTests
 
 			IRef Ref3 = StorageClient.Refs[(NamespaceId, BucketId, RefId3)];
 
-			IoHash RootHash3 = Bundle.RootHash;
 			BundleObject RootObject3 = CbSerializer.Deserialize<BundleObject>(Ref3.Value);
 			Assert.AreEqual(1, RootObject3.Exports.Count);
 			Assert.AreEqual(1, RootObject3.ImportObjects.Count);
-			Assert.AreEqual(RootHash3, RootObject3.Exports[0].Hash);
 
 			IoHash LeafHash3 = RootObject3.ImportObjects[0].Object.Hash;
 			Assert.AreNotEqual(LeafHash1, LeafHash3);
@@ -249,6 +311,25 @@ namespace HordeServerTests
 			Assert.AreEqual(0, LeafObject3.ImportObjects.Count);
 		}
 
+		[TestMethod]
+		public async Task CoreAppendTest()
+		{
+			byte[] Data = new byte[4096];
+			new Random(0).NextBytes(Data);
+
+			ChunkingOptions Options = new ChunkingOptions();
+			Options.LeafOptions = new ChunkingOptionsForNodeType(16, 64, 256);
+
+			FileNode Node = new FileNode();
+			for (int Idx = 0; Idx < Data.Length; Idx++)
+			{
+				Node.Append(Data.AsMemory(Idx, 1), Options);
+
+				byte[] OutputData = await Node.ToByteArrayAsync();
+				Assert.IsTrue(Data.AsMemory(0, Idx + 1).Span.SequenceEqual(OutputData.AsSpan(0, Idx + 1)));
+			}
+		}
+
 		[TestMethod]
 		public async Task FixedSizeChunkingTests()
 		{
@@ -279,15 +360,16 @@ namespace HordeServerTests
 				Data[Idx] = (byte)Idx;
 			}
 
-			Bundle<FileNode> Bundle = new Bundle<FileNode>(StorageClient, NamespaceId, new BundleOptions(), null);
+			Bundle<FileNode> Bundle = new Bundle<FileNode>(StorageClient, NamespaceId, new FileNode(), new BundleOptions(), null);
+			FileNode Root = Bundle.Root;
 
 			const int NumIterations = 100;
 			for (int Idx = 0; Idx < NumIterations; Idx++)
 			{
-				Bundle.Root.Append(Data, Options);
+				Root.Append(Data, Options);
 			}
 
-			byte[] Result = await Bundle.Root.ToByteArrayAsync();
+			byte[] Result = await Root.ToByteArrayAsync();
 			Assert.AreEqual(NumIterations * Data.Length, Result.Length);
 
 			for (int Idx = 0; Idx < NumIterations; Idx++)
@@ -296,7 +378,7 @@ namespace HordeServerTests
 				Assert.IsTrue(SpanData.Span.SequenceEqual(Data));
 			}
 
-			await CheckSizes(Bundle.Root, Options, true);
+			await CheckSizes(Root, Options, true);
 		}
 
 		async Task CheckSizes(FileNode Node, ChunkingOptions Options, bool Rightmost)
@@ -311,13 +393,74 @@ namespace HordeServerTests
 				Assert.IsTrue(Rightmost || Node.Payload.Length >= Options.InteriorOptions.MinSize);
 				Assert.IsTrue(Node.Payload.Length <= Options.InteriorOptions.MaxSize);
 
-				int ChildCount = Node.GetChildCount();
+				int ChildCount = Node.Children.Count;
 				for (int Idx = 0; Idx < ChildCount; Idx++)
 				{
-					FileNode ChildNode = await Node.GetChildNodeAsync(Idx);
+					FileNode ChildNode = await Node.Children[Idx].GetAsync();
 					await CheckSizes(ChildNode, Options, Idx == ChildCount - 1);
 				}
 			}
 		}
+
+		[TestMethod]
+		public async Task SpillTestAsync()
+		{
+			BundleOptions Options = new BundleOptions();
+			Options.MaxBlobSize = 1;
+
+			Bundle<DirectoryNode> Bundle = new Bundle<DirectoryNode>(StorageClient, NamespaceId, new DirectoryNode(), Options, null);
+			DirectoryNode Root = Bundle.Root;
+
+			long TotalLength = 0;
+			for (int IdxA = 0; IdxA < 10; IdxA++)
+			{
+				DirectoryNode NodeA = Root.AddDirectory($"{IdxA}");
+				for (int IdxB = 0; IdxB < 10; IdxB++)
+				{
+					DirectoryNode NodeB = NodeA.AddDirectory($"{IdxB}");
+					for (int IdxC = 0; IdxC < 10; IdxC++)
+					{
+						DirectoryNode NodeC = NodeB.AddDirectory($"{IdxC}");
+						for (int IdxD = 0; IdxD < 10; IdxD++)
+						{
+							FileNode File = NodeC.CreateFile($"{IdxD}", 0);
+							byte[] Data = Encoding.UTF8.GetBytes($"This is file {IdxA}/{IdxB}/{IdxC}/{IdxD}");
+							TotalLength += Data.Length;
+							File.Append(Data, new ChunkingOptions());
+						}
+					}
+
+					int OldWorkingSetSize = GetWorkingSetSize(Bundle.Root);
+					await Bundle.TrimAsync(20);
+					int NewWorkingSetSize = GetWorkingSetSize(Bundle.Root);
+					Assert.IsTrue(NewWorkingSetSize <= OldWorkingSetSize);
+					Assert.IsTrue(NewWorkingSetSize <= 20);
+				}
+			}
+
+			Assert.IsTrue(StorageClient.Blobs.Count > 0);
+			Assert.IsTrue(StorageClient.Refs.Count == 0);
+
+			RefId RefId = new RefId("ref");
+			await Bundle.WriteAsync(BucketId, RefId, true, DateTime.UtcNow);
+
+			Assert.AreEqual(TotalLength, Root.Length);
+
+			Assert.IsTrue(StorageClient.Blobs.Count > 0);
+			Assert.IsTrue(StorageClient.Refs.Count == 1);
+		}
+
+		int GetWorkingSetSize(BundleNode Node)
+		{
+			int Size = 0;
+			foreach (BundleNodeRef NodeRef in Node.GetReferences())
+			{
+				if (NodeRef.Node != null)
+				{
+					Size += 1 + GetWorkingSetSize(NodeRef.Node);
+				}
+			}
+			return Size;
+		}
 	}
 }
diff --git a/Engine/Source/Programs/Shared/EpicGames.Horde/Bundles/Bundle.cs b/Engine/Source/Programs/Shared/EpicGames.Horde/Bundles/Bundle.cs
--- a/Engine/Source/Programs/Shared/EpicGames.Horde/Bundles/Bundle.cs
+++ b/Engine/Source/Programs/Shared/EpicGames.Horde/Bundles/Bundle.cs
@@ -43,6 +43,16 @@ namespace EpicGames.Horde.Bundles
 		/// Minimum size of a block to be compressed
 		/// </summary>
 		public int MinCompressionPacketSize { get; set; } = 16 * 1024;
+
+		/// <summary>
+		/// Number of nodes to retain in the working set after performing a partial flush
+		/// </summary>
+		public int TrimIgnoreCount { get; set; } = 150;
+
+		/// <summary>
+		/// Proportion of used data under which to trigger a re-pack
+		/// </summary>
+		public float RepackRatio { get; set; } = 0.6f;
 	}
 
 	/// <summary>
@@ -102,7 +112,7 @@ namespace EpicGames.Horde.Bundles
 			public int TotalCost { get; }
 
 			public bool Mounted { get; set; }
-			public HashSet<BlobInfo> ReferencedBy { get; } = new HashSet<BlobInfo>(); // Used for repacking
+			public List<NodeInfo>? LiveNodes { get; set; }
 
 			public BlobInfo(IoHash Hash, DateTime CreationTimeUtc, int TotalCost)
 			{
@@ -112,14 +122,15 @@ namespace EpicGames.Horde.Bundles
 			}
  		}
 
-		IStorageClient StorageClient;
-		NamespaceId NamespaceId;
+		readonly IStorageClient StorageClient;
+		readonly NamespaceId NamespaceId;
+		
+		object RootCacheKey = new object();
 
 		/// <summary>
-		/// Hash of the root object
+		/// Reference to the root node in the bundle
 		/// </summary>
-		public IoHash RootHash { get; protected set; }
-		object RootCacheKey = new object();
+		public BundleNode Root { get; private set; }
 
 		Dictionary<IoHash, NodeInfo> HashToNode = new Dictionary<IoHash, NodeInfo>();
 		Dictionary<IoHash, BlobInfo> HashToBlob = new Dictionary<IoHash, BlobInfo>();
@@ -136,25 +147,175 @@ namespace EpicGames.Horde.Bundles
 		/// <summary>
 		/// Constructor
 		/// </summary>
-		/// <param name="StorageClient">Client interface for the storage backend</param>
-		/// <param name="NamespaceId">Namespace for storing blobs in the storage system</param>
-		/// <param name="Options">Options for controlling how things are serialized</param>
-		/// <param name="Cache">Cache for storing decompressed objects, serialized blobs.</param>
-		public Bundle(IStorageClient StorageClient, NamespaceId NamespaceId, BundleOptions Options, IMemoryCache? Cache)
+		protected Bundle(IStorageClient StorageClient, NamespaceId NamespaceId, BundleNode Root, BundleOptions Options, IMemoryCache? Cache)
 		{
 			this.StorageClient = StorageClient;
 			this.NamespaceId = NamespaceId;
+			this.Root = Root;
 			this.Options = Options;
 			this.Cache = Cache;
 			this.HashToNode = new Dictionary<IoHash, NodeInfo>();
 		}
 
+		/// <summary>
+		/// Constructor
+		/// </summary>
+		/// <param name="StorageClient">Client interface for the storage backend</param>
+		/// <param name="NamespaceId">Namespace for storing blobs in the storage system</param>
+		/// <param name="Object">Object containing the root node for this bundle</param>
+		/// <param name="Root">Reference to the root node for the bundle</param>
+		/// <param name="Options">Options for controlling how things are serialized</param>
+		/// <param name="Cache">Cache for storing decompressed objects, serialized blobs.</param>
+		protected Bundle(IStorageClient StorageClient, NamespaceId NamespaceId, BundleObject Object, BundleNode Root, BundleOptions Options, IMemoryCache? Cache)
+			: this(StorageClient, NamespaceId, Root, Options, Cache)
+		{
+			RegisterObject(null, Object, Object.Data);
+		}
+
+		/// <summary>
+		/// Creates a new typed bundle with a default-constructed root node
+		/// </summary>
+		/// <typeparam name="T">The node type</typeparam>
+		/// <param name="StorageClient">Client interface for the storage backend</param>
+		/// <param name="NamespaceId">Namespace for storing blobs in the storage system</param>
+		/// <param name="Options">Options for controlling how things are serialized</param>
+		/// <param name="Cache">Cache for storing decompressed objects, serialized blobs.</param>
+		/// <returns>New bundle instance</returns>
+		public static Bundle<T> Create<T>(IStorageClient StorageClient, NamespaceId NamespaceId, BundleOptions Options, IMemoryCache? Cache) where T : BundleNode, new()
+		{
+			return Create(StorageClient, NamespaceId, new T(), Options, Cache);
+		}
+
+		/// <summary>
+		/// Creates a new typed bundle with a specific root node
+		/// </summary>
+		/// <typeparam name="T">The node type</typeparam>
+		/// <param name="StorageClient">Client interface for the storage backend</param>
+		/// <param name="NamespaceId">Namespace for storing blobs in the storage system</param>
+		/// <param name="Root">Reference to the root node for the bundle</param>
+		/// <param name="Options">Options for controlling how things are serialized</param>
+		/// <param name="Cache">Cache for storing decompressed objects, serialized blobs.</param>
+		/// <returns>New bundle instance</returns>
+		public static Bundle<T> Create<T>(IStorageClient StorageClient, NamespaceId NamespaceId, T Root, BundleOptions Options, IMemoryCache? Cache) where T : BundleNode
+		{
+			return new Bundle<T>(StorageClient, NamespaceId, Root, Options, Cache);
+		}
+
+		/// <summary>
+		/// Reads a bundle from storage
+		/// </summary>
+		public static async ValueTask<Bundle<T>> ReadAsync<T>(IStorageClient StorageClient, NamespaceId NamespaceId, BucketId BucketId, RefId RefId, BundleOptions Options, IMemoryCache? Cache) where T : BundleNode
+		{
+			Bundle<T> NewBundle = new Bundle<T>(StorageClient, NamespaceId, null!, Options, Cache);
+
+			BundleObject Object = await StorageClient.GetRefAsync<BundleObject>(NamespaceId, BucketId, RefId);
+			NewBundle.RegisterObject(null, Object, Object.Data);
+
+			IoHash RootHash = Object.Exports[Object.Exports.Count - 1].Hash;
+			BundleNodeRef<T> RootRef = new BundleNodeRef<T>(NewBundle, RootHash);
+
+			((Bundle)NewBundle).Root = await RootRef.GetAsync();
+			return NewBundle;
+		}
+
+		/// <summary>
+		/// Flush a subset of nodes in the working set.
+		/// </summary>
+		private List<NodeInfo> SerializeWorkingSet(int IgnoreCount)
+		{
+			// Find all the nodes that need to be flushed, in order
+			List<BundleNodeRef> NodeRefs = new List<BundleNodeRef>();
+			LinkWorkingSet(Root, null, NodeRefs);
+			NodeRefs.SortBy(x => x.LastModifiedTime);
+
+			// Keep track of the nodes that have been serialized so far. Nodes in-memory may hash to the same serialized data, so we need to deduplicate them.
+			List<NodeInfo> DirtyNodes = new List<NodeInfo>();
+			HashSet<NodeInfo> DirtyNodesSet = new HashSet<NodeInfo>();
+
+			// Write them all to storage
+			for (int Idx = 0; Idx + IgnoreCount < NodeRefs.Count; Idx++)
+			{
+				BundleNodeRef NodeRef = NodeRefs[Idx];
+
+				BundleNode? Node = NodeRef.Node;
+				Debug.Assert(Node != null);
+
+				NodeInfo NodeInfo = WriteNode(Node);
+				NodeRef.MarkAsClean(NodeInfo.Hash);
+
+				if (NodeInfo.Blob == null && DirtyNodesSet.Add(NodeInfo))
+				{
+					DirtyNodes.Add(NodeInfo);
+				}
+			}
+			return DirtyNodes;
+		}
+
+		/// <summary>
+		/// Traverse a tree of <see cref="BundleNodeRef"/> objects and update the incoming reference to each node
+		/// </summary>
+		/// <param name="Node">The reference to update</param>
+		/// <param name="IncomingRef">The parent reference</param>
+		/// <param name="NodeRefs">List to accumulate the list of all modified references</param>
+		private void LinkWorkingSet(BundleNode Node, BundleNodeRef? IncomingRef, List<BundleNodeRef> NodeRefs)
+		{
+			Node.IncomingRef = IncomingRef;
+
+			foreach (BundleNodeRef ChildRef in Node.GetReferences())
+			{
+				ChildRef.Bundle ??= this;
+				ChildRef.ParentRef ??= IncomingRef;
+
+				if (ChildRef.Node != null && ChildRef.LastModifiedTime != 0)
+				{
+					LinkWorkingSet(ChildRef.Node, ChildRef, NodeRefs);
+					if (IncomingRef != null)
+					{
+						IncomingRef.LastModifiedTime = Math.Max(IncomingRef.LastModifiedTime, ChildRef.LastModifiedTime + 1);
+					}
+					NodeRefs.Add(ChildRef);
+				}
+			}
+		}
+
+		/// <summary>
+		/// Perform an incremental flush of the tree
+		/// </summary>
+		public Task TrimAsync() => TrimAsync(Options.TrimIgnoreCount);
+
+		/// <summary>
+		/// Perform an incremental flush of the tree
+		/// </summary>
+		public async Task TrimAsync(int IgnoreCount)
+		{
+			// Timestamp for any flushed objects
+			DateTime UtcNow = DateTime.UtcNow;
+
+			// Find all the nodes that need to be flushed, in order
+			List<NodeInfo> WriteNodes = SerializeWorkingSet(IgnoreCount);
+			WriteNodes.SortBy(x => x.Rank);
+
+			// Write them all to storage
+			int MinNodeIdx = 0;
+			int NextBlobCost = 0;
+			for (int Idx = 0; Idx < WriteNodes.Count; Idx++)
+			{
+				NodeInfo Node = WriteNodes[Idx];
+				if (Idx > MinNodeIdx && NextBlobCost + Node.Cost > Options.MaxBlobSize)
+				{
+					await WriteObjectAsync(WriteNodes.Slice(MinNodeIdx, Idx - MinNodeIdx), UtcNow);
+					MinNodeIdx = Idx;
+					NextBlobCost = 0;
+				}
+			}
+		}
+
 		/// <summary>
 		/// Gets data for the node with the given hash
 		/// </summary>
 		/// <param name="Hash">Hash of the node to return data for</param>
 		/// <returns>The node data</returns>
-		public ValueTask<ReadOnlyMemory<byte>> GetDataAsync(IoHash Hash)
+		internal ValueTask<ReadOnlyMemory<byte>> GetDataAsync(IoHash Hash)
 		{
 			return GetDataAsync(HashToNode[Hash]);
 		}
@@ -186,13 +347,27 @@ namespace EpicGames.Horde.Bundles
 			}
 		}
 
+		NodeInfo WriteNode(BundleNode Node)
+		{
+			ReadOnlyMemory<byte> Data = Node.Serialize();
+
+			List<IoHash> References = new List<IoHash>();
+			foreach (BundleNodeRef Reference in Node.GetReferences())
+			{
+				Debug.Assert(Reference.Hash != IoHash.Zero);
+				References.Add(Reference.Hash);
+			}
+
+			return WriteNode(Data, References);
+		}
+
 		/// <summary>
 		/// Write a node's data
 		/// </summary>
 		/// <param name="Data">The node data</param>
 		/// <param name="References">Hashes for referenced nodes</param>
 		/// <returns>Hash of the data</returns>
-		public IoHash WriteNode(ReadOnlyMemory<byte> Data, IEnumerable<IoHash> References)
+		NodeInfo WriteNode(ReadOnlyMemory<byte> Data, IEnumerable<IoHash> References)
 		{
 			IoHash Hash = IoHash.Compute(Data.Span);
 			NodeInfo[] NodeReferences = References.Select(x => HashToNode[x]).Distinct().OrderBy(x => x.Hash).ToArray();
@@ -211,7 +386,7 @@ namespace EpicGames.Horde.Bundles
 			}
 
 			Node.SetStandalone(Data, NodeReferences);
-			return Hash;
+			return Node;
 		}
 
 		async ValueTask MountBlobAsync(BlobInfo Blob)
@@ -281,18 +456,6 @@ namespace EpicGames.Horde.Bundles
 			return DecodedData;
 		}
 
-		/// <summary>
-		/// Reads a ref into the bundle
-		/// </summary>
-		/// <param name="BucketId">Bucket containing the ref</param>
-		/// <param name="RefId">Reference id</param>
-		public virtual async Task ReadAsync(BucketId BucketId, RefId RefId)
-		{
-			BundleObject Object = await StorageClient.GetRefAsync<BundleObject>(NamespaceId, BucketId, RefId);
-			RegisterObject(null, Object, Object.Data);
-			RootHash = Object.Exports[Object.Exports.Count - 1].Hash;
-		}
-
 		NodeInfo FindOrAddNode(IoHash Hash, int Rank, int Cost)
 		{
 			NodeInfo? Node;
@@ -346,90 +509,66 @@ namespace EpicGames.Horde.Bundles
 		/// <param name="UtcNow"></param>
 		public virtual async Task WriteAsync(BucketId BucketId, RefId RefId, bool Compact, DateTime UtcNow)
 		{
-			NodeInfo RootNode = HashToNode[RootHash];
+			// Completely flush the entire working set into NodeInfo objects
+			SerializeWorkingSet(0);
+
+			// Serialize the root node
+			NodeInfo RootNode = WriteNode(Root);
 
 			// If we're compacting the output, see if there are any other blocks we should merge
 			if (Compact)
 			{
 				// Find the live set of objects
-				HashSet<NodeInfo> LiveNodes = new HashSet<NodeInfo>();
-				await FindFullLiveSetAsync(RootNode, LiveNodes);
+				List<NodeInfo> LiveNodes = new List<NodeInfo>();
+				await FindLiveNodesAsync(RootNode, LiveNodes, new HashSet<NodeInfo>());
 
-				// Find the live set of blobs
-				HashSet<BlobInfo> LiveBlobs = new HashSet<BlobInfo>();
+				// Find the live set of blobs, and the cost of nodes within them
+				List<BlobInfo> LiveBlobs = new List<BlobInfo>();
 				foreach (NodeInfo LiveNode in LiveNodes)
 				{
-					if (LiveNode.Blob != null)
-					{
-						LiveBlobs.Add(LiveNode.Blob);
-					}
-				}
-
-				// Clear the reference list on each one
-				foreach (BlobInfo LiveBlob in LiveBlobs)
-				{
-					LiveBlob.ReferencedBy.Clear();
-				}
-
-				// Find the first reference from each node to other blobs
-				foreach (NodeInfo NodeInfo in LiveNodes)
-				{
-					if (NodeInfo.Rank > 0 && NodeInfo.Blob != null)
+					BlobInfo? LiveBlob = LiveNode.Blob;
+					if (LiveBlob != null)
 					{
-						foreach (NodeInfo ReferencedNode in NodeInfo.References!)
+						if (LiveBlob.LiveNodes == null)
 						{
-							BlobInfo? ReferencedBlob = ReferencedNode.Blob;
-							if (ReferencedBlob != null)
-							{
-								ReferencedBlob.ReferencedBy.Add(NodeInfo.Blob);
-							}
+							LiveBlob.LiveNodes = new List<NodeInfo>();
+							LiveBlobs.Add(LiveBlob);
 						}
+						LiveBlob.LiveNodes.Add(LiveNode);
 					}
 				}
 
-				// Clear out the reference data on each blob
-				HashSet<BlobInfo> VisitedBlobs = new HashSet<BlobInfo>();
+				// Figure out which blobs need to be repacked
 				foreach (BlobInfo LiveBlob in LiveBlobs)
 				{
-					ExpandReferences(LiveBlob, VisitedBlobs);
-				}
+					int LiveCost = LiveBlob.LiveNodes.Sum(x => x.Cost);
+					int MinLiveCost = (int)(LiveBlob.TotalCost * Options.RepackRatio);
 
-				// Find the total cost of all the current blobs, then loop through the blobs trying to find a more optimal arrangement
-				List<NodeInfo> SortedLiveNodes = LiveNodes.OrderByDescending(x => x.Rank).ThenByDescending(x => x.Hash).ToList();
-				for (; ; )
-				{
-					BlobInfo? MergeBlob = null;
-					double MergeCost = GetNewCostHeuristic(SortedLiveNodes, new HashSet<BlobInfo>(), UtcNow);
-
-					foreach (BlobInfo LiveBlob in LiveBlobs)
+					if (LiveCost < MinLiveCost)
 					{
-						double PotentialCost = GetNewCostHeuristic(SortedLiveNodes, LiveBlob.ReferencedBy, UtcNow);
-						if (PotentialCost < MergeCost)
+						foreach (NodeInfo Node in LiveBlob.LiveNodes!)
 						{
-							MergeBlob = LiveBlob;
-							MergeCost = PotentialCost;
+							ReadOnlyMemory<byte> Data = await GetDataAsync(Node);
+							Node.SetStandalone(Data, Node.References!);
 						}
 					}
+				}
 
-					// Bail out if we didn't find anything to merge
-					if (MergeBlob == null)
-					{
-						break;
-					}
-
-					// Make sure this blob is resident (it may not have been loaded above if it's a leaf blob)
-					await MountBlobAsync(MergeBlob);
-
-					// Find all the live nodes in this blob
-					List<NodeInfo> MergeNodes = LiveNodes.Where(x => x.Blob == MergeBlob).ToList();
-					foreach (NodeInfo MergeNode in MergeNodes)
+				// Walk backwards through the list of live nodes and make anything that references a repacked node also standalone
+				for (int Idx = LiveNodes.Count - 1; Idx >= 0; Idx--)
+				{
+					NodeInfo LiveNode = LiveNodes[Idx];
+					if (LiveNode.References != null && LiveNode.References.Any(x => x.Blob == null))
 					{
-						ReadOnlyMemory<byte> Data = await GetDataAsync(MergeNode);
-						MergeNode.SetStandalone(Data, MergeNode.References!);
+						ReadOnlyMemory<byte> Data = await GetDataAsync(LiveNode);
+						LiveNode.SetStandalone(Data, LiveNode.References!);
 					}
+				}
 
-					// Remove the merged blob from the live list
-					LiveBlobs.Remove(MergeBlob);
+				// Clear the reference list on each blob
+				foreach (BlobInfo LiveBlob in LiveBlobs)
+				{
+					LiveBlob.LiveNodes = null;
 				}
 			}
 
@@ -478,32 +617,22 @@ namespace EpicGames.Horde.Bundles
 			}
 		}
 
-		async Task FindFullLiveSetAsync(NodeInfo Node, HashSet<NodeInfo> LiveNodes)
+		async Task FindLiveNodesAsync(NodeInfo Node, List<NodeInfo> LiveNodes, HashSet<NodeInfo> LiveNodeSet)
 		{
-			if (LiveNodes.Add(Node) && Node.Rank > 0)
+			if (LiveNodeSet.Add(Node))
 			{
-				if (Node.Blob != null)
+				LiveNodes.Add(Node);
+				if (Node.Rank > 0)
 				{
-					await MountBlobAsync(Node.Blob);
-				}
-				foreach (NodeInfo Reference in Node.References!)
-				{
-					await FindFullLiveSetAsync(Reference, LiveNodes);
-				}
-			}
-		}
-
-		void ExpandReferences(BlobInfo Blob, HashSet<BlobInfo> VisitedBlobs)
-		{
-			if (VisitedBlobs.Add(Blob))
-			{
-				BlobInfo[] ReferencedBy = Blob.ReferencedBy.ToArray();
-				foreach (BlobInfo ReferencedBlob in ReferencedBy)
-				{
-					ExpandReferences(ReferencedBlob, VisitedBlobs);
-					Blob.ReferencedBy.UnionWith(ReferencedBlob.ReferencedBy);
+					if (Node.Blob != null)
+					{
+						await MountBlobAsync(Node.Blob);
+					}
+					foreach (NodeInfo Reference in Node.References!)
+					{
+						await FindLiveNodesAsync(Reference, LiveNodes, LiveNodeSet);
+					}
 				}
-				Blob.ReferencedBy.Add(Blob);
 			}
 		}
 
@@ -642,87 +771,6 @@ namespace EpicGames.Horde.Bundles
 				Buffer = NewBuffer;
 			}
 		}
-
-		/// <inheritdoc cref="GetCostHeuristic(int, TimeSpan)"/>
-		double GetCostHeuristic(BlobInfo Info, DateTime UtcNow) => GetCostHeuristic(Info.TotalCost, UtcNow - Info.CreationTimeUtc);
-
-		/// <summary>
-		/// Gets the cost of packaging the given set of nodes which are either not in a blob or are in a blob listed to be collapsed
-		/// </summary>
-		/// <param name="Nodes">The full list of nodes</param>
-		/// <param name="CollapseBlobs">Set of blobs that should be collapsed</param>
-		/// <param name="UtcNow">The current time</param>
-		/// <returns></returns>
-		double GetNewCostHeuristic(List<NodeInfo> Nodes, HashSet<BlobInfo> CollapseBlobs, DateTime UtcNow)
-		{
-			double Cost = 0.0;
-			foreach (BlobInfo CollapseBlob in CollapseBlobs)
-			{
-				Cost -= GetCostHeuristic(CollapseBlob, UtcNow);
-			}
-
-			int BlobSize = 0;
-			for (int Idx = Nodes.Count - 1; Idx >= 0; Idx--)
-			{
-				NodeInfo Node = Nodes[Idx];
-				if (Node.Blob == null || CollapseBlobs.Contains(Node.Blob))
-				{
-					if (BlobSize > 0 && BlobSize + Node.Data.Length > Options.MaxBlobSize)
-					{
-						Cost += GetCostHeuristic(BlobSize, TimeSpan.Zero);
-						BlobSize = 0;
-					}
-					BlobSize += Node.Data.Length;
-				}
-			}
-			Cost += GetCostHeuristic(BlobSize, TimeSpan.Zero);
-			return Cost;
-		}
-
-		/// <summary>
-		/// Heuristic which estimates the cost of a particular blob. This is used to compare scenarios of merging blobs to reduce download
-		/// size against keeping older blobs which a lot of agents already have.
-		/// </summary>
-		/// <param name="Size">Size of the blob</param>
-		/// <param name="Age">Age of the blob</param>
-		/// <returns>Heuristic for the cost of a blob</returns>
-		static double GetCostHeuristic(int Size, TimeSpan Age)
-		{
-			// Time overhead to starting a download
-			const double DownloadInit = 0.1;
-
-			// Download speed for agents, in bytes/sec
-			const double DownloadRate = 1024 * 1024;
-
-			// Probability of an agent having to download everything. Prevents bias against keeping a large number of files.
-			const double CleanSyncProbability = 0.2;
-
-			// Average length of time between agents having to update
-			TimeSpan AverageCoherence = TimeSpan.FromHours(4.0);
-
-			// Scale the age into a -1.0 -> 1.0 range around AverageCoherence
-			double ScaledAge = (AverageCoherence - Age).TotalSeconds / AverageCoherence.TotalSeconds;
-
-			// Get the probability of agents having to sync this blob based on its age. This is modeled as a logistic function (1 / (1 + e^-x))
-			// with value of 0.5 at AverageCoherence, and MaxInterval at zero.
-
-			// Find the scale factor for the 95% interval
-			//    1 / (1 + e^-x) = MaxInterval
-			//    e^-x = (1 / MaxInterval) - 1
-			//    x = -ln((1 / MaxInterval) - 1)
-			const double MaxInterval = 0.95;
-			double SigmoidScale = -Math.Log((1.0 / MaxInterval) - 1.0);
-
-			// Find the probability of having to sync this 
-			double Param = ScaledAge * SigmoidScale;
-			double Probability = 1.0 / (1.0 + Math.Exp(-Param));
-
-			// Scale the probability against having to do a full sync
-			Probability = CleanSyncProbability + (Probability * (1.0 - CleanSyncProbability));
-
-			// Compute the final cost estimate; the amount of time we expect agents to spend downloading the file
-			return Probability * (DownloadInit + (Size / DownloadRate));
-		}
 	}
 
 	/// <summary>
@@ -731,49 +779,17 @@ namespace EpicGames.Horde.Bundles
 	/// <typeparam name="T">The root node type</typeparam>
 	public class Bundle<T> : Bundle where T : BundleNode
 	{
-		static BundleNodeFactory<T> Factory = CreateFactory();
-
 		/// <summary>
-		/// Object at the root of the bundle
+		/// Node at the root of this tree
 		/// </summary>
-		public T Root { get; private set; }
+		public new T Root => (T)base.Root;
 
 		/// <summary>
 		/// Constructor
 		/// </summary>
-		public Bundle(IStorageClient StorageClient, NamespaceId NamespaceId, BundleOptions Options, IMemoryCache? Cache)
-			: base(StorageClient, NamespaceId, Options, Cache)
-		{
-			Root = Factory.CreateRoot(this);
-		}
-
-		/// <summary>
-		/// Create a factory for manipulating T instances
-		/// </summary>
-		/// <returns></returns>
-		static BundleNodeFactory<T> CreateFactory()
-		{
-			BundleNodeFactoryAttribute? Attribute = typeof(T).GetCustomAttribute<BundleNodeFactoryAttribute>();
-			if (Attribute == null)
-			{
-				throw new InvalidOperationException("Nodes used as the root of a bundle must have a BundleNodeFactoryAttribute");
-			}
-			return (BundleNodeFactory<T>)Activator.CreateInstance(Attribute.Type)!;
-		}
-
-		/// <inheritdoc/>
-		public override async Task ReadAsync(BucketId BucketId, RefId RefId)
-		{
-			await base.ReadAsync(BucketId, RefId);
-			ReadOnlyMemory<byte> Data = await GetDataAsync(RootHash);
-			Root = Factory.ParseRoot(this, RootHash, Data);
-		}
-
-		/// <inheritdoc/>
-		public override async Task WriteAsync(BucketId BucketId, RefId RefId, bool Compact, DateTime UtcNow)
+		internal Bundle(IStorageClient StorageClient, NamespaceId NamespaceId, T Root, BundleOptions Options, IMemoryCache? Cache)
+			: base(StorageClient, NamespaceId, Root, Options, Cache)
 		{
-			RootHash = Root.Serialize();
-			await base.WriteAsync(BucketId, RefId, Compact, UtcNow);
 		}
 	}
 }
diff --git a/Engine/Source/Programs/Shared/EpicGames.Horde/Bundles/BundleNode.cs b/Engine/Source/Programs/Shared/EpicGames.Horde/Bundles/BundleNode.cs
--- a/Engine/Source/Programs/Shared/EpicGames.Horde/Bundles/BundleNode.cs
+++ b/Engine/Source/Programs/Shared/EpicGames.Horde/Bundles/BundleNode.cs
@@ -3,6 +3,8 @@
 using EpicGames.Core;
 using System;
 using System.Collections.Generic;
+using System.Diagnostics;
+using System.Reflection;
 using System.Text;
 using System.Threading.Tasks;
 
@@ -14,114 +16,240 @@ namespace EpicGames.Horde.Bundles.Nodes
 	public abstract class BundleNode
 	{
 		/// <summary>
-		/// The bundle that owns this node
+		/// Cached incoming reference to the owner of this node. When the contents of this node are flushed to storage, this is modified to become a weak reference
+		/// and hash value. When modified, it becomes a strong reference and zero hash.
 		/// </summary>
-		public Bundle Owner { get; }
+		internal BundleNodeRef? IncomingRef { get; set; }
 
 		/// <summary>
-		/// The node's parent. Each expanded instance of the node in the tree will have a unique object, so each will have a singular parent.
+		/// Mark this node as dirty
 		/// </summary>
-		protected BundleNode? Parent { get; set; }
+		protected void MarkAsDirty() => IncomingRef?.MarkAsDirty();
 
 		/// <summary>
-		/// Hash of this data. Set to zero once modified.
+		/// Serialize this node into a sequence of bytes
 		/// </summary>
-		private IoHash Hash { get; set; }
+		/// <returns>Serialized data for this node</returns>
+		public abstract ReadOnlyMemory<byte> Serialize();
 
 		/// <summary>
-		/// Constructor
+		/// Enumerates all the child references from this node
+		/// </summary>
+		/// <returns>Children of this node</returns>
+		public abstract IEnumerable<BundleNodeRef> GetReferences();
+	}
+
+	/// <summary>
+	/// Stores a reference from a parent to child node, which can be resurrected after the child node is flushed to storage if subsequently modified.
+	/// </summary>
+	public class BundleNodeRef
+	{
+		/// <summary>
+		/// The bundle that owns this reference. Used to realize hashes into full nodes, and set when we do a full scan of the tree from the bundle.
+		/// </summary>
+		internal Bundle? Bundle;
+
+		/// <summary>
+		/// Cached reference to the parent node
+		/// </summary>
+		internal BundleNodeRef? ParentRef;
+
+		/// <summary>
+		/// Strong reference to the current node
+		/// </summary>
+		protected BundleNode? StrongRef;
+
+		/// <summary>
+		/// Weak reference to the child node. Maintained after the object has been flushed.
+		/// </summary>
+		protected WeakReference<BundleNode>? WeakRef { get; set; }
+
+		/// <summary>
+		/// Last time that the node was modified
+		/// </summary>
+		internal long LastModifiedTime;
+
+		/// <summary>
+		/// Hash of the node. May be set to <see cref="IoHash.Zero"/> if the node in memory has been modified.
+		/// </summary>
+		public IoHash Hash { get; private set; }
+
+		/// <summary>
+		/// Reference to the node.
 		/// </summary>
-		public BundleNode(Bundle Owner, BundleNode? Parent)
-			: this(Owner, Parent, IoHash.Zero)
+		public BundleNode? Node
 		{
+			get => StrongRef;
+			set { StrongRef = value; WeakRef = null; }
 		}
 
 		/// <summary>
-		/// Constructor
+		/// Creates a reference to a node with the given hash
 		/// </summary>
-		public BundleNode(Bundle Owner, BundleNode? Parent, IoHash Hash)
+		/// <param name="Bundle">The bundle that contains the reference</param>
+		/// <param name="Hash">Hash of the referenced node</param>
+		public BundleNodeRef(Bundle Bundle, IoHash Hash)
 		{
-			this.Owner = Owner;
-			this.Parent = Parent;
+			this.Bundle = Bundle;
 			this.Hash = Hash;
 		}
 
 		/// <summary>
-		/// Checks if the current node is dirty
+		/// Creates a reference to the given node
+		/// </summary>
+		/// <param name="Node">The referenced node</param>
+		public BundleNodeRef(BundleNode Node)
+		{
+			this.StrongRef = Node;
+			this.LastModifiedTime = Stopwatch.GetTimestamp();
+		}
+
+		/// <summary>
+		/// Detach this reference from its current owner. Used in situations where a node needs to be reparented.
 		/// </summary>
-		/// <returns>True if the node is dirty</returns>
-		public bool IsDirty() => Hash == IoHash.Zero;
+		public void Detach()
+		{
+			MarkAsDirty();
+			ParentRef = null;
+		}
 
 		/// <summary>
-		/// Mark this node as dirty
+		/// Marks this reference as dirty
 		/// </summary>
-		protected void MarkAsDirty()
+		internal void MarkAsDirty()
 		{
-			if (!IsDirty())
+			if (LastModifiedTime == 0)
 			{
-				Hash = IoHash.Zero;
-				Parent?.MarkAsDirty();
+				this.LastModifiedTime = Stopwatch.GetTimestamp();
+				if (StrongRef == null)
+				{
+					if (WeakRef == null || !WeakRef.TryGetTarget(out StrongRef))
+					{
+						throw new InvalidOperationException("Unable to resolve weak reference to node");
+					}
+					this.WeakRef = null;
+				}
+				ParentRef?.MarkAsDirty();
 			}
 		}
 
 		/// <summary>
-		/// Serialize this node to a buffer
+		/// Marks this reference as clean
 		/// </summary>
-		/// <returns>Hash of the serialized node</returns>
-		public IoHash Serialize()
+		internal void MarkAsClean(IoHash Hash)
 		{
-			if (IsDirty())
+			this.Hash = Hash;
+
+			if (StrongRef != null)
 			{
-				Hash = SerializeDirty();
+				Collapse();
+
+				this.WeakRef = new WeakReference<BundleNode>(StrongRef);
+				this.StrongRef = null;
 			}
-			return Hash;
+
+			this.LastModifiedTime = 0;
 		}
 
 		/// <summary>
-		/// Serialize this node to the bundle
+		/// Callback for when a reference is collapsed to a hash value. At the time of calling, both the hash and node data will be valid.
 		/// </summary>
-		/// <returns>Hash of the serialized node</returns>
-		protected abstract IoHash SerializeDirty();
+		protected virtual void Collapse()
+		{
+		}
 	}
 
 	/// <summary>
-	/// Attribute used to define a factory for a particular node type
+	/// Strongly typed reference to a <see cref="BundleNode"/>
 	/// </summary>
-	public class BundleNodeFactoryAttribute : Attribute
+	/// <typeparam name="T">Type of the node</typeparam>
+	public class BundleNodeRef<T> : BundleNodeRef where T : BundleNode
 	{
+		static readonly BundleNodeDeserializer<T> Deserializer = CreateDeserializer();
+
+		static BundleNodeDeserializer<T> CreateDeserializer()
+		{
+			BundleNodeDeserializerAttribute? Attribute = typeof(T).GetCustomAttribute<BundleNodeDeserializerAttribute>();
+			if (Attribute == null)
+			{
+				throw new InvalidOperationException($"No serializer is defined for {typeof(T).Name}");
+			}
+			return (BundleNodeDeserializer<T>)Activator.CreateInstance(Attribute.Type)!;
+		}
+
+		/// <inheritdoc cref="BundleNodeRef.Node"/>
+		public new T? Node
+		{
+			get => (T?)base.Node;
+			set => base.Node = value;
+		}
+
 		/// <summary>
-		/// The factory type. Should be derived from <see cref="BundleNodeFactory{T}"/>
+		/// Constructor
 		/// </summary>
-		public Type Type { get; }
+		/// <param name="Bundle">The bundle that contains the reference</param>
+		/// <param name="Hash">Hash of the referenced node</param>
+		public BundleNodeRef(Bundle Bundle, IoHash Hash) : base(Bundle, Hash)
+		{
+		}
 
 		/// <summary>
 		/// Constructor
 		/// </summary>
-		public BundleNodeFactoryAttribute(Type Type) => this.Type = Type;
+		/// <param name="Node">The referenced node</param>
+		public BundleNodeRef(T Node) : base(Node)
+		{
+		}
+
+		/// <summary>
+		/// Gets the referenced node, loading it from storage if necessary
+		/// </summary>
+		/// <returns>The parsed node</returns>
+		public async ValueTask<T> GetAsync()
+		{
+			if (StrongRef == null)
+			{
+				if (WeakRef != null && WeakRef.TryGetTarget(out StrongRef))
+				{
+					WeakRef = null;
+				}
+				else
+				{
+					StrongRef = Deserializer.Deserialize(Bundle!, await Bundle!.GetDataAsync(Hash));
+				}
+			}
+			return (T)StrongRef;
+		}
 	}
 
 	/// <summary>
-	/// Factory class for creating node types
+	/// Attribute used to define a factory for a particular node type
 	/// </summary>
-	/// <typeparam name="T">The type of node returned</typeparam>
-	public abstract class BundleNodeFactory<T> where T : BundleNode
+	public class BundleNodeDeserializerAttribute : Attribute
 	{
 		/// <summary>
-		/// Creates a default value for the factory
+		/// The factory type. Should be derived from <see cref="BundleNodeDeserializer{T}"/>
 		/// </summary>
-		/// <param name="Bundle"></param>
-		/// <returns>New instance of the node type</returns>
-		public abstract T CreateRoot(Bundle Bundle);
+		public Type Type { get; }
 
 		/// <summary>
-		/// Parse an object from data
+		/// Constructor
 		/// </summary>
-		/// <param name="Bundle">The owning bundle</param>
-		/// <param name="Hash">Hash of the data</param>
-		/// <param name="Data"></param>
-		/// <returns>New instance of the node type</returns>
-		public abstract T ParseRoot(Bundle Bundle, IoHash Hash, ReadOnlyMemory<byte> Data);
+		public BundleNodeDeserializerAttribute(Type Type) => this.Type = Type;
 	}
-}
-
 
+	/// <summary>
+	/// Factory class for deserializing node types
+	/// </summary>
+	/// <typeparam name="T">The type of node returned</typeparam>
+	public abstract class BundleNodeDeserializer<T> where T : BundleNode
+	{
+		/// <summary>
+		/// Deserializes data from the given bundle
+		/// </summary>
+		/// <param name="Bundle">The bundle containing the data</param>
+		/// <param name="Data">Data to deserialize</param>
+		/// <returns>New node parsed from the data</returns>
+		public abstract T Deserialize(Bundle Bundle, ReadOnlyMemory<byte> Data);
+	}
+}
diff --git a/Engine/Source/Programs/Shared/EpicGames.Horde/Bundles/Nodes/DirectoryNode.cs b/Engine/Source/Programs/Shared/EpicGames.Horde/Bundles/Nodes/DirectoryNode.cs
--- a/Engine/Source/Programs/Shared/EpicGames.Horde/Bundles/Nodes/DirectoryNode.cs
+++ b/Engine/Source/Programs/Shared/EpicGames.Horde/Bundles/Nodes/DirectoryNode.cs
@@ -5,6 +5,7 @@ using EpicGames.Serialization;
 using Microsoft.Extensions.Logging;
 using System;
 using System.Collections.Generic;
+using System.Diagnostics;
 using System.Diagnostics.CodeAnalysis;
 using System.IO;
 using System.Linq;
@@ -16,18 +17,8 @@ namespace EpicGames.Horde.Bundles.Nodes
 	/// <summary>
 	/// Flags for a directory entry
 	/// </summary>
-	public enum DirectoryEntryFlags : byte
+	public enum FileEntryFlags : byte
 	{
-		/// <summary>
-		/// This item is a directory
-		/// </summary>
-		Directory = 1,
-
-		/// <summary>
-		/// This item is a file
-		/// </summary>
-		File = 2,
-
 		/// <summary>
 		/// Indicates that the referenced file is executable
 		/// </summary>
@@ -55,144 +46,151 @@ namespace EpicGames.Horde.Bundles.Nodes
 	}
 
 	/// <summary>
-	/// Entry within a directory
+	/// Entry for a file within a directory node
 	/// </summary>
-	public class DirectoryEntry
+	public class FileEntry : BundleNodeRef<FileNode>
 	{
 		/// <summary>
-		/// Directory that owns this entry
-		/// </summary>
-		public DirectoryNode Directory { get; }
-
-		/// <summary>
-		/// Name of the entry
+		/// Name of this file
 		/// </summary>
 		public Utf8String Name { get; }
 
 		/// <summary>
-		/// Flags for this entry
+		/// Flags for this file
 		/// </summary>
-		public DirectoryEntryFlags Flags { get; }
+		public FileEntryFlags Flags { get; }
 
 		/// <summary>
-		/// Hash of this entry
+		/// Length of this entry
 		/// </summary>
-		IoHash Hash;
+		public long Length => (Node == null) ? CachedLength : Node.Length;
 
 		/// <summary>
-		/// Cached logical representation of this object
+		/// Cached length of this node
 		/// </summary>
-		BundleNode? Node;
+		long CachedLength;
 
 		/// <summary>
 		/// Constructor
 		/// </summary>
-		internal DirectoryEntry(DirectoryNode Directory, Utf8String Name, DirectoryEntryFlags Flags, IoHash Hash)
+		public FileEntry(Utf8String Name, FileEntryFlags Flags, FileNode Node)
+			: base(Node)
 		{
-			this.Directory = Directory;
 			this.Name = Name;
 			this.Flags = Flags;
-			this.Hash = Hash;
 		}
 
 		/// <summary>
 		/// Constructor
 		/// </summary>
-		internal DirectoryEntry(DirectoryNode Directory, Utf8String Name, DirectoryEntryFlags Flags, BundleNode Node)
+		public FileEntry(Bundle Bundle, Utf8String Name, FileEntryFlags Flags, long Length, IoHash Hash)
+			: base(Bundle, Hash)
 		{
-			this.Directory = Directory;
 			this.Name = Name;
 			this.Flags = Flags;
-			this.Node = Node;
+			this.CachedLength = Length;
 		}
 
-		internal IoHash Serialize()
+		/// <inheritdoc/>
+		protected override void Collapse()
 		{
-			if (Node != null)
-			{
-				Hash = Node.Serialize();
-			}
-			return Hash;
+			CachedLength = Node!.Length;
 		}
+	}
+
+	/// <summary>
+	/// Entry for a directory within a directory node
+	/// </summary>
+	public class DirectoryEntry : BundleNodeRef<DirectoryNode>
+	{
+		/// <summary>
+		/// Name of this directory
+		/// </summary>
+		public Utf8String Name { get; }
+
+		/// <summary>
+		/// Length of this directory tree
+		/// </summary>
+		public long Length => (Node == null) ? CachedLength : Node.Length;
 
 		/// <summary>
-		/// Gets the file associated with this entry
+		/// Cached value for the length of this tree
 		/// </summary>
-		/// <returns>File object associated with this entry</returns>
-		public async ValueTask<FileNode> GetFileAsync()
+		long CachedLength;
+
+		/// <summary>
+		/// Constructor
+		/// </summary>
+		public DirectoryEntry(Utf8String Name, DirectoryNode Node)
+			: base(Node)
 		{
-			if ((Flags & DirectoryEntryFlags.File) == 0)
-			{
-				throw new InvalidCastException();
-			}
-			if (Node == null)
-			{
-				ReadOnlyMemory<byte> Data = await Directory.Owner.GetDataAsync(Hash);
-				Node = new FileNode(Directory.Owner, Directory, Hash, Data);
-			}
-			return (FileNode)Node;
+			this.Name = Name;
 		}
 
 		/// <summary>
-		/// Gets the directory associated with this entry
+		/// Constructor
 		/// </summary>
-		/// <returns>Directory object associated with this entry</returns>
-		public async ValueTask<DirectoryNode> GetDirectoryAsync()
+		public DirectoryEntry(Bundle Bundle, Utf8String Name, long Length, IoHash Hash)
+			: base(Bundle, Hash)
 		{
-			if ((Flags & DirectoryEntryFlags.Directory) == 0)
-			{
-				throw new InvalidCastException();
-			}
-			if (Node == null)
-			{
-				ReadOnlyMemory<byte> Data = await Directory.Owner.GetDataAsync(Hash);
-				Node = DirectoryNode.Deserialize(Directory.Owner, Directory, Hash, Data.Span);
-			}
-			return (DirectoryNode)Node;
+			this.Name = Name;
+			this.CachedLength = Length;
+		}
+
+		/// <inheritdoc/>
+		protected override void Collapse()
+		{
+			CachedLength = Node!.Length;
 		}
 	}
 
 	/// <summary>
 	/// A directory node
 	/// </summary>
-	[BundleNodeFactory(typeof(DirectoryNodeFactory))]
+	[BundleNodeDeserializer(typeof(DirectoryNodeDeserializer))]
 	public class DirectoryNode : BundleNode
 	{
 		internal const byte TypeId = (byte)'d';
 
-		Dictionary<Utf8String, DirectoryEntry> NameToEntry = new Dictionary<Utf8String, DirectoryEntry>();
+		Dictionary<Utf8String, FileEntry> NameToFileEntry = new Dictionary<Utf8String, FileEntry>();
+		Dictionary<Utf8String, DirectoryEntry> NameToDirectoryEntry = new Dictionary<Utf8String, DirectoryEntry>();
 
 		/// <summary>
-		/// Entries within this directory
+		/// Total size of this directory
 		/// </summary>
-		public IReadOnlyCollection<DirectoryEntry> Entries => NameToEntry.Values;
+		public long Length => Files.Sum(x => x.Length) + Directories.Sum(x => x.Length);
 
-		internal DirectoryNode(Bundle Owner, BundleNode? Parent)
-			: base(Owner, Parent)
-		{
-		}
+		/// <summary>
+		/// All the files within this directory
+		/// </summary>
+		public IReadOnlyCollection<FileEntry> Files => NameToFileEntry.Values;
 
-		internal DirectoryNode(Bundle Owner, BundleNode? Parent, IoHash Hash)
-			: base(Owner, Parent, Hash)
-		{
-		}
+		/// <summary>
+		/// All the subdirectories within this directory
+		/// </summary>
+		public IReadOnlyCollection<DirectoryEntry> Directories => NameToDirectoryEntry.Values;
 
 		/// <summary>
 		/// Clear the contents of this directory
 		/// </summary>
 		public void Clear()
 		{
-			NameToEntry.Clear();
+			NameToFileEntry.Clear();
+			NameToDirectoryEntry.Clear();
 			MarkAsDirty();
 		}
 
 		/// <summary>
-		/// Attempt to get a directory entry with the given name
+		/// Check whether an entry with the given name exists in this directory
 		/// </summary>
-		/// <param name="Name">Name of the entry to find</param>
-		/// <param name="Entry">Receives the entry on success, or null on failure</param>
-		/// <returns>True if a matching entry was found</returns>
-		public bool TryGetEntry(Utf8String Name, [NotNullWhen(true)] out DirectoryEntry? Entry) => NameToEntry.TryGetValue(Name, out Entry);
+		/// <param name="Name">Name of the entry to search for</param>
+		/// <returns>True if the entry exists</returns>
+		public bool Contains(Utf8String Name) => TryGetFileEntry(Name, out _) || TryGetDirectoryEntry(Name, out _);
+
+		/// <inheritdoc/>
+		public override IEnumerable<BundleNodeRef> GetReferences() => Enumerable.Concat<BundleNodeRef>(Files, Directories);
+
+		#region File operations
 
 		/// <summary>
 		/// Adds a new directory with the given name
@@ -200,17 +198,17 @@ namespace EpicGames.Horde.Bundles.Nodes
 		/// <param name="Name">Name of the new directory</param>
 		/// <param name="Flags">Flags for the new file</param>
 		/// <returns>The new directory object</returns>
-		public FileNode CreateFile(Utf8String Name, DirectoryEntryFlags Flags)
+		public FileNode CreateFile(Utf8String Name, FileEntryFlags Flags)
 		{
-			if((Flags & DirectoryEntryFlags.Directory) != 0)
+			if(TryGetDirectoryEntry(Name, out _))
 			{
-				throw new ArgumentException("Directory flag cannot be specified for new file");
+				throw new ArgumentException($"A directory with the name {Name} already exists");
 			}
 
-			FileNode NewNode = new FileNode(Owner, this);
+			FileNode NewNode = new FileNode();
 
-			DirectoryEntry Entry = new DirectoryEntry(this, Name, Flags | DirectoryEntryFlags.File, NewNode);
-			NameToEntry[Name] = Entry;
+			FileEntry Entry = new FileEntry(Name, Flags, NewNode);
+			NameToFileEntry[Name] = Entry;
 			MarkAsDirty();
 
 			return NewNode;
@@ -222,7 +220,7 @@ namespace EpicGames.Horde.Bundles.Nodes
 		/// <param name="Path">Path to the file</param>
 		/// <param name="Flags">Flags for the new file</param>
 		/// <returns>The new directory object</returns>
-		public async ValueTask<FileNode> CreateFileByPathAsync(Utf8String Path, DirectoryEntryFlags Flags)
+		public async ValueTask<FileNode> CreateFileByPathAsync(Utf8String Path, FileEntryFlags Flags)
 		{
 			DirectoryNode Directory = this;
 
@@ -242,12 +240,35 @@ namespace EpicGames.Horde.Bundles.Nodes
 			}
 		}
 
+		/// <summary>
+		/// Attempts to get a file entry with the given name
+		/// </summary>
+		/// <param name="Name">Name of the file</param>
+		/// <param name="Entry">Entry for the file</param>
+		/// <returns>True if the file was found</returns>
+		public bool TryGetFileEntry(Utf8String Name, [NotNullWhen(true)] out FileEntry? Entry) => NameToFileEntry.TryGetValue(Name, out Entry);
+
+		/// <summary>
+		/// Deletes the file entry with the given name
+		/// </summary>
+		/// <param name="Name">Name of the entry to delete</param>
+		/// <returns>True if the entry was found, false otherwise</returns>
+		public bool DeleteFile(Utf8String Name)
+		{
+			if (NameToFileEntry.Remove(Name))
+			{
+				MarkAsDirty();
+				return true;
+			}
+			return false;
+		}
+
 		/// <summary>
 		/// 
 		/// </summary>
 		/// <param name="Path"></param>
 		/// <returns></returns>
-		public async ValueTask DeleteFileByPathAsync(Utf8String Path)
+		public async ValueTask<bool> DeleteFileByPathAsync(Utf8String Path)
 		{
 			DirectoryNode Directory = this;
 
@@ -257,8 +278,7 @@ namespace EpicGames.Horde.Bundles.Nodes
 				int Length = RemainingPath.IndexOf('/');
 				if (Length == -1)
 				{
-					Directory.Delete(RemainingPath);
-					return;
+					return Directory.DeleteFile(RemainingPath);
 				}
 				if (Length > 0)
 				{
@@ -268,6 +288,10 @@ namespace EpicGames.Horde.Bundles.Nodes
 			}
 		}
 
+		#endregion
+
+		#region Directory operations
+
 		/// <summary>
 		/// Adds a new directory with the given name
 		/// </summary>
@@ -275,15 +299,28 @@ namespace EpicGames.Horde.Bundles.Nodes
 		/// <returns>The new directory object</returns>
 		public DirectoryNode AddDirectory(Utf8String Name)
 		{
-			DirectoryNode NewNode = new DirectoryNode(Owner, this);
+			if (TryGetFileEntry(Name, out _))
+			{
+				throw new ArgumentException($"A file with the name '{Name}' already exists in this directory", nameof(Name));
+			}
+
+			DirectoryNode Node = new DirectoryNode();
 
-			DirectoryEntry Entry = new DirectoryEntry(this, Name, DirectoryEntryFlags.Directory, NewNode);
-			NameToEntry.Add(Name, Entry);
+			DirectoryEntry Entry = new DirectoryEntry(Name, Node);
+			NameToDirectoryEntry.Add(Name, Entry);
 			MarkAsDirty();
 
-			return NewNode;
+			return Node;
 		}
 
+		/// <summary>
+		/// Attempts to get a directory entry with the given name
+		/// </summary>
+		/// <param name="Name">Name of the directory</param>
+		/// <param name="Entry">Entry for the directory</param>
+		/// <returns>True if the directory was found</returns>
+		public bool TryGetDirectoryEntry(Utf8String Name, [NotNullWhen(true)] out DirectoryEntry? Entry) => NameToDirectoryEntry.TryGetValue(Name, out Entry);
+
 		/// <summary>
 		/// Tries to get a directory with the given name
 		/// </summary>
@@ -291,9 +328,9 @@ namespace EpicGames.Horde.Bundles.Nodes
 		/// <returns>The new directory object</returns>
 		public async ValueTask<DirectoryNode?> FindDirectoryAsync(Utf8String Name)
 		{
-			if (TryGetEntry(Name, out DirectoryEntry? Entry))
+			if (TryGetDirectoryEntry(Name, out DirectoryEntry? Entry))
 			{
-				return await Entry.GetDirectoryAsync();
+				return await Entry.GetAsync();
 			}
 			else
 			{
@@ -317,9 +354,9 @@ namespace EpicGames.Horde.Bundles.Nodes
 		/// </summary>
 		/// <param name="Name">Name of the entry to delete</param>
 		/// <returns>True if the entry was found, false otherwise</returns>
-		public bool Delete(Utf8String Name)
+		public bool DeleteDirectory(Utf8String Name)
 		{
-			if(NameToEntry.Remove(Name))
+			if(NameToDirectoryEntry.Remove(Name))
 			{
 				MarkAsDirty();
 				return true;
@@ -327,41 +364,69 @@ namespace EpicGames.Horde.Bundles.Nodes
 			return false;
 		}
 
+		#endregion
+
 		/// <inheritdoc/>
-		protected override IoHash SerializeDirty()
+		public override ReadOnlyMemory<byte> Serialize()
 		{
+			// Measure the required size of the write buffer
 			int Size = 1;
-			foreach (DirectoryEntry Entry in NameToEntry.Values)
+			Size += VarInt.Measure(NameToFileEntry.Count);
+			foreach (FileEntry Entry in NameToFileEntry.Values)
 			{
-				Size += 1 + (Entry.Name.Length + 1) + IoHash.NumBytes;
+				Size += (Entry.Name.Length + 1) + 1 + VarInt.Measure((ulong)Entry.Length) + IoHash.NumBytes;
+			}
+			Size += VarInt.Measure(NameToDirectoryEntry.Count);
+			foreach (DirectoryEntry Entry in NameToDirectoryEntry.Values)
+			{
+				Size += (Entry.Name.Length + 1) + VarInt.Measure((ulong)Entry.Length) + IoHash.NumBytes;
 			}
 
+			// Allocate the buffer and copy the node to it
 			byte[] Data = new byte[Size];
 			Span<byte> Span = Data.AsSpan();
 
 			Span[0] = TypeId;
 			Span = Span.Slice(1);
 
-			List<IoHash> References = new List<IoHash>();
-			foreach (DirectoryEntry Entry in NameToEntry.Values.OrderBy(x => x.Name))
+			int FileCountBytes = VarInt.Write(Span, NameToFileEntry.Count);
+			Span = Span.Slice(FileCountBytes);
+
+			foreach (FileEntry FileEntry in NameToFileEntry.Values.OrderBy(x => x.Name))
 			{
-				IoHash EntryHash = Entry.Serialize();
-				References.Add(EntryHash);
+				FileEntry.Name.Span.CopyTo(Span);
+				Span = Span.Slice(FileEntry.Name.Length + 1);
 
-				Span[0] = (byte)Entry.Flags;
+				Span[0] = (byte)FileEntry.Flags;
 				Span = Span.Slice(1);
 
-				Entry.Name.Span.CopyTo(Span);
-				Span = Span.Slice(Entry.Name.Length + 1);
+				int LengthBytes = VarInt.Write(Span, FileEntry.Length);
+				Span = Span.Slice(LengthBytes);
+
+				FileEntry.Hash.CopyTo(Span);
+				Span = Span.Slice(IoHash.NumBytes);
+			}
 
-				EntryHash.CopyTo(Span);
+			int DirectoryCountBytes = VarInt.Write(Span, NameToDirectoryEntry.Count);
+			Span = Span.Slice(DirectoryCountBytes);
+
+			foreach (DirectoryEntry DirectoryEntry in NameToDirectoryEntry.Values.OrderBy(x => x.Name))
+			{
+				DirectoryEntry.Name.Span.CopyTo(Span);
+				Span = Span.Slice(DirectoryEntry.Name.Length + 1);
+
+				int LengthBytes = VarInt.Write(Span, DirectoryEntry.Length);
+				Span = Span.Slice(LengthBytes);
+
+				DirectoryEntry.Hash.CopyTo(Span);
 				Span = Span.Slice(IoHash.NumBytes);
 			}
 
-			return Owner.WriteNode(Data, References);
+			Debug.Assert(Span.Length == 0);
+			return Data;
 		}
 
-		internal static DirectoryNode Deserialize(Bundle Owner, BundleNode? Parent, IoHash Hash, ReadOnlySpan<byte> Span)
+		internal static DirectoryNode Deserialize(Bundle Bundle, ReadOnlySpan<byte> Span)
 		{
 			if(Span[0] != TypeId)
 			{
@@ -370,21 +435,49 @@ namespace EpicGames.Horde.Bundles.Nodes
 
 			Span = Span.Slice(1);
 
-			DirectoryNode Node = new DirectoryNode(Owner, Parent);
-			while (Span.Length > 0)
+			DirectoryNode Node = new DirectoryNode();
+
+			int FileCount = (int)VarInt.Read(Span, out int FileCountBytes);
+			Span = Span.Slice(FileCountBytes);
+
+			Node.NameToFileEntry.EnsureCapacity(FileCount);
+			for (int Idx = 0; Idx < FileCount; Idx++)
 			{
-				DirectoryEntryFlags Flags = (DirectoryEntryFlags)Span[0];
+				int NameLen = Span.IndexOf((byte)0);
+				Utf8String Name = new Utf8String(Span.Slice(0, NameLen).ToArray());
+				Span = Span.Slice(NameLen + 1);
+
+				FileEntryFlags Flags = (FileEntryFlags)Span[0];
 				Span = Span.Slice(1);
 
-				int Length = Span.IndexOf((byte)0);
-				Utf8String Name = new Utf8String(Span.Slice(0, Length).ToArray());
-				Span = Span.Slice(Length + 1);
+				long Length = (long)VarInt.Read(Span, out int LengthBytes);
+				Span = Span.Slice(LengthBytes);
+
+				IoHash Hash = new IoHash(Span);
+				Span = Span.Slice(IoHash.NumBytes);
+
+				Node.NameToFileEntry[Name] = new FileEntry(Bundle, Name, Flags, Length, Hash);
+			}
+
+			int DirectoryCount = (int)VarInt.Read(Span, out int DirectoryCountBytes);
+			Span = Span.Slice(DirectoryCountBytes);
 
-				IoHash EntryHash = new IoHash(Span);
+			Node.NameToDirectoryEntry.EnsureCapacity(DirectoryCount);
+			for (int Idx = 0; Idx < DirectoryCount; Idx++)
+			{
+				int NameLen = Span.IndexOf((byte)0);
+				Utf8String Name = new Utf8String(Span.Slice(0, NameLen).ToArray());
+				Span = Span.Slice(NameLen + 1);
+
+				long Length = (long)VarInt.Read(Span, out int LengthBytes);
+				Span = Span.Slice(LengthBytes);
+
+				IoHash Hash = new IoHash(Span);
 				Span = Span.Slice(IoHash.NumBytes);
 
-				Node.NameToEntry[Name] = new DirectoryEntry(Node, Name, Flags, EntryHash);
+				Node.NameToDirectoryEntry[Name] = new DirectoryEntry(Bundle, Name, Length, Hash);
 			}
+
 			return Node;
 		}
 
@@ -406,7 +499,7 @@ namespace EpicGames.Horde.Bundles.Nodes
 			foreach (FileInfo FileInfo in DirectoryInfo.EnumerateFiles())
 			{
 				Logger.LogInformation("Adding {File}", FileInfo.FullName);
-				FileNode FileNode = CreateFile(FileInfo.Name, DirectoryEntryFlags.File);
+				FileNode FileNode = CreateFile(FileInfo.Name, 0);
  				await FileNode.CopyFromFileAsync(FileInfo, Options);
 			}
 		}
@@ -422,20 +515,17 @@ namespace EpicGames.Horde.Bundles.Nodes
 			DirectoryInfo.Create();
 
 			List<Task> Tasks = new List<Task>();
-			foreach (DirectoryEntry Entry in Entries)
+			foreach (FileEntry FileEntry in Files)
 			{
-				if ((Entry.Flags & DirectoryEntryFlags.Directory) != 0)
-				{
-					DirectoryInfo SubDirectoryInfo = DirectoryInfo.CreateSubdirectory(Entry.Name.ToString());
-					DirectoryNode SubDirectoryNode = await Entry.GetDirectoryAsync();
-					Tasks.Add(Task.Run(() => SubDirectoryNode.CopyToDirectoryAsync(SubDirectoryInfo, Logger)));
-				}
-				else
-				{
-					FileInfo FileInfo = new FileInfo(Path.Combine(DirectoryInfo.FullName, Entry.Name.ToString()));
-					FileNode FileNode = await Entry.GetFileAsync();
-					Tasks.Add(Task.Run(() => FileNode.CopyToFileAsync(FileInfo)));
-				}
+				FileInfo FileInfo = new FileInfo(Path.Combine(DirectoryInfo.FullName, FileEntry.Name.ToString()));
+				FileNode FileNode = await FileEntry.GetAsync();
+				Tasks.Add(Task.Run(() => FileNode.CopyToFileAsync(FileInfo)));
+			}
+			foreach (DirectoryEntry DirectoryEntry in Directories)
+			{
+				DirectoryInfo SubDirectoryInfo = DirectoryInfo.CreateSubdirectory(DirectoryEntry.Name.ToString());
+				DirectoryNode SubDirectoryNode = await DirectoryEntry.GetAsync();
+				Tasks.Add(Task.Run(() => SubDirectoryNode.CopyToDirectoryAsync(SubDirectoryInfo, Logger)));
 			}
 
 			await Task.WhenAll(Tasks);
@@ -445,18 +535,12 @@ namespace EpicGames.Horde.Bundles.Nodes
 	/// <summary>
 	/// Factory for creating and serializing <see cref="DirectoryNode"/> objects
 	/// </summary>
-	public class DirectoryNodeFactory : BundleNodeFactory<DirectoryNode>
+	public class DirectoryNodeDeserializer : BundleNodeDeserializer<DirectoryNode>
 	{
 		/// <inheritdoc/>
-		public override DirectoryNode CreateRoot(Bundle Bundle)
-		{
-			return new DirectoryNode(Bundle, null);
-		}
-
-		/// <inheritdoc/>
-		public override DirectoryNode ParseRoot(Bundle Bundle, IoHash Hash, ReadOnlyMemory<byte> Data)
+		public override DirectoryNode Deserialize(Bundle Bundle, ReadOnlyMemory<byte> Data)
 		{
-			return DirectoryNode.Deserialize(Bundle, null, Hash, Data.Span);
+			return DirectoryNode.Deserialize(Bundle, Data.Span);
 		}
 	}
 }
diff --git a/Engine/Source/Programs/Shared/EpicGames.Horde/Bundles/Nodes/FileNode.cs b/Engine/Source/Programs/Shared/EpicGames.Horde/Bundles/Nodes/FileNode.cs
--- a/Engine/Source/Programs/Shared/EpicGames.Horde/Bundles/Nodes/FileNode.cs
+++ b/Engine/Source/Programs/Shared/EpicGames.Horde/Bundles/Nodes/FileNode.cs
@@ -73,7 +73,7 @@ namespace EpicGames.Horde.Bundles.Nodes
 	/// Chunks are pushed into a tree hierarchy as data is appended to the root, with nodes of the tree also split along content-aware boundaries with <see cref="IoHash.NumBytes"/> granularity.
 	/// Once a chunk has been written to storage, it is treated as immutable.
 	/// </summary>
-	[BundleNodeFactory(typeof(FileNodeFactory))]
+	[BundleNodeDeserializer(typeof(FileNodeDeserializer))]
 	public sealed class FileNode : BundleNode
 	{
 		const byte TypeId = (byte)'c';
@@ -81,34 +81,55 @@ namespace EpicGames.Horde.Bundles.Nodes
 		internal int Depth { get; private set; }
 		internal ReadOnlyMemory<byte> Payload { get; private set; }
 		ReadOnlyMemory<byte> Data; // Full serialized data, including type id and header fields.
+		IoHash? Hash; // Hash of the serialized data buffer.
 
 		// In-memory state
-		uint Hash;
-		byte[]? WriteBuffer; // Null if the node is complete and does not support writes
-		List<FileNode?>? ChildNodes;
+		bool IsRoot = true;
+		uint RollingHash;
+		byte[] WriteBuffer = Array.Empty<byte>();
+		List<BundleNodeRef<FileNode>>? ChildNodeRefs;
 
 		/// <summary>
-		/// Constructor
+		/// Length of this tree
 		/// </summary>
-		public FileNode(Bundle Owner, BundleNode? Parent)
-			: this(Owner, Parent, 0, new ChunkingOptions())
-		{
-		}
+		public long Length { get; private set; }
+
+		/// <summary>
+		/// Whether this node is read-only
+		/// </summary>
+		public bool IsComplete() => Hash != null;
 
 		/// <summary>
-		/// Create a writable file node
+		/// Accessor for the children of this node
 		/// </summary>
-		private FileNode(Bundle Owner, BundleNode? Parent, int Depth, ChunkingOptions Options)
-			: base(Owner, Parent)
+		public IReadOnlyList<BundleNodeRef<FileNode>> Children => (IReadOnlyList<BundleNodeRef<FileNode>>?)ChildNodeRefs ?? Array.Empty<BundleNodeRef<FileNode>>();
+
+		/// <inheritdoc/>
+		public override ReadOnlyMemory<byte> Serialize()
 		{
-			MakeWritable(Depth, Options);
+			if (ChildNodeRefs != null)
+			{
+				Span<byte> Span = GetWritableSpan(VarInt.Measure((ulong)Length) + ChildNodeRefs.Count * IoHash.NumBytes, 0);
+
+				int LengthBytes = VarInt.Write(Span, Length);
+				Span = Span.Slice(LengthBytes);
+
+				foreach (BundleNodeRef<FileNode> ChildNodeRef in ChildNodeRefs)
+				{
+					ChildNodeRef.Hash.CopyTo(Span);
+					Span = Span.Slice(IoHash.NumBytes);
+				}
+			}
+			return Data;
 		}
 
+		/// <inheritdoc/>
+		public override IEnumerable<BundleNodeRef> GetReferences() => Children;
+
 		/// <summary>
 		/// Create a file node from deserialized data
 		/// </summary>
-		public FileNode(Bundle Owner, BundleNode? Parent, IoHash Hash, ReadOnlyMemory<byte> Data)
-			: base(Owner, Parent, Hash)
+		public static FileNode Deserialize(Bundle Bundle, ReadOnlyMemory<byte> Data)
 		{
 			ReadOnlySpan<byte> Span = Data.Span;
 			if (Span[0] != TypeId)
@@ -116,82 +137,64 @@ namespace EpicGames.Horde.Bundles.Nodes
 				throw new InvalidDataException("Invalid type id");
 			}
 
-			Depth = (int)VarInt.Read(Span.Slice(1), out int DepthBytes);
-			int HeaderLength = 1 + DepthBytes;
+			FileNode Node = new FileNode();
+			Node.Data = Data;
+			Node.Hash = IoHash.Compute(Data.Span);
+			Node.Depth = (int)VarInt.Read(Span.Slice(1), out int DepthBytes);
 
-			Payload = Data.Slice(HeaderLength);
-			this.Data = Data;
+			int HeaderLength = 1 + DepthBytes;
+			Node.Payload = Data.Slice(HeaderLength);
 
-			if (Depth > 0)
+			if (Node.Depth == 0)
 			{
-				ChildNodes = new List<FileNode?>(Payload.Length / IoHash.NumBytes);
-				ChildNodes.AddRange(Enumerable.Repeat<FileNode?>(null, Payload.Length / IoHash.NumBytes));
-			}
-		}
-
-		private void MakeWritable(int Depth, ChunkingOptions Options)
-		{
-			int WriteBufferSize = 1 + VarInt.Measure(Depth);
-			if (Depth == 0)
-			{
-				WriteBufferSize += Options.LeafOptions.MaxSize;
+				Node.Length = Node.Payload.Length;
 			}
 			else
 			{
-				WriteBufferSize += Options.InteriorOptions.MaxSize;
-			}
-
-			byte[] WriteBuffer = new byte[WriteBufferSize];
-			WriteBuffer[0] = TypeId;
-			int DepthBytes = VarInt.Write(WriteBuffer.AsSpan(1), Depth);
-			int HeaderLength = 1 + DepthBytes;
-
-			this.Depth = Depth;
-			this.Payload = ReadOnlyMemory<byte>.Empty;
-			this.Data = WriteBuffer.AsMemory(0, HeaderLength);
-			this.Hash = 0;
-			this.WriteBuffer = WriteBuffer;
-			this.ChildNodes = (Depth > 0) ? new List<FileNode?>() : null;
-		}
+				Node.Length = (long)VarInt.Read(Span, out int LengthBytes);
+				Span = Span.Slice(LengthBytes);
 
-		/// <summary>
-		/// Returns the number of child nodes to this chunk
-		/// </summary>
-		public int GetChildCount() => (ChildNodes == null)? 0 : ChildNodes.Count;
-
-		/// <summary>
-		/// Gets a child node at the given index
-		/// </summary>
-		/// <param name="Index">Index of the child node</param>
-		/// <returns>Node containing data for the given child</returns>
-		public async ValueTask<FileNode> GetChildNodeAsync(int Index)
-		{
-			if(ChildNodes == null)
-			{
-				throw new InvalidOperationException("Node does not contain any children");
-			}
+				Node.ChildNodeRefs = new List<BundleNodeRef<FileNode>>(Node.Payload.Length / IoHash.NumBytes);
 
-			FileNode? ChildNode = ChildNodes[Index];
-			if (ChildNode == null)
-			{
-				IoHash Hash = new IoHash(Payload.Span.Slice(Index * IoHash.NumBytes));
-				ReadOnlyMemory<byte> Data = await Owner.GetDataAsync(Hash);
-				ChildNode = new FileNode(Owner, this, Hash, Data);
-				ChildNodes[Index] = ChildNode;
+				Span = Span.Slice(HeaderLength);
+				while (Span.Length > 0)
+				{
+					IoHash ChildHash = new IoHash(Span);
+					Node.ChildNodeRefs.Add(new BundleNodeRef<FileNode>(Bundle, ChildHash));
+				}
 			}
-			return ChildNode;
+			return Node;
 		}
 
 		/// <summary>
-		/// Enumerate the children from this node
+		/// Gets a writable span of the given size.
 		/// </summary>
+		/// <param name="SpanSize">Required size of the writable span</param>
+		/// <param name="HintMaxPayloadSize">Hint for the maximum size of the payload, to avoid having to reallocate the buffer. May be zero if unknown.</param>
 		/// <returns></returns>
-		public async IAsyncEnumerable<FileNode> GetChildNodesAsync()
+		private Span<byte> GetWritableSpan(int SpanSize, int HintMaxPayloadSize)
 		{
-			for (int Idx = 0; Idx < GetChildCount(); Idx++)
+			Debug.Assert(!IsComplete());
+
+			int HeaderLength = 1 + VarInt.Measure(Depth);
+			int PrevPayloadLength = Payload.Length;
+
+			int MinWriteBufferSize = HeaderLength + Math.Max(PrevPayloadLength + SpanSize, HintMaxPayloadSize);
+			if (MinWriteBufferSize > WriteBuffer.Length)
 			{
-				yield return await GetChildNodeAsync(Idx);
+				byte[] NewWriteBuffer = new byte[MinWriteBufferSize];
+
+				NewWriteBuffer[0] = TypeId;
+				VarInt.Write(NewWriteBuffer.AsSpan(1), Depth);
+				Payload.Span.CopyTo(NewWriteBuffer.AsSpan(HeaderLength));
+
+				WriteBuffer = NewWriteBuffer;
 			}
+
+			Data = WriteBuffer.AsMemory(0, HeaderLength + PrevPayloadLength + SpanSize);
+			Payload = Data.Slice(HeaderLength);
+
+			return WriteBuffer.AsSpan(HeaderLength + PrevPayloadLength, SpanSize);
 		}
 
 		/// <summary>
@@ -247,12 +250,12 @@ namespace EpicGames.Horde.Bundles.Nodes
 		/// <param name="OutputStream">The output stream to receive the data</param>
 		public async Task CopyToStreamAsync(Stream OutputStream)
 		{
-			if (Depth > 0)
+			if (ChildNodeRefs != null)
 			{
-				for (int Idx = 0; Idx < GetChildCount(); Idx++)
+				foreach (BundleNodeRef<FileNode> ChildNodeRef in ChildNodeRefs)
 				{
-					FileNode Node = await GetChildNodeAsync(Idx);
-					await Node.CopyToStreamAsync(OutputStream);
+					FileNode ChildNode = await ChildNodeRef.GetAsync();
+					await ChildNode.CopyToStreamAsync(OutputStream);
 				}
 			}
 			else
@@ -281,9 +284,9 @@ namespace EpicGames.Horde.Bundles.Nodes
 		/// <param name="Options">Settings for chunking the data</param>
 		public void Append(ReadOnlyMemory<byte> Input, ChunkingOptions Options)
 		{
-			if (Parent is FileNode)
+			if (!IsRoot)
 			{
-				throw new InvalidOperationException("Data may only be appended to the root of a tree of file nodes.");
+				throw new InvalidOperationException("Data may only be appended to the root of a file node tree");
 			}
 
 			for (; ; )
@@ -296,29 +299,38 @@ namespace EpicGames.Horde.Bundles.Nodes
 				}
 
 				// Increase the height of the tree by pushing the contents of this node into a new child node
-				FileNode NewNode = new FileNode(Owner, this, IoHash.Zero, Data);
-				if (ChildNodes != null)
+				FileNode NewNode = new FileNode();
+				NewNode.IsRoot = false;
+				NewNode.Depth = Depth;
+				NewNode.Payload = Payload;
+				NewNode.Data = Data;
+				NewNode.ChildNodeRefs = ChildNodeRefs;
+
+				// Detach all the child refs
+				if (ChildNodeRefs != null)
 				{
-					for (int Idx = 0; Idx < ChildNodes.Count; Idx++)
+					foreach (BundleNodeRef<FileNode> ChildNodeRef in ChildNodeRefs)
 					{
-						FileNode? ChildNode = ChildNodes[Idx];
-						if (ChildNode != null)
-						{
-							ChildNode.Parent = NewNode;
-							NewNode.ChildNodes![Idx] = ChildNode;
-						}
+						ChildNodeRef.Detach();
 					}
+					ChildNodeRefs = null;
 				}
 
-				// Reinitialize the current node
-				MakeWritable(Depth + 1, Options);
-				AddChildNode(NewNode);
+				// Increase the depth and reset the current node
+				Depth++;
+				ChildNodeRefs = new List<BundleNodeRef<FileNode>>();
+				Payload = ReadOnlyMemory<byte>.Empty;
+				Data = ReadOnlyMemory<byte>.Empty;
+				WriteBuffer = Array.Empty<byte>();
+
+				// Append the new node as a child
+				ChildNodeRefs.Add(new BundleNodeRef<FileNode>(NewNode));
 			}
 		}
 
 		private ReadOnlyMemory<byte> AppendToNode(ReadOnlyMemory<byte> Data, ChunkingOptions Options)
 		{
-			if (WriteBuffer == null || Data.Length == 0)
+			if (Data.Length == 0 || IsComplete())
 			{
 				return Data;
 			}
@@ -335,58 +347,134 @@ namespace EpicGames.Horde.Bundles.Nodes
 
 		private ReadOnlyMemory<byte> AppendToLeafNode(ReadOnlyMemory<byte> NewData, ChunkingOptions Options)
 		{
-			ReadOnlySpan<byte> NewSpan = NewData.Span;
+			// Fast path for appending data to the buffer up to the chunk window size
+			int WindowSize = Options.LeafOptions.MinSize;
+			if (Payload.Length < WindowSize)
+			{
+				int AppendLength = Math.Min(WindowSize - Payload.Length, NewData.Length);
+				AppendLeafData(NewData.Span.Slice(0, AppendLength), Options);
+				NewData = NewData.Slice(AppendLength);
+			}
 
-			uint HashThreshold = (uint)((1L << 32) / Options.LeafOptions.TargetSize);
-			for (int Length = 0; Length < NewData.Length; Length++)
+			// Cap the maximum amount of data to append to this node
+			int MaxLength = Math.Min(NewData.Length, Options.LeafOptions.MaxSize - Payload.Length);
+			if (MaxLength > 0)
 			{
-				Hash = BuzHash.Add(Hash, NewSpan[Length]);
+				int AppendLength = HashScan(NewData.Span.Slice(0, MaxLength), Options);
+				AppendLeafData(NewData.Span.Slice(0, AppendLength), Options);
+				NewData = NewData.Slice(AppendLength);
+			}
+
+			return NewData;
+		}
 
-				int BlockSize = Payload.Length + Length;
-				if ((BlockSize >= Options.LeafOptions.MinSize && Hash < HashThreshold) || (BlockSize >= Options.LeafOptions.MaxSize))
+		private int HashScan(ReadOnlySpan<byte> NewSpan, ChunkingOptions Options)
+		{
+			int WindowSize = Options.LeafOptions.MinSize;
+			Debug.Assert(Payload.Length >= WindowSize);
+
+			int Length = 0;
+
+			// If this is the first time we're hashing the data, compute the hash of the existing window
+			if (Payload.Length == WindowSize)
+			{
+				RollingHash = BuzHash.Add(RollingHash, Payload.Span);
+			}
+
+			// Get the threshold for the rolling hash
+			uint RollingHashThreshold = (uint)((1L << 32) / Options.LeafOptions.TargetSize);
+
+			// Get the remaining part of the payload span leading into the new data
+			ReadOnlySpan<byte> PayloadSpan = Payload.Span.Slice(Payload.Length - WindowSize);
+
+			// Step the window through the tail end of the existing payload window. In this state, update the hash to remove data from the current payload, and add data from the new payload.
+			int SplitLength = Math.Min(NewSpan.Length, WindowSize);
+			for (; Length < SplitLength; Length++)
+			{
+				RollingHash = BuzHash.Add(RollingHash, NewSpan[Length]);
+				if (RollingHash < RollingHashThreshold)
+				{
+					return Length;
+				}
+				RollingHash = BuzHash.Sub(RollingHash, PayloadSpan[Length], WindowSize);
+			}
+
+			// Step through the new window.
+			for (; Length < NewSpan.Length; Length++)
+			{
+				RollingHash = BuzHash.Add(RollingHash, NewSpan[Length]);
+				if (RollingHash < RollingHashThreshold)
+				{
+					return Length;
+				}
+				RollingHash = BuzHash.Sub(RollingHash, NewSpan[Length - WindowSize], WindowSize);
+			}
+
+			return Length;
+		}
+
+		private int AppendInternal(ReadOnlySpan<byte> TailSpan, ReadOnlySpan<byte> HeadSpan, uint RollingHashThreshold, int WindowSize, int MaxLength, ChunkingOptions Options)
+		{
+			for (int Length = 0; Length < MaxLength; Length++)
+			{
+				RollingHash = BuzHash.Add(RollingHash, HeadSpan[Length]);
+
+				if (RollingHash < RollingHashThreshold)
 				{
-					AppendData(NewSpan.Slice(0, Length));
-					WriteBuffer = null;
-					return NewData.Slice(Length);
+					AppendLeafData(HeadSpan.Slice(0, Length), Options);
+					Hash = IoHash.Compute(Data.Span);
+					return Length;
 				}
+
+				RollingHash = BuzHash.Sub(RollingHash, TailSpan[Length], WindowSize);
 			}
+			return -1;
+		}
 
-			AppendData(NewSpan);
-			return ReadOnlyMemory<byte>.Empty;
+		private void AppendLeafData(ReadOnlySpan<byte> Data, ChunkingOptions Options)
+		{
+			Span<byte> WriteSpan = GetWritableSpan(Data.Length, Options.LeafOptions.MaxSize);
+			Data.CopyTo(WriteSpan);
+			Length += Data.Length;
 		}
 
 		private ReadOnlyMemory<byte> AppendToInteriorNode(ReadOnlyMemory<byte> NewData, ChunkingOptions Options)
 		{
 			for (; ; )
 			{
-				Debug.Assert(ChildNodes != null);
+				Debug.Assert(ChildNodeRefs != null);
 
 				// Try to write to the last node
-				if (ChildNodes.Count > 0)
+				if (ChildNodeRefs.Count > 0)
 				{
-					FileNode? LastNode = ChildNodes[^1];
+					FileNode? LastNode = ChildNodeRefs[^1].Node;
 					if (LastNode != null)
 					{
+						// Update the length to match the new node
+						Length -= LastNode.Length;
 						NewData = LastNode.AppendToNode(NewData, Options);
-						if (LastNode.WriteBuffer == null)
-						{
-							Span<byte> LastHashSpan = WriteBuffer.AsSpan(Data.Length - IoHash.NumBytes);
+						Length += LastNode.Length;
 
+						// If the last node is complete, write it to the buffer
+						if (LastNode.IsComplete())
+						{
 							// Add it to the write buffer
-							IoHash LastHash = LastNode.Serialize();
-							LastHash.CopyTo(LastHashSpan);
+							Span<byte> LastHashSpan = GetWritableSpan(IoHash.NumBytes, Options.InteriorOptions.MaxSize);
+							LastNode.Hash!.Value.CopyTo(LastHashSpan);
 
 							// Update the hash
-							Hash = BuzHash.Add(Hash, LastHashSpan);
+							RollingHash = BuzHash.Add(RollingHash, LastHashSpan);
 
 							// Check if it's time to finish this chunk
 							uint HashThreshold = (uint)(((1L << 32) * IoHash.NumBytes) / Options.LeafOptions.TargetSize);
-							if ((Payload.Length >= Options.InteriorOptions.MinSize && Hash < HashThreshold) || (Payload.Length >= Options.InteriorOptions.MaxSize))
+							if ((Payload.Length >= Options.InteriorOptions.MinSize && RollingHash < HashThreshold) || (Payload.Length >= Options.InteriorOptions.MaxSize))
 							{
-								WriteBuffer = null;
+								Hash = IoHash.Compute(Data.Span);
 								return NewData;
 							}
 						}
+
+						// Bail out if there's nothing left to write
 						if (NewData.Length == 0)
 						{
 							return NewData;
@@ -394,86 +482,20 @@ namespace EpicGames.Horde.Bundles.Nodes
 					}
 				}
 
-				// Create a new child node
-				FileNode NewNode = new FileNode(Owner, this, Depth - 1, Options);
-				AddChildNode(NewNode);
+				// Add a new child node
+				FileNode ChildNode = new FileNode();
+				ChildNode.IsRoot = false;
+				ChildNodeRefs.Add(new BundleNodeRef<FileNode>(ChildNode));
 			}
 		}
-
-		void AddChildNode(FileNode NewNode)
-		{
-			Debug.Assert(ChildNodes != null);
-			ChildNodes.Add(NewNode);
-			ExpandData(IoHash.NumBytes);
-		}
-
-		void ExpandData(int Size)
-		{
-			Debug.Assert(WriteBuffer != null);
-
-			int NewDataLength = Data.Length + Size;
-			if (NewDataLength > WriteBuffer.Length)
-			{
-				Array.Resize(ref WriteBuffer, NewDataLength + 4096);
-			}
-
-			int HeaderLength = Data.Length - Payload.Length;
-			Data = WriteBuffer.AsMemory(0, NewDataLength);
-			Payload = Data.Slice(HeaderLength);
-		}
-
-		void AppendData(ReadOnlySpan<byte> NewData)
-		{
-			ExpandData(NewData.Length);
-			NewData.CopyTo(WriteBuffer.AsSpan(Data.Length - NewData.Length));
-		}
-
-		/// <inheritdoc/>
-		protected override IoHash SerializeDirty()
-		{
-			// If we're still writing to the last node, flush it to the write buffer
-			if (ChildNodes != null && ChildNodes.Count > 0)
-			{
-				FileNode? LastNode = ChildNodes[^1];
-				if (LastNode != null && LastNode.WriteBuffer != null)
-				{
-					Span<byte> LastHashSpan = WriteBuffer.AsSpan(Data.Length - IoHash.NumBytes);
-					IoHash LastHash = LastNode.Serialize();
-					LastHash.CopyTo(LastHashSpan);
-				}
-			}
-
-			// Do not allow this node to be modified from this point out
-			WriteBuffer = null;
-
-			// Write the current data
-			List<IoHash> References = new List<IoHash>();
-			if (Depth > 0)
-			{
-				for (ReadOnlySpan<byte> Span = Payload.Span; Span.Length > 0; Span = Span.Slice(IoHash.NumBytes))
-				{
-					References.Add(new IoHash(Span));
-				}
-			}
-			return Owner.WriteNode(Data, References);
-		}
 	}
 
 	/// <summary>
 	/// Factory class for file nodes
 	/// </summary>
-	public class FileNodeFactory : BundleNodeFactory<FileNode>
+	public class FileNodeDeserializer : BundleNodeDeserializer<FileNode>
 	{
 		/// <inheritdoc/>
-		public override FileNode CreateRoot(Bundle Bundle)
-		{
-			return new FileNode(Bundle, null);
-		}
-
-		/// <inheritdoc/>
-		public override FileNode ParseRoot(Bundle Bundle, IoHash Hash, ReadOnlyMemory<byte> Data)
-		{
-			return new FileNode(Bundle, null, Hash, Data);
-		}
+		public override FileNode Deserialize(Bundle Bundle, ReadOnlyMemory<byte> Data) => FileNode.Deserialize(Bundle, Data);
 	}
 }
