1) Updating UHT in UBT parser with recent RigVM changes.
2) Added new error test for changes
3) Added new success test for RigVM in general.

diff --git a/Engine/Source/Programs/Shared/EpicGames.UHT/Parsers/UhtScriptStructParser.cs b/Engine/Source/Programs/Shared/EpicGames.UHT/Parsers/UhtScriptStructParser.cs
--- a/Engine/Source/Programs/Shared/EpicGames.UHT/Parsers/UhtScriptStructParser.cs
+++ b/Engine/Source/Programs/Shared/EpicGames.UHT/Parsers/UhtScriptStructParser.cs
@@ -199,32 +199,47 @@ namespace EpicGames.UHT.Parsers
 					.RequireList('(', ')')
 					.Optional("virtual")
 					.RequireIdentifier((ref UhtToken Identifier) => MethodInfo.ReturnType = Identifier.Value.ToString())
-					.RequireIdentifier((ref UhtToken Identifier) => MethodInfo.Name = Identifier.Value.ToString())
+					.RequireIdentifier((ref UhtToken Identifier) => MethodInfo.Name = Identifier.Value.ToString());
+
+				bool bIsGetUpgradeInfo = false;
+				if (MethodInfo.ReturnType == "FRigVMStructUpgradeInfo" && MethodInfo.Name == "GetUpgradeInfo")
+				{
+					StructInfo.bHasGetUpgradeInfoMethod = true;
+					bIsGetUpgradeInfo = true;
+				}
+
+				TopScope.TokenReader
 					.RequireList('(', ')', ',', false, (IEnumerable<UhtToken> Tokens) =>
 					{
-						StringViewBuilder Builder = new StringViewBuilder();
-						UhtToken LastToken = new UhtToken();
-						foreach (UhtToken Token in Tokens)
+						if (!bIsGetUpgradeInfo)
 						{
-							if (Token.IsSymbol('='))
+							StringViewBuilder Builder = new StringViewBuilder();
+							UhtToken LastToken = new UhtToken();
+							foreach (UhtToken Token in Tokens)
 							{
-								break;
-							}
-							if (LastToken)
-							{
-								if (Builder.Length != 0)
+								if (Token.IsSymbol('='))
+								{
+									break;
+								}
+								if (LastToken)
 								{
-									Builder.Append(' ');
+									if (Builder.Length != 0)
+									{
+										Builder.Append(' ');
+									}
+									Builder.Append(LastToken.Value);
 								}
-								Builder.Append(LastToken.Value);
+								LastToken = Token;
 							}
-							LastToken = Token;
+							MethodInfo.Parameters.Add(new UhtRigVMParameter(LastToken.Value.ToString(), Builder.ToString()));
 						}
-						MethodInfo.Parameters.Add(new UhtRigVMParameter(LastToken.Value.ToString(), Builder.ToString()));
 					})
 					.ConsumeUntil(';');
 
-				StructInfo.Methods.Add(MethodInfo);
+				if (!bIsGetUpgradeInfo)
+				{
+					StructInfo.Methods.Add(MethodInfo);
+				}
 			}
 		}
 	}
diff --git a/Engine/Source/Programs/Shared/EpicGames.UHT/Types/UhtScriptStruct.cs b/Engine/Source/Programs/Shared/EpicGames.UHT/Types/UhtScriptStruct.cs
--- a/Engine/Source/Programs/Shared/EpicGames.UHT/Types/UhtScriptStruct.cs
+++ b/Engine/Source/Programs/Shared/EpicGames.UHT/Types/UhtScriptStruct.cs
@@ -282,6 +282,7 @@ namespace EpicGames.UHT.Types
 	 */
 	public class UhtRigVMStructInfo
 	{
+		public bool bHasGetUpgradeInfoMethod = false;
 		public string Name { get; set; } = string.Empty;
 		public List<UhtRigVMParameter> Members = new List<UhtRigVMParameter>();
 		public List<UhtRigVMMethodInfo> Methods { get; set; } = new List<UhtRigVMMethodInfo>();
@@ -449,6 +450,15 @@ namespace EpicGames.UHT.Types
 				this.LogError($"Struct '{this.SourceName}' has an invalid Unreal prefix, expecting '{ExpectedName}");
 			}
 
+			// Validate RigVM
+			if (this.RigVMStructInfo != null && this.MetaData.ContainsKey("Deprecated") && !this.RigVMStructInfo.bHasGetUpgradeInfoMethod)
+			{
+				this.LogError($"RigVMStruct '{this.SourceName}' is marked as deprecated but is missing 'GetUpgradeInfo method.");
+				this.LogError("Please implement a method like below:");
+				this.LogError("RIGVM_METHOD()");
+				this.LogError("virtual FRigVMStructUpgradeInfo GetUpgradeInfo() const override;");					
+			}
+
 			return Options |= UhtValidationOptions.Shadowing;
 		}
 		#endregion
