Remove debug statement from UHT

#rnx

diff --git a/Engine/Source/Programs/Shared/EpicGames.UHT/Types/UhtClass.cs b/Engine/Source/Programs/Shared/EpicGames.UHT/Types/UhtClass.cs
--- a/Engine/Source/Programs/Shared/EpicGames.UHT/Types/UhtClass.cs
+++ b/Engine/Source/Programs/Shared/EpicGames.UHT/Types/UhtClass.cs
@@ -444,11 +444,6 @@ namespace EpicGames.UHT.Types
 								break;
 							}
 
-							if (Property.SourceName == "SourceModels")
-							{
-								Debugger.Break();
-							}
-
 							if (!Property.PropertyExportFlags.HasAnyFlags(UhtPropertyExportFlags.GetterSpecifiedNone))
 							{
 								if (Property.PropertyExportFlags.HasAnyFlags(UhtPropertyExportFlags.GetterSpecified))
