Horde: Add IStorageClient.HasBlobAsync

#rnx

diff --git a/Engine/Source/Programs/Horde/Horde.Build/Storage/Services/BasicStorageClient.cs b/Engine/Source/Programs/Horde/Horde.Build/Storage/Services/BasicStorageClient.cs
--- a/Engine/Source/Programs/Horde/Horde.Build/Storage/Services/BasicStorageClient.cs
+++ b/Engine/Source/Programs/Horde/Horde.Build/Storage/Services/BasicStorageClient.cs
@@ -93,6 +93,12 @@ namespace Horde.Build.Storage.Services
 			await StorageBackend.WriteBytesAsync(Path, Data);
 		}
 
+		/// <inheritdoc/>
+		public async Task<bool> HasBlobAsync(NamespaceId NamespaceId, IoHash Hash, CancellationToken CancellationToken)
+		{
+			return await StorageBackend.ExistsAsync(GetFullBlobPath(NamespaceId, Hash));
+		}
+
 		/// <inheritdoc/>
 		public async Task<IRef> GetRefAsync(NamespaceId NamespaceId, BucketId BucketId, RefId RefId, CancellationToken CancellationToken)
 		{
diff --git a/Engine/Source/Programs/Shared/EpicGames.Horde/Storage/IStorageClient.cs b/Engine/Source/Programs/Shared/EpicGames.Horde/Storage/IStorageClient.cs
--- a/Engine/Source/Programs/Shared/EpicGames.Horde/Storage/IStorageClient.cs
+++ b/Engine/Source/Programs/Shared/EpicGames.Horde/Storage/IStorageClient.cs
@@ -176,6 +176,15 @@ namespace EpicGames.Horde.Storage
 		/// <param name="CancellationToken">Cancellation token for the operation</param>
 		Task WriteBlobAsync(NamespaceId NamespaceId, IoHash Hash, Stream Stream, CancellationToken CancellationToken = default);
 
+		/// <summary>
+		/// Checks if the given blob exists
+		/// </summary>
+		/// <param name="NamespaceId">Namespace to operate on</param>
+		/// <param name="Hash">Hash of the blob</param>
+		/// <param name="CancellationToken">Cancellation token for the operation</param>
+		/// <returns>True if the blob exists, false if it did not exist</returns>
+		Task<bool> HasBlobAsync(NamespaceId NamespaceId, IoHash Hash, CancellationToken CancellationToken = default);
+
 		#endregion
 		#region Refs
 
@@ -190,13 +199,13 @@ namespace EpicGames.Horde.Storage
 		Task<IRef> GetRefAsync(NamespaceId NamespaceId, BucketId BucketId, RefId RefId, CancellationToken CancellationToken = default);
 
 		/// <summary>
-		/// Gets the given reference
+		/// Checks if the given reference exists
 		/// </summary>
 		/// <param name="NamespaceId">Namespace identifier</param>
 		/// <param name="BucketId">Bucket identifier</param>
 		/// <param name="RefId">Ref identifier</param>
 		/// <param name="CancellationToken">Cancellation token for the operation</param>
-		/// <returns>The reference data if the ref exists</returns>
+		/// <returns>True if the ref exists, false if it did not exist</returns>
 		Task<bool> HasRefAsync(NamespaceId NamespaceId, BucketId BucketId, RefId RefId, CancellationToken CancellationToken = default);
 
 		/// <summary>
diff --git a/Engine/Source/Programs/Shared/EpicGames.Horde/Storage/Impl/FileStorageClient.cs b/Engine/Source/Programs/Shared/EpicGames.Horde/Storage/Impl/FileStorageClient.cs
--- a/Engine/Source/Programs/Shared/EpicGames.Horde/Storage/Impl/FileStorageClient.cs
+++ b/Engine/Source/Programs/Shared/EpicGames.Horde/Storage/Impl/FileStorageClient.cs
@@ -79,6 +79,12 @@ namespace EpicGames.Horde.Storage.Impl
 			}
 		}
 
+		/// <inheritdoc/>
+		public Task<bool> HasBlobAsync(NamespaceId NamespaceId, IoHash Hash, CancellationToken CancellationToken = default)
+		{
+			return Task.FromResult(FileReference.Exists(GetBlobFile(NamespaceId, Hash)));
+		}
+
 		/// <inheritdoc/>
 		public Task<bool> DeleteRefAsync(NamespaceId NamespaceId, BucketId BucketId, RefId RefId, CancellationToken CancellationToken = default)
 		{
diff --git a/Engine/Source/Programs/Shared/EpicGames.Horde/Storage/Impl/HttpStorageClient.cs b/Engine/Source/Programs/Shared/EpicGames.Horde/Storage/Impl/HttpStorageClient.cs
--- a/Engine/Source/Programs/Shared/EpicGames.Horde/Storage/Impl/HttpStorageClient.cs
+++ b/Engine/Source/Programs/Shared/EpicGames.Horde/Storage/Impl/HttpStorageClient.cs
@@ -117,6 +117,27 @@ namespace EpicGames.Horde.Storage.Impl
 			Response.EnsureSuccessStatusCode();
 		}
 
+		/// <inheritdoc/>
+		public async Task<bool> HasBlobAsync(NamespaceId NamespaceId, IoHash Hash, CancellationToken CancellationToken = default)
+		{
+			using HttpRequestMessage Request = new HttpRequestMessage(HttpMethod.Head, $"api/v1/blobs/{NamespaceId}/{Hash}");
+
+			HttpResponseMessage Response = await HttpClient.SendAsync(Request, CancellationToken);
+			try
+			{
+				if (Response.StatusCode != System.Net.HttpStatusCode.NotFound)
+				{
+					Response.EnsureSuccessStatusCode();
+				}
+				return Response.StatusCode == System.Net.HttpStatusCode.OK;
+			}
+			catch
+			{
+				Response.Dispose();
+				throw;
+			}
+		}
+
 		#endregion
 
 		#region Refs
diff --git a/Engine/Source/Programs/Shared/EpicGames.Horde/Storage/Impl/MemoryStorageClient.cs b/Engine/Source/Programs/Shared/EpicGames.Horde/Storage/Impl/MemoryStorageClient.cs
--- a/Engine/Source/Programs/Shared/EpicGames.Horde/Storage/Impl/MemoryStorageClient.cs
+++ b/Engine/Source/Programs/Shared/EpicGames.Horde/Storage/Impl/MemoryStorageClient.cs
@@ -14,7 +14,7 @@ using System.Threading.Tasks;
 namespace EpicGames.Horde.Storage.Impl
 {
 	/// <summary>
-	/// Implementation of <see cref="IStorageClient"/> which writes directly to the local filesystem for testing. Not intended for production use.
+	/// Implementation of <see cref="IStorageClient"/> which writes directly to memory for testing. Not intended for production use.
 	/// </summary>
 	public class MemoryStorageClient : IStorageClient
 	{
@@ -34,14 +34,10 @@ namespace EpicGames.Horde.Storage.Impl
 			}
 		}
 
-		/// <summary>
-		/// Blobs that have been written to storage
-		/// </summary>
+		/// <inheritdoc/>
 		public Dictionary<(NamespaceId, IoHash), ReadOnlyMemory<byte>> Blobs { get; } = new Dictionary<(NamespaceId, IoHash), ReadOnlyMemory<byte>>();
 
-		/// <summary>
-		/// Refs that have been written to storage
-		/// </summary>
+		/// <inheritdoc/>
 		public Dictionary<(NamespaceId, BucketId, RefId), IRef> Refs { get; } = new Dictionary<(NamespaceId, BucketId, RefId), IRef>();
 
 		/// <inheritdoc/>
@@ -61,6 +57,12 @@ namespace EpicGames.Horde.Storage.Impl
 			}
 		}
 
+		/// <inheritdoc/>
+		public Task<bool> HasBlobAsync(NamespaceId NamespaceId, IoHash Hash, CancellationToken CancellationToken = default)
+		{
+			return Task.FromResult(Blobs.ContainsKey((NamespaceId, Hash)));
+		}
+
 		/// <inheritdoc/>
 		public Task<bool> DeleteRefAsync(NamespaceId NamespaceId, BucketId BucketId, RefId RefId, CancellationToken CancellationToken = default)
 		{
@@ -76,7 +78,7 @@ namespace EpicGames.Horde.Storage.Impl
 		/// <inheritdoc/>
 		public Task<IRef> GetRefAsync(NamespaceId NamespaceId, BucketId BucketId, RefId RefId, CancellationToken CancellationToken = default)
 		{
-			return Task.FromResult<IRef>(Refs[(NamespaceId, BucketId, RefId)]);
+			return Task.FromResult(Refs[(NamespaceId, BucketId, RefId)]);
 		}
 
 		/// <inheritdoc/>
