Horde: Remove parallel extraction of directories from bundles. Bundles can page in and out individual blobs as needed, and is currently not thread safe (even for read-only operations).

diff --git a/Engine/Source/Programs/Shared/EpicGames.Horde/Bundles/Nodes/DirectoryNode.cs b/Engine/Source/Programs/Shared/EpicGames.Horde/Bundles/Nodes/DirectoryNode.cs
--- a/Engine/Source/Programs/Shared/EpicGames.Horde/Bundles/Nodes/DirectoryNode.cs
+++ b/Engine/Source/Programs/Shared/EpicGames.Horde/Bundles/Nodes/DirectoryNode.cs
@@ -514,21 +514,23 @@ namespace EpicGames.Horde.Bundles.Nodes
 		{
 			DirectoryInfo.Create();
 
-			List<Task> Tasks = new List<Task>();
+//			List<Task> Tasks = new List<Task>();
 			foreach (FileEntry FileEntry in Files)
 			{
 				FileInfo FileInfo = new FileInfo(Path.Combine(DirectoryInfo.FullName, FileEntry.Name.ToString()));
 				FileNode FileNode = await FileEntry.GetAsync();
-				Tasks.Add(Task.Run(() => FileNode.CopyToFileAsync(FileInfo)));
+				await FileNode.CopyToFileAsync(FileInfo);
+//				Tasks.Add(Task.Run(() => FileNode.CopyToFileAsync(FileInfo)));
 			}
 			foreach (DirectoryEntry DirectoryEntry in Directories)
 			{
 				DirectoryInfo SubDirectoryInfo = DirectoryInfo.CreateSubdirectory(DirectoryEntry.Name.ToString());
 				DirectoryNode SubDirectoryNode = await DirectoryEntry.GetAsync();
-				Tasks.Add(Task.Run(() => SubDirectoryNode.CopyToDirectoryAsync(SubDirectoryInfo, Logger)));
+				await SubDirectoryNode.CopyToDirectoryAsync(SubDirectoryInfo, Logger);
+//				Tasks.Add(Task.Run(() => SubDirectoryNode.CopyToDirectoryAsync(SubDirectoryInfo, Logger)));
 			}
 
-			await Task.WhenAll(Tasks);
+//			await Task.WhenAll(Tasks);
 		}
 	}
 
