Horde: Reduce memory pressure when serailizing large commit trees to bundles.

* Buffers allocated for file nodes are now trimmed to their minimum required size to reduce memory after having been completed.
* Fixed accounting for total size when adding new nodes into a blob.
* Fixed issue where nodes which did not get rolled into a blob for writing to storage during an incremental trim operation not being reconsidered on following iterations.

diff --git a/Engine/Source/Programs/Horde/Horde.Build/Perforce/Impl/CommitService.cs b/Engine/Source/Programs/Horde/Horde.Build/Perforce/Impl/CommitService.cs
--- a/Engine/Source/Programs/Horde/Horde.Build/Perforce/Impl/CommitService.cs
+++ b/Engine/Source/Programs/Horde/Horde.Build/Perforce/Impl/CommitService.cs
@@ -913,6 +913,7 @@ namespace HordeServer.Commits.Impl
 			DirectoryNode Contents = CommitBundle.Root;
 			if (Options.WriteStubFiles)
 			{
+				int Count = 0;
 				await foreach (PerforceResponse<SyncRecord> Record in Perforce.StreamCommandAsync<SyncRecord>("sync", new[] { "-k", $"{QueryPath}@{Change}" }, null, default))
 				{
 					SyncRecord SyncRecord = Record.Data;
@@ -926,6 +927,12 @@ namespace HordeServer.Commits.Impl
 
 					byte[] Data = Encoding.UTF8.GetBytes($"{SyncRecord.DepotFile}#{SyncRecord.Revision}");
 					File.Append(Data, Options.Chunking);
+
+					if (++Count > 1000)
+					{
+						await CommitBundle.TrimAsync();
+						Count = 0;
+					}
 				}
 			}
 			else
diff --git a/Engine/Source/Programs/Shared/EpicGames.Horde/Bundles/Bundle.cs b/Engine/Source/Programs/Shared/EpicGames.Horde/Bundles/Bundle.cs
--- a/Engine/Source/Programs/Shared/EpicGames.Horde/Bundles/Bundle.cs
+++ b/Engine/Source/Programs/Shared/EpicGames.Horde/Bundles/Bundle.cs
@@ -221,17 +221,13 @@ namespace EpicGames.Horde.Bundles
 		/// <summary>
 		/// Flush a subset of nodes in the working set.
 		/// </summary>
-		private List<NodeInfo> SerializeWorkingSet(int IgnoreCount)
+		private void SerializeWorkingSet(int IgnoreCount)
 		{
 			// Find all the nodes that need to be flushed, in order
 			List<BundleNodeRef> NodeRefs = new List<BundleNodeRef>();
 			LinkWorkingSet(Root, null, NodeRefs);
 			NodeRefs.SortBy(x => x.LastModifiedTime);
 
-			// Keep track of the nodes that have been serialized so far. Nodes in-memory may hash to the same serialized data, so we need to deduplicate them.
-			List<NodeInfo> DirtyNodes = new List<NodeInfo>();
-			HashSet<NodeInfo> DirtyNodesSet = new HashSet<NodeInfo>();
-
 			// Write them all to storage
 			for (int Idx = 0; Idx + IgnoreCount < NodeRefs.Count; Idx++)
 			{
@@ -242,13 +238,7 @@ namespace EpicGames.Horde.Bundles
 
 				NodeInfo NodeInfo = WriteNode(Node);
 				NodeRef.MarkAsClean(NodeInfo.Hash);
-
-				if (NodeInfo.Blob == null && DirtyNodesSet.Add(NodeInfo))
-				{
-					DirtyNodes.Add(NodeInfo);
-				}
 			}
-			return DirtyNodes;
 		}
 
 		/// <summary>
@@ -288,13 +278,16 @@ namespace EpicGames.Horde.Bundles
 		/// </summary>
 		public async Task TrimAsync(int IgnoreCount)
 		{
+			// Serialize nodes currently in the working set
+			SerializeWorkingSet(IgnoreCount);
+
+			// Find all the nodes that can be written, in order
+			List<NodeInfo> WriteNodes = new List<NodeInfo>();
+			FindNodesToWrite(Root, WriteNodes, new HashSet<NodeInfo>());
+
 			// Timestamp for any flushed objects
 			DateTime UtcNow = DateTime.UtcNow;
 
-			// Find all the nodes that need to be flushed, in order
-			List<NodeInfo> WriteNodes = SerializeWorkingSet(IgnoreCount);
-			WriteNodes.SortBy(x => x.Rank);
-
 			// Write them all to storage
 			int MinNodeIdx = 0;
 			int NextBlobCost = 0;
@@ -307,6 +300,34 @@ namespace EpicGames.Horde.Bundles
 					MinNodeIdx = Idx;
 					NextBlobCost = 0;
 				}
+				NextBlobCost += Node.Cost;
+			}
+		}
+
+		void FindNodesToWrite(BundleNode Node, List<NodeInfo> WriteNodes, HashSet<NodeInfo> WriteNodesSet)
+		{
+			foreach (BundleNodeRef ChildRef in Node.GetReferences())
+			{
+				if (ChildRef.Node != null)
+				{
+					FindNodesToWrite(ChildRef.Node, WriteNodes, WriteNodesSet);
+				}
+				else
+				{
+					FindNodeInfosToWrite(HashToNode[ChildRef.Hash], WriteNodes, WriteNodesSet);
+				}
+			}
+		}
+
+		void FindNodeInfosToWrite(NodeInfo Node, List<NodeInfo> WriteNodes, HashSet<NodeInfo> WriteNodesSet)
+		{
+			if (Node.Blob == null && WriteNodesSet.Add(Node))
+			{
+				foreach (NodeInfo ChildNode in Node.References!)
+				{
+					FindNodeInfosToWrite(ChildNode, WriteNodes, WriteNodesSet);
+				}
+				WriteNodes.Add(Node);
 			}
 		}
 
diff --git a/Engine/Source/Programs/Shared/EpicGames.Horde/Bundles/Nodes/FileNode.cs b/Engine/Source/Programs/Shared/EpicGames.Horde/Bundles/Nodes/FileNode.cs
--- a/Engine/Source/Programs/Shared/EpicGames.Horde/Bundles/Nodes/FileNode.cs
+++ b/Engine/Source/Programs/Shared/EpicGames.Horde/Bundles/Nodes/FileNode.cs
@@ -107,8 +107,25 @@ namespace EpicGames.Horde.Bundles.Nodes
 		/// <inheritdoc/>
 		public override ReadOnlyMemory<byte> Serialize()
 		{
-			if (ChildNodeRefs != null)
+			if (ChildNodeRefs == null)
+			{
+				// Compact the buffer to reflect the amount of serialized data
+				const int MaxSlackPct = 20;
+
+				int SlackBytes = WriteBuffer.Length - Data.Length;
+				int SlackPct = (SlackBytes * 100) / WriteBuffer.Length;
+
+				if (SlackPct > MaxSlackPct)
+				{
+					WriteBuffer = new byte[Data.Length];
+					Data.CopyTo(WriteBuffer);
+					Data = WriteBuffer;
+					Payload = Data.Slice(Data.Length - Payload.Length);
+				}
+			}
+			else
 			{
+				// Allocate data with the required size and write it immediately
 				Span<byte> Span = GetWritableSpan(VarInt.Measure((ulong)Length) + ChildNodeRefs.Count * IoHash.NumBytes, 0);
 
 				int LengthBytes = VarInt.Write(Span, Length);
