Horde: Add a storage client which traces interactions with the underlying storage system to the log.

diff --git a/Engine/Source/Programs/Shared/EpicGames.Horde/Storage/Impl/TracingStorageClient.cs b/Engine/Source/Programs/Shared/EpicGames.Horde/Storage/Impl/TracingStorageClient.cs
new file mode 100644
--- /dev/null
+++ b/Engine/Source/Programs/Shared/EpicGames.Horde/Storage/Impl/TracingStorageClient.cs
@@ -0,0 +1,88 @@
+// Copyright Epic Games, Inc. All Rights Reserved.
+
+using EpicGames.Core;
+using EpicGames.Serialization;
+using Microsoft.Extensions.Logging;
+using System;
+using System.Collections.Generic;
+using System.IO;
+using System.Linq;
+using System.Text;
+using System.Threading;
+using System.Threading.Tasks;
+
+namespace EpicGames.Horde.Storage.Impl
+{
+	/// <summary>
+	/// Implementation of <see cref="IStorageClient"/> which writes directly to the local filesystem for testing. Not intended for production use.
+	/// </summary>
+	public class TracingStorageClient : IStorageClient
+	{
+		IStorageClient Inner;
+		ILogger Logger;
+
+		/// <summary>
+		/// Constructor
+		/// </summary>
+		/// <param name="Inner"></param>
+		/// <param name="Logger"></param>
+		public TracingStorageClient(IStorageClient Inner, ILogger Logger)
+		{
+			this.Inner = Inner;
+			this.Logger = Logger;
+		}
+
+		/// <inheritdoc/>
+		public Task<Stream> ReadBlobAsync(NamespaceId NamespaceId, IoHash Hash, CancellationToken CancellationToken = default)
+		{
+			Logger.LogDebug("Reading blob {NamespaceId}/{Hash}", NamespaceId, Hash);
+			return Inner.ReadBlobAsync(NamespaceId, Hash, CancellationToken);
+		}
+
+		/// <inheritdoc/>
+		public Task WriteBlobAsync(NamespaceId NamespaceId, IoHash Hash, Stream Stream, CancellationToken CancellationToken = default)
+		{
+			Logger.LogDebug("Writing blob {NamespaceId}/{Hash}", NamespaceId, Hash);
+			return Inner.WriteBlobAsync(NamespaceId, Hash, Stream, CancellationToken);
+		}
+
+		/// <inheritdoc/>
+		public Task<bool> DeleteRefAsync(NamespaceId NamespaceId, BucketId BucketId, RefId RefId, CancellationToken CancellationToken = default)
+		{
+			Logger.LogDebug("Deleting ref {NamespaceId}/{BucketId}/{RefId}", NamespaceId, BucketId, RefId);
+			return Inner.DeleteRefAsync(NamespaceId, BucketId, RefId, CancellationToken);
+		}
+
+		/// <inheritdoc/>
+		public Task<List<RefId>> FindMissingRefsAsync(NamespaceId NamespaceId, BucketId BucketId, List<RefId> RefIds, CancellationToken CancellationToken = default)
+		{
+			return Inner.FindMissingRefsAsync(NamespaceId, BucketId, RefIds, CancellationToken);
+		}
+
+		/// <inheritdoc/>
+		public Task<IRef> GetRefAsync(NamespaceId NamespaceId, BucketId BucketId, RefId RefId, CancellationToken CancellationToken = default)
+		{
+			Logger.LogDebug("Getting ref {NamespaceId}/{BucketId}/{RefId}", NamespaceId, BucketId, RefId);
+			return Inner.GetRefAsync(NamespaceId, BucketId, RefId, CancellationToken);
+		}
+
+		/// <inheritdoc/>
+		public Task<bool> HasRefAsync(NamespaceId NamespaceId, BucketId BucketId, RefId RefId, CancellationToken CancellationToken = default)
+		{
+			return Inner.HasRefAsync(NamespaceId, BucketId, RefId, CancellationToken);
+		}
+
+		/// <inheritdoc/>
+		public Task<List<IoHash>> TryFinalizeRefAsync(NamespaceId NamespaceId, BucketId BucketId, RefId RefId, IoHash Hash, CancellationToken CancellationToken = default)
+		{
+			return Inner.TryFinalizeRefAsync(NamespaceId, BucketId, RefId, Hash, CancellationToken);
+		}
+
+		/// <inheritdoc/>
+		public Task<List<IoHash>> TrySetRefAsync(NamespaceId NamespaceId, BucketId BucketId, RefId RefId, CbObject Value, CancellationToken CancellationToken = default)
+		{
+			Logger.LogDebug("Setting ref {NamespaceId}/{BucketId}/{RefId}", NamespaceId, BucketId, RefId);
+			return Inner.TrySetRefAsync(NamespaceId, BucketId, RefId, Value, CancellationToken);
+		}
+	}
+}
