Horde: Missing implementation for HasBlobAsync on TracingStorageClient.

diff --git a/Engine/Source/Programs/Shared/EpicGames.Horde/Storage/Impl/TracingStorageClient.cs b/Engine/Source/Programs/Shared/EpicGames.Horde/Storage/Impl/TracingStorageClient.cs
--- a/Engine/Source/Programs/Shared/EpicGames.Horde/Storage/Impl/TracingStorageClient.cs
+++ b/Engine/Source/Programs/Shared/EpicGames.Horde/Storage/Impl/TracingStorageClient.cs
@@ -46,6 +46,12 @@ namespace EpicGames.Horde.Storage.Impl
 			return Inner.WriteBlobAsync(NamespaceId, Hash, Stream, CancellationToken);
 		}
 
+		/// <inheritdoc/>
+		public Task<bool> HasBlobAsync(NamespaceId NamespaceId, IoHash Hash, CancellationToken CancellationToken = default)
+		{
+			return Inner.HasBlobAsync(NamespaceId, Hash, CancellationToken);
+		}
+
 		/// <inheritdoc/>
 		public Task<bool> DeleteRefAsync(NamespaceId NamespaceId, BucketId BucketId, RefId RefId, CancellationToken CancellationToken = default)
 		{
