Horde: Rework node serialization to return ReadOnlySequence<byte> instances. This allows splitting the payload data across multiple allocations, and avoids having to allocate slack space/compact data when serializing. Fragmentation in the large object heap was causing OOM conditions on ECS tasks when mirroring P4 content.

diff --git a/Engine/Source/Programs/Shared/EpicGames.Horde/Bundles/Bundle.cs b/Engine/Source/Programs/Shared/EpicGames.Horde/Bundles/Bundle.cs
--- a/Engine/Source/Programs/Shared/EpicGames.Horde/Bundles/Bundle.cs
+++ b/Engine/Source/Programs/Shared/EpicGames.Horde/Bundles/Bundle.cs
@@ -8,6 +8,7 @@ using K4os.Compression.LZ4;
 using Microsoft.Extensions.Caching.Memory;
 using Microsoft.Extensions.Logging;
 using System;
+using System.Buffers;
 using System.Buffers.Binary;
 using System.Collections;
 using System.Collections.Generic;
@@ -72,7 +73,7 @@ namespace EpicGames.Horde.Bundles
 
 			public BlobInfo? Blob;
 			public BundleExport? Export;
-			public ReadOnlyMemory<byte> Data;
+			public ReadOnlySequence<byte> Data;
 			public NodeInfo[]? References;
 
 			public NodeInfo(IoHash Hash, int Rank, int Cost)
@@ -88,7 +89,7 @@ namespace EpicGames.Horde.Bundles
 				this.Blob = Blob;
 			}
 
-			void Set(BlobInfo? Blob, BundleExport? Export, ReadOnlyMemory<byte> Data, NodeInfo[]? References)
+			void Set(BlobInfo? Blob, BundleExport? Export, ReadOnlySequence<byte> Data, NodeInfo[]? References)
 			{
 				this.Blob = Blob;
 				this.Export = Export;
@@ -96,9 +97,9 @@ namespace EpicGames.Horde.Bundles
 				this.References = References;
 			}
 
-			public void SetImport(BlobInfo Blob) => Set(Blob, null, ReadOnlyMemory<byte>.Empty, null);
-			public void SetExport(BlobInfo? Blob, BundleExport Export, ReadOnlyMemory<byte> Data, NodeInfo[] References) => Set(Blob, Export, Data, References);
-			public void SetStandalone(ReadOnlyMemory<byte> Data, NodeInfo[] References) => Set(null, null, Data, References);
+			public void SetImport(BlobInfo Blob) => Set(Blob, null, ReadOnlySequence<byte>.Empty, null);
+			public void SetExport(BlobInfo? Blob, BundleExport Export, ReadOnlySequence<byte> Data, NodeInfo[] References) => Set(Blob, Export, Data, References);
+			public void SetStandalone(ReadOnlySequence<byte> Data, NodeInfo[] References) => Set(null, null, Data, References);
 		}
 
 		/// <summary>
@@ -336,7 +337,7 @@ namespace EpicGames.Horde.Bundles
 		/// </summary>
 		/// <param name="Hash">Hash of the node to return data for</param>
 		/// <returns>The node data</returns>
-		internal ValueTask<ReadOnlyMemory<byte>> GetDataAsync(IoHash Hash)
+		internal ValueTask<ReadOnlySequence<byte>> GetDataAsync(IoHash Hash)
 		{
 			return GetDataAsync(HashToNode[Hash]);
 		}
@@ -346,7 +347,7 @@ namespace EpicGames.Horde.Bundles
 		/// </summary>
 		/// <param name="Node">The node to get the data for</param>
 		/// <returns>The node data</returns>
-		async ValueTask<ReadOnlyMemory<byte>> GetDataAsync(NodeInfo Node)
+		async ValueTask<ReadOnlySequence<byte>> GetDataAsync(NodeInfo Node)
 		{
 			if (Node.Blob == null)
 			{
@@ -356,7 +357,7 @@ namespace EpicGames.Horde.Bundles
 				}
 				else
 				{
-					return DecodePacket(RootCacheKey, Node.Data, Node.Export.Packet).Slice(Node.Export.Offset, Node.Export.Length);
+					return new ReadOnlySequence<byte>(DecodePacket(RootCacheKey, Node.Data, Node.Export.Packet).Slice(Node.Export.Offset, Node.Export.Length));
 				}
 			}
 			else
@@ -364,13 +365,13 @@ namespace EpicGames.Horde.Bundles
 				await MountBlobAsync(Node.Blob);
 				Debug.Assert(Node.Export != null);
 				BundleObject Object = await GetObjectAsync(Node.Blob.Hash);
-				return DecodePacket(Node.Blob.Hash, Object.Data, Node.Export.Packet).Slice(Node.Export.Offset, Node.Export.Length);
+				return new ReadOnlySequence<byte>(DecodePacket(Node.Blob.Hash, new ReadOnlySequence<byte>(Object.Data), Node.Export.Packet).Slice(Node.Export.Offset, Node.Export.Length));
 			}
 		}
 
 		NodeInfo WriteNode(BundleNode Node)
 		{
-			ReadOnlyMemory<byte> Data = Node.Serialize();
+			ReadOnlySequence<byte> Data = Node.Serialize();
 
 			List<IoHash> References = new List<IoHash>();
 			foreach (BundleNodeRef Reference in Node.GetReferences())
@@ -388,9 +389,9 @@ namespace EpicGames.Horde.Bundles
 		/// <param name="Data">The node data</param>
 		/// <param name="References">Hashes for referenced nodes</param>
 		/// <returns>Hash of the data</returns>
-		NodeInfo WriteNode(ReadOnlyMemory<byte> Data, IEnumerable<IoHash> References)
+		NodeInfo WriteNode(ReadOnlySequence<byte> Data, IEnumerable<IoHash> References)
 		{
-			IoHash Hash = IoHash.Compute(Data.Span);
+			IoHash Hash = IoHash.Compute(Data);
 			NodeInfo[] NodeReferences = References.Select(x => HashToNode[x]).Distinct().OrderBy(x => x.Hash).ToArray();
 
 			NodeInfo? Node;
@@ -402,7 +403,7 @@ namespace EpicGames.Horde.Bundles
 					Rank = Math.Max(Rank, Reference.Rank + 1);
 				}
 
-				Node = new NodeInfo(Hash, Rank, Data.Length);
+				Node = new NodeInfo(Hash, Rank, (int)Data.Length);
 				HashToNode[Hash] = Node;
 			}
 
@@ -452,15 +453,17 @@ namespace EpicGames.Horde.Bundles
 		/// <param name="Data">Raw blob data</param>
 		/// <param name="Packet">The decoded block location and size</param>
 		/// <returns>The decoded data</returns>
-		ReadOnlyMemory<byte> DecodePacket(object BlobKey, ReadOnlyMemory<byte> Data, BundleCompressionPacket Packet)
+		ReadOnlyMemory<byte> DecodePacket(object BlobKey, ReadOnlySequence<byte> Data, BundleCompressionPacket Packet)
 		{
 			object CacheKey = (BlobKey, Packet.Offset);
 
 			ReadOnlyMemory<byte> DecodedData;
 			if (Cache == null || !Cache.TryGetValue(CacheKey, out DecodedData))
 			{
+				Debug.Assert(Data.IsSingleSegment);
+
 				byte[] DecodedBuffer = new byte[Packet.DecodedLength];
-				LZ4Codec.Decode(Data.Span.Slice(Packet.Offset, Packet.EncodedLength), DecodedBuffer);
+				LZ4Codec.Decode(Data.FirstSpan.Slice(Packet.Offset, Packet.EncodedLength), DecodedBuffer);
 
 				DecodedData = DecodedBuffer;
 
@@ -517,7 +520,7 @@ namespace EpicGames.Horde.Bundles
 			for(int Idx = 0; Idx < Object.Exports.Count; Idx++)
 			{
 				NodeInfo Node = Nodes[Idx];
-				Node.SetExport(Blob, Object.Exports[Idx], Data, Object.Exports[Idx].References.Select(x => Nodes[x]).ToArray());
+				Node.SetExport(Blob, Object.Exports[Idx], new ReadOnlySequence<byte>(Data), Object.Exports[Idx].References.Select(x => Nodes[x]).ToArray());
 			}
 		}
 
@@ -569,7 +572,7 @@ namespace EpicGames.Horde.Bundles
 					{
 						foreach (NodeInfo Node in LiveBlob.LiveNodes!)
 						{
-							ReadOnlyMemory<byte> Data = await GetDataAsync(Node);
+							ReadOnlySequence<byte> Data = await GetDataAsync(Node);
 							Node.SetStandalone(Data, Node.References!);
 						}
 					}
@@ -581,7 +584,7 @@ namespace EpicGames.Horde.Bundles
 					NodeInfo LiveNode = LiveNodes[Idx];
 					if (LiveNode.References != null && LiveNode.References.Any(x => x.Blob == null))
 					{
-						ReadOnlyMemory<byte> Data = await GetDataAsync(LiveNode);
+						ReadOnlySequence<byte> Data = await GetDataAsync(LiveNode);
 						LiveNode.SetStandalone(Data, LiveNode.References!);
 					}
 				}
@@ -665,7 +668,7 @@ namespace EpicGames.Horde.Bundles
 			foreach (NodeInfo Node in Nodes)
 			{
 				Node.Blob = null;
-				Node.Data = Object.Data;
+				Node.Data = new ReadOnlySequence<byte>(Object.Data);
 			}
 
 			RootCacheKey = new object();
@@ -680,7 +683,7 @@ namespace EpicGames.Horde.Bundles
 			foreach (NodeInfo Node in Nodes)
 			{
 				Node.Blob = Blob;
-				Node.Data = ReadOnlyMemory<byte>.Empty;
+				Node.Data = ReadOnlySequence<byte>.Empty;
 			}
 		}
 
@@ -728,7 +731,7 @@ namespace EpicGames.Horde.Bundles
 			BundleCompressionPacket Packet = new BundleCompressionPacket(0);
 			foreach (NodeInfo Node in Nodes)
 			{
-				ReadOnlyMemory<byte> NodeData = await GetDataAsync(Node);
+				ReadOnlySequence<byte> NodeData = await GetDataAsync(Node);
 
 				// If we can't fit this data into the current block, flush the contents of it first
 				if (BlockSize > 0 && BlockSize + NodeData.Length > Options.MinCompressionPacketSize)
@@ -740,21 +743,24 @@ namespace EpicGames.Horde.Bundles
 
 				// Create the export for this node
 				int[] References = Node.References.Select(x => NodeToIndex[x]).ToArray();
-				Node.Export = new BundleExport(Node.Hash, Node.Rank, Node.Cost, Packet, Packet.DecodedLength, NodeData.Length, References);
+				Node.Export = new BundleExport(Node.Hash, Node.Rank, Node.Cost, Packet, Packet.DecodedLength, (int)NodeData.Length, References);
 				Object.Exports.Add(Node.Export);
 
 				// Write out the new block
 				int Offset = Packet.EncodedLength;
-				if (NodeData.Length < Options.MinCompressionPacketSize)
+				if (NodeData.Length < Options.MinCompressionPacketSize || !NodeData.IsSingleSegment)
 				{
-					int RequiredSize = Math.Max(BlockSize + NodeData.Length, (int)(Options.MaxBlobSize * 1.2));
+					int RequiredSize = Math.Max(BlockSize + (int)NodeData.Length, (int)(Options.MaxBlobSize * 1.2));
 					CreateFreeSpace(ref BlockBuffer, BlockSize, RequiredSize);
-					NodeData.CopyTo(BlockBuffer.AsMemory(BlockSize));
-					BlockSize += NodeData.Length;
+					foreach (ReadOnlyMemory<byte> NodeSegment in NodeData)
+					{
+						NodeSegment.CopyTo(BlockBuffer.AsMemory(BlockSize));
+						BlockSize += NodeSegment.Length;
+					}
 				}
 				else
 				{
-					FlushPacket(NodeData, Packet);
+					FlushPacket(NodeData.First, Packet);
 					Packet = new BundleCompressionPacket(Packet.Offset + Packet.EncodedLength);
 				}
 			}
diff --git a/Engine/Source/Programs/Shared/EpicGames.Horde/Bundles/BundleNode.cs b/Engine/Source/Programs/Shared/EpicGames.Horde/Bundles/BundleNode.cs
--- a/Engine/Source/Programs/Shared/EpicGames.Horde/Bundles/BundleNode.cs
+++ b/Engine/Source/Programs/Shared/EpicGames.Horde/Bundles/BundleNode.cs
@@ -2,6 +2,7 @@
 
 using EpicGames.Core;
 using System;
+using System.Buffers;
 using System.Collections.Generic;
 using System.Diagnostics;
 using System.Reflection;
@@ -30,7 +31,7 @@ namespace EpicGames.Horde.Bundles.Nodes
 		/// Serialize this node into a sequence of bytes
 		/// </summary>
 		/// <returns>Serialized data for this node</returns>
-		public abstract ReadOnlyMemory<byte> Serialize();
+		public abstract ReadOnlySequence<byte> Serialize();
 
 		/// <summary>
 		/// Enumerates all the child references from this node
@@ -251,5 +252,25 @@ namespace EpicGames.Horde.Bundles.Nodes
 		/// <param name="Data">Data to deserialize</param>
 		/// <returns>New node parsed from the data</returns>
 		public abstract T Deserialize(Bundle Bundle, ReadOnlyMemory<byte> Data);
+
+		/// <summary>
+		/// Deserializes data from the given bundle
+		/// </summary>
+		/// <param name="Bundle">The bundle containing the data</param>
+		/// <param name="Data">Data to deserialize</param>
+		/// <returns>New node parsed from the data</returns>
+		public T Deserialize(Bundle Bundle, ReadOnlySequence<byte> Data)
+		{
+			ReadOnlyMemory<byte> Memory;
+			if (Data.IsSingleSegment)
+			{
+				Memory = Data.First;
+			}
+			else
+			{
+				Memory = Data.ToArray();
+			}
+			return Deserialize(Bundle, Memory);
+		}
 	}
 }
diff --git a/Engine/Source/Programs/Shared/EpicGames.Horde/Bundles/Nodes/DirectoryNode.cs b/Engine/Source/Programs/Shared/EpicGames.Horde/Bundles/Nodes/DirectoryNode.cs
--- a/Engine/Source/Programs/Shared/EpicGames.Horde/Bundles/Nodes/DirectoryNode.cs
+++ b/Engine/Source/Programs/Shared/EpicGames.Horde/Bundles/Nodes/DirectoryNode.cs
@@ -4,6 +4,7 @@ using EpicGames.Core;
 using EpicGames.Serialization;
 using Microsoft.Extensions.Logging;
 using System;
+using System.Buffers;
 using System.Collections.Generic;
 using System.Diagnostics;
 using System.Diagnostics.CodeAnalysis;
@@ -367,7 +368,7 @@ namespace EpicGames.Horde.Bundles.Nodes
 		#endregion
 
 		/// <inheritdoc/>
-		public override ReadOnlyMemory<byte> Serialize()
+		public override ReadOnlySequence<byte> Serialize()
 		{
 			// Measure the required size of the write buffer
 			int Size = 1;
@@ -423,7 +424,7 @@ namespace EpicGames.Horde.Bundles.Nodes
 			}
 
 			Debug.Assert(Span.Length == 0);
-			return Data;
+			return new ReadOnlySequence<byte>(Data);
 		}
 
 		internal static DirectoryNode Deserialize(Bundle Bundle, ReadOnlySpan<byte> Span)
diff --git a/Engine/Source/Programs/Shared/EpicGames.Horde/Bundles/Nodes/FileNode.cs b/Engine/Source/Programs/Shared/EpicGames.Horde/Bundles/Nodes/FileNode.cs
--- a/Engine/Source/Programs/Shared/EpicGames.Horde/Bundles/Nodes/FileNode.cs
+++ b/Engine/Source/Programs/Shared/EpicGames.Horde/Bundles/Nodes/FileNode.cs
@@ -3,6 +3,7 @@
 using EpicGames.Core;
 using EpicGames.Serialization;
 using System;
+using System.Buffers;
 using System.Collections.Generic;
 using System.Diagnostics;
 using System.IO;
@@ -78,15 +79,30 @@ namespace EpicGames.Horde.Bundles.Nodes
 	{
 		const byte TypeId = (byte)'c';
 
+		class DataSegment : ReadOnlySequenceSegment<byte>
+		{
+			public DataSegment(long RunningIndex, ReadOnlyMemory<byte> Data)
+			{
+				this.RunningIndex = RunningIndex;
+				this.Memory = Data;
+			}
+
+			public void SetNext(DataSegment Next)
+			{
+				this.Next = Next;
+			}
+		}
+
 		internal int Depth { get; private set; }
-		internal ReadOnlyMemory<byte> Payload { get; private set; }
-		ReadOnlyMemory<byte> Data; // Full serialized data, including type id and header fields.
+		internal ReadOnlySequence<byte> Payload { get; private set; }
+		ReadOnlySequence<byte> Data; // Full serialized data, including type id and header fields.
 		IoHash? Hash; // Hash of the serialized data buffer.
 
 		// In-memory state
 		bool IsRoot = true;
 		uint RollingHash;
-		byte[] WriteBuffer = Array.Empty<byte>();
+		DataSegment? FirstSegment;
+		DataSegment? LastSegment;
 		List<BundleNodeRef<FileNode>>? ChildNodeRefs;
 
 		/// <summary>
@@ -105,38 +121,9 @@ namespace EpicGames.Horde.Bundles.Nodes
 		public IReadOnlyList<BundleNodeRef<FileNode>> Children => (IReadOnlyList<BundleNodeRef<FileNode>>?)ChildNodeRefs ?? Array.Empty<BundleNodeRef<FileNode>>();
 
 		/// <inheritdoc/>
-		public override ReadOnlyMemory<byte> Serialize()
+		public override ReadOnlySequence<byte> Serialize()
 		{
-			if (ChildNodeRefs == null)
-			{
-				// Compact the buffer to reflect the amount of serialized data
-				const int MaxSlackPct = 20;
-
-				int SlackBytes = WriteBuffer.Length - Data.Length;
-				int SlackPct = (SlackBytes * 100) / WriteBuffer.Length;
-
-				if (SlackPct > MaxSlackPct)
-				{
-					WriteBuffer = new byte[Data.Length];
-					Data.CopyTo(WriteBuffer);
-					Data = WriteBuffer;
-					Payload = Data.Slice(Data.Length - Payload.Length);
-				}
-			}
-			else
-			{
-				// Allocate data with the required size and write it immediately
-				Span<byte> Span = GetWritableSpan(VarInt.Measure((ulong)Length) + ChildNodeRefs.Count * IoHash.NumBytes, 0);
-
-				int LengthBytes = VarInt.Write(Span, Length);
-				Span = Span.Slice(LengthBytes);
-
-				foreach (BundleNodeRef<FileNode> ChildNodeRef in ChildNodeRefs)
-				{
-					ChildNodeRef.Hash.CopyTo(Span);
-					Span = Span.Slice(IoHash.NumBytes);
-				}
-			}
+			MarkComplete();
 			return Data;
 		}
 
@@ -155,12 +142,13 @@ namespace EpicGames.Horde.Bundles.Nodes
 			}
 
 			FileNode Node = new FileNode();
-			Node.Data = Data;
+			Node.Data = new ReadOnlySequence<byte>(Data);
 			Node.Hash = IoHash.Compute(Data.Span);
 			Node.Depth = (int)VarInt.Read(Span.Slice(1), out int DepthBytes);
 
 			int HeaderLength = 1 + DepthBytes;
-			Node.Payload = Data.Slice(HeaderLength);
+			ReadOnlyMemory<byte> Payload = Data.Slice(HeaderLength);
+			Node.Payload = new ReadOnlySequence<byte>(Payload);
 
 			if (Node.Depth == 0)
 			{
@@ -171,7 +159,7 @@ namespace EpicGames.Horde.Bundles.Nodes
 				Node.Length = (long)VarInt.Read(Span, out int LengthBytes);
 				Span = Span.Slice(LengthBytes);
 
-				Node.ChildNodeRefs = new List<BundleNodeRef<FileNode>>(Node.Payload.Length / IoHash.NumBytes);
+				Node.ChildNodeRefs = new List<BundleNodeRef<FileNode>>(Payload.Length / IoHash.NumBytes);
 
 				Span = Span.Slice(HeaderLength);
 				while (Span.Length > 0)
@@ -183,37 +171,6 @@ namespace EpicGames.Horde.Bundles.Nodes
 			return Node;
 		}
 
-		/// <summary>
-		/// Gets a writable span of the given size.
-		/// </summary>
-		/// <param name="SpanSize">Required size of the writable span</param>
-		/// <param name="HintMaxPayloadSize">Hint for the maximum size of the payload, to avoid having to reallocate the buffer. May be zero if unknown.</param>
-		/// <returns></returns>
-		private Span<byte> GetWritableSpan(int SpanSize, int HintMaxPayloadSize)
-		{
-			Debug.Assert(!IsComplete());
-
-			int HeaderLength = 1 + VarInt.Measure(Depth);
-			int PrevPayloadLength = Payload.Length;
-
-			int MinWriteBufferSize = HeaderLength + Math.Max(PrevPayloadLength + SpanSize, HintMaxPayloadSize);
-			if (MinWriteBufferSize > WriteBuffer.Length)
-			{
-				byte[] NewWriteBuffer = new byte[MinWriteBufferSize];
-
-				NewWriteBuffer[0] = TypeId;
-				VarInt.Write(NewWriteBuffer.AsSpan(1), Depth);
-				Payload.Span.CopyTo(NewWriteBuffer.AsSpan(HeaderLength));
-
-				WriteBuffer = NewWriteBuffer;
-			}
-
-			Data = WriteBuffer.AsMemory(0, HeaderLength + PrevPayloadLength + SpanSize);
-			Payload = Data.Slice(HeaderLength);
-
-			return WriteBuffer.AsSpan(HeaderLength + PrevPayloadLength, SpanSize);
-		}
-
 		/// <summary>
 		/// Serialize this node and its children into a byte array
 		/// </summary>
@@ -277,7 +234,10 @@ namespace EpicGames.Horde.Bundles.Nodes
 			}
 			else
 			{
-				await OutputStream.WriteAsync(Payload);
+				foreach (ReadOnlyMemory<byte> PayloadSegment in Payload)
+				{
+					await OutputStream.WriteAsync(PayloadSegment);
+				}
 			}
 		}
 
@@ -317,12 +277,19 @@ namespace EpicGames.Horde.Bundles.Nodes
 
 				// Increase the height of the tree by pushing the contents of this node into a new child node
 				FileNode NewNode = new FileNode();
-				NewNode.IsRoot = false;
 				NewNode.Depth = Depth;
 				NewNode.Payload = Payload;
 				NewNode.Data = Data;
+				NewNode.Hash = Hash;
+
+				NewNode.IsRoot = false;
+				NewNode.RollingHash = RollingHash;
+				NewNode.FirstSegment = FirstSegment;
+				NewNode.LastSegment = LastSegment;
 				NewNode.ChildNodeRefs = ChildNodeRefs;
 
+				NewNode.Length = Length;
+
 				// Detach all the child refs
 				if (ChildNodeRefs != null)
 				{
@@ -335,10 +302,15 @@ namespace EpicGames.Horde.Bundles.Nodes
 
 				// Increase the depth and reset the current node
 				Depth++;
+				Payload = ReadOnlySequence<byte>.Empty;
+				Data = ReadOnlySequence<byte>.Empty;
+				Hash = null;
+
+				IsRoot = true;
+				RollingHash = BuzHash.Add(0, NewNode.Hash!.Value.ToByteArray());
+				FirstSegment = null;
+				LastSegment = null;
 				ChildNodeRefs = new List<BundleNodeRef<FileNode>>();
-				Payload = ReadOnlyMemory<byte>.Empty;
-				Data = ReadOnlyMemory<byte>.Empty;
-				WriteBuffer = Array.Empty<byte>();
 
 				// Append the new node as a child
 				ChildNodeRefs.Add(new BundleNodeRef<FileNode>(NewNode));
@@ -368,91 +340,147 @@ namespace EpicGames.Horde.Bundles.Nodes
 			int WindowSize = Options.LeafOptions.MinSize;
 			if (Payload.Length < WindowSize)
 			{
-				int AppendLength = Math.Min(WindowSize - Payload.Length, NewData.Length);
+				int AppendLength = Math.Min(WindowSize - (int)Payload.Length, NewData.Length);
 				AppendLeafData(NewData.Span.Slice(0, AppendLength), Options);
 				NewData = NewData.Slice(AppendLength);
 			}
 
 			// Cap the maximum amount of data to append to this node
-			int MaxLength = Math.Min(NewData.Length, Options.LeafOptions.MaxSize - Payload.Length);
+			int MaxLength = Math.Min(NewData.Length, Options.LeafOptions.MaxSize - (int)Payload.Length);
 			if (MaxLength > 0)
 			{
-				int AppendLength = HashScan(NewData.Span.Slice(0, MaxLength), Options);
-				AppendLeafData(NewData.Span.Slice(0, AppendLength), Options);
-				NewData = NewData.Slice(AppendLength);
+				ReadOnlySpan<byte> InputSpan = NewData.Span.Slice(0, MaxLength);
+				int Length = AppendLeafDataToChunkBoundary(InputSpan, Options);
+				NewData = NewData.Slice(Length);
 			}
 
+			// Compute the hash if this node is complete
+			if (Payload.Length == Options.LeafOptions.MaxSize)
+			{
+				MarkComplete();
+			}
 			return NewData;
 		}
 
-		private int HashScan(ReadOnlySpan<byte> NewSpan, ChunkingOptions Options)
+		private int AppendLeafDataToChunkBoundary(ReadOnlySpan<byte> InputSpan, ChunkingOptions Options)
 		{
 			int WindowSize = Options.LeafOptions.MinSize;
 			Debug.Assert(Payload.Length >= WindowSize);
 
-			int Length = 0;
-
 			// If this is the first time we're hashing the data, compute the hash of the existing window
 			if (Payload.Length == WindowSize)
 			{
-				RollingHash = BuzHash.Add(RollingHash, Payload.Span);
+				foreach (ReadOnlyMemory<byte> Segment in Payload)
+				{
+					RollingHash = BuzHash.Add(RollingHash, Segment.Span);
+				}
 			}
 
 			// Get the threshold for the rolling hash
 			uint RollingHashThreshold = (uint)((1L << 32) / Options.LeafOptions.TargetSize);
 
-			// Get the remaining part of the payload span leading into the new data
-			ReadOnlySpan<byte> PayloadSpan = Payload.Span.Slice(Payload.Length - WindowSize);
+			// Length of the data taken from the input span, updated as we step through it.
+			int Length = 0;
 
 			// Step the window through the tail end of the existing payload window. In this state, update the hash to remove data from the current payload, and add data from the new payload.
-			int SplitLength = Math.Min(NewSpan.Length, WindowSize);
-			for (; Length < SplitLength; Length++)
+			int SplitLength = Math.Min(InputSpan.Length, WindowSize);
+			ReadOnlySequence<byte> SplitPayload = Payload.Slice(Payload.Length - WindowSize, SplitLength);
+
+			foreach (ReadOnlyMemory<byte> PayloadSegment in SplitPayload)
 			{
-				RollingHash = BuzHash.Add(RollingHash, NewSpan[Length]);
-				if (RollingHash < RollingHashThreshold)
+				int BaseLength = Length;
+				int SpanLength = Length + PayloadSegment.Length;
+
+				ReadOnlySpan<byte> PayloadSpan = PayloadSegment.Span;
+				for (; Length < SpanLength; Length++)
 				{
-					return Length;
+					RollingHash = BuzHash.Add(RollingHash, InputSpan[Length]);
+					if (RollingHash < RollingHashThreshold)
+					{
+						AppendLeafData(InputSpan.Slice(0, Length), Options);
+						MarkComplete();
+						return Length;
+					}
+					RollingHash = BuzHash.Sub(RollingHash, PayloadSpan[Length - BaseLength], WindowSize);
 				}
-				RollingHash = BuzHash.Sub(RollingHash, PayloadSpan[Length], WindowSize);
 			}
 
 			// Step through the new window.
-			for (; Length < NewSpan.Length; Length++)
+			for (; Length < InputSpan.Length; Length++)
 			{
-				RollingHash = BuzHash.Add(RollingHash, NewSpan[Length]);
+				RollingHash = BuzHash.Add(RollingHash, InputSpan[Length]);
 				if (RollingHash < RollingHashThreshold)
 				{
+					AppendLeafData(InputSpan.Slice(0, Length), Options);
+					MarkComplete();
 					return Length;
 				}
-				RollingHash = BuzHash.Sub(RollingHash, NewSpan[Length - WindowSize], WindowSize);
+				RollingHash = BuzHash.Sub(RollingHash, InputSpan[Length - WindowSize], WindowSize);
 			}
 
-			return Length;
+			AppendLeafData(InputSpan, Options);
+			return InputSpan.Length;
 		}
 
-		private int AppendInternal(ReadOnlySpan<byte> TailSpan, ReadOnlySpan<byte> HeadSpan, uint RollingHashThreshold, int WindowSize, int MaxLength, ChunkingOptions Options)
+		private void AppendLeafData(ReadOnlySpan<byte> LeafData, ChunkingOptions Options)
 		{
-			for (int Length = 0; Length < MaxLength; Length++)
+			int DepthBytes = VarInt.Measure(Depth);
+
+			if (LastSegment == null)
 			{
-				RollingHash = BuzHash.Add(RollingHash, HeadSpan[Length]);
+				byte[] Buffer = new byte[1 + DepthBytes + LeafData.Length];
+				Buffer[0] = TypeId;
 
-				if (RollingHash < RollingHashThreshold)
-				{
-					AppendLeafData(HeadSpan.Slice(0, Length), Options);
-					Hash = IoHash.Compute(Data.Span);
-					return Length;
-				}
+				VarInt.Write(Buffer.AsSpan(1, DepthBytes), Depth);
+				LeafData.CopyTo(Buffer.AsSpan(1 + DepthBytes));
 
-				RollingHash = BuzHash.Sub(RollingHash, TailSpan[Length], WindowSize);
+				FirstSegment = LastSegment = new DataSegment(0, Buffer);
 			}
-			return -1;
+			else
+			{
+				DataSegment NewSegment = new DataSegment(LastSegment.RunningIndex + LastSegment.Memory.Length, LeafData.ToArray());
+				LastSegment.SetNext(NewSegment);
+				LastSegment = NewSegment;
+			}
+
+			Data = new ReadOnlySequence<byte>(FirstSegment!, 0, LastSegment, LastSegment.Memory.Length);
+			Payload = Data.Slice(1 + DepthBytes);
+
+			Length += LeafData.Length;
 		}
 
-		private void AppendLeafData(ReadOnlySpan<byte> Data, ChunkingOptions Options)
+		private void MarkComplete()
 		{
-			Span<byte> WriteSpan = GetWritableSpan(Data.Length, Options.LeafOptions.MaxSize);
-			Data.CopyTo(WriteSpan);
-			Length += Data.Length;
+			if (Hash == null)
+			{
+				if (ChildNodeRefs != null)
+				{
+					int DepthBytes = VarInt.Measure(Depth);
+					int LengthBytes = VarInt.Measure((ulong)Length);
+
+					byte[] NewData = new byte[1 + DepthBytes + LengthBytes + (ChildNodeRefs.Count * IoHash.NumBytes)];
+					Data = new ReadOnlySequence<byte>(NewData);
+					Payload = new ReadOnlySequence<byte>(NewData.AsMemory(1 + DepthBytes));
+
+					NewData[0] = TypeId;
+					Span<byte> Span = NewData.AsSpan(1);
+
+					VarInt.Write(Span, Depth);
+					Span = Span.Slice(DepthBytes);
+
+					VarInt.Write(Span, Length);
+					Span = Span.Slice(LengthBytes);
+
+					foreach (BundleNodeRef<FileNode> ChildNodeRef in ChildNodeRefs)
+					{
+						ChildNodeRef.Hash.CopyTo(Span);
+						Span = Span.Slice(IoHash.NumBytes);
+					}
+
+					Debug.Assert(Span.Length == 0);
+				}
+				Hash = IoHash.Compute(Data);
+			}
 		}
 
 		private ReadOnlyMemory<byte> AppendToInteriorNode(ReadOnlyMemory<byte> NewData, ChunkingOptions Options)
@@ -475,18 +503,15 @@ namespace EpicGames.Horde.Bundles.Nodes
 						// If the last node is complete, write it to the buffer
 						if (LastNode.IsComplete())
 						{
-							// Add it to the write buffer
-							Span<byte> LastHashSpan = GetWritableSpan(IoHash.NumBytes, Options.InteriorOptions.MaxSize);
-							LastNode.Hash!.Value.CopyTo(LastHashSpan);
-
 							// Update the hash
-							RollingHash = BuzHash.Add(RollingHash, LastHashSpan);
+							byte[] HashData = LastNode.Hash!.Value.ToByteArray();
+							RollingHash = BuzHash.Add(RollingHash, HashData);
 
 							// Check if it's time to finish this chunk
 							uint HashThreshold = (uint)(((1L << 32) * IoHash.NumBytes) / Options.LeafOptions.TargetSize);
 							if ((Payload.Length >= Options.InteriorOptions.MinSize && RollingHash < HashThreshold) || (Payload.Length >= Options.InteriorOptions.MaxSize))
 							{
-								Hash = IoHash.Compute(Data.Span);
+								MarkComplete();
 								return NewData;
 							}
 						}
diff --git a/Engine/Source/Programs/Shared/EpicGames.IoHash/IoHash.cs b/Engine/Source/Programs/Shared/EpicGames.IoHash/IoHash.cs
--- a/Engine/Source/Programs/Shared/EpicGames.IoHash/IoHash.cs
+++ b/Engine/Source/Programs/Shared/EpicGames.IoHash/IoHash.cs
@@ -72,10 +72,10 @@ namespace EpicGames.Core
 		}
 
 		/// <summary>
-		/// Creates a content hash for a block of data, using a given algorithm.
+		/// Creates the IoHash for a block of data.
 		/// </summary>
 		/// <param name="Data">Data to compute the hash for</param>
-		/// <returns>New content hash instance containing the hash of the data</returns>
+		/// <returns>New hash instance containing the hash of the data</returns>
 		public static IoHash Compute(ReadOnlySpan<byte> Data)
 		{
 			byte[] Output = new byte[32];
@@ -83,6 +83,28 @@ namespace EpicGames.Core
 			return new IoHash(Output);
 		}
 
+		/// <summary>
+		/// Creates the IoHash for a block of data.
+		/// </summary>
+		/// <param name="Sequence">Data to compute the hash for</param>
+		/// <returns>New hash instance containing the hash of the data</returns>
+		public static IoHash Compute(ReadOnlySequence<byte> Sequence)
+		{
+			if (Sequence.IsSingleSegment)
+			{
+				return Compute(Sequence.FirstSpan);
+			}
+
+			using (Blake3.Hasher Hasher = Blake3.Hasher.New())
+			{
+				foreach (ReadOnlyMemory<byte> Segment in Sequence)
+				{
+					Hasher.Update(Segment.Span);
+				}
+				return FromBlake3(Hasher);
+			}
+		}
+
 		/// <summary>
 		/// Parses a digest from the given hex string
 		/// </summary>
