EpicGames.Core: Fix NullReferenceException when framework process stdout & err pipes aren't merged

#rnx

diff --git a/Engine/Source/Programs/Shared/EpicGames.Core/ManagedProcess.cs b/Engine/Source/Programs/Shared/EpicGames.Core/ManagedProcess.cs
--- a/Engine/Source/Programs/Shared/EpicGames.Core/ManagedProcess.cs
+++ b/Engine/Source/Programs/Shared/EpicGames.Core/ManagedProcess.cs
@@ -950,14 +950,11 @@ namespace EpicGames.Core
 			StdOutText?.Dispose();
 			StdErr?.Dispose();
 			StdErrText?.Dispose();
+			AccountingProcessGroup?.Dispose();
 
-			if (AccountingProcessGroup != null) 
-			{
-				AccountingProcessGroup.Dispose();
-			}
 			if(FrameworkProcess != null)
 			{
-				FrameworkCancellationSource!.Cancel();
+				FrameworkCancellationSource?.Cancel();
 
 				if (!FrameworkProcess.HasExited)
 				{
@@ -968,7 +965,7 @@ namespace EpicGames.Core
 				FrameworkProcess.Dispose();
 				FrameworkProcess = null;
 
-				FrameworkOutputTask!.ContinueWith(x => FrameworkCancellationSource.Dispose());
+				FrameworkOutputTask?.ContinueWith(x => FrameworkCancellationSource?.Dispose());
 			}
 		}
 
