Minor bug fix to file reader which hadn't been well tested.

#rn trivial
#rnx

diff --git a/Engine/Source/Programs/Shared/EpicGames.UHT/Utils/UhtFileManager.cs b/Engine/Source/Programs/Shared/EpicGames.UHT/Utils/UhtFileManager.cs
--- a/Engine/Source/Programs/Shared/EpicGames.UHT/Utils/UhtFileManager.cs
+++ b/Engine/Source/Programs/Shared/EpicGames.UHT/Utils/UhtFileManager.cs
@@ -102,7 +102,7 @@ namespace EpicGames.UHT.Utils
 
 						// Try to read the whole file into a buffer created by hand.  This avoids a LOT of memory allocations which in turn reduces the
 						// GC stress on the system.  Removing the StreamReader would be nice in the future.
-						long RawFileLength = fs.Length + 32;
+						long RawFileLength = fs.Length;
 						UhtBuffer InitialBuffer = UhtBuffer.Borrow((int)RawFileLength);
 						int ReadLength = sr.Read(InitialBuffer.Memory.Span);
 						if (sr.EndOfStream)
@@ -165,12 +165,12 @@ namespace EpicGames.UHT.Utils
 			}
 		}
 
-		///// <summary>
-		///// Read the given source file
-		///// </summary>
-		///// <param name="FilePath">Full file path</param>
-		///// <param name="Contents">Contents of the file</param>
-		///// <returns>True if the file was read, false if not</returns>
+		/// <summary>
+		/// Read the given source file
+		/// </summary>
+		/// <param name="FilePath">Full file path</param>
+		/// <param name="Contents">Contents of the file</param>
+		/// <returns>True if the file was read, false if not</returns>
 		private bool ReadFile(string FilePath, out StringView Contents)
 		{
 			// Exceptions are very expensive.  Don't bother trying to open the file if it doesn't exist
@@ -189,21 +189,17 @@ namespace EpicGames.UHT.Utils
 
 						// Try to read the whole file into a buffer created by hand.  This avoids a LOT of memory allocations which in turn reduces the
 						// GC stress on the system.  Removing the StreamReader would be nice in the future.
-						long RawFileLength = fs.Length + 32;
-						Memory<char> InitialBuffer = new Memory<char>(new char[RawFileLength]);
-						int ReadLength = sr.Read(InitialBuffer.Span);
+						long RawFileLength = fs.Length;
+						char[] InitialBuffer = new char[RawFileLength];
+						int ReadLength = sr.Read(InitialBuffer, 0, (int)RawFileLength);
 						if (sr.EndOfStream)
 						{
-							Contents = new StringView(InitialBuffer, 0, ReadLength);
+							Contents = new StringView(new ReadOnlyMemory<char>(InitialBuffer, 0, ReadLength));
 						}
 						else
 						{
 							string Remaining = sr.ReadToEnd();
-							long TotalSize = ReadLength + Remaining.Length;
-							Memory<char> Combined = new Memory<char>(new char[TotalSize]);
-							Buffer.BlockCopy(InitialBuffer.Span.ToArray(), 0, Combined.Span.ToArray(), 0, ReadLength * sizeof(char));
-							Buffer.BlockCopy(Remaining.ToArray(), 0, Combined.Span.ToArray(), ReadLength * sizeof(char), Remaining.Length * sizeof(char));
-							Contents = new StringView(Combined);
+							Contents = new StringView(String.Concat(new ReadOnlySpan<char>(InitialBuffer, 0, ReadLength), Remaining));
 						}
 						return true;
 					}
