Horde: Remove temp memory allocation for creating IoHash objects.

diff --git a/Engine/Source/Programs/Shared/EpicGames.IoHash/IoHash.cs b/Engine/Source/Programs/Shared/EpicGames.IoHash/IoHash.cs
--- a/Engine/Source/Programs/Shared/EpicGames.IoHash/IoHash.cs
+++ b/Engine/Source/Programs/Shared/EpicGames.IoHash/IoHash.cs
@@ -66,7 +66,7 @@ namespace EpicGames.Core
 		/// <param name="Hasher">The hasher to construct from</param>
 		public static IoHash FromBlake3(Blake3.Hasher Hasher)
 		{
-			byte[] Output = new byte[32];
+			Span<byte> Output = stackalloc byte[32];
 			Hasher.Finalize(Output);
 			return new IoHash(Output);
 		}
@@ -78,7 +78,7 @@ namespace EpicGames.Core
 		/// <returns>New hash instance containing the hash of the data</returns>
 		public static IoHash Compute(ReadOnlySpan<byte> Data)
 		{
-			byte[] Output = new byte[32];
+			Span<byte> Output = stackalloc byte[32];
 			Blake3.Hasher.Hash(Data, Output);
 			return new IoHash(Output);
 		}
