Horde: Fix double-free when using intercepted IO in native perforce library.

diff --git a/Engine/Source/Programs/Shared/EpicGames.Perforce.Native/Client.cpp b/Engine/Source/Programs/Shared/EpicGames.Perforce.Native/Client.cpp
--- a/Engine/Source/Programs/Shared/EpicGames.Perforce.Native/Client.cpp
+++ b/Engine/Source/Programs/Shared/EpicGames.Perforce.Native/Client.cpp
@@ -608,8 +608,7 @@ public:
 		{
 			if (GzipInst->OutputFull())
 			{
-				FFileSys::Write(Buffer, (int)(GzipInst->os - Buffer), e);
-				GzipInst->os = Buffer;
+				Flush(e);
 			}
 			if (e->Test() || !GzipInst->Uncompress(e) || GzipInst->InputEmpty())
 			{
@@ -620,12 +619,17 @@ public:
 
 	virtual void Close(Error* e) override
 	{
-		if (GzipInst && mode == FOM_WRITE && GzipInst->os > Buffer)
+		Flush(e);
+		FFileSys::Close(e);
+	}
+
+	void Flush(Error* e)
+	{
+		if (GzipInst && GzipInst->os > Buffer)
 		{
 			FFileSys::Write(Buffer, (int)(GzipInst->os - Buffer), e);
+			GzipInst->os = Buffer;
 		}
-
-		FFileSys::Close(e);
 	}
 };
 
