Updating UHT documentation

#rnx

diff --git a/Engine/Source/Programs/Shared/EpicGames.UHT/Exporters/CodeGen/UhtHeaderCodeGeneratorHFile.cs b/Engine/Source/Programs/Shared/EpicGames.UHT/Exporters/CodeGen/UhtHeaderCodeGeneratorHFile.cs
--- a/Engine/Source/Programs/Shared/EpicGames.UHT/Exporters/CodeGen/UhtHeaderCodeGeneratorHFile.cs
+++ b/Engine/Source/Programs/Shared/EpicGames.UHT/Exporters/CodeGen/UhtHeaderCodeGeneratorHFile.cs
@@ -642,7 +642,7 @@ namespace EpicGames.UHT.Exporters.CodeGen
 			return Builder;
 		}
 
-		private StringBuilder AppendCallbackParametersDecls(StringBuilder Builder,UhtClass Class, List<UhtFunction> CallbackFunctions)
+		private StringBuilder AppendCallbackParametersDecls(StringBuilder Builder, UhtClass Class, List<UhtFunction> CallbackFunctions)
 		{
 			if (CallbackFunctions.Count > 0)
 			{
@@ -1241,8 +1241,8 @@ namespace EpicGames.UHT.Exporters.CodeGen
 				Builder.Append('\t').AppendMacroName(this, Class, AccessorsMacroSuffix).Append(" \\\r\n");
 				if (bHasCallbacks)
 				{
-						Builder.Append('\t').AppendMacroName(this, Class, CallbackWrappersMacroSuffix).Append(" \\\r\n");
-					}
+					Builder.Append('\t').AppendMacroName(this, Class, CallbackWrappersMacroSuffix).Append(" \\\r\n");
+				}
 				if (bIsInterface)
 				{
 					if (bIsLegacy)
@@ -1394,4 +1394,26 @@ namespace EpicGames.UHT.Exporters.CodeGen
 		}
 		#endregion
 	}
+
+	static class UhtClassExtensions
+	{
+		public static IEnumerable<UhtProperty> EnumerateReplicatedProperties(this UhtClass Class, bool bIncludeSuper)
+		{
+			if (bIncludeSuper && Class.SuperClass != null)
+			{
+				foreach (UhtProperty Property in Class.SuperClass.EnumerateReplicatedProperties(true))
+				{
+					yield return Property;
+				}
+			}
+
+			foreach (UhtProperty Property in Class.Properties)
+			{
+				if (Property.PropertyFlags.HasAnyFlags(EPropertyFlags.Net))
+				{
+					yield return Property;
+				}
+			}
+		}
+	}
 }
diff --git a/Engine/Source/Programs/Shared/EpicGames.UHT/Types/Properties/UhtByteProperty.cs b/Engine/Source/Programs/Shared/EpicGames.UHT/Types/Properties/UhtByteProperty.cs
--- a/Engine/Source/Programs/Shared/EpicGames.UHT/Types/Properties/UhtByteProperty.cs
+++ b/Engine/Source/Programs/Shared/EpicGames.UHT/Types/Properties/UhtByteProperty.cs
@@ -153,7 +153,7 @@ namespace EpicGames.UHT.Types
 
 		#region Keyword
 		[UhtPropertyType(Keyword = "uint8", Options = UhtPropertyTypeOptions.Simple | UhtPropertyTypeOptions.Immediate)]
-		public static UhtProperty? UInt8Property(UhtPropertyResolvePhase ResolvePhase, UhtPropertySettings PropertySettings, IUhtTokenReader TokenReader, UhtToken MatchedToken)
+		private static UhtProperty? UInt8Property(UhtPropertyResolvePhase ResolvePhase, UhtPropertySettings PropertySettings, IUhtTokenReader TokenReader, UhtToken MatchedToken)
 		{
 			if (PropertySettings.bIsBitfield)
 			{
diff --git a/Engine/Source/Programs/Shared/EpicGames.UHT/Types/UhtClass.cs b/Engine/Source/Programs/Shared/EpicGames.UHT/Types/UhtClass.cs
--- a/Engine/Source/Programs/Shared/EpicGames.UHT/Types/UhtClass.cs
+++ b/Engine/Source/Programs/Shared/EpicGames.UHT/Types/UhtClass.cs
@@ -19,6 +19,10 @@ namespace EpicGames.UHT.Types
 	[Flags]
 	public enum UhtClassExportFlags : UInt32
 	{
+
+		/// <summary>
+		/// No export flags
+		/// </summary>
 		None = 0,
 
 		/// <summary>
@@ -36,16 +40,34 @@ namespace EpicGames.UHT.Types
 		/// </summary>
 		HasReplciatedProperties = SelfHasReplicatedProperties | SuperHasReplicatedProperties,
 
+		/// <summary>
+		/// Custom constructor specifier present
+		/// </summary>
 		HasCustomConstructor = 1 << 2,
 
+		/// <summary>
+		/// A default constructor was found in the class
+		/// </summary>
 		HasDefaultConstructor = 1 << 3,
 
+		/// <summary>
+		/// An object initializer constructor was found in the class
+		/// </summary>
 		HasObjectInitializerConstructor = 1 << 4,
 
+		/// <summary>
+		/// A custom vtable helper constructor was found
+		/// </summary>
 		HasCustomVTableHelperConstructor = 1 << 5,
 
+		/// <summary>
+		/// A constructor was found
+		/// </summary>
 		HasConstructor = 1 << 6,
 
+		/// <summary>
+		/// GetLifetimeReplicatedProps was found in the class
+		/// </summary>
 		HasGetLifetimeReplicatedProps = 1 << 7,
 	}
 
@@ -91,19 +113,52 @@ namespace EpicGames.UHT.Types
 		}
 	}
 
+	/// <summary>
+	/// Type of the class
+	/// </summary>
 	public enum UhtClassType
 	{
+		/// <summary>
+		/// Class is a UCLASS
+		/// </summary>
 		Class,
+
+		/// <summary>
+		/// Class is a UINTERFACE
+		/// </summary>
 		Interface,
+
+		/// <summary>
+		/// Class is the native interface for a UINTERFACE
+		/// </summary>
 		NativeInterface,
 	}
 
+	/// <summary>
+	/// Type of archive serializer found
+	/// </summary>
 	[Flags]
 	public enum UhtSerializerArchiveType
 	{
+
+		/// <summary>
+		/// No serializer found
+		/// </summary>
 		None = 0,
+
+		/// <summary>
+		/// Archive serializer found
+		/// </summary>
 		Archive = 1 << 0,
+
+		/// <summary>
+		/// Structured archive serializer found
+		/// </summary>
 		StructuredArchiveRecord = 1 << 1,
+
+		/// <summary>
+		/// Mask of all serializer types
+		/// </summary>
 		All = Archive | StructuredArchiveRecord,
 	}
 
@@ -149,60 +204,139 @@ namespace EpicGames.UHT.Types
 		}
 	}
 
+	/// <summary>
+	/// A skipped declaration
+	/// </summary>
 	public struct UhtDeclaration
 	{
+
+		/// <summary>
+		/// Compiler directives when declaration was parsed
+		/// </summary>
 		public UhtCompilerDirective CompilerDirectives;
+
+		/// <summary>
+		/// Collection of tokens parsed in the declaration
+		/// </summary>
 		public UhtToken[] Tokens;
 	}
 
+	/// <summary>
+	/// Represents a declaration found by name
+	/// </summary>
 	public struct UhtFoundDeclaration
 	{
+
+		/// <summary>
+		/// Compiler directives when declaration was parsed
+		/// </summary>
 		public UhtCompilerDirective CompilerDirectives;
+
+		/// <summary>
+		/// Collection of tokens parsed in the declaration
+		/// </summary>
 		public UhtToken[] Tokens;
+
+		/// <summary>
+		/// Token index for the matching name
+		/// </summary>
 		public int NameTokenIndex;
+
+		/// <summary>
+		/// True if "virtual" was found prior to the matching name
+		/// </summary>
 		public bool bIsVirtual;
 	}
 
+	/// <summary>
+	/// Instance of a UCLASS, UINTERFACE, or a native interface
+	/// </summary>
 	[UhtEngineClass(Name = "Class")]
 	public class UhtClass : UhtStruct
 	{
 		private static UhtSpecifierValidatorTable ClassSpecifierValidatorTable = UhtSpecifierValidatorTables.Instance.Get(UhtTableNames.Class);
 		private static UhtSpecifierValidatorTable InterfaceSpecifierValidatorTable = UhtSpecifierValidatorTables.Instance.Get(UhtTableNames.Interface);
 		private static UhtSpecifierValidatorTable NativeInterfaceSpecifierValidatorTable = UhtSpecifierValidatorTables.Instance.Get(UhtTableNames.NativeInterface);
+		private List<UhtDeclaration>? DeclarationsInternal = null;
 
+		/// <summary>
+		/// Configuration section
+		/// </summary>
 		public string Config { get; set; } = String.Empty;
+
+		/// <summary>
+		/// If needed, the #if block define for the serializer
+		/// </summary>
 		public string EnclosingDefine { get; set; } = String.Empty;
 
+		/// <summary>
+		/// The class within
+		/// </summary>
 		[JsonConverter(typeof(UhtTypeSourceNameJsonConverter<UhtClass>))]
 		public UhtClass ClassWithin { get; set; }
 
+		/// <summary>
+		/// Class flags
+		/// </summary>
 		[JsonConverter(typeof(JsonStringEnumConverter))]
 		public EClassFlags ClassFlags { get; set; } = EClassFlags.None;
 
+		/// <summary>
+		/// Class cast flags
+		/// </summary>
 		[JsonConverter(typeof(JsonStringEnumConverter))]
 		public EClassCastFlags ClassCastFlags { get; set; } = EClassCastFlags.None;
 
+		/// <summary>
+		/// Export flags not present in the engine
+		/// </summary>
 		[JsonConverter(typeof(JsonStringEnumConverter))]
 		public UhtClassExportFlags ClassExportFlags { get; set; } = UhtClassExportFlags.None;
 
+		/// <summary>
+		/// Type of the class
+		/// </summary>
 		public UhtClassType ClassType = UhtClassType.Class;
 
+		/// <summary>
+		/// Type of archivers present
+		/// </summary>
 		[JsonConverter(typeof(JsonStringEnumConverter))]
 		public UhtSerializerArchiveType SerializerArchiveType { get; set; } = UhtSerializerArchiveType.None;
 
+		/// <summary>
+		/// Line number of the prolog
+		/// </summary>
 		public int PrologLineNumber { get; set; } = -1;
 
+		/// <summary>
+		/// Collection of functions and other declarations found in the class
+		/// </summary>
 		[JsonIgnore]
 		public IList<UhtDeclaration>? Declarations => this.DeclarationsInternal;
+
+		/// <summary>
+		/// If this, this class is a UINTERFACE and NativeInterface is the associated native interface
+		/// </summary>
 		public UhtClass? NativeInterface = null;
 
+		/// <summary>
+		/// Line number o the generated body statement
+		/// </summary>
 		public int GeneratedBodyLineNumber { get; set; } = -1;
 
+		/// <summary>
+		/// Access of the generated body
+		/// </summary>
 		[JsonConverter(typeof(JsonStringEnumConverter))]
 		public UhtAccessSpecifier GeneratedBodyAccessSpecifier { get; set; } = UhtAccessSpecifier.None;
 
-		public bool bHasGeneratedBody = false; // If not true, then GENERATED_UCLASS_BODY was used
+		/// <summary>
+		/// True if GENERATED_BODY was used.  If false, GENERATED_UCLASS_BODY was used.
+		/// </summary>
+		public bool bHasGeneratedBody = false;
 
+		/// <inheritdoc/>
 		[JsonIgnore]
 		public override UhtEngineType EngineType
 		{
@@ -249,17 +383,29 @@ namespace EpicGames.UHT.Types
 			}
 		}
 
+		/// <summary>
+		/// The super class
+		/// </summary>
 		[JsonConverter(typeof(UhtNullableTypeSourceNameJsonConverter<UhtClass>))]
 		public UhtClass? SuperClass => (UhtClass?)this.Super;
 
+		/// <summary>
+		/// Construct a new instance of the class
+		/// </summary>
+		/// <param name="Outer">The outer type</param>
+		/// <param name="LineNumber">Line number where class begins</param>
 		public UhtClass(UhtType Outer, int LineNumber) : base(Outer, LineNumber)
 		{
 			this.ClassWithin = this;
 		}
 
+		/// <summary>
+		/// True if this class inherits from AActor
+		/// </summary>
 		[JsonIgnore]
 		public bool bIsActorClass => IsChildOf(this.Session.AActor);
 
+		///<inheritdoc/>
 		[JsonIgnore]
 		public override string EngineNamePrefix
 		{ 
@@ -289,8 +435,6 @@ namespace EpicGames.UHT.Types
 			}
 		}
 
-		private List<UhtDeclaration>? DeclarationsInternal = null;
-
 		/// <summary>
 		/// Add the given list of tokens as a possible declaration
 		/// </summary>
@@ -349,25 +493,6 @@ namespace EpicGames.UHT.Types
 			return false;
 		}
 
-		public IEnumerable<UhtProperty> EnumerateReplicatedProperties(bool bIncludeSuper)
-		{
-			if (bIncludeSuper && this.SuperClass != null)
-			{
-				foreach (UhtProperty Property in this.SuperClass.EnumerateReplicatedProperties(true))
-				{
-					yield return Property;
-				}
-			}
-
-			foreach (UhtProperty Property in this.Properties)
-			{
-				if (Property.PropertyFlags.HasAnyFlags(EPropertyFlags.Net))
-				{
-					yield return Property;
-				}
-			}
-		}
-
 		/// <summary>
 		/// Checks to see if the class or any super class has the given flags.
 		/// </summary>
@@ -393,6 +518,7 @@ namespace EpicGames.UHT.Types
 			public bool bFound;
 		}
 
+		/// <inheritdoc/>
 		protected override bool ResolveSelf(UhtResolvePhase Phase)
 		{
 			bool bResult = base.ResolveSelf(Phase);
@@ -769,6 +895,7 @@ namespace EpicGames.UHT.Types
 		#endregion
 
 		#region Validation support
+		/// <inheritdoc/>
 		protected override UhtValidationOptions Validate(UhtValidationOptions Options)
 		{
 			Options = base.Validate(Options);
diff --git a/Engine/Source/Programs/Shared/EpicGames.UHT/Types/UhtProperty.cs b/Engine/Source/Programs/Shared/EpicGames.UHT/Types/UhtProperty.cs
--- a/Engine/Source/Programs/Shared/EpicGames.UHT/Types/UhtProperty.cs
+++ b/Engine/Source/Programs/Shared/EpicGames.UHT/Types/UhtProperty.cs
@@ -140,8 +140,18 @@ namespace EpicGames.UHT.Types
 		Member,
 	};
 
+	/// <summary>
+	/// Helper methods for the property category
+	/// </summary>
 	public static class UhtPropertyCategoryExtensions
 	{
+
+		/// <summary>
+		/// Return the hint text for the property category
+		/// </summary>
+		/// <param name="PropertyCategory">Property category</param>
+		/// <returns>The user facing hint text</returns>
+		/// <exception cref="UhtIceException">Unexpected category</exception>
 		public static string GetHintText(this UhtPropertyCategory PropertyCategory)
 		{
 			switch (PropertyCategory)
@@ -160,51 +170,80 @@ namespace EpicGames.UHT.Types
 					throw new UhtIceException("Unknown variable category");
 			}
 		}
-
-		public static bool IsMemberProperty(this UhtPropertyCategory PropertyCategory)
-		{
-			return PropertyCategory == UhtPropertyCategory.Member;
-		}
-
-		public static bool IsFunctionProperty(this UhtPropertyCategory PropertyCategory)
-		{
-			return PropertyCategory != UhtPropertyCategory.Member;
-		}
-
-		public static bool IsFunctionReturn(this UhtPropertyCategory PropertyCategory)
-		{
-			return PropertyCategory == UhtPropertyCategory.Return;
-		}
-
-		public static bool IsFunctionArgument(this UhtPropertyCategory PropertyCategory)
-		{
-			return PropertyCategory == UhtPropertyCategory.RegularParameter || PropertyCategory == UhtPropertyCategory.ReplicatedParameter;
-		}
 	}
 
+	/// <summary>
+	/// Allocator used for container
+	/// </summary>
 	public enum UhtPropertyAllocator
 	{
+		/// <summary>
+		/// Default allocator
+		/// </summary>
 		Default,
+
+		/// <summary>
+		/// Memory image allocator
+		/// </summary>
 		MemoryImage
 	};
 
+	/// <summary>
+	/// Type of pointer
+	/// </summary>
 	public enum UhtPointerType
 	{
+
+		/// <summary>
+		/// No pointer specified
+		/// </summary>
 		None,
+
+		// Native pointer specified
 		Native,
 	}
 
+	/// <summary>
+	/// Size type of an integer
+	/// </summary>
 	public enum UhtPropertyIntType
 	{
+
+		/// <summary>
+		/// Property is not an integer
+		/// </summary>
 		None,
+
+		/// <summary>
+		/// Property is a sized integer
+		/// </summary>
 		Sized,
+
+		/// <summary>
+		/// Property is an unsized integer
+		/// </summary>
 		Unsized,
 	};
 
+	/// <summary>
+	/// Type of reference
+	/// </summary>
 	public enum UhtPropertyRefQualifier
 	{
+
+		/// <summary>
+		/// Property is not a reference
+		/// </summary>
 		None,
+
+		/// <summary>
+		/// Property is a const reference
+		/// </summary>
 		ConstRef,
+
+		/// <summary>
+		/// Property is a non-const reference
+		/// </summary>
 		NonConstRef,
 	};
 
@@ -214,6 +253,10 @@ namespace EpicGames.UHT.Types
 	[Flags]
 	public enum UhtPropertyOptions
 	{
+
+		/// <summary>
+		/// No property options
+		/// </summary>
 		None = 0,
 
 		/// <summary>
@@ -264,9 +307,15 @@ namespace EpicGames.UHT.Types
 		}
 	}
 
+	/// <summary>
+	/// Property capabilities.  Use the caps system instead of patterns such as "Property is UhtObjectProperty"
+	/// </summary>
 	[Flags]
 	public enum UhtPropertyCaps
 	{
+		/// <summary>
+		/// No property caps
+		/// </summary>
 		None = 0,
 
 		/// <summary>
@@ -453,8 +502,17 @@ namespace EpicGames.UHT.Types
 		GetterSetterArg,
 	}
 
+	/// <summary>
+	/// Extension methods for the property text type
+	/// </summary>
 	public static class UhtPropertyTextTypeExtensions
 	{
+
+		/// <summary>
+		/// Test to see if the text type is for a function
+		/// </summary>
+		/// <param name="TextType">Type of text</param>
+		/// <returns>True if the text type is a function</returns>
 		public static bool IsParameter(this UhtPropertyTextType TextType)
 		{
 			return
@@ -464,10 +522,22 @@ namespace EpicGames.UHT.Types
 		}
 	}
 
+	//ETSTODO - This can be removed since FunctionThunkParameterArgType is only used in conjunction with UhtPGetArgumentType
+	/// <summary>
+	/// Specifies how the PGet argument type is to be formatted
+	/// </summary>
 	public enum UhtPGetArgumentType
 	{
+		/// <summary>
+		/// </summary>
 		None,
+
+		/// <summary>
+		/// </summary>
 		EngineClass,
+
+		/// <summary>
+		/// </summary>
 		TypeText,
 	}
 
diff --git a/Engine/Source/Programs/Shared/EpicGames.UHT/Types/UhtScriptStruct.cs b/Engine/Source/Programs/Shared/EpicGames.UHT/Types/UhtScriptStruct.cs
--- a/Engine/Source/Programs/Shared/EpicGames.UHT/Types/UhtScriptStruct.cs
+++ b/Engine/Source/Programs/Shared/EpicGames.UHT/Types/UhtScriptStruct.cs
@@ -15,15 +15,45 @@ namespace EpicGames.UHT.Types
 	[Flags]
 	public enum UhtRigVMParameterFlags : UInt32
 	{
-		None			= 0,
-		Constant		= 0x00000001,
-		Input			= 0x00000002,
-		Output			= 0x00000004,
-		Singleton		= 0x00000008,
+		/// <summary>
+		/// No RigVM flags
+		/// </summary>
+		None = 0,
+
+		/// <summary>
+		/// "Constant" metadata specified
+		/// </summary>
+		Constant = 0x00000001,
+
+		/// <summary>
+		/// "Input" metadata specified
+		/// </summary>
+		Input = 0x00000002,
+
+		/// <summary>
+		/// "Output" metadata specified
+		/// </summary>
+		Output = 0x00000004,
+
+		/// <summary>
+		/// "Singleton" metadata specified
+		/// </summary>
+		Singleton = 0x00000008,
+
+		/// <summary>
+		/// Set if the property is editor only
+		/// </summary>
 		EditorOnly		= 0x00000010,
-		IsEnum			= 0x00010000,
-		IsArray			= 0x00020000,
-		IsFixedArray	= 0x00040000,
+
+		/// <summary>
+		/// Set if the property is an enum
+		/// </summary>
+		IsEnum = 0x00010000,
+
+		/// <summary>
+		/// Set if the property is an array
+		/// </summary>
+		IsArray = 0x00020000,
 	}
 
 	/// <summary>
@@ -68,47 +98,91 @@ namespace EpicGames.UHT.Types
 		}
 	}
 
-	/**
-	 * The FRigVMParameter represents a single parameter of a method
-	 * marked up with RIGVM_METHOD.
-	 * Each parameter can be marked with Constant, Input or Output
-	 * metadata - this struct simplifies access to that information.
-	 */
-
+	/// <summary>
+	/// The FRigVMParameter represents a single parameter of a method
+	/// marked up with RIGVM_METHOD.
+	/// Each parameter can be marked with Constant, Input or Output
+	/// metadata - this struct simplifies access to that information.
+	/// </summary>
 	public class UhtRigVMParameter
 	{
+		/// <summary>
+		/// Property associated with the RigVM parameter
+		/// </summary>
 		public readonly UhtProperty? Property;
+
+		/// <summary>
+		/// Name of the property
+		/// </summary>
 		public string Name { get; } = String.Empty;
+
+		/// <summary>
+		/// Type of the property
+		/// </summary>
 		public string Type { get; } = String.Empty;
-		public readonly string Getter = String.Empty;
+
+		/// <summary>
+		/// Cast name
+		/// </summary>
 		public readonly string? CastName = null;
+
+		/// <summary>
+		/// Cast type
+		/// </summary>
 		public readonly string? CastType = null;
+
+		/// <summary>
+		/// Flags associated with the parameter
+		/// </summary>
 		public UhtRigVMParameterFlags ParameterFlags = UhtRigVMParameterFlags.None;
 
+		/// <summary>
+		/// True if the parameter is marked as "Constant"
+		/// </summary>
 		[JsonIgnore]
 		public bool bConstant => this.ParameterFlags.HasAnyFlags(UhtRigVMParameterFlags.Constant);
 
+		/// <summary>
+		/// True if the parameter is marked as "Input"
+		/// </summary>
 		[JsonIgnore]
 		public bool bInput => this.ParameterFlags.HasAnyFlags(UhtRigVMParameterFlags.Input);
 
+		/// <summary>
+		/// True if the parameter is marked as "Output"
+		/// </summary>
 		[JsonIgnore]
 		public bool bOutput => this.ParameterFlags.HasAnyFlags(UhtRigVMParameterFlags.Output);
 
+		/// <summary>
+		/// True if the parameter is marked as "Singleton"
+		/// </summary>
 		[JsonIgnore]
 		public bool bSingleton => this.ParameterFlags.HasAnyFlags(UhtRigVMParameterFlags.Singleton);
 
+		/// <summary>
+		/// True if the parameter is editor only
+		/// </summary>
 		[JsonIgnore]
 		public bool bEditorOnly => this.ParameterFlags.HasAnyFlags(UhtRigVMParameterFlags.EditorOnly);
 
+		/// <summary>
+		/// True if the parameter is an enum
+		/// </summary>
 		[JsonIgnore]
 		public bool bIsEnum => this.ParameterFlags.HasAnyFlags(UhtRigVMParameterFlags.IsEnum);
 
+		/// <summary>
+		/// True if the parameter is an array
+		/// </summary>
 		[JsonIgnore]
 		public bool bIsArray => this.ParameterFlags.HasAnyFlags(UhtRigVMParameterFlags.IsArray);
 
-		[JsonIgnore]
-		public bool bIsFixedArray => this.ParameterFlags.HasAnyFlags(UhtRigVMParameterFlags.IsFixedArray);
-
+		/// <summary>
+		/// Create a new RigVM parameter from a property
+		/// </summary>
+		/// <param name="Property">Source property</param>
+		/// <param name="Index">Parameter index.  Used to create a unique cast name.</param>
 		public UhtRigVMParameter(UhtProperty Property, int Index)
 		{
 			this.Property = Property;
@@ -119,7 +193,6 @@ namespace EpicGames.UHT.Types
 			this.ParameterFlags |= Property.MetaData.ContainsKey("Output") ? UhtRigVMParameterFlags.Output : UhtRigVMParameterFlags.None;
 			this.ParameterFlags |= Property.bIsEditorOnlyProperty ? UhtRigVMParameterFlags.EditorOnly : UhtRigVMParameterFlags.None;
 			this.ParameterFlags |= Property.MetaData.ContainsKey("Singleton") ? UhtRigVMParameterFlags.Singleton : UhtRigVMParameterFlags.None;
-			this.Getter = "GetRef";
 
 			if (Property.MetaData.ContainsKey("Visible"))
 			{
@@ -148,32 +221,58 @@ namespace EpicGames.UHT.Types
 				{
 					string ExtendedType = this.ExtendedType(false);
 					this.CastName = $"{this.Name}_{Index}_Array";
-					this.CastType = $"TArrayView<const {ExtendedType.AsMemory(1, ExtendedType.Length - 2)}>";
+					this.CastType = $"TArrayView<const {ExtendedType.Substring(1, ExtendedType.Length - 2)}>";
 				}
 			}
 		}
 
+		/// <summary>
+		/// Create a new parameter
+		/// </summary>
+		/// <param name="Name">Name of the parameter</param>
+		/// <param name="Type">Type of the parameter</param>
 		public UhtRigVMParameter(string Name, string Type)
 		{
 			this.Name = Name;
 			this.Type = Type;
 		}
 
+		/// <summary>
+		/// Get the name of the parameter
+		/// </summary>
+		/// <param name="bCastName">If true, return the cast name</param>
+		/// <returns>Parameter name</returns>
 		public string NameOriginal(bool bCastName = false)
 		{
 			return bCastName && this.CastName != null ? this.CastName : this.Name;
 		}
 
+		/// <summary>
+		/// Get the type of the parameter
+		/// </summary>
+		/// <param name="bCastType">If true, return the cast type</param>
+		/// <returns>Parameter type</returns>
 		public string TypeOriginal(bool bCastType = false)
 		{
 			return bCastType && this.CastType != null ? this.CastType : this.Type;
 		}
 
+		/// <summary>
+		/// Get the full declaration (type and name)
+		/// </summary>
+		/// <param name="bCastType">If true, return the cast type</param>
+		/// <param name="bCastName">If true, return the cast name</param>
+		/// <returns>Parameter declaration</returns>
 		public string Declaration(bool bCastType = false, bool bCastName = false)
 		{
 			return $"{TypeOriginal(bCastType)} {NameOriginal(bCastName)}";
 		}
 
+		/// <summary>
+		/// Return the base type without any template arguments
+		/// </summary>
+		/// <param name="bCastType">If true, return the cast type</param>
+		/// <returns>Base parameter type</returns>
 		public string BaseType(bool bCastType = false)
 		{
 			string String = TypeOriginal(bCastType);
@@ -189,6 +288,11 @@ namespace EpicGames.UHT.Types
 			}
 		}
 
+		/// <summary>
+		/// Template arguments of the type or type if not a template type.
+		/// </summary>
+		/// <param name="bCastType">If true, return the cast type</param>
+		/// <returns>Template arguments of the type</returns>
 		public string ExtendedType(bool bCastType = false)
 		{
 			string String = TypeOriginal(bCastType);
@@ -204,6 +308,11 @@ namespace EpicGames.UHT.Types
 			}
 		}
 
+		/// <summary>
+		/// Return the type with a const reference
+		/// </summary>
+		/// <param name="bCastType">If true, return the cast type</param>
+		/// <returns>Type with a const reference</returns>
 		public string TypeConstRef(bool bCastType = false)
 		{
 			string String = TypeNoRef(bCastType);
@@ -217,12 +326,22 @@ namespace EpicGames.UHT.Types
 			}
 		}
 
+		/// <summary>
+		/// Return the type with a reference
+		/// </summary>
+		/// <param name="bCastType">If true, return the cast type</param>
+		/// <returns>Type with a reference</returns>
 		public string TypeRef(bool bCastType = false)
 		{
 			string String = TypeNoRef(bCastType);
 			return $"{String}&";
 		}
 
+		/// <summary>
+		/// Return the type without reference
+		/// </summary>
+		/// <param name="bCastType">If true, return the cast type</param>
+		/// <returns>Type without the reference</returns>
 		public string TypeNoRef(bool bCastType = false)
 		{
 			string String = TypeOriginal(bCastType);
@@ -236,55 +355,105 @@ namespace EpicGames.UHT.Types
 			}
 		}
 
+		/// <summary>
+		/// Return the type as a reference
+		/// </summary>
+		/// <param name="bCastType">If true, return the cast type</param>
+		/// <returns>Type as a reference</returns>
 		public string TypeVariableRef(bool bCastType = false)
 		{
 			return IsConst() ? TypeConstRef(bCastType) : TypeRef(bCastType);
 		}
 
+		/// <summary>
+		/// Return a variable declaration for the parameter
+		/// </summary>
+		/// <param name="bCastType">If true, return the cast type</param>
+		/// <param name="bCastName">If true, return the cast name</param>
+		/// <returns>Parameter as a variable declaration</returns>
 		public string Variable(bool bCastType = false, bool bCastName = false)
 		{
 			return $"{TypeVariableRef(bCastType)} {NameOriginal(bCastName)}";
 		}
 
+		/// <summary>
+		/// True if the parameter is constant
+		/// </summary>
+		/// <returns>True if the parameter is constant</returns>
 		public bool IsConst()
 		{
 			return this.bConstant || (this.bInput && !this.bOutput);
 		}
 
+		/// <summary>
+		/// Return true if the parameter requires a cast
+		/// </summary>
+		/// <returns>True if the parameter requires a cast</returns>
 		public bool RequiresCast()
 		{
 			return this.CastType != null && this.CastName != null;
 		}
 	}
 
-	/**
-	 * A single info dataset for a function marked with RIGVM_METHOD.
-	 * This struct provides access to its name, the return type and all parameters.
-	 */
+	/// <summary>
+	/// A single info dataset for a function marked with RIGVM_METHOD.
+	/// This struct provides access to its name, the return type and all parameters.
+	/// </summary>
 	public class UhtRigVMMethodInfo
 	{
 		private static string NoPrefixInternal = string.Empty;
 		private static string ReturnPrefixInternal = "return ";
 
+		/// <summary>
+		/// Return type of the method
+		/// </summary>
 		public string ReturnType { get; set; } = string.Empty;
+
+		/// <summary>
+		/// Name of the method
+		/// </summary>
 		public string Name { get; set; } = string.Empty;
+
+		/// <summary>
+		/// Method parameters
+		/// </summary>
 		public List<UhtRigVMParameter> Parameters { get; set; } = new List<UhtRigVMParameter>();
 
+		/// <summary>
+		/// If the method has a return value, return "return".  Otherwise return nothing.
+		/// </summary>
+		/// <returns>Prefix required for the return value</returns>
 		public string ReturnPrefix()
 		{
 			return (ReturnType.Length == 0 || ReturnType == "void") ? NoPrefixInternal : ReturnPrefixInternal;
 		}
 	}
 
-	/**
-	 * An info dataset providing access to all functions marked with RIGVM_METHOD
-	 * for each struct.
-	 */
+	/// <summary>
+	/// An info dataset providing access to all functions marked with RIGVM_METHOD
+	/// for each struct.
+	/// </summary>
 	public class UhtRigVMStructInfo
 	{
+
+		/// <summary>
+		/// True if the GetUpgradeInfoMethod was found. 
+		/// </summary>
 		public bool bHasGetUpgradeInfoMethod = false;
+
+		/// <summary>
+		/// Engine name of the owning script struct
+		/// </summary>
 		public string Name { get; set; } = string.Empty;
+
+		/// <summary>
+		/// List of the members
+		/// </summary>
 		public List<UhtRigVMParameter> Members = new List<UhtRigVMParameter>();
+
+		/// <summary>
+		/// List of the methods
+		/// </summary>
 		public List<UhtRigVMMethodInfo> Methods { get; set; } = new List<UhtRigVMMethodInfo>();
 	};
 
@@ -294,10 +463,30 @@ namespace EpicGames.UHT.Types
 	[Flags]
 	public enum UhtScriptStructExportFlags : UInt32
 	{
+
+		/// <summary>
+		/// No export flags
+		/// </summary>
 		None = 0,
+
+		/// <summary>
+		/// "HasDefaults" specifier present
+		/// </summary>
 		HasDefaults = 1 << 0,
+
+		/// <summary>
+		/// "HasNoOpConstructor" specifier present
+		/// </summary>
 		HasNoOpConstructor = 1 << 1,
+
+		/// <summary>
+		/// "IsAlwaysAccessible" specifier present
+		/// </summary>
 		IsAlwaysAccessible = 1 << 2,
+
+		/// <summary>
+		/// "IsCoreType" specifier present
+		/// </summary>
 		IsCoreType = 1 << 3,
 	}
 
@@ -343,40 +532,68 @@ namespace EpicGames.UHT.Types
 		}
 	}
 
-
+	/// <summary>
+	/// Represents the USTRUCT object
+	/// </summary>
 	[UhtEngineClass(Name = "ScriptStruct")]
 	public class UhtScriptStruct : UhtStruct
 	{
 		private static UhtSpecifierValidatorTable ScriptStructSpecifierValidatorTable = UhtSpecifierValidatorTables.Instance.Get(UhtTableNames.ScriptStruct);
 
+		/// <summary>
+		/// Script struct engine flags
+		/// </summary>
 		[JsonConverter(typeof(JsonStringEnumConverter))]
 		public EStructFlags ScriptStructFlags { get; set; } = EStructFlags.NoFlags;
 
+		/// <summary>
+		/// UHT only script struct falgs
+		/// </summary>
 		[JsonConverter(typeof(JsonStringEnumConverter))]
 		public UhtScriptStructExportFlags ScriptStructExportFlags { get; set; } = UhtScriptStructExportFlags.None;
 
+		/// <summary>
+		/// Line number where GENERATED_BODY/GENERATED_USTRUCT_BODY macro was found
+		/// </summary>
 		public int MacroDeclaredLineNumber { get; set; } = -1;
 
+		/// <summary>
+		/// RigVM structure info
+		/// </summary>
 		public UhtRigVMStructInfo? RigVMStructInfo { get; set; } = null;
 
+		/// <inheritdoc/>
 		[JsonIgnore]
 		public override UhtEngineType EngineType => UhtEngineType.ScriptStruct;
 
+		/// <summary>
+		/// True if the struct has the "HasDefaults" specifier
+		/// </summary>
 		[JsonIgnore]
 		public bool bHasDefaults => this.ScriptStructExportFlags.HasAnyFlags(UhtScriptStructExportFlags.HasDefaults);
 
+		/// <summary>
+		/// True if the struct has the "IsAlwaysAccessible" specifier
+		/// </summary>
 		[JsonIgnore]
 		public bool bIsAlwaysAccessible => this.ScriptStructExportFlags.HasAnyFlags(UhtScriptStructExportFlags.IsAlwaysAccessible);
 
+		/// <summary>
+		/// True if the struct has the "IsCoreType" specifier
+		/// </summary>
 		[JsonIgnore]
 		public bool bIsCoreType => this.ScriptStructExportFlags.HasAnyFlags(UhtScriptStructExportFlags.IsCoreType);
 
+		/// <summary>
+		/// True if the struct has the "HasNoOpConstructor" specifier
+		/// </summary>
 		[JsonIgnore]
 		public bool bHasNoOpConstructor => this.ScriptStructExportFlags.HasAnyFlags(UhtScriptStructExportFlags.HasNoOpConstructor);
 
 		/// <inheritdoc/>
 		public override string EngineClassName { get => "ScriptStruct"; }
 
+		/// <inheritdoc/>
 		[JsonIgnore]
 		public override string EngineNamePrefix
 		{
@@ -390,9 +607,17 @@ namespace EpicGames.UHT.Types
 		[JsonIgnore]
 		protected override UhtSpecifierValidatorTable? SpecifierValidatorTable { get => UhtScriptStruct.ScriptStructSpecifierValidatorTable; }
 
+		/// <summary>
+		/// Super struct
+		/// </summary>
 		[JsonConverter(typeof(UhtNullableTypeSourceNameJsonConverter<UhtScriptStruct>))]
 		public UhtScriptStruct? SuperScriptStruct => (UhtScriptStruct?)this.Super;
 
+		/// <summary>
+		/// Construct a new script struct
+		/// </summary>
+		/// <param name="Outer">Outer type</param>
+		/// <param name="LineNumber">Line number of the definition</param>
 		public UhtScriptStruct(UhtType Outer, int LineNumber) : base(Outer, LineNumber)
 		{
 		}
@@ -431,6 +656,7 @@ namespace EpicGames.UHT.Types
 		}
 
 		#region Validation support
+		/// <inheritdoc/>
 		protected override UhtValidationOptions Validate(UhtValidationOptions Options)
 		{
 			Options = base.Validate(Options);
diff --git a/Engine/Source/Programs/Shared/EpicGames.UHT/Types/UhtSourceFile.cs b/Engine/Source/Programs/Shared/EpicGames.UHT/Types/UhtSourceFile.cs
--- a/Engine/Source/Programs/Shared/EpicGames.UHT/Types/UhtSourceFile.cs
+++ b/Engine/Source/Programs/Shared/EpicGames.UHT/Types/UhtSourceFile.cs
@@ -6,34 +6,86 @@ using System.IO;
 
 namespace EpicGames.UHT.Types
 {
+
+	/// <summary>
+	/// Represents a source file in UHT
+	/// </summary>
 	public class UhtSourceFile : IUhtMessageSite, IUhtMessageSource
 	{
+		private UhtSimpleMessageSite MessageSite;
+		private UhtSourceFragment SourceFragment;
+
+		/// <summary>
+		/// The session associated with the source
+		/// </summary>
 		public readonly UhtSession Session;
+
+		/// <summary>
+		/// The full file path
+		/// </summary>
 		public readonly string FilePath;
+
+		/// <summary>
+		/// The file name
+		/// </summary>
 		public readonly string FileName;
+
+		/// <summary>
+		/// The contents of the source file
+		/// </summary>
 		public StringView Data { get => SourceFragment.Data; }
+
+		/// <summary>
+		/// If this source is from a fragment, the original source file
+		/// </summary>
 		public UhtSourceFile? FragmentSourceFile { get => SourceFragment.SourceFile; }
+
+		/// <summary>
+		/// If this source is from a fragment, the starting line of the fragment
+		/// </summary>
 		public int FragmentLineNumber { get => SourceFragment.LineNumber; }
-		public string FullFilePath { get => SourceFragment.FilePath; }
 
-		private UhtSimpleMessageSite MessageSite;
-		private UhtSourceFragment SourceFragment;
+		/// <summary>
+		/// Full file path of the source file
+		/// </summary>
+		public string FullFilePath { get => SourceFragment.FilePath; }
 
 		#region IUHTMessageSite implementation
+		/// <inheritdoc/>
 		public IUhtMessageSession MessageSession => this.MessageSite.MessageSession;
+
+		/// <inheritdoc/>
 		public IUhtMessageSource? MessageSource => this.MessageSite.MessageSource;
+
+		/// <inheritdoc/>
 		public IUhtMessageLineNumber? MessageLineNumber => null;
 		#endregion
 
 		#region IUHTMessageSource implementation
+		/// <inheritdoc/>
 		public string MessageFilePath { get => this.FilePath; }
+
+		/// <inheritdoc/>
 		public string MessageFullFilePath { get => this.FullFilePath; }
+
+		/// <inheritdoc/>
 		public bool bMessageIsFragment { get => this.FragmentSourceFile != null; }
+
+		/// <inheritdoc/>
 		public string MessageFragmentFilePath { get => this.FragmentSourceFile != null ? this.FragmentSourceFile.MessageFilePath : ""; }
+
+		/// <inheritdoc/>
 		public string MessageFragmentFullFilePath { get => this.FragmentSourceFile != null ? this.FragmentSourceFile.MessageFullFilePath : ""; }
+
+		/// <inheritdoc/>
 		public int MessageFragmentLineNumber { get => this.FragmentLineNumber; }
 		#endregion
 
+		/// <summary>
+		/// Construct a new instance of a source file
+		/// </summary>
+		/// <param name="Session">The owning session</param>
+		/// <param name="FilePath">The full file path</param>
 		public UhtSourceFile(UhtSession Session, string FilePath)
 		{
 			this.Session = Session;
@@ -42,6 +94,9 @@ namespace EpicGames.UHT.Types
 			this.MessageSite = new UhtSimpleMessageSite(this.Session, this);
 		}
 
+		/// <summary>
+		/// Read the contents of the source file
+		/// </summary>
 		public virtual void Read()
 		{
 			this.SourceFragment = Session.ReadSource(this.FilePath);
diff --git a/Engine/Source/Programs/Shared/EpicGames.UHT/Types/UhtStruct.cs b/Engine/Source/Programs/Shared/EpicGames.UHT/Types/UhtStruct.cs
--- a/Engine/Source/Programs/Shared/EpicGames.UHT/Types/UhtStruct.cs
+++ b/Engine/Source/Programs/Shared/EpicGames.UHT/Types/UhtStruct.cs
@@ -10,9 +10,17 @@ using System.Text.Json.Serialization;
 
 namespace EpicGames.UHT.Types
 {
+	/// <summary>
+	/// Base class for all types that contain properties and functions.
+	/// Also support the ability to specifier super class and base classes
+	/// </summary>
 	[UhtEngineClass(Name = "Struct")]
 	public abstract class UhtStruct : UhtField
 	{
+
+		/// <summary>
+		/// Generated code version of the type.  Set via specifiers
+		/// </summary>
 		public EGeneratedCodeVersion GeneratedCodeVersion = EGeneratedCodeVersion.None;
 
 		/// <summary>
@@ -61,20 +69,35 @@ namespace EpicGames.UHT.Types
 			}
 		}
 
+		/// <inheritdoc/>
 		[JsonIgnore]
 		public virtual string EngineNamePrefix { get => "F"; }
 
 		/// <inheritdoc/>
 		public override string EngineClassName { get => "Struct"; }
 
+		/// <summary>
+		/// Super type
+		/// </summary>
 		public UhtStruct? Super = null;
 
+		/// <summary>
+		/// Base types
+		/// </summary>
 		[JsonConverter(typeof(UhtNullableTypeListJsonConverter<UhtStruct>))]
 		public List<UhtStruct>? Bases { get; set; } = null;
 
+		/// <summary>
+		/// Super struct type
+		/// </summary>
 		[JsonIgnore]
 		public UhtStruct? SuperStruct => Super;
 
+		/// <summary>
+		/// Construct a new instance
+		/// </summary>
+		/// <param name="Outer">Outer type</param>
+		/// <param name="LineNumber">Line number where definition begins</param>
 		public UhtStruct(UhtType Outer, int LineNumber) : base(Outer, LineNumber)
 		{
 		}
@@ -82,7 +105,7 @@ namespace EpicGames.UHT.Types
 		/// <summary>
 		/// Test to see if the given struct is derived from the base structure
 		/// </summary>
-		/// <param name="SomeBase">Base structure.</param>
+		/// <param name="Base">Base structure.</param>
 		/// <returns>True if the given structure is the specified base or derives from the base.  If the base is null, the false is returned.</returns>
 		public bool IsChildOf(UhtStruct? Base)
 		{
@@ -101,6 +124,7 @@ namespace EpicGames.UHT.Types
 		}
 
 		#region Resolution support
+		/// <inheritdoc/>
 		protected override void ResolveSuper(UhtResolvePhase ResolvePhase)
 		{
 			if (this.Super != null)
@@ -146,6 +170,12 @@ namespace EpicGames.UHT.Types
 		#endregion
 
 		#region Super and base binding helper methods
+		/// <summary>
+		/// Resolve super identifier
+		/// </summary>
+		/// <param name="SuperIdentifier">Token that represent the super</param>
+		/// <param name="FindOptions">Find options to restrict types</param>
+		/// <exception cref="UhtException">Thrown if super can not be found</exception>
 		public void BindAndResolveSuper(ref UhtToken SuperIdentifier, UhtFindOptions FindOptions)
 		{
 			if (SuperIdentifier)
@@ -161,6 +191,12 @@ namespace EpicGames.UHT.Types
 			}
 		}
 
+		/// <summary>
+		/// Resolve bases.  Unlike super, this routine will not generate an error if the type can not be found.
+		/// Having unrecognized types is expected.
+		/// </summary>
+		/// <param name="BaseIdentifiers">Collection of bases</param>
+		/// <param name="FindOptions">Options to restrict types being searched</param>
 		public void BindAndResolveBases(List<UhtToken[]>? BaseIdentifiers, UhtFindOptions FindOptions)
 		{
 			if (BaseIdentifiers != null)
@@ -185,6 +221,7 @@ namespace EpicGames.UHT.Types
 		#endregion
 
 		#region Validation support
+		/// <inheritdoc/>
 		protected override UhtValidationOptions Validate(UhtValidationOptions Options)
 		{
 			Options = base.Validate(Options);
@@ -218,6 +255,7 @@ namespace EpicGames.UHT.Types
 			return MinValue <= MaxValue;
 		}
 
+		/// <inheritdoc/>
 		protected override void ValidateDocumentationPolicy(UhtDocumentationPolicy Policy)
 		{
 			if (Policy.bClassOrStructCommentRequired)
@@ -300,7 +338,7 @@ namespace EpicGames.UHT.Types
 			}
 		}
 
-		protected void ValidateSparseClassData()
+		private void ValidateSparseClassData()
 		{
 			// Fetch the data types
 			string[]? SparseClassDataTypes = this.MetaData.GetStringArray(UhtNames.SparseClassDataTypes);
diff --git a/Engine/Source/Programs/Shared/EpicGames.UHT/Types/UhtType.cs b/Engine/Source/Programs/Shared/EpicGames.UHT/Types/UhtType.cs
--- a/Engine/Source/Programs/Shared/EpicGames.UHT/Types/UhtType.cs
+++ b/Engine/Source/Programs/Shared/EpicGames.UHT/Types/UhtType.cs
@@ -294,34 +294,111 @@ namespace EpicGames.UHT.Types
 		}
 	}
 
+	/// <summary>
+	/// Options to customize how symbols are found
+	/// </summary>
 	[Flags]
 	public enum UhtFindOptions
 	{
-		EngineName = 1 << 0,        // This is the default
+		/// <summary>
+		/// Search using the engine name.  This is the default if neither engine/source is specified.
+		/// </summary>
+		EngineName = 1 << 0,
+
+		/// <summary>
+		/// Search using the source name.  This can not be set if EngineName is also set.
+		/// </summary>
 		SourceName = 1 << 1,
 
+		/// <summary>
+		/// Search for UhtEnum types
+		/// </summary>
 		Enum = 1 << 8,
+
+		/// <summary>
+		/// Search for UhtScriptStruct types
+		/// </summary>
 		ScriptStruct = 1 << 9,
+		
+		/// <summary>
+		/// Search for UhtClass types
+		/// </summary>
 		Class = 1 << 10,
+
+		/// <summary>
+		/// Search for UhtFunction types that are delegate functions
+		/// </summary>
 		DelegateFunction = 1 << 11,
+
+		/// <summary>
+		/// Search for UhtFunction types that are not delegate functions
+		/// </summary>
 		Function = 1 << 12,
+
+		/// <summary>
+		/// Search for UhtProperty types
+		/// </summary>
 		Property = 1 << 13,
 
+		/// <summary>
+		/// Search for any properties or functions
+		/// </summary>
+		PropertyOrFunction = Property | Function | DelegateFunction,
+
+		/// <summary>
+		/// Do not search the super chain
+		/// </summary>
 		NoParents = 1 << 16,
+
+		/// <summary>
+		/// Do not search the outer chain
+		/// </summary>
 		NoOuter = 1 << 17,
+
+		/// <summary>
+		/// Do not search include files
+		/// </summary>
 		NoIncludes = 1 << 18,
+
+		/// <summary>
+		/// Do not do a last resort global namespace search
+		/// </summary>
 		NoGlobal = 1 << 19,
+
+		/// <summary>
+		/// Do not search the children of the starting point
+		/// </summary>
 		NoSelf = 1 << 20,
+
+		/// <summary>
+		/// Only search the parent chain
+		/// </summary>
+		ParentsOnly = NoOuter | NoIncludes | NoGlobal | NoSelf,
+
+		/// <summary>
+		/// Search only the children of the starting type
+		/// </summary>
+		SelfOnly = NoParents | NoOuter | NoIncludes | NoGlobal,
+
+		/// <summary>
+		/// Case compare (usually for SourceName searches)
+		/// </summary>
 		CaseCompare = 1 << 21,
+
+		/// <summary>
+		/// Caseless compare (usually for EngineName searches)
+		/// </summary>
 		CaselessCompare = 1 << 22,
-		ParentsOnly = NoOuter | NoIncludes | NoGlobal | NoSelf,
 
+		/// <summary>
+		/// Mask of name flags
+		/// </summary>
 		NamesMask = EngineName | SourceName,
+
+		/// <summary>
+		/// Mask of type flags
+		/// </summary>
 		TypesMask = Enum | ScriptStruct | Class | DelegateFunction | Function | Property,
-		ObjectTypesMask = Enum | ScriptStruct | Class,
-		DelgeateTypesMask = DelegateFunction,
-		PropertyOrFunction = Property | Function | DelegateFunction,
-		SelfOnly = NoParents | NoOuter | NoIncludes | NoGlobal,
 	}
 
 	/// <summary>
@@ -366,9 +443,16 @@ namespace EpicGames.UHT.Types
 		}
 	}
 
+	/// <summary>
+	/// Options during validation checks
+	/// </summary>
 	[Flags]
 	public enum UhtValidationOptions
 	{
+
+		/// <summary>
+		/// No validation options
+		/// </summary>
 		None,
 
 		/// <summary>
@@ -440,23 +524,63 @@ namespace EpicGames.UHT.Types
 		}
 	}
 
+	/// <summary>
+	/// Access specifier for type
+	/// </summary>
 	public enum UhtAccessSpecifier
 	{
+
+		/// <summary>
+		/// No access specifier seen
+		/// </summary>
 		None,
+
+		/// <summary>
+		/// Public type
+		/// </summary>
 		Public,
+
+		/// <summary>
+		/// Private type
+		/// </summary>
 		Private,
+
+		/// <summary>
+		/// Protected type
+		/// </summary>
 		Protected,
 	}
 
+	/// <summary>
+	/// Requested resolve phase
+	/// </summary>
 	public enum UhtResolvePhase : Int32
 	{
+		/// <summary>
+		/// Check for any type that is invalid.  Invalid items will be removed from the children. 
+		/// Used to detect things that looked like native interfaces but no interface was found
+		/// </summary>
 		InvalidCheck = 1,
+
+		/// <summary>
+		/// Resolve any super references
+		/// </summary>
 		Bases,
+
+		/// <summary>
+		/// Resolve property types that reference other types
+		/// </summary>
 		Properties,
+
+		/// <summary>
+		/// Final resolve phase for anything that needs fully resolved properties.
+		/// </summary>
 		Final,
-		Count,
 	}
 
+	/// <summary>
+	/// Flags for required documentation
+	/// </summary>
 	public class UhtDocumentationPolicy
 	{
 		private static UhtDocumentationPolicy StrictDocumentationPolicy = new UhtDocumentationPolicy
@@ -479,14 +603,43 @@ namespace EpicGames.UHT.Types
 			bFloatRangesRequired = false,
 		};
 
+		/// <summary>
+		/// If true, this is a valid documentation policy that should be respected
+		/// </summary>
 		public bool bPolicySet;
+
+		/// <summary>
+		/// A tool tip must be defined
+		/// </summary>
 		public bool bClassOrStructCommentRequired;
+
+		/// <summary>
+		/// Functions must have a tool tip
+		/// </summary>
 		public bool bFunctionToolTipsRequired;
+
+		/// <summary>
+		/// Property members must have a tool tip
+		/// </summary>
 		public bool bMemberToolTipsRequired;
+
+		/// <summary>
+		/// Function parameters must have a tool tip
+		/// </summary>
 		public bool bParameterToolTipsRequired;
+
+		/// <summary>
+		/// Float properties must have a range specified
+		/// </summary>
 		public bool bFloatRangesRequired;
 
-		public static bool TryGet(UhtType Type, string DocumentationPolicyName, out UhtDocumentationPolicy Policy)
+		/// <summary>
+		/// Return the documentation policy given the name
+		/// </summary>
+		/// <param name="DocumentationPolicyName">Name of the policy</param>
+		/// <param name="Policy">Policy settings</param>
+		/// <returns>True if the name is valid, false if not</returns>
+		public static bool TryGet(string DocumentationPolicyName, out UhtDocumentationPolicy Policy)
 		{
 			if (DocumentationPolicyName == "Strict")
 			{
@@ -499,6 +652,10 @@ namespace EpicGames.UHT.Types
 		}
 	};
 
+	/// <summary>
+	/// Base type for all UHT types.
+	/// These mostly map to UnrealEngine types.
+	/// </summary>
 	public abstract class UhtType : IUhtMessageSite, IUhtMessageLineNumber
 	{
 		/// <summary>
@@ -637,12 +794,15 @@ namespace EpicGames.UHT.Types
 		private int ResolveState = 0;
 
 		#region IUHTMessageSite implementation
+		/// <inheritdoc/>
 		[JsonIgnore]
 		public virtual IUhtMessageSession MessageSession => this.HeaderFile.MessageSession;
 
+		/// <inheritdoc/>
 		[JsonIgnore]
 		public virtual IUhtMessageSource? MessageSource => this.HeaderFile.MessageSource;
 
+		/// <inheritdoc/>
 		[JsonIgnore]
 		public IUhtMessageLineNumber? MessageLineNumber => this;
 		#endregion
@@ -725,6 +885,12 @@ namespace EpicGames.UHT.Types
 		}
 
 		#region Type resolution
+
+		/// <summary>
+		/// Resolve the type.  To customize behavior, override ResolveSuper, ResolveSelf, or ResolveChildren.
+		/// </summary>
+		/// <param name="Phase">Resolution phase</param>
+		/// <returns>False if the type fails any invalid checks during InvalidPhase check.</returns>
 		public bool Resolve(UhtResolvePhase Phase)
 		{
 			int InitState = (int)Phase << 1;
@@ -770,15 +936,28 @@ namespace EpicGames.UHT.Types
 			}
 		}
 
+		/// <summary>
+		/// Resolve any super types
+		/// </summary>
+		/// <param name="Phase">Resolution phase</param>
 		protected virtual void ResolveSuper(UhtResolvePhase Phase)
 		{
 		}
 
+		/// <summary>
+		/// Resolve self
+		/// </summary>
+		/// <param name="Phase">Resolution phase</param>
+		/// <returns>False if the type fails any invalid checks during InvalidPhase check.</returns>
 		protected virtual bool ResolveSelf(UhtResolvePhase Phase)
 		{
 			return true;
 		}
 
+		/// <summary>
+		/// Resolve children
+		/// </summary>
+		/// <param name="Phase">Resolution phase</param>
 		protected virtual void ResolveChildren(UhtResolvePhase Phase)
 		{
 			if (this.ChildrenInternal != null)
@@ -864,7 +1043,11 @@ namespace EpicGames.UHT.Types
 		#endregion
 
 		#region Validation support
-		// Validate the property settings
+		/// <summary>
+		/// Validate the type settings
+		/// </summary>
+		/// <param name="Options">Validation options</param>
+		/// <returns>Updated validation options.  This will be used to validate the children.</returns>
 		protected virtual UhtValidationOptions Validate(UhtValidationOptions Options)
 		{
 			ValidateMetaData();
@@ -893,7 +1076,7 @@ namespace EpicGames.UHT.Types
 
 		private void ValidateDocumentationPolicy()
 		{
-			if (UhtDocumentationPolicy.TryGet(this, this.DocumentationPolicyName, out UhtDocumentationPolicy Policy))
+			if (UhtDocumentationPolicy.TryGet(this.DocumentationPolicyName, out UhtDocumentationPolicy Policy))
 			{
 				if (Policy.bPolicySet)
 				{
@@ -906,10 +1089,19 @@ namespace EpicGames.UHT.Types
 			}
 		}
 
+		/// <summary>
+		/// Validate the documentation policy
+		/// </summary>
+		/// <param name="Policy">The expected documentation policy</param>
 		protected virtual void ValidateDocumentationPolicy(UhtDocumentationPolicy Policy)
 		{
 		}
 
+		/// <summary>
+		/// Validate the given type.  Validates itself and the children
+		/// </summary>
+		/// <param name="Type">Type to validate</param>
+		/// <param name="Options">Validation options</param>
 		public static void ValidateType(UhtType Type, UhtValidationOptions Options)
 		{
 			// Invoke the new method
@@ -924,6 +1116,11 @@ namespace EpicGames.UHT.Types
 		#endregion
 
 		#region Name and user facing text support
+		/// <summary>
+		/// Return the path name of the type which includes the outers
+		/// </summary>
+		/// <param name="StopOuter">Type to stop at when generating the path</param>
+		/// <return>The path string</return>
 		public string GetPathName(UhtType? StopOuter = null)
 		{
 			StringBuilder Builder = new StringBuilder();
@@ -931,6 +1128,11 @@ namespace EpicGames.UHT.Types
 			return Builder.ToString();
 		}
 
+		/// <summary>
+		/// Return the path name of the type which includes the outers
+		/// </summary>
+		/// <param name="Builder">Destination builder</param>
+		/// <param name="StopOuter">Type to stop at when generating the path</param>
 		public virtual void GetPathName(StringBuilder Builder, UhtType? StopOuter = null)
 		{
 			GetPathNameInternal(Builder, StopOuter);
@@ -940,7 +1142,7 @@ namespace EpicGames.UHT.Types
 			}
 		}
 
-		public virtual void GetPathNameInternal(StringBuilder Builder, UhtType? StopOuter = null)
+		private void GetPathNameInternal(StringBuilder Builder, UhtType? StopOuter = null)
 		{
 			if (this != StopOuter)
 			{
@@ -970,6 +1172,10 @@ namespace EpicGames.UHT.Types
 			}
 		}
 
+		/// <summary>
+		/// Get the full type name which is the engine class name followed by the path name
+		/// </summary>
+		/// <returns>Full name string</returns>
 		public string GetFullName()
 		{
 			StringBuilder Builder = new StringBuilder();
@@ -977,6 +1183,10 @@ namespace EpicGames.UHT.Types
 			return Builder.ToString();
 		}
 
+		/// <summary>
+		/// Get the full type name which is the engine class name followed by the path name
+		/// </summary>
+		/// <param name="Builder">Destination builder</param>
 		public virtual void GetFullName(StringBuilder Builder)
 		{
 			Builder.Append(this.EngineClassName);
@@ -1037,11 +1247,18 @@ namespace EpicGames.UHT.Types
 		}
 		#endregion
 
+		/// <summary>
+		/// Collect all the references for the type
+		/// </summary>
+		/// <param name="Collector">Object collecting the references</param>
 		public virtual void CollectReferences(IUhtReferenceCollector Collector)
 		{
 		}
 	}
 
+	/// <summary>
+	/// String builder extensions for UhtTypes
+	/// </summary>
 	public static class UhtTypeStringBuilderExtensions
 	{
 		/// <summary>
