More updates to UHT documentation

#rnx

diff --git a/Engine/Source/Programs/Shared/EpicGames.UHT/Parsers/UhtPropertyParser.cs b/Engine/Source/Programs/Shared/EpicGames.UHT/Parsers/UhtPropertyParser.cs
--- a/Engine/Source/Programs/Shared/EpicGames.UHT/Parsers/UhtPropertyParser.cs
+++ b/Engine/Source/Programs/Shared/EpicGames.UHT/Parsers/UhtPropertyParser.cs
@@ -236,7 +236,8 @@ namespace EpicGames.UHT.Parsers
 
 		public UhtPropertySpecifierContext(UhtParsingScope Scope, IUhtMessageSite MessageSite) : base(Scope, MessageSite, UhtMetaData.Empty)
 		{
-			this.PropertySettings = new UhtPropertySettings(Scope.ScopeType, 1, UhtPropertyCategory.Member, 0);
+			this.PropertySettings = new UhtPropertySettings();
+			this.PropertySettings.Reset(Scope.ScopeType, 1, UhtPropertyCategory.Member, 0);
 		}
 	}
 
diff --git a/Engine/Source/Programs/Shared/EpicGames.UHT/Types/Properties/UhtEnumProperty.cs b/Engine/Source/Programs/Shared/EpicGames.UHT/Types/Properties/UhtEnumProperty.cs
--- a/Engine/Source/Programs/Shared/EpicGames.UHT/Types/Properties/UhtEnumProperty.cs
+++ b/Engine/Source/Programs/Shared/EpicGames.UHT/Types/Properties/UhtEnumProperty.cs
@@ -251,7 +251,8 @@ namespace EpicGames.UHT.Types
 
 		private UhtProperty CreateUnderlyingProperty()
 		{
-			UhtPropertySettings PropertySettings = new UhtPropertySettings(this, 0, this.PropertyCategory, 0);
+			UhtPropertySettings PropertySettings = new UhtPropertySettings();
+			PropertySettings.Reset(this, 0, this.PropertyCategory, 0);
 			PropertySettings.SourceName = "UnderlyingType";
 			switch (this.UnderlyingType)
 			{
diff --git a/Engine/Source/Programs/Shared/EpicGames.UHT/Types/UhtMetaData.cs b/Engine/Source/Programs/Shared/EpicGames.UHT/Types/UhtMetaData.cs
--- a/Engine/Source/Programs/Shared/EpicGames.UHT/Types/UhtMetaData.cs
+++ b/Engine/Source/Programs/Shared/EpicGames.UHT/Types/UhtMetaData.cs
@@ -8,22 +8,53 @@ using System.Diagnostics.CodeAnalysis;
 
 namespace EpicGames.UHT.Types
 {
+
+	/// <summary>
+	/// UhtEnum types need to convert a meta data name index to the enum name
+	/// </summary>
 	public interface IUhtMetaDataKeyConversion
 	{
+
+		/// <summary>
+		/// Convert a name and index into a full meta data name
+		/// </summary>
+		/// <param name="Name">Base name of the meta data</param>
+		/// <param name="NameIndex">Index of the meta data.  -1 for the root object.</param>
+		/// <returns>Complete meta data key name</returns>
 		string GetMetaDataKey(string Name, int NameIndex);
 	}
 
+	/// <summary>
+	/// Uniquely identifies a meta data element
+	/// </summary>
 	public struct UhtMetaDataKey
 	{
+
+		/// <summary>
+		/// The name of the meta data
+		/// </summary>
 		public string Name;
+
+		/// <summary>
+		/// The index of the meta data name (i.e. enum value index) or -1 for the owning object meta data
+		/// </summary>
 		public int Index;
 
+		/// <summary>
+		/// Construct a new meta data key
+		/// </summary>
+		/// <param name="Name">Meta data name</param>
+		/// <param name="Index">Meta data index</param>
 		public UhtMetaDataKey(string Name, int Index = UhtMetaData.INDEX_NONE)
 		{
 			this.Name = Name;
 			this.Index = Index;
 		}
 
+		/// <summary>
+		/// Convert the key to a string
+		/// </summary>
+		/// <returns>String representation</returns>
 		public override string ToString()
 		{
 			if (Index == UhtMetaData.INDEX_NONE)
@@ -37,8 +68,17 @@ namespace EpicGames.UHT.Types
 		}
 	}
 
+	/// <summary>
+	/// Comparer for meta data keys
+	/// </summary>
 	class UhtMetaDataKeyComparer : IEqualityComparer<UhtMetaDataKey>, IComparer<UhtMetaDataKey>
 	{
+		/// <summary>
+		/// Compare two keys
+		/// </summary>
+		/// <param name="x">First key</param>
+		/// <param name="y">Second key</param>
+		/// <returns>-1, 0, or 1 depending on the relationship</returns>
 		public int Compare([AllowNull] UhtMetaDataKey x, [AllowNull] UhtMetaDataKey y)
 		{
 			if (x.Index < y.Index)
@@ -52,11 +92,22 @@ namespace EpicGames.UHT.Types
 			return string.Compare(x.Name, y.Name, StringComparison.OrdinalIgnoreCase);
 		}
 
+		/// <summary>
+		/// Test to see if two meta data keys are equal
+		/// </summary>
+		/// <param name="x">First key</param>
+		/// <param name="y">Second key</param>
+		/// <returns>True if the keys match</returns>
 		public bool Equals([AllowNull] UhtMetaDataKey x, [AllowNull] UhtMetaDataKey y)
 		{
 			return x.Index == y.Index && x.Name.Equals(y.Name, StringComparison.OrdinalIgnoreCase);
 		}
 
+		/// <summary>
+		/// Get the hash code of the meta data key
+		/// </summary>
+		/// <param name="obj">Key</param>
+		/// <returns>Hash code</returns>
 		public int GetHashCode([DisallowNull] UhtMetaDataKey obj)
 		{
 			return obj.Name.GetHashCode() ^ obj.Name.GetHashCode(StringComparison.OrdinalIgnoreCase);
@@ -106,8 +157,15 @@ namespace EpicGames.UHT.Types
 		/// </summary>
 		internal IUhtMetaDataKeyConversion? KeyConversion = null;
 
+		/// <summary>
+		/// Index for a meta data key associated with the owning object
+		/// </summary>
 		public const int INDEX_NONE = -1;
 
+		/// <summary>
+		/// Construct new meta data
+		/// </summary>
+		/// <param name="MessageSite">Message site for generating errors</param>
 		public UhtMetaData(IUhtMessageSite? MessageSite)
 		{
 			this.MessageSite = MessageSite;
@@ -179,6 +237,7 @@ namespace EpicGames.UHT.Types
 		/// Attempt to get the value associated with the key
 		/// </summary>
 		/// <param name="Name">Name of the meta data</param>
+		/// <param name="Value">Found value</param>
 		/// <returns>True if the key is found</returns>
 		public bool TryGetValue(string Name, [NotNullWhen(true)] out string? Value)
 		{
@@ -189,6 +248,7 @@ namespace EpicGames.UHT.Types
 		/// Attempt to get the value associated with the key including parent meta data
 		/// </summary>
 		/// <param name="Name">Name of the meta data</param>
+		/// <param name="Value">Found value</param>
 		/// <returns>True if the key is found</returns>
 		public bool TryGetValueHierarchical(string Name, [NotNullWhen(true)] out string? Value)
 		{
@@ -200,6 +260,7 @@ namespace EpicGames.UHT.Types
 		/// </summary>
 		/// <param name="Name">Name of the meta data</param>
 		/// <param name="NameIndex">Enumeration value index or -1 for the type's meta data</param>
+		/// <param name="Value">Found value</param>
 		/// <returns>True if the key is found</returns>
 		public bool TryGetValue(string Name, int NameIndex, [NotNullWhen(true)] out string? Value)
 		{
@@ -216,6 +277,7 @@ namespace EpicGames.UHT.Types
 		/// </summary>
 		/// <param name="Name">Name of the meta data</param>
 		/// <param name="NameIndex">Enumeration value index or -1 for the type's meta data</param>
+		/// <param name="Value">Found value</param>
 		/// <returns>True if the key is found</returns>
 		public bool TryGetValueHierarchical(string Name, int NameIndex, [NotNullWhen(true)] out string? Value)
 		{
@@ -230,6 +292,12 @@ namespace EpicGames.UHT.Types
 			return false;
 		}
 
+		/// <summary>
+		/// Get the string value of the given meta data key or the default value.
+		/// </summary>
+		/// <param name="Name">Name of the meta data key</param>
+		/// <param name="NameIndex">Index of the meta data key</param>
+		/// <returns>Meta data value or empty string if not found.</returns>
 		public string GetValueOrDefault(string Name, int NameIndex = INDEX_NONE)
 		{
 			string? Out;
@@ -240,6 +308,12 @@ namespace EpicGames.UHT.Types
 			return String.Empty;
 		}
 
+		/// <summary>
+		/// Get the string value of the given meta data searching the whole meta data chain.
+		/// </summary>
+		/// <param name="Name">Name of the meta data key</param>
+		/// <param name="NameIndex">Index of the meta data key</param>
+		/// <returns>Meta data value or empty string if not found.</returns>
 		public string GetValueOrDefaultHierarchical(string Name, int NameIndex = INDEX_NONE)
 		{
 			string? Out;
@@ -250,16 +324,34 @@ namespace EpicGames.UHT.Types
 			return string.Empty;
 		}
 
+		/// <summary>
+		/// Get the boolean value of the given meta data.
+		/// </summary>
+		/// <param name="Name">Name of the meta data key</param>
+		/// <param name="NameIndex">Index of the meta data key</param>
+		/// <returns>Boolean value or false if not found</returns>
 		public bool GetBoolean(string Name, int NameIndex = INDEX_NONE)
 		{
 			return UhtFCString.ToBool(GetValueOrDefault(Name, NameIndex));
 		}
 
+		/// <summary>
+		/// Get the boolean value of the given meta data searching the whole meta data chain.
+		/// </summary>
+		/// <param name="Name">Name of the meta data key</param>
+		/// <param name="NameIndex">Index of the meta data key</param>
+		/// <returns>Boolean value or false if not found</returns>
 		public bool GetBooleanHierarchical(string Name, int NameIndex = INDEX_NONE)
 		{
 			return UhtFCString.ToBool(GetValueOrDefaultHierarchical(Name, NameIndex));
 		}
 
+		/// <summary>
+		/// Get the double value of the given meta data.
+		/// </summary>
+		/// <param name="Name">Name of the meta data key</param>
+		/// <param name="NameIndex">Index of the meta data key</param>
+		/// <returns>Double value or zero if not found</returns>
 		public double GetDouble(string Name, int NameIndex = INDEX_NONE)
 		{
 			string Value = GetValueOrDefault(Name, NameIndex);
@@ -271,6 +363,12 @@ namespace EpicGames.UHT.Types
 			return Result;
 		}
 
+		/// <summary>
+		/// Get the string array value of the given meta data.
+		/// </summary>
+		/// <param name="Name">Name of the meta data key</param>
+		/// <param name="NameIndex">Index of the meta data key</param>
+		/// <returns>String array or null if not found</returns>
 		public string[]? GetStringArray(string Name, int NameIndex = INDEX_NONE)
 		{
 			string? Temp;
@@ -281,6 +379,12 @@ namespace EpicGames.UHT.Types
 			return null;
 		}
 
+		/// <summary>
+		/// Get the string array value of the given meta data searching the whole meta data chain.
+		/// </summary>
+		/// <param name="Name">Name of the meta data key</param>
+		/// <param name="NameIndex">Index of the meta data key</param>
+		/// <returns>String array or null if not found</returns>
 		public string[]? GetStringArrayHierarchical(string Name, int NameIndex = INDEX_NONE)
 		{
 			string? Temp;
@@ -291,11 +395,22 @@ namespace EpicGames.UHT.Types
 			return null;
 		}
 
+		/// <summary>
+		/// Add new meta data
+		/// </summary>
+		/// <param name="Name">Name of the meta data key</param>
+		/// <param name="Value">Value of the meta data</param>
 		public void Add(string Name, string Value)
 		{
 			Add(Name, INDEX_NONE, Value);
 		}
 
+		/// <summary>
+		/// Add new meta data
+		/// </summary>
+		/// <param name="Name">Name of the meta data key</param>
+		/// <param name="NameIndex">Index of the meta data key</param>
+		/// <param name="Value">Value of the meta data</param>
 		public void Add(string Name, int NameIndex, string Value)
 		{
 			string RemappedName;
@@ -309,16 +424,34 @@ namespace EpicGames.UHT.Types
 			GetDictionary()[new UhtMetaDataKey(Name, NameIndex)] = Value;
 		}
 
+		/// <summary>
+		/// Add new meta data
+		/// </summary>
+		/// <param name="Name">Name of the meta data key</param>
+		/// <param name="Value">Value of the meta data</param>
+		/// <param name="NameIndex">Index of the meta data key</param>
 		public void Add(string Name, bool Value, int NameIndex = INDEX_NONE)
 		{
 			Add(Name, NameIndex, Value ? "true" : "false");
 		}
 
+		/// <summary>
+		/// Add new meta data
+		/// </summary>
+		/// <param name="Name">Name of the meta data key</param>
+		/// <param name="Strings">Value of the meta data</param>
+		/// <param name="Separator">Separator to use to join the strings</param>
 		public void Add(string Name, List<string> Strings, char Separator = ' ')
 		{
 			Add(Name, string.Join(Separator, Strings));
 		}
 
+		/// <summary>
+		/// Add new meta data if there are strings in the value.
+		/// </summary>
+		/// <param name="Name">Name of the meta data key</param>
+		/// <param name="Strings">Value of the meta data</param>
+		/// <param name="Separator">Separator to use to join the strings</param>
 		public void AddIfNotEmpty(string Name, List<string> Strings, char Separator = ' ')
 		{
 			if (Strings.Count != 0)
@@ -327,6 +460,10 @@ namespace EpicGames.UHT.Types
 			}
 		}
 
+		/// <summary>
+		/// Add the meta data from another meta data block
+		/// </summary>
+		/// <param name="MetaData"></param>
 		public void Add(UhtMetaData MetaData)
 		{
 			if (MetaData.Dictionary != null)
@@ -339,6 +476,11 @@ namespace EpicGames.UHT.Types
 			}
 		}
 
+		/// <summary>
+		/// Remove the given meta data
+		/// </summary>
+		/// <param name="Name">Name of the meta data key</param>
+		/// <param name="NameIndex">Index of the meta data key</param>
 		public void Remove(string Name, int NameIndex = INDEX_NONE)
 		{
 			if (this.Dictionary != null)
@@ -347,20 +489,23 @@ namespace EpicGames.UHT.Types
 			}
 		}
 
-		private SortedList<UhtMetaDataKey, string> GetDictionary()
-		{
-			if (this.Dictionary == null)
-			{
-				this.Dictionary = new SortedList<UhtMetaDataKey, string>(Comparer);
-			}
-			return this.Dictionary;
-		}
-
+		/// <summary>
+		/// Given a key, return the full meta data name
+		/// </summary>
+		/// <param name="Key">Meta data key</param>
+		/// <returns>Full meta data name</returns>
 		public string GetKeyString(UhtMetaDataKey Key)
 		{
 			return GetKeyString(Key.Name, Key.Index);
 		}
 
+		/// <summary>
+		/// Given a key, return the full meta data name
+		/// </summary>
+		/// <param name="Name">Name of the meta data key</param>
+		/// <param name="NameIndex">Index of the meta data key</param>
+		/// <returns>Full meta data name</returns>
+		/// <exception cref="UhtIceException">Thrown if an index is supplied (not -1) and no key conversion interface is set</exception>
 		public string GetKeyString(string Name, int NameIndex)
 		{
 			if (NameIndex == INDEX_NONE)
@@ -398,5 +543,14 @@ namespace EpicGames.UHT.Types
 			}
 			return Out;
 		}
+
+		private SortedList<UhtMetaDataKey, string> GetDictionary()
+		{
+			if (this.Dictionary == null)
+			{
+				this.Dictionary = new SortedList<UhtMetaDataKey, string>(Comparer);
+			}
+			return this.Dictionary;
+		}
 	}
 }
diff --git a/Engine/Source/Programs/Shared/EpicGames.UHT/Types/UhtObject.cs b/Engine/Source/Programs/Shared/EpicGames.UHT/Types/UhtObject.cs
--- a/Engine/Source/Programs/Shared/EpicGames.UHT/Types/UhtObject.cs
+++ b/Engine/Source/Programs/Shared/EpicGames.UHT/Types/UhtObject.cs
@@ -6,10 +6,21 @@ using EpicGames.UHT.Utils;
 
 namespace EpicGames.UHT.Types
 {
+	/// <summary>
+	/// Represents a UObject in the engine
+	/// </summary>
 	[UhtEngineClass(Name = "Object")]
 	public abstract class UhtObject : UhtType
 	{
+
+		/// <summary>
+		/// Internal object flags.
+		/// </summary>
 		public EInternalObjectFlags InternalObjectFlags = EInternalObjectFlags.None;
+
+		/// <summary>
+		/// Unique index of the object
+		/// </summary>
 		public readonly int ObjectTypeIndex;
 
 		/// <summary>
@@ -21,11 +32,20 @@ namespace EpicGames.UHT.Types
 		/// <inheritdoc/>
 		public override string EngineClassName { get => "Object"; }
 
+		/// <summary>
+		/// Construct a new instance of the object
+		/// </summary>
+		/// <param name="Session">Session of the object</param>
 		protected UhtObject(UhtSession Session) : base(Session)
 		{
 			ObjectTypeIndex = this.Session.GetNextObjectTypeIndex();
 		}
 
+		/// <summary>
+		/// Construct a new instance of the object
+		/// </summary>
+		/// <param name="Outer">Outer object</param>
+		/// <param name="LineNumber">Line number where object is defined</param>
 		protected UhtObject(UhtType Outer, int LineNumber) : base(Outer, LineNumber)
 		{
 			ObjectTypeIndex = this.Session.GetNextObjectTypeIndex();
diff --git a/Engine/Source/Programs/Shared/EpicGames.UHT/Types/UhtPackage.cs b/Engine/Source/Programs/Shared/EpicGames.UHT/Types/UhtPackage.cs
--- a/Engine/Source/Programs/Shared/EpicGames.UHT/Types/UhtPackage.cs
+++ b/Engine/Source/Programs/Shared/EpicGames.UHT/Types/UhtPackage.cs
@@ -9,22 +9,47 @@ using UnrealBuildBase;
 
 namespace EpicGames.UHT.Types
 {
+
+	/// <summary>
+	/// Represents a UPackage in the engine
+	/// </summary>
 	[UhtEngineClass(Name = "Package")]
 	public class UhtPackage : UhtObject
 	{
+
+		/// <summary>
+		/// Unique index of the package
+		/// </summary>
 		public readonly int PackageTypeIndex;
-		public EPackageFlags PackageFlags { get; internal set; } = EPackageFlags.None;
-		public UHTManifest.Module Module { get; internal set; }
 
+		/// <summary>
+		/// Engine package flags
+		/// </summary>
+		public EPackageFlags PackageFlags { get; set; } = EPackageFlags.None;
+
+		/// <summary>
+		/// UHT module of the package (1 to 1 relationship)
+		/// </summary>
+		public UHTManifest.Module Module { get; set; }
+
+		/// <inheritdoc/>
 		[JsonIgnore]
 		public override UhtEngineType EngineType => UhtEngineType.Package;
 
+		/// <inheritdoc/>
+		public override string EngineClassName { get => "Package"; }
+
+		/// <inheritdoc/>
 		[JsonIgnore]
 		public override UhtPackage Package => this;
 
+		/// <inheritdoc/>
 		[JsonIgnore]
 		public override UhtHeaderFile HeaderFile => throw new NotImplementedException();
 
+		/// <summary>
+		/// True if the package is part of the engine
+		/// </summary>
 		[JsonIgnore]
 		public bool bIsPartOfEngine
 		{
@@ -52,15 +77,23 @@ namespace EpicGames.UHT.Types
 			}
 		}
 
+		/// <summary>
+		/// True if the package is a plugin
+		/// </summary>
 		[JsonIgnore]
 		public bool bIsPlugin => this.Module.BaseDirectory.Replace('\\', '/').Contains("/Plugins/");
 
-
+		/// <summary>
+		/// Short name of the package (without the /Script/)
+		/// </summary>
 		public string ShortName;
 
-		/// <inheritdoc/>
-		public override string EngineClassName { get => "Package"; }
-
+		/// <summary>
+		/// Construct a new instance of a package
+		/// </summary>
+		/// <param name="Session">Running session</param>
+		/// <param name="Module">Source module of the package</param>
+		/// <param name="PackageFlags">Assorted package flags</param>
 		public UhtPackage(UhtSession Session, UHTManifest.Module Module, EPackageFlags PackageFlags) : base(Session)
 		{
 			this.Module = Module;
diff --git a/Engine/Source/Programs/Shared/EpicGames.UHT/Types/UhtProperty.cs b/Engine/Source/Programs/Shared/EpicGames.UHT/Types/UhtProperty.cs
--- a/Engine/Source/Programs/Shared/EpicGames.UHT/Types/UhtProperty.cs
+++ b/Engine/Source/Programs/Shared/EpicGames.UHT/Types/UhtProperty.cs
@@ -199,7 +199,9 @@ namespace EpicGames.UHT.Types
 		/// </summary>
 		None,
 
-		// Native pointer specified
+		/// <summary>
+		/// Native pointer specified
+		/// </summary>
 		Native,
 	}
 
@@ -541,51 +543,192 @@ namespace EpicGames.UHT.Types
 		TypeText,
 	}
 
+	/// <summary>
+	/// Property member context provides extra context when formatting the member declaration and definitions.
+	/// </summary>
 	public interface IUhtPropertyMemberContext
 	{
+		/// <summary>
+		/// The outer/owning structure
+		/// </summary>
 		public UhtStruct OuterStruct { get; }
+
+		/// <summary>
+		/// The outer/owning structure source name.  In some cases, this can differ from the OuterStruct.SourceName
+		/// </summary>
 		public string OuterStructSourceName { get; }
+
+		/// <summary>
+		/// Name of the statics block for static definitions for the outer object
+		/// </summary>
 		public string StaticsName { get; }
+
+		/// <summary>
+		/// Prefix to apply to declaration names
+		/// </summary>
 		public string NamePrefix { get; }
+
+		/// <summary>
+		/// Suffix to add to the declaration name
+		/// </summary>
 		public string MetaDataSuffix { get; }
+
+		/// <summary>
+		/// Return the hash code for a given object
+		/// </summary>
+		/// <param name="Object">Object in question</param>
+		/// <returns>Hash code</returns>
 		public uint GetTypeHash(UhtObject Object);
+
+		/// <summary>
+		/// Return the singleton name for the given object
+		/// </summary>
+		/// <param name="Object">Object in question</param>
+		/// <param name="bRegistered">If true, the singleton that returns the registered object is returned.</param>
+		/// <returns>Singleton function name</returns>
 		public string GetSingletonName(UhtObject? Object, bool bRegistered);
 	}
 
+	/// <summary>
+	/// Property settings is a transient object used during the parsing of properties
+	/// </summary>
 	public class UhtPropertySettings
 	{
+
+		/// <summary>
+		/// Source name of the property
+		/// </summary>
 		public string SourceName = String.Empty;
+
+		/// <summary>
+		/// Engine name of the property
+		/// </summary>
 		public string EngineName = String.Empty;
+
+		/// <summary>
+		/// Property's meta data
+		/// </summary>
 		public UhtMetaData MetaData = UhtMetaData.Empty;
+
+		/// <summary>
+		/// Property outer object
+		/// </summary>
 		public UhtType Outer;
+
+		/// <summary>
+		/// Line number of the property declaration
+		/// </summary>
 		public int LineNumber;
+
+		/// <summary>
+		/// Property category
+		/// </summary>
 		public UhtPropertyCategory PropertyCategory;
+
+		/// <summary>
+		/// Engine property flags
+		/// </summary>
 		public EPropertyFlags PropertyFlags;
+
+		/// <summary>
+		/// Property flags not allowed by the context of the property parsing
+		/// </summary>
 		public EPropertyFlags DisallowPropertyFlags;
+
+		/// <summary>
+		/// UHT specified property flags
+		/// </summary>
 		public UhtPropertyExportFlags PropertyExportFlags;
+
+		/// <summary>
+		/// Allocator used for containers
+		/// </summary>
 		public UhtPropertyAllocator Allocator;
+
+		/// <summary>
+		/// Options for property parsing
+		/// </summary>
 		public UhtPropertyOptions Options;
+
+		/// <summary>
+		/// Property pointer type
+		/// </summary>
 		public UhtPointerType PointerType;
+
+		/// <summary>
+		/// Replication notify name
+		/// </summary>
 		public string? RepNotifyName;
+
+		/// <summary>
+		/// If set, the array size of the property
+		/// </summary>
 		public string? ArrayDimensions;
+
+		/// <summary>
+		/// Getter method
+		/// </summary>
 		public string? Setter;
+
+		/// <summary>
+		/// Setter method
+		/// </summary>
 		public string? Getter;
+
+		/// <summary>
+		/// Default value of the property
+		/// </summary>
 		public List<UhtToken>? DefaultValueTokens;
+
+		/// <summary>
+		/// If true, the property is a bit field
+		/// </summary>
 		public bool bIsBitfield;
 
+		/// <summary>
+		/// Construct a new, uninitialized version of the property settings
+		/// </summary>
 #pragma warning disable CS8618 // Non-nullable field must contain a non-null value when exiting constructor. Consider declaring as nullable.
 		public UhtPropertySettings()
 #pragma warning restore CS8618 // Non-nullable field must contain a non-null value when exiting constructor. Consider declaring as nullable.
 		{
 		}
 
-#pragma warning disable CS8618 // Non-nullable field must contain a non-null value when exiting constructor. Consider declaring as nullable.
-		public UhtPropertySettings(UhtType Outer, int LineNumber, UhtPropertyCategory PropertyCategory, EPropertyFlags DisallowPropertyFlags)
-#pragma warning restore CS8618 // Non-nullable field must contain a non-null value when exiting constructor. Consider declaring as nullable.
+		/// <summary>
+		/// Construct property settings based on the property settings for a parent container
+		/// </summary>
+		/// <param name="ParentPropertySettings">Parent container property</param>
+		/// <param name="SourceName">Name of the property</param>
+		/// <param name="MessageSite">Message site used to construct meta data object</param>
+		public UhtPropertySettings(UhtPropertySettings ParentPropertySettings, string SourceName, IUhtMessageSite MessageSite)
 		{
-			Reset(Outer, LineNumber, PropertyCategory, DisallowPropertyFlags);
+			this.SourceName = SourceName;
+			this.EngineName = SourceName;
+			this.MetaData = new UhtMetaData(MessageSite);
+			this.Outer = ParentPropertySettings.Outer;
+			this.LineNumber = ParentPropertySettings.LineNumber;
+			this.PropertyCategory = ParentPropertySettings.PropertyCategory;
+			this.PropertyFlags = ParentPropertySettings.PropertyFlags;
+			this.DisallowPropertyFlags = ParentPropertySettings.DisallowPropertyFlags;
+			this.PropertyExportFlags = UhtPropertyExportFlags.Public;
+			this.RepNotifyName = null;
+			this.Allocator = UhtPropertyAllocator.Default;
+			this.Options = ParentPropertySettings.Options;
+			this.PointerType = UhtPointerType.None;
+			this.ArrayDimensions = null;
+			this.DefaultValueTokens = null;
+			this.Setter = null;
+			this.Getter = null;
+			this.bIsBitfield = false;
 		}
 
+		/// <summary>
+		/// Reset the property settings.  Used on a cached property settings object
+		/// </summary>
+		/// <param name="Outer">Outer/owning type</param>
+		/// <param name="LineNumber">Line number of property</param>
+		/// <param name="PropertyCategory">Category of property</param>
+		/// <param name="DisallowPropertyFlags">Property flags that are not allowed</param>
 		public void Reset(UhtType Outer, int LineNumber, UhtPropertyCategory PropertyCategory, EPropertyFlags DisallowPropertyFlags)
 		{
 			this.SourceName = string.Empty;
@@ -607,13 +750,12 @@ namespace EpicGames.UHT.Types
 			this.bIsBitfield = false;
 		}
 
-#pragma warning disable CS8618 // Non-nullable field must contain a non-null value when exiting constructor. Consider declaring as nullable.
-		public UhtPropertySettings(UhtProperty Property, UhtPropertyOptions Options)
-#pragma warning restore CS8618 // Non-nullable field must contain a non-null value when exiting constructor. Consider declaring as nullable.
-		{
-			Reset(Property, Options);
-		}
-
+		/// <summary>
+		/// Reset property settings based on the given property.  Used to prepare a cached property settings for parsing.
+		/// </summary>
+		/// <param name="Property">Source property</param>
+		/// <param name="Options">Property options</param>
+		/// <exception cref="UhtIceException">Thrown if the input property doesn't have an outer</exception>
 		public void Reset(UhtProperty Property, UhtPropertyOptions Options)
 		{
 			if (Property.Outer == null)
@@ -639,31 +781,12 @@ namespace EpicGames.UHT.Types
 			this.Getter = Property.Getter;
 			this.bIsBitfield = Property.bIsBitfield;
 		}
-
-		public UhtPropertySettings(UhtPropertySettings ParentPropertySettings, string SourceName, IUhtMessageSite MessageSite)
-		{
-			this.SourceName = SourceName;
-			this.EngineName = SourceName;
-			this.MetaData = new UhtMetaData(MessageSite);
-			this.Outer = ParentPropertySettings.Outer;
-			this.LineNumber = ParentPropertySettings.LineNumber;
-			this.PropertyCategory = ParentPropertySettings.PropertyCategory;
-			this.PropertyFlags = ParentPropertySettings.PropertyFlags;
-			this.DisallowPropertyFlags = ParentPropertySettings.DisallowPropertyFlags;
-			this.PropertyExportFlags = UhtPropertyExportFlags.Public;
-			this.RepNotifyName = null;
-			this.Allocator = UhtPropertyAllocator.Default;
-			this.Options = ParentPropertySettings.Options;
-			this.PointerType = UhtPointerType.None;
-			this.ArrayDimensions = null;
-			this.DefaultValueTokens = null;
-			this.Setter = null;
-			this.Getter = null;
-			this.bIsBitfield = false;
-		}
 	}
 
-	[UhtEngineClass(Name = "Field", IsProperty = true)]
+	/// <summary>
+	/// Represents FProperty fields
+	/// </summary>
+	[UhtEngineClass(Name = "Field", IsProperty = true)] // This is here just so it is defined
 	[UhtEngineClass(Name = "Property", IsProperty = true)]
 	public abstract class UhtProperty : UhtType
 	{
@@ -724,27 +847,80 @@ namespace EpicGames.UHT.Types
 		protected virtual UhtPGetArgumentType PGetTypeArgument { get => UhtPGetArgumentType.None; }
 		#endregion
 
+		/// <summary>
+		/// Property category
+		/// </summary>
 		public UhtPropertyCategory PropertyCategory = UhtPropertyCategory.Member;
 
+		/// <summary>
+		/// Engine property flags
+		/// </summary>
 		[JsonConverter(typeof(JsonStringEnumConverter))]
 		public EPropertyFlags PropertyFlags { get; set; }
 
+		/// <summary>
+		/// Capabilities of the property.  Use caps system instead of testing for specific property types.
+		/// </summary>
 		[JsonIgnore]
 		[JsonConverter(typeof(JsonStringEnumConverter))]
 		public UhtPropertyCaps PropertyCaps { get; set; }
 
+		/// <summary>
+		/// Engine flags that are disallowed on this property
+		/// </summary>
 		public EPropertyFlags DisallowPropertyFlags = EPropertyFlags.None;
+
+		/// <summary>
+		/// UHT specified property flags
+		/// </summary>
 		public UhtPropertyExportFlags PropertyExportFlags = UhtPropertyExportFlags.Public;
+
+		/// <summary>
+		/// Reference type of the property
+		/// </summary>
 		public UhtPropertyRefQualifier RefQualifier = UhtPropertyRefQualifier.None;
+
+		/// <summary>
+		/// Pointer type of the property
+		/// </summary>
 		public UhtPointerType PointerType = UhtPointerType.None;
+
+		/// <summary>
+		/// Allocator to be used with containers
+		/// </summary>
 		public UhtPropertyAllocator Allocator = UhtPropertyAllocator.Default;
+
+		/// <summary>
+		/// Replication notify name
+		/// </summary>
 		public string? RepNotifyName = null;
+
+		/// <summary>
+		/// Fixed array size
+		/// </summary>
 		public string? ArrayDimensions = null;
+
+		/// <summary>
+		/// Property setter
+		/// </summary>
 		public string? Setter = null;
+
+		/// <summary>
+		/// Property getter
+		/// </summary>
 		public string? Getter = null;
+
+		/// <summary>
+		/// Default value of property
+		/// </summary>
 		public List<UhtToken>? DefaultValueTokens = null;
+
+		/// <summary>
+		/// If true, this property is a bit field
+		/// </summary>
 		public bool bIsBitfield = false;
 
+		///<inheritdoc/>
 		[JsonIgnore]
 		public override UhtEngineType EngineType => UhtEngineType.Property;
 
@@ -771,19 +947,33 @@ namespace EpicGames.UHT.Types
 			}
 		}
 
-		// Other attributes
+		/// <summary>
+		/// If true, the property is a fixed/static array
+		/// </summary>
 		[JsonIgnore]
 		public bool bIsStaticArray => !string.IsNullOrEmpty(this.ArrayDimensions);
 
+		/// <summary>
+		/// If true, the property is editor only
+		/// </summary>
 		[JsonIgnore]
 		public bool bIsEditorOnlyProperty => this.PropertyFlags.HasAnyFlags(EPropertyFlags.DevelopmentAssets);
 
+		/// <summary>
+		/// Construct a new property
+		/// </summary>
+		/// <param name="Outer">Outer type of the property</param>
+		/// <param name="LineNumber">Line number where property was declared</param>
 		public UhtProperty(UhtType Outer, int LineNumber) : base(Outer, LineNumber)
 		{
 			this.PropertyFlags = EPropertyFlags.None;
 			this.PropertyCaps = UhtPropertyCaps.CanBeContainerValue | UhtPropertyCaps.CanBeContainerKey | UhtPropertyCaps.CanHaveConfig;
 		}
 
+		/// <summary>
+		/// Construct a new property
+		/// </summary>
+		/// <param name="PropertySettings">Property settings from parsing</param>
 		public UhtProperty(UhtPropertySettings PropertySettings) : base(PropertySettings.Outer, PropertySettings.LineNumber, PropertySettings.MetaData)
 		{
 			this.SourceName = PropertySettings.SourceName;
@@ -901,7 +1091,7 @@ namespace EpicGames.UHT.Types
 		/// Append the required code to declare the property as a member
 		/// </summary>
 		/// <param name="Builder">Output builder</param>
-		/// <param name="NamePrefix">Prefix to the property name</param>
+		/// <param name="Context">Current context</param>
 		/// <param name="Name">Name of the property.  This is needed in some cases where the name in the declarations doesn't match the property name.</param>
 		/// <param name="NameSuffix">Suffix to the property name</param>
 		/// <param name="Tabs">Number of tabs prefix the line with</param>
@@ -912,7 +1102,7 @@ namespace EpicGames.UHT.Types
 		/// Append the required code to declare the property as a member
 		/// </summary>
 		/// <param name="Builder">Output builder</param>
-		/// <param name="NamePrefix">Prefix to the property name</param>
+		/// <param name="Context">Current context</param>
 		/// <param name="Name">Name of the property.  This is needed in some cases where the name in the declarations doesn't match the property name.</param>
 		/// <param name="NameSuffix">Suffix to the property name</param>
 		/// <param name="Tabs">Number of tabs prefix the line with</param>
@@ -1166,9 +1356,9 @@ namespace EpicGames.UHT.Types
 		public virtual void AppendObjectHashes(StringBuilder Builder, int StartingLength, IUhtPropertyMemberContext Context)
 		{
 		}
-#endregion
+		#endregion
 
-#region Parsing support
+		#region Parsing support
 		/// <summary>
 		/// Parse a default value for the property and return a sanitized string representation.
 		/// 
@@ -1178,9 +1368,10 @@ namespace EpicGames.UHT.Types
 		/// <param name="InnerDefaultValue">Sanitized representation of default value</param>
 		/// <returns>True if a default value was parsed.</returns>
 		public abstract bool SanitizeDefaultValue(IUhtTokenReader DefaultValueReader, StringBuilder InnerDefaultValue);
-#endregion
+		#endregion
 
-#region Resolution support
+		#region Resolution support
+		/// <inheritdoc/>
 		protected override bool ResolveSelf(UhtResolvePhase Phase)
 		{
 			bool bResult = base.ResolveSelf(Phase);
@@ -1292,9 +1483,10 @@ namespace EpicGames.UHT.Types
 		{
 			return false;
 		}
-#endregion
+		#endregion
 
-#region Validation support
+		#region Validation support
+		/// <inheritdoc/>
 		protected override UhtValidationOptions Validate(UhtValidationOptions Options)
 		{
 			Options = base.Validate(Options);
@@ -1352,7 +1544,13 @@ namespace EpicGames.UHT.Types
 			this.LogError($"{this.PropertyCategory.GetHintText()}: '{this.SourceName}' cannot be defined in '{this.Outer?.SourceName}' as it is already defined in scope '{Shadows.Outer?.SourceName}' (shadowing is not allowed)");
 		}
 
-		// Validate the property settings
+		/// <summary>
+		/// Validate the property settings
+		/// </summary>
+		/// <param name="OuterStruct">The outer structure for the property.  For properties inside containers, this will be the owning struct of the container</param>
+		/// <param name="OutermostProperty">Outer most property being validated.  For properties in containers, 
+		/// this will be the container property.  For properties outside of containers or the container itself, this will be the property.</param>
+		/// <param name="Options"></param>
 		public virtual void Validate(UhtStruct OuterStruct, UhtProperty OutermostProperty, UhtValidationOptions Options)
 		{
 			// Check for deprecation
@@ -1602,9 +1800,9 @@ namespace EpicGames.UHT.Types
 			ErrorType = null;
 			return false;
 		}
-#endregion
+		#endregion
 
-#region Reference support
+		#region Reference support
 		/// <inheritdoc/>
 		public override void CollectReferences(IUhtReferenceCollector Collector)
 		{
@@ -1615,33 +1813,70 @@ namespace EpicGames.UHT.Types
 			}
 		}
 
+		/// <summary>
+		/// Collect the references for the property.  This method is used by container properties to
+		/// collect the contained property's references.
+		/// </summary>
+		/// <param name="Collector">Reference collector</param>
+		/// <param name="bIsTemplateProperty">If true, this is a property in a container</param>
 		public virtual void CollectReferencesInternal(IUhtReferenceCollector Collector, bool bIsTemplateProperty)
 		{
 		}
 
+		/// <summary>
+		/// Return a string of all the forward declarations.
+		/// This method should be removed once C++ UHT is removed.
+		/// This information can be collected via CollectReferences.
+		/// </summary>
+		/// <returns></returns>
 		public virtual string? GetForwardDeclarations()
 		{
 			return null;
 		}
 
+		/// <summary>
+		/// Enumerate all reference types.  This method is used exclusively by FindNoExportStructsRecursive
+		/// </summary>
+		/// <returns>Type enumerator</returns>
 		public virtual IEnumerable<UhtType> EnumerateReferencedTypes()
 		{
 			return Enumerable.Empty<UhtType>();
 		}
-#endregion
+		#endregion
 
-#region RigVM support methods
+		#region RigVM support methods
+		/// <summary>
+		/// Return the RigVM type string.
+		/// </summary>
+		/// <param name="ParameterFlags">Parameter flags that can be updated for specific types</param>
+		/// <returns>RigVM type string</returns>
 		public abstract string? GetRigVMType(ref UhtRigVMParameterFlags ParameterFlags);
-#endregion
+		#endregion
 
-#region Helper methods
+		#region Helper methods
+		/// <summary>
+		/// Generate a new name suffix based on the current suffix and the new suffix
+		/// </summary>
+		/// <param name="OuterSuffix">Current suffix</param>
+		/// <param name="NewSuffix">Suffix to be added</param>
+		/// <returns>Combination of the two suffixes</returns>
 		protected string GetNameSuffix(string OuterSuffix, string NewSuffix)
 		{
 			return string.IsNullOrEmpty(OuterSuffix) ? NewSuffix : $"{OuterSuffix}{NewSuffix}";
 		}
 
+		/// <summary>
+		/// Test to see if the two properties are the same type
+		/// </summary>
+		/// <param name="Other">Other property to test</param>
+		/// <returns>True if the properies are the same type</returns>
 		public abstract bool IsSameType(UhtProperty Other);
 
+		/// <summary>
+		/// Test to see if the two properties are the same type and ConstParm/OutParm flags somewhat match
+		/// </summary>
+		/// <param name="Other">The other property to test</param>
+		/// <returns></returns>
 		public bool MatchesType(UhtProperty Other)
 		{
 			if (this.PropertyFlags.HasAnyFlags(EPropertyFlags.OutParm))
@@ -1658,77 +1893,183 @@ namespace EpicGames.UHT.Types
 			}
 			return IsSameType(Other);
 		}
-#endregion
+		#endregion
 	}
 
+	/// <summary>
+	/// Assorted StringBuilder extensions for properties
+	/// </summary>
 	public static class UhtPropertyStringBuilderExtensions
 	{
+
+		/// <summary>
+		/// Add the given property text
+		/// </summary>
+		/// <param name="Builder">Destination builder</param>
+		/// <param name="Property">Property in question</param>
+		/// <param name="TextType">Type of text to append</param>
+		/// <param name="bIsTemplateArgument">If true, this is a template argument property</param>
+		/// <returns>Destination builder</returns>
 		public static StringBuilder AppendPropertyText(this StringBuilder Builder, UhtProperty Property, UhtPropertyTextType TextType, bool bIsTemplateArgument = false)
 		{
 			return Property.AppendText(Builder, TextType, bIsTemplateArgument);
 		}
 
+		/// <summary>
+		/// Append the member declaration
+		/// </summary>
+		/// <param name="Builder">Destination builder</param>
+		/// <param name="Property">Property in question</param>
+		/// <param name="Context">Context of the property</param>
+		/// <param name="Name">Property name</param>
+		/// <param name="NameSuffix">Name suffix</param>
+		/// <param name="Tabs">Number of tabs in the formatting</param>
+		/// <returns>Destination builder</returns>
 		public static StringBuilder AppendMemberDecl(this StringBuilder Builder, UhtProperty Property, IUhtPropertyMemberContext Context, string Name, string NameSuffix, int Tabs)
 		{
 			return Property.AppendMemberDecl(Builder, Context, Name, NameSuffix, Tabs);
 		}
 
+		/// <summary>
+		/// Append the member definition
+		/// </summary>
+		/// <param name="Builder">Destination builder</param>
+		/// <param name="Property">Property in question</param>
+		/// <param name="Context">Context of the property</param>
+		/// <param name="Name">Property name</param>
+		/// <param name="NameSuffix">Name suffix</param>
+		/// <param name="Offset">Offset of the property in the parent</param>
+		/// <param name="Tabs">Number of tabs in the formatting</param>
+		/// <returns>Destination builder</returns>
 		public static StringBuilder AppendMemberDef(this StringBuilder Builder, UhtProperty Property, IUhtPropertyMemberContext Context, string Name, string NameSuffix, string? Offset, int Tabs)
 		{
 			return Property.AppendMemberDef(Builder, Context, Name, NameSuffix, Offset, Tabs);
 		}
 
+		/// <summary>
+		/// Append the member pointer
+		/// </summary>
+		/// <param name="Builder">Destination builder</param>
+		/// <param name="Property">Property in question</param>
+		/// <param name="Context">Context of the property</param>
+		/// <param name="Name">Property name</param>
+		/// <param name="NameSuffix">Name suffix</param>
+		/// <param name="Tabs">Number of tabs in the formatting</param>
+		/// <returns>Destination builder</returns>
 		public static StringBuilder AppendMemberPtr(this StringBuilder Builder, UhtProperty Property, IUhtPropertyMemberContext Context, string Name, string NameSuffix, int Tabs)
 		{
 			return Property.AppendMemberPtr(Builder, Context, Name, NameSuffix, Tabs);
 		}
 
+		/// <summary>
+		/// Append the full declaration including such things as 'const', '*', and '&'
+		/// </summary>
+		/// <param name="Builder">Destination builder</param>
+		/// <param name="Property">Property in question</param>
+		/// <param name="TextType">Type of text to append</param>
+		/// <param name="bSkipParameterName">If true, don't append the parameter name</param>
+		/// <returns>Destination builder</returns>
 		public static StringBuilder AppendFullDecl(this StringBuilder Builder, UhtProperty Property, UhtPropertyTextType TextType, bool bSkipParameterName)
 		{
 			return Property.AppendFullDecl(Builder, TextType, bSkipParameterName);
 		}
 
+		/// <summary>
+		/// Append the user facing declaration.
+		/// </summary>
+		/// <param name="Builder">Destination builder</param>
+		/// <param name="Property">Property in question</param>
+		/// <returns>Destination builder</returns>
 		public static StringBuilder AppendUserFacingDecl(this StringBuilder Builder, UhtProperty Property)
 		{
 			return Property.AppendText(Builder, UhtPropertyTextType.UserFacing);
 		}
 
+		/// <summary>
+		/// Append the function thunk parameter get
+		/// </summary>
+		/// <param name="Builder">Destination builder</param>
+		/// <param name="Property">Property in question</param>
+		/// <returns>Destination builder</returns>
 		public static StringBuilder AppendFunctionThunkParameterGet(this StringBuilder Builder, UhtProperty Property)
 		{
 			return Property.AppendFunctionThunkParameterGet(Builder);
 		}
 
+		/// <summary>
+		/// Append the function thunk parameter array type
+		/// </summary>
+		/// <param name="Builder">Destination builder</param>
+		/// <param name="Property">Property in question</param>
+		/// <returns>Destination builder</returns>
 		public static StringBuilder AppendFunctionThunkParameterArrayType(this StringBuilder Builder, UhtProperty Property)
 		{
 			return Property.AppendText(Builder, UhtPropertyTextType.FunctionThunkParameterArrayType);
 		}
 
+		/// <summary>
+		/// Append the function thunk parameter argument
+		/// </summary>
+		/// <param name="Builder">Destination builder</param>
+		/// <param name="Property">Property in question</param>
+		/// <returns>Destination builder</returns>
 		public static StringBuilder AppendFunctionThunkParameterArg(this StringBuilder Builder, UhtProperty Property)
 		{
 			return Property.AppendFunctionThunkParameterArg(Builder);
 		}
 
+		/// <summary>
+		/// Append the function thunk return
+		/// </summary>
+		/// <param name="Builder">Destination builder</param>
+		/// <param name="Property">Property in question</param>
+		/// <returns>Destination builder</returns>
 		public static StringBuilder AppendFunctionThunkReturn(this StringBuilder Builder, UhtProperty Property)
 		{
 			return Property.AppendText(Builder, UhtPropertyTextType.FunctionThunkReturn);
 		}
 
+		/// <summary>
+		/// Append the function thunk parameter name
+		/// </summary>
+		/// <param name="Builder">Destination builder</param>
+		/// <param name="Property">Property in question</param>
+		/// <returns>Destination builder</returns>
 		public static StringBuilder AppendFunctionThunkParameterName(this StringBuilder Builder, UhtProperty Property)
 		{
 			return Property.AppendFunctionThunkParameterName(Builder);
 		}
 
+		/// <summary>
+		/// Append the sparse type
+		/// </summary>
+		/// <param name="Builder">Destination builder</param>
+		/// <param name="Property">Property in question</param>
+		/// <returns>Destination builder</returns>
 		public static StringBuilder AppendSparse(this StringBuilder Builder, UhtProperty Property)
 		{
 			return Property.AppendText(Builder, UhtPropertyTextType.Sparse);
 		}
 
+		/// <summary>
+		/// Append the property's null constructor arg
+		/// </summary>
+		/// <param name="Builder">Destination builder</param>
+		/// <param name="Property">Property in question</param>
+		/// <param name="bIsInitializer">If true this is in an initializer context.</param>
+		/// <returns>Destination builder</returns>
 		public static StringBuilder AppendNullConstructorArg(this StringBuilder Builder, UhtProperty Property, bool bIsInitializer)
 		{
 			Property.AppendNullConstructorArg(Builder, bIsInitializer);
 			return Builder;
 		}
 
+		/// <summary>
+		/// Append the replication notify function or a 'nullptr'
+		/// </summary>
+		/// <param name="Builder">Destination builder</param>
+		/// <param name="Property">Property in question</param>
+		/// <returns>Destination builder</returns>
 		public static StringBuilder AppendNotifyFunc(this StringBuilder Builder, UhtProperty Property)
 		{
 			if (Property.RepNotifyName != null)
@@ -1742,12 +2083,25 @@ namespace EpicGames.UHT.Types
 			return Builder;
 		}
 
+		/// <summary>
+		/// Append the parameter flags
+		/// </summary>
+		/// <param name="Builder">Destination builder</param>
+		/// <param name="PropertyFlags">Property flags</param>
+		/// <returns>Destination builder</returns>
 		public static StringBuilder AppendFlags(this StringBuilder Builder, EPropertyFlags PropertyFlags)
 		{
 			PropertyFlags &= ~EPropertyFlags.ComputedFlags;
 			return Builder.Append("(EPropertyFlags)0x").AppendFormat("{0:x16}", (ulong)PropertyFlags);
 		}
 
+		/// <summary>
+		/// Append the property array dimensions as a CPP_ARRAY_DIM macro or '1' if not a fixed array.
+		/// </summary>
+		/// <param name="Builder">Destination builder</param>
+		/// <param name="Property">Property in question</param>
+		/// <param name="Context">Context of the property</param>
+		/// <returns>Destination builder</returns>
 		public static StringBuilder AppendArrayDim(this StringBuilder Builder, UhtProperty Property, IUhtPropertyMemberContext Context)
 		{
 			if (Property.ArrayDimensions != null)
@@ -1769,6 +2123,7 @@ namespace EpicGames.UHT.Types
 		/// <param name="StartingLength">Initial length of the builder prior to appending the hashes</param>
 		/// <param name="Context">Context used to lookup the hashes</param>
 		/// <param name="Object">Object being appended</param>
+		/// <returns>Destination builder</returns>
 		public static StringBuilder AppendObjectHash(this StringBuilder Builder, int StartingLength, UhtType ReferingType, IUhtPropertyMemberContext Context, UhtObject? Object)
 		{
 			if (Object == null)
@@ -1808,27 +2163,55 @@ namespace EpicGames.UHT.Types
 		/// <param name="ReferingType">Type asking for an object hash</param>
 		/// <param name="Context">Context used to lookup the hashes</param>
 		/// <param name="Object">Object being appended</param>
+		/// <returns>Destination builder</returns>
 		public static StringBuilder AppendObjectHash(this StringBuilder Builder, UhtType ReferingType, IUhtPropertyMemberContext Context, UhtObject? Object)
 		{
 			return Builder.AppendObjectHash(Builder.Length, ReferingType, Context, Object);
 		}
 
+		/// <summary>
+		/// Append the object hashes for all referenced objects
+		/// </summary>
+		/// <param name="Builder">Destination builder</param>
+		/// <param name="Property">Property in question</param>
+		/// <param name="Context">Context of the property</param>
+		/// <returns>Destination builder</returns>
 		public static StringBuilder AppendObjectHashes(this StringBuilder Builder, UhtProperty Property, IUhtPropertyMemberContext Context)
 		{
 			Property.AppendObjectHashes(Builder, Builder.Length, Context);
 			return Builder;
 		}
 
+		/// <summary>
+		/// Append the singleton name for the given type
+		/// </summary>
+		/// <param name="Builder">Destination builder</param>
+		/// <param name="Context">Context of the property</param>
+		/// <param name="Type">Type to append</param>
+		/// <param name="bRegistered">If true, append the registered type singleton.</param>
+		/// <returns>Destination builder</returns>
 		public static StringBuilder AppendSingletonName(this StringBuilder Builder, IUhtPropertyMemberContext Context, UhtObject? Type, bool bRegistered)
 		{
 			return Builder.Append(Context.GetSingletonName(Type, bRegistered));
 		}
 
+		/// <summary>
+		/// Append the getter wrapper name
+		/// </summary>
+		/// <param name="Builder">Destination builder</param>
+		/// <param name="Property">Property in question</param>
+		/// <returns>Destination builder</returns>
 		public static StringBuilder AppendPropertyGetterWrapperName(this StringBuilder Builder, UhtProperty Property)
 		{
 			return Builder.Append("Get").Append(Property.SourceName).Append("_WrapperImpl");
 		}
 
+		/// <summary>
+		/// Append the setter wrapper name
+		/// </summary>
+		/// <param name="Builder">Destination builder</param>
+		/// <param name="Property">Property in question</param>
+		/// <returns>Destination builder</returns>
 		public static StringBuilder AppendPropertySetterWrapperName(this StringBuilder Builder, UhtProperty Property)
 		{
 			return Builder.Append("Set").Append(Property.SourceName).Append("_WrapperImpl");
