Horde: More improvements to Bundles.

- Now does all string manipulations in UTF8 rather than converting to and fro with data received from Perforce
- Added support for reading a stream out of a CbWriter object. This allows pushing data directly to Horde Storage from the piecemeal 1kb buffers it allocates without having to concatenate everything into a contiguous array first, removing some big allocations and data copying on the common path for serializing objects into Horde Storage.
- Refactored the way that BundleNodeRef objects work so they don't need to track the Bundle that owns them, so we can add some serialization/deserialization tests without having to set up as much framework code.

diff --git a/Engine/Source/Programs/Horde/Horde.Agent/Commands/Bundles/CreateCommand.cs b/Engine/Source/Programs/Horde/Horde.Agent/Commands/Bundles/CreateCommand.cs
--- a/Engine/Source/Programs/Horde/Horde.Agent/Commands/Bundles/CreateCommand.cs
+++ b/Engine/Source/Programs/Horde/Horde.Agent/Commands/Bundles/CreateCommand.cs
@@ -72,7 +72,7 @@ namespace HordeAgent.Commands.Bundles
 
 			Bundle<DirectoryNode> NewBundle = Bundle.Create<DirectoryNode>(StorageClient, NamespaceId, new BundleOptions(), null);
 			await NewBundle.Root.CopyFromDirectoryAsync(InputDir.ToDirectoryInfo(), new ChunkingOptions(), Logger);
-			await NewBundle.WriteAsync(BucketId, RefId, false, DateTime.UtcNow);
+			await NewBundle.WriteAsync(BucketId, RefId, false);
 
 			return 0;
 		}
diff --git a/Engine/Source/Programs/Horde/Horde.Agent/Commands/Bundles/ExtractCommand.cs b/Engine/Source/Programs/Horde/Horde.Agent/Commands/Bundles/ExtractCommand.cs
--- a/Engine/Source/Programs/Horde/Horde.Agent/Commands/Bundles/ExtractCommand.cs
+++ b/Engine/Source/Programs/Horde/Horde.Agent/Commands/Bundles/ExtractCommand.cs
@@ -39,7 +39,7 @@ namespace HordeAgent.Commands.Bundles
 			using (IMemoryCache Cache = new MemoryCache(new MemoryCacheOptions { SizeLimit = 1024 * 1024 * 1000 }))
 			{
 				Bundle<DirectoryNode> NewBundle = await Bundle.ReadAsync<DirectoryNode>(StorageClient, NamespaceId, BucketId, RefId, new BundleOptions(), Cache);
-				await NewBundle.Root.CopyToDirectoryAsync(OutputDir.ToDirectoryInfo(), Logger);
+				await NewBundle.Root.CopyToDirectoryAsync(NewBundle, OutputDir.ToDirectoryInfo(), Logger);
 			}
 			return 0;
 		}
diff --git a/Engine/Source/Programs/Horde/Horde.Build/Commands/Bundles/CreateCommand.cs b/Engine/Source/Programs/Horde/Horde.Build/Commands/Bundles/CreateCommand.cs
--- a/Engine/Source/Programs/Horde/Horde.Build/Commands/Bundles/CreateCommand.cs
+++ b/Engine/Source/Programs/Horde/Horde.Build/Commands/Bundles/CreateCommand.cs
@@ -50,7 +50,7 @@ namespace HordeServer.Commands.Bundles
 
 			Bundle<DirectoryNode> NewBundle = Bundle.Create<DirectoryNode>(StorageClient, NamespaceId, new BundleOptions(), null);
 			await NewBundle.Root.CopyFromDirectoryAsync(InputDir.ToDirectoryInfo(), new ChunkingOptions(), Logger);
-			await NewBundle.WriteAsync(BucketId, RefId, false, DateTime.UtcNow);
+			await NewBundle.WriteAsync(BucketId, RefId, false);
 
 			return 0;
 		}
diff --git a/Engine/Source/Programs/Horde/Horde.Build/Commands/Bundles/ExtractCommand.cs b/Engine/Source/Programs/Horde/Horde.Build/Commands/Bundles/ExtractCommand.cs
--- a/Engine/Source/Programs/Horde/Horde.Build/Commands/Bundles/ExtractCommand.cs
+++ b/Engine/Source/Programs/Horde/Horde.Build/Commands/Bundles/ExtractCommand.cs
@@ -39,7 +39,7 @@ namespace HordeServer.Commands.Bundles
 			using (IMemoryCache Cache = new MemoryCache(new MemoryCacheOptions { SizeLimit = 1024 * 1024 * 1000 }))
 			{
 				Bundle<DirectoryNode> NewBundle = await Bundle.ReadAsync<DirectoryNode>(StorageClient, NamespaceId, BucketId, RefId, new BundleOptions(), Cache);
-				await NewBundle.Root.CopyToDirectoryAsync(OutputDir.ToDirectoryInfo(), Logger);
+				await NewBundle.Root.CopyToDirectoryAsync(NewBundle, OutputDir.ToDirectoryInfo(), Logger);
 			}
 			return 0;
 		}
diff --git a/Engine/Source/Programs/Horde/Horde.Build/Perforce/Impl/CommitService.cs b/Engine/Source/Programs/Horde/Horde.Build/Perforce/Impl/CommitService.cs
--- a/Engine/Source/Programs/Horde/Horde.Build/Perforce/Impl/CommitService.cs
+++ b/Engine/Source/Programs/Horde/Horde.Build/Perforce/Impl/CommitService.cs
@@ -36,6 +36,7 @@ using Horde.Build.Utilities;
 using EpicGames.Horde.Bundles;
 using EpicGames.Horde.Bundles.Nodes;
 using Microsoft.Extensions.Caching.Memory;
+using System.Runtime;
 
 namespace HordeServer.Commits.Impl
 {
@@ -72,7 +73,7 @@ namespace HordeServer.Commits.Impl
 		/// <summary>
 		/// Write Perforce stub files rather than including the actual content.
 		/// </summary>
-		public bool WriteStubFiles { get; set; } = true;
+		public bool WriteStubFiles { get; set; } = false;
 
 		/// <summary>
 		/// Namespace to store content replicated from Perforce
@@ -915,6 +916,8 @@ namespace HordeServer.Commits.Impl
 			Logger.LogInformation("Updating client {Client} to changelist {Change}", ClientInfo.Client.Name, Change);
 			ClientInfo.Change = -1;
 
+			Utf8String ClientRoot = new Utf8String(ClientInfo.Client.Root);
+
 			DirectoryNode Contents = CommitBundle.Root;
 			if (Options.WriteStubFiles)
 			{
@@ -922,13 +925,13 @@ namespace HordeServer.Commits.Impl
 				await foreach (PerforceResponse<SyncRecord> Record in Perforce.StreamCommandAsync<SyncRecord>("sync", new[] { "-k", $"{QueryPath}@{Change}" }, null, default))
 				{
 					SyncRecord SyncRecord = Record.Data;
-					if (!SyncRecord.Path.StartsWith(ClientInfo.Client.Root, StringComparison.Ordinal))
+					if (!SyncRecord.Path.StartsWith(ClientRoot))
 					{
 						throw new ArgumentException($"Unable to make path {ClientInfo.Client.Root} relative to client root {ClientInfo.Client.Root}");
 					}
 
-					string Path = SyncRecord.Path.Substring(ClientInfo.Client.Root.Length).Replace('\\', '/');
-					FileNode File = await Contents.CreateFileByPathAsync(Path, FileEntryFlags.PerforceDepotPathAndRevision);
+					Utf8String Path = SyncRecord.Path.Substring(ClientRoot.Length);
+					FileNode File = await Contents.AddFileByPathAsync(CommitBundle, Path, FileEntryFlags.PerforceDepotPathAndRevision);
 
 					byte[] Data = Encoding.UTF8.GetBytes($"{SyncRecord.DepotFile}#{SyncRecord.Revision}");
 					File.Append(Data, Options.Chunking);
@@ -942,21 +945,34 @@ namespace HordeServer.Commits.Impl
 			}
 			else
 			{
+				const long TrimInterval = 1024 * 1024 * 256;
+
+				long DataSize = 0;
+				long TrimDataSize = TrimInterval;
+
 				Dictionary<int, FileNode> Files = new Dictionary<int, FileNode>();
 				await foreach (PerforceResponse Response in Perforce.StreamCommandAsync("sync", Array.Empty<string>(), new string[] { $"{QueryPath}@{Change}" }, null, typeof(SyncRecord), true, default))
 				{
+					PerforceError? Error = Response.Error;
+					if (Error != null)
+					{
+						Logger.LogWarning("Perforce: {Message}", Error.Data);
+						continue;
+					}
+
 					PerforceIo? Io = Response.Io;
 					if (Io != null)
 					{
 						if (Io.Command == PerforceIoCommand.Open)
 						{
-							string Path = GetClientRelativePath(Io.Payload, ClientInfo.Client.Root);
-							Files[Io.File] = await Contents.CreateFileByPathAsync(Path, 0);
+							Utf8String Path = GetClientRelativePath(Io.Payload, ClientInfo.Client.Root);
+							Files[Io.File] = await Contents.AddFileByPathAsync(CommitBundle, Path, 0);
 						}
 						else if (Io.Command == PerforceIoCommand.Write)
 						{
 							FileNode File = Files[Io.File];
 							File.Append(Io.Payload, Options.Chunking);
+							DataSize += Io.Payload.Length;
 						}
 						else if (Io.Command == PerforceIoCommand.Close)
 						{
@@ -964,8 +980,16 @@ namespace HordeServer.Commits.Impl
 						}
 						else if (Io.Command == PerforceIoCommand.Unlink)
 						{
-							string Path = GetClientRelativePath(Io.Payload, ClientInfo.Client.Root);
-							await Contents.DeleteFileByPathAsync(Path);
+							Utf8String Path = GetClientRelativePath(Io.Payload, ClientInfo.Client.Root);
+							await Contents.DeleteFileByPathAsync(CommitBundle, Path);
+						}
+
+						if (DataSize > TrimDataSize)
+						{
+							Logger.LogInformation("Trimming working set after receiving {NumBytes:n0}mb...", DataSize / (1024 * 1024));
+							await CommitBundle.TrimAsync();
+							TrimDataSize = DataSize + TrimInterval;
+							Logger.LogInformation("Trimming complete. Next trim at {NextNumBytes:n0}mb.", TrimDataSize / (1024 * 1024));
 						}
 					}
 				}
@@ -976,7 +1000,7 @@ namespace HordeServer.Commits.Impl
 			// Return the new root object
 			RefId RefId = CommitTree.GetDefaultRefId(Change);
 			Logger.LogInformation("Writing ref {RefId} for {StreamId} change {Change}", RefId, Stream.Id, Change);
-			await CommitBundle.WriteAsync(GetStreamBucketId(Stream.Id), RefId, true, DateTime.UtcNow);
+			await CommitBundle.WriteAsync(GetStreamBucketId(Stream.Id), RefId, true);
 
 			return new CommitTree(Stream.Id, Change, RefId);
 		}
@@ -1000,17 +1024,21 @@ namespace HordeServer.Commits.Impl
 			}
 		}
 
-		static string GetClientRelativePath(ReadOnlyMemory<byte> Data, string ClientRoot)
+		static Utf8String GetClientRelativePath(ReadOnlyMemory<byte> Data, Utf8String ClientRoot)
 		{
-			ReadOnlySpan<byte> Span = Data.Span;
+			int Length = Data.Span.IndexOf((byte)0);
+			if (Length != -1)
+			{
+				Data = Data.Slice(0, Length);
+			}
 
-			string Path = Encoding.UTF8.GetString(Span.Slice(0, Span.IndexOf((byte)0)));
-			if (!Path.StartsWith(ClientRoot, StringComparison.Ordinal))
+			Utf8String Path = new Utf8String(Data);
+			if (!Path.StartsWith(ClientRoot, Utf8StringComparer.Ordinal))
 			{
 				throw new ArgumentException($"Unable to make path {Path} relative to client root {ClientRoot}");
 			}
 
-			return Path.Substring(ClientRoot.Length).Replace('\\', '/').TrimStart('/');
+			return Path.Substring(ClientRoot.Length).Clone();
 		}
 
 		class ReplicationClient
diff --git a/Engine/Source/Programs/Shared/EpicGames.Horde.Tests/BundleTests.cs b/Engine/Source/Programs/Shared/EpicGames.Horde.Tests/BundleTests.cs
--- a/Engine/Source/Programs/Shared/EpicGames.Horde.Tests/BundleTests.cs
+++ b/Engine/Source/Programs/Shared/EpicGames.Horde.Tests/BundleTests.cs
@@ -88,32 +88,26 @@ namespace HordeServerTests
 		[TestMethod]
 		public void SerializationTests()
 		{
-			DateTime UtcNow = DateTime.UtcNow;
-
 			List<BundleImport> Imports = new List<BundleImport>();
 			Imports.Add(new BundleImport(IoHash.Compute(Encoding.UTF8.GetBytes("import1")), 0, 123));
 			Imports.Add(new BundleImport(IoHash.Compute(Encoding.UTF8.GetBytes("import2")), 1, 456));
 
 			BundleObject OldObject = new BundleObject();
-			OldObject.CreationTimeUtc = DateTime.UtcNow;
-			OldObject.ImportObjects.Add(new BundleImportObject(UtcNow - TimeSpan.FromMinutes(5.0), new CbObjectAttachment(IoHash.Compute(new byte[] { 1, 2, 3 })), 12345, Imports));
-			OldObject.Exports.Add(new BundleExport(IoHash.Compute(Encoding.UTF8.GetBytes("export1")), 2, 5, new BundleCompressionPacket(0) { EncodedLength = 20, DecodedLength = 40 }, 0, 40, new int[] { 1, 2 }));
-			OldObject.Exports.Add(new BundleExport(IoHash.Compute(Encoding.UTF8.GetBytes("export2")), 3, 6, new BundleCompressionPacket(20) { EncodedLength = 10, DecodedLength = 20 }, 0, 20, new int[] { -1 }));
+			OldObject.ImportObjects.Add(new BundleImportObject(new CbObjectAttachment(IoHash.Compute(new byte[] { 1, 2, 3 })), 12345, Imports));
+			OldObject.Exports.Add(new BundleExport(IoHash.Compute(Encoding.UTF8.GetBytes("export1")), 2, new BundleCompressionPacket(0) { EncodedLength = 20, DecodedLength = 40 }, 0, 40, new int[] { 1, 2 }));
+			OldObject.Exports.Add(new BundleExport(IoHash.Compute(Encoding.UTF8.GetBytes("export2")), 3, new BundleCompressionPacket(20) { EncodedLength = 10, DecodedLength = 20 }, 0, 20, new int[] { -1 }));
 
 			CbObject SerializedObject = CbSerializer.Serialize(OldObject);
 			ReadOnlyMemory<byte> SerializedData = SerializedObject.GetView();
 
 			BundleObject NewObject = CbSerializer.Deserialize<BundleObject>(SerializedData);
 
-			Assert.AreEqual(OldObject.CreationTimeUtc, NewObject.CreationTimeUtc);
-
 			Assert.AreEqual(OldObject.ImportObjects.Count, NewObject.ImportObjects.Count);
 			for (int Idx = 0; Idx < OldObject.ImportObjects.Count; Idx++)
 			{
 				BundleImportObject OldObjectImport = OldObject.ImportObjects[Idx];
 				BundleImportObject NewObjectImport = NewObject.ImportObjects[Idx];
 
-				Assert.AreEqual(OldObjectImport.CreationTimeUtc, NewObjectImport.CreationTimeUtc);
 				Assert.AreEqual(OldObjectImport.Object, NewObjectImport.Object);
 				Assert.AreEqual(OldObjectImport.TotalCost, NewObjectImport.TotalCost);
 				Assert.AreEqual(OldObjectImport.Imports.Count, NewObjectImport.Imports.Count);
@@ -125,7 +119,7 @@ namespace HordeServerTests
 
 					Assert.AreEqual(OldImport.Hash, NewImport.Hash);
 					Assert.AreEqual(OldImport.Rank, NewImport.Rank);
-					Assert.AreEqual(OldImport.Cost, NewImport.Cost);
+					Assert.AreEqual(OldImport.Length, NewImport.Length);
 				}
 			}
 
@@ -137,7 +131,7 @@ namespace HordeServerTests
 
 				Assert.AreEqual(OldExport.Hash, NewExport.Hash);
 				Assert.AreEqual(OldExport.Rank, NewExport.Rank);
-				Assert.AreEqual(OldExport.Cost, NewExport.Cost);
+				Assert.AreEqual(OldExport.Length, NewExport.Length);
 				Assert.AreEqual(OldExport.Packet.Offset, NewExport.Packet.Offset);
 				Assert.AreEqual(OldExport.Packet.EncodedLength, NewExport.Packet.EncodedLength);
 				Assert.AreEqual(OldExport.Packet.DecodedLength, NewExport.Packet.DecodedLength);
@@ -155,7 +149,7 @@ namespace HordeServerTests
 			DirectoryNode Node2 = Node.AddDirectory("world");
 
 			RefId RefId = new RefId("testref");
-			await NewBundle.WriteAsync(BucketId, RefId, false, DateTime.MinValue);
+			await NewBundle.WriteAsync(BucketId, RefId, false);
 
 			// Should be stored inline
 			Assert.AreEqual(1, StorageClient.Refs.Count);
@@ -173,12 +167,12 @@ namespace HordeServerTests
 
 			Assert.AreEqual(0, NewBundle.Root.Files.Count);
 			Assert.AreEqual(1, NewBundle.Root.Directories.Count);
-			DirectoryNode? OutputNode = await NewBundle.Root.FindDirectoryAsync("hello");
+			DirectoryNode? OutputNode = await NewBundle.FindDirectoryAsync(NewBundle.Root, "hello");
 			Assert.IsNotNull(OutputNode);
 
 			Assert.AreEqual(0, OutputNode!.Files.Count);
 			Assert.AreEqual(1, OutputNode!.Directories.Count);
-			DirectoryNode? OutputNode2 = await OutputNode!.FindDirectoryAsync("world");
+			DirectoryNode? OutputNode2 = await NewBundle.FindDirectoryAsync(OutputNode, "world");
 			Assert.IsNotNull(OutputNode2);
 
 			Assert.AreEqual(0, OutputNode2!.Files.Count);
@@ -202,7 +196,7 @@ namespace HordeServerTests
 			Assert.AreEqual(0, StorageClient.Blobs.Count);
 
 			RefId RefId = new RefId("ref");
-			await Bundle.WriteAsync(BucketId, RefId, false, DateTime.UtcNow);
+			await Bundle.WriteAsync(BucketId, RefId, false);
 
 			Assert.AreEqual(1, StorageClient.Refs.Count);
 			Assert.AreEqual(1, StorageClient.Blobs.Count);
@@ -223,31 +217,29 @@ namespace HordeServerTests
 			DirectoryNode Node4 = Node3.AddDirectory("node4");
 
 			RefId RefId = new RefId("ref");
-			await InitialBundle.WriteAsync(BucketId, RefId, false, DateTime.UtcNow);
+			await InitialBundle.WriteAsync(BucketId, RefId, false);
 
 			Assert.AreEqual(1, StorageClient.Refs.Count);
 			Assert.AreEqual(4, StorageClient.Blobs.Count);
 
 			Bundle<DirectoryNode> NewBundle = await Bundle.ReadAsync<DirectoryNode>(StorageClient, NamespaceId, BucketId, RefId, Options, null);
 
-			DirectoryNode? NewNode1 = await NewBundle.Root.FindDirectoryAsync("node1");
+			DirectoryNode? NewNode1 = await NewBundle.FindDirectoryAsync(NewBundle.Root, "node1");
 			Assert.IsNotNull(NewNode1);
 
-			DirectoryNode? NewNode2 = await NewNode1!.FindDirectoryAsync("node2");
+			DirectoryNode? NewNode2 = await NewBundle.FindDirectoryAsync(NewNode1!, "node2");
 			Assert.IsNotNull(NewNode2);
 
-			DirectoryNode? NewNode3 = await NewNode2!.FindDirectoryAsync("node3");
+			DirectoryNode? NewNode3 = await NewBundle.FindDirectoryAsync(NewNode2!, "node3");
 			Assert.IsNotNull(NewNode3);
 
-			DirectoryNode? NewNode4 = await NewNode3!.FindDirectoryAsync("node4");
+			DirectoryNode? NewNode4 = await NewBundle.FindDirectoryAsync(NewNode3!, "node4");
 			Assert.IsNotNull(NewNode4);
 		}
 
 		[TestMethod]
 		public async Task CompactTest()
 		{
-			DateTime Time = DateTime.UtcNow;
-
 			BundleOptions Options = new BundleOptions();
 			Options.MaxBlobSize = 1024 * 1024;
 			Options.MaxInlineBlobSize = 1;
@@ -260,7 +252,7 @@ namespace HordeServerTests
 			DirectoryNode Node4 = Bundle.Root.AddDirectory("node4"); // same contents as node 3
 
 			RefId RefId1 = new RefId("ref1");
-			await Bundle.WriteAsync(BucketId, RefId1, false, Time);
+			await Bundle.WriteAsync(BucketId, RefId1, false);
 
 			Assert.AreEqual(1, StorageClient.Refs.Count);
 			Assert.AreEqual(1, StorageClient.Blobs.Count);
@@ -280,7 +272,7 @@ namespace HordeServerTests
 			Bundle.Root.DeleteDirectory("node1");
 
 			RefId RefId2 = new RefId("ref2");
-			await Bundle.WriteAsync(BucketId, RefId2, false, Time);
+			await Bundle.WriteAsync(BucketId, RefId2, false);
 
 			IRef Ref2 = StorageClient.Refs[(NamespaceId, BucketId, RefId2)];
 
@@ -295,7 +287,7 @@ namespace HordeServerTests
 
 			// Repack it and check that we make a new object
 			RefId RefId3 = new RefId("ref3");
-			await Bundle.WriteAsync(BucketId, RefId3, true, Time + TimeSpan.FromDays(7));
+			await Bundle.WriteAsync(BucketId, RefId3, true);
 
 			IRef Ref3 = StorageClient.Refs[(NamespaceId, BucketId, RefId3)];
 
@@ -325,7 +317,7 @@ namespace HordeServerTests
 			{
 				Node.Append(Data.AsMemory(Idx, 1), Options);
 
-				byte[] OutputData = await Node.ToByteArrayAsync();
+				byte[] OutputData = await Node.ToByteArrayAsync(null!);
 				Assert.IsTrue(Data.AsMemory(0, Idx + 1).Span.SequenceEqual(OutputData.AsSpan(0, Idx + 1)));
 			}
 		}
@@ -369,7 +361,7 @@ namespace HordeServerTests
 				Root.Append(Data, Options);
 			}
 
-			byte[] Result = await Root.ToByteArrayAsync();
+			byte[] Result = await Root.ToByteArrayAsync(null!);
 			Assert.AreEqual(NumIterations * Data.Length, Result.Length);
 
 			for (int Idx = 0; Idx < NumIterations; Idx++)
@@ -378,10 +370,10 @@ namespace HordeServerTests
 				Assert.IsTrue(SpanData.Span.SequenceEqual(Data));
 			}
 
-			await CheckSizes(Root, Options, true);
+			await CheckSizes(Bundle, Root, Options, true);
 		}
 
-		async Task CheckSizes(FileNode Node, ChunkingOptions Options, bool Rightmost)
+		async Task CheckSizes(Bundle Bundle, FileNode Node, ChunkingOptions Options, bool Rightmost)
 		{
 			if (Node.Depth == 0)
 			{
@@ -396,8 +388,8 @@ namespace HordeServerTests
 				int ChildCount = Node.Children.Count;
 				for (int Idx = 0; Idx < ChildCount; Idx++)
 				{
-					FileNode ChildNode = await Node.Children[Idx].GetAsync();
-					await CheckSizes(ChildNode, Options, Idx == ChildCount - 1);
+					FileNode ChildNode = await Bundle.GetAsync(Node.Children[Idx]);
+					await CheckSizes(Bundle, ChildNode, Options, Idx == ChildCount - 1);
 				}
 			}
 		}
@@ -423,7 +415,7 @@ namespace HordeServerTests
 						DirectoryNode NodeC = NodeB.AddDirectory($"{IdxC}");
 						for (int IdxD = 0; IdxD < 10; IdxD++)
 						{
-							FileNode File = NodeC.CreateFile($"{IdxD}", 0);
+							FileNode File = NodeC.AddFile($"{IdxD}", 0);
 							byte[] Data = Encoding.UTF8.GetBytes($"This is file {IdxA}/{IdxB}/{IdxC}/{IdxD}");
 							TotalLength += Data.Length;
 							File.Append(Data, new ChunkingOptions());
@@ -442,7 +434,7 @@ namespace HordeServerTests
 			Assert.IsTrue(StorageClient.Refs.Count == 0);
 
 			RefId RefId = new RefId("ref");
-			await Bundle.WriteAsync(BucketId, RefId, true, DateTime.UtcNow);
+			await Bundle.WriteAsync(BucketId, RefId, true);
 
 			Assert.AreEqual(TotalLength, Root.Length);
 
diff --git a/Engine/Source/Programs/Shared/EpicGames.Horde.Tests/NodeTests.cs b/Engine/Source/Programs/Shared/EpicGames.Horde.Tests/NodeTests.cs
new file mode 100644
--- /dev/null
+++ b/Engine/Source/Programs/Shared/EpicGames.Horde.Tests/NodeTests.cs
@@ -0,0 +1,108 @@
+// Copyright Epic Games, Inc. All Rights Reserved.
+
+using EpicGames.Core;
+using EpicGames.Horde.Bundles.Nodes;
+using Microsoft.VisualStudio.TestTools.UnitTesting;
+using System;
+using System.Buffers;
+using System.Collections.Generic;
+using System.Diagnostics.CodeAnalysis;
+using System.Linq;
+using System.Text;
+
+namespace EpicGames.Horde.Tests
+{
+	[TestClass]
+	public class NodeTests
+	{
+		[TestMethod]
+		public void TestDirectoryNode()
+		{
+			DirectoryNode Node = new DirectoryNode();
+			{
+				DirectoryNode Foo = Node.AddDirectory("Foo");
+
+				FileNode File1 = Foo.AddFile("first file.txt", 0);
+				File1.Append(Encoding.UTF8.GetBytes("first file file contents"), new ChunkingOptions());
+
+				DirectoryNode Bar = Node.AddDirectory("Bar");
+
+				FileNode File2 = Foo.AddFile("second file.txt", 0);
+				File2.Append(Encoding.UTF8.GetBytes("second file's contents"), new ChunkingOptions());
+			}
+
+			Dictionary<IoHash, ReadOnlyMemory<byte>> HashToData = new Dictionary<IoHash, ReadOnlyMemory<byte>>();
+			IoHash Hash = Serialize(Node, HashToData);
+
+			DirectoryNode Node2 = DeserializeDirectory(Hash, HashToData);
+			{
+				Assert.AreEqual(2, Node2.Directories.Count);
+				Assert.AreEqual(0, Node2.Files.Count);
+
+				Assert.IsTrue(Node2.TryGetDirectoryEntry("Foo", out DirectoryEntry? FooEntry));
+				Assert.IsNotNull(FooEntry);
+				DirectoryNode? Foo = FooEntry!.Node!;
+				Assert.IsNotNull(Foo);
+
+				Assert.AreEqual(0, Foo.Directories.Count);
+				Assert.AreEqual(2, Foo.Files.Count);
+
+				Assert.IsTrue(Foo.TryGetFileEntry("first file.txt", out FileEntry? File1));
+				Assert.IsTrue(Foo.TryGetFileEntry("second file.txt", out FileEntry? File2));
+
+				Assert.IsTrue(Node2.TryGetDirectoryEntry("Bar", out DirectoryEntry? BarEntry));
+				Assert.IsNotNull(BarEntry);
+				DirectoryNode? Bar = BarEntry!.Node!;
+				Assert.IsNotNull(Bar);
+			}
+		}
+
+		static IoHash Serialize(BundleNode Node, Dictionary<IoHash, ReadOnlyMemory<byte>> HashToData)
+		{
+			foreach (BundleNodeRef Ref in Node.GetReferences())
+			{
+				if (Ref.Node != null)
+				{
+					IoHash RefHash = Serialize(Ref.Node, HashToData);
+					Ref.MarkAsClean(RefHash);
+				}
+			}
+
+			byte[] Data = Node.Serialize().ToArray();
+			IoHash Hash = IoHash.Compute(Data);
+			HashToData[Hash] = Data;
+			return Hash;
+		}
+
+		static DirectoryNode DeserializeDirectory(IoHash Hash, Dictionary<IoHash, ReadOnlyMemory<byte>> HashToData)
+		{
+			ReadOnlyMemory<byte> Data = HashToData[Hash];
+			DirectoryNode Node = DirectoryNode.Deserialize(Data);
+
+			foreach (FileEntry FileEntry in Node.Files)
+			{
+				FileEntry.Node = DeserializeFile(FileEntry.Hash, HashToData);
+			}
+
+			foreach (DirectoryEntry DirectoryEntry in Node.Directories)
+			{
+				DirectoryEntry.Node = DeserializeDirectory(DirectoryEntry.Hash, HashToData);
+			}
+
+			return Node;
+		}
+
+		static FileNode DeserializeFile(IoHash Hash, Dictionary<IoHash, ReadOnlyMemory<byte>> HashToData)
+		{
+			ReadOnlyMemory<byte> Data = HashToData[Hash];
+			FileNode Node = FileNode.Deserialize(Data);
+
+			foreach (BundleNodeRef<FileNode> ChildRef in Node.Children)
+			{
+				ChildRef.Node = DeserializeFile(Hash, HashToData);
+			}
+
+			return Node;
+		}
+	}
+}
diff --git a/Engine/Source/Programs/Shared/EpicGames.Horde/Bundles/Bundle.cs b/Engine/Source/Programs/Shared/EpicGames.Horde/Bundles/Bundle.cs
--- a/Engine/Source/Programs/Shared/EpicGames.Horde/Bundles/Bundle.cs
+++ b/Engine/Source/Programs/Shared/EpicGames.Horde/Bundles/Bundle.cs
@@ -69,22 +69,27 @@ namespace EpicGames.Horde.Bundles
 		{
 			public IoHash Hash { get; }
 			public int Rank { get; }
-			public int Cost { get; }
+			public int Length { get; }
 
 			public BlobInfo? Blob;
 			public BundleExport? Export;
 			public ReadOnlySequence<byte> Data;
 			public NodeInfo[]? References;
 
-			public NodeInfo(IoHash Hash, int Rank, int Cost)
+			public NodeInfo(IoHash Hash, int Rank, int Length)
 			{
+				if (Length <= 0)
+				{
+					throw new ArgumentException("Node length must be greater than zero", nameof(Length));
+				}
+
 				this.Hash = Hash;
 				this.Rank = Rank;
-				this.Cost = Cost;
+				this.Length = Length;
 			}
 
 			public NodeInfo(BlobInfo Blob, BundleImport Import)
-				: this(Import.Hash, Import.Rank, Import.Cost)
+				: this(Import.Hash, Import.Rank, Import.Length)
 			{
 				this.Blob = Blob;
 			}
@@ -109,23 +114,21 @@ namespace EpicGames.Horde.Bundles
 		class BlobInfo
 		{
 			public IoHash Hash { get; }
-			public DateTime CreationTimeUtc { get; }
 			public int TotalCost { get; }
 
 			public bool Mounted { get; set; }
 			public List<NodeInfo>? LiveNodes { get; set; }
 
-			public BlobInfo(IoHash Hash, DateTime CreationTimeUtc, int TotalCost)
+			public BlobInfo(IoHash Hash, int TotalCost)
 			{
 				this.Hash = Hash;
-				this.CreationTimeUtc = CreationTimeUtc;
 				this.TotalCost = TotalCost;
 			}
- 		}
+		}
 
 		readonly IStorageClient StorageClient;
 		readonly NamespaceId NamespaceId;
-		
+
 		object RootCacheKey = new object();
 
 		/// <summary>
@@ -137,6 +140,9 @@ namespace EpicGames.Horde.Bundles
 		Dictionary<IoHash, BlobInfo> HashToBlob = new Dictionary<IoHash, BlobInfo>();
 		IMemoryCache? Cache;
 
+		CbWriter SharedWriter = new CbWriter();
+		byte[] SharedBlobBuffer = Array.Empty<byte>();
+
 		/// <summary>
 		/// Options for the bundle
 		/// </summary>
@@ -213,9 +219,9 @@ namespace EpicGames.Horde.Bundles
 			NewBundle.RegisterObject(null, Object, Object.Data);
 
 			IoHash RootHash = Object.Exports[Object.Exports.Count - 1].Hash;
-			BundleNodeRef<T> RootRef = new BundleNodeRef<T>(NewBundle, RootHash);
+			BundleNodeRef<T> RootRef = new BundleNodeRef<T>(RootHash);
 
-			((Bundle)NewBundle).Root = await RootRef.GetAsync();
+			((Bundle)NewBundle).Root = await NewBundle.GetAsync(RootRef);
 			return NewBundle;
 		}
 
@@ -254,7 +260,6 @@ namespace EpicGames.Horde.Bundles
 
 			foreach (BundleNodeRef ChildRef in Node.GetReferences())
 			{
-				ChildRef.Bundle ??= this;
 				ChildRef.ParentRef ??= IncomingRef;
 
 				if (ChildRef.Node != null && ChildRef.LastModifiedTime != 0)
@@ -286,22 +291,19 @@ namespace EpicGames.Horde.Bundles
 			List<NodeInfo> WriteNodes = new List<NodeInfo>();
 			FindNodesToWrite(Root, WriteNodes, new HashSet<NodeInfo>());
 
-			// Timestamp for any flushed objects
-			DateTime UtcNow = DateTime.UtcNow;
-
 			// Write them all to storage
 			int MinNodeIdx = 0;
 			int NextBlobCost = 0;
 			for (int Idx = 0; Idx < WriteNodes.Count; Idx++)
 			{
 				NodeInfo Node = WriteNodes[Idx];
-				if (Idx > MinNodeIdx && NextBlobCost + Node.Cost > Options.MaxBlobSize)
+				if (Idx > MinNodeIdx && NextBlobCost + Node.Length > Options.MaxBlobSize)
 				{
-					await WriteObjectAsync(WriteNodes.Slice(MinNodeIdx, Idx - MinNodeIdx), UtcNow);
+					await WriteObjectAsync(WriteNodes.Slice(MinNodeIdx, Idx - MinNodeIdx));
 					MinNodeIdx = Idx;
 					NextBlobCost = 0;
 				}
-				NextBlobCost += Node.Cost;
+				NextBlobCost += Node.Length;
 			}
 		}
 
@@ -332,6 +334,22 @@ namespace EpicGames.Horde.Bundles
 			}
 		}
 
+		/// <summary>
+		/// Gets the node referenced by a <see cref="BundleNodeRef{T}"/>, reading it from storage and deserializing it if necessary.
+		/// </summary>
+		/// <typeparam name="T">Type of node to return</typeparam>
+		/// <param name="Ref">The reference to get</param>
+		/// <returns>The deserialized node</returns>
+		public async ValueTask<T> GetAsync<T>(BundleNodeRef<T> Ref) where T : BundleNode
+		{
+			if (!Ref.MakeStrongRef())
+			{
+				ReadOnlySequence<byte> Data = await GetDataAsync(Ref.Hash);
+				Ref.Node = BundleNode.Deserialize<T>(Data);
+			}
+			return Ref.Node!;
+		}
+
 		/// <summary>
 		/// Gets data for the node with the given hash
 		/// </summary>
@@ -496,7 +514,7 @@ namespace EpicGames.Horde.Bundles
 			List<NodeInfo> Nodes = new List<NodeInfo>();
 			foreach (BundleExport Export in Object.Exports)
 			{
-				NodeInfo Node = FindOrAddNode(Export.Hash, Export.Rank, Export.Cost);
+				NodeInfo Node = FindOrAddNode(Export.Hash, Export.Rank, Export.Length);
 				Nodes.Add(Node);
 			}
 			foreach (BundleImportObject ImportObject in Object.ImportObjects)
@@ -504,12 +522,12 @@ namespace EpicGames.Horde.Bundles
 				BlobInfo? ImportBlob;
 				if (!HashToBlob.TryGetValue(ImportObject.Object.Hash, out ImportBlob))
 				{
-					ImportBlob = new BlobInfo(ImportObject.Object.Hash, ImportObject.CreationTimeUtc, ImportObject.TotalCost);
+					ImportBlob = new BlobInfo(ImportObject.Object.Hash, ImportObject.TotalCost);
 					HashToBlob[ImportObject.Object.Hash] = ImportBlob;
 				}
 				foreach (BundleImport Import in ImportObject.Imports)
 				{
-					NodeInfo Node = FindOrAddNode(Import.Hash, Import.Rank, Import.Cost);
+					NodeInfo Node = FindOrAddNode(Import.Hash, Import.Rank, Import.Length);
 					if(Node.Blob == null)
 					{
 						Node.SetImport(ImportBlob);
@@ -530,8 +548,7 @@ namespace EpicGames.Horde.Bundles
 		/// <param name="BucketId"></param>
 		/// <param name="RefId"></param>
 		/// <param name="Compact"></param>
-		/// <param name="UtcNow"></param>
-		public virtual async Task WriteAsync(BucketId BucketId, RefId RefId, bool Compact, DateTime UtcNow)
+		public virtual async Task WriteAsync(BucketId BucketId, RefId RefId, bool Compact)
 		{
 			// Completely flush the entire working set into NodeInfo objects
 			SerializeWorkingSet(0);
@@ -565,7 +582,7 @@ namespace EpicGames.Horde.Bundles
 				// Figure out which blobs need to be repacked
 				foreach (BlobInfo LiveBlob in LiveBlobs)
 				{
-					int LiveCost = LiveBlob.LiveNodes.Sum(x => x.Cost);
+					int LiveCost = LiveBlob.LiveNodes.Sum(x => x.Length);
 					int MinLiveCost = (int)(LiveBlob.TotalCost * Options.RepackRatio);
 
 					if (LiveCost < MinLiveCost)
@@ -604,27 +621,27 @@ namespace EpicGames.Horde.Bundles
 			NodeInfo[] WriteNodes = NewNodes.Where(x => x.Blob == null).OrderBy(x => x.Rank).ThenBy(x => x.Hash).ToArray();
 
 			int MinIdx = 0;
-			int NextBlobCost = WriteNodes[0].Cost;
+			int NextBlobCost = WriteNodes[0].Length;
 
 			for (int MaxIdx = 1; MaxIdx < WriteNodes.Length; MaxIdx++)
 			{
-				NextBlobCost += WriteNodes[MaxIdx].Cost;
+				NextBlobCost += WriteNodes[MaxIdx].Length;
 				if (NextBlobCost > Options.MaxBlobSize)
 				{
-					await WriteObjectAsync(WriteNodes.Slice(MinIdx, MaxIdx - MinIdx), UtcNow);
+					await WriteObjectAsync(WriteNodes.Slice(MinIdx, MaxIdx - MinIdx));
 					MinIdx = MaxIdx;
-					NextBlobCost = WriteNodes[MaxIdx].Cost;
+					NextBlobCost = WriteNodes[MaxIdx].Length;
 				}
 			}
 
 			if (NextBlobCost > Options.MaxInlineBlobSize && MinIdx + 1 < WriteNodes.Length)
 			{
-				await WriteObjectAsync(WriteNodes.Slice(MinIdx, WriteNodes.Length - 1 - MinIdx), UtcNow);
+				await WriteObjectAsync(WriteNodes.Slice(MinIdx, WriteNodes.Length - 1 - MinIdx));
 				MinIdx = WriteNodes.Length - 1;
 			}
 
 			// Write the final ref
-			await WriteRefAsync(WriteNodes.Slice(MinIdx), BucketId, RefId, UtcNow);
+			await WriteRefAsync(WriteNodes.Slice(MinIdx), BucketId, RefId);
 		}
 
 		private void FindModifiedLiveSet(NodeInfo Node, HashSet<NodeInfo> Nodes)
@@ -660,10 +677,11 @@ namespace EpicGames.Horde.Bundles
 			}
 		}
 
-		async Task WriteRefAsync(IReadOnlyList<NodeInfo> Nodes, BucketId BucketId, RefId RefId, DateTime UtcNow)
+		async Task WriteRefAsync(IReadOnlyList<NodeInfo> Nodes, BucketId BucketId, RefId RefId)
 		{
-			BundleObject Object = await CreateObjectAsync(Nodes, UtcNow);
+			BundleObject Object = await CreateObjectAsync(Nodes);
 			await StorageClient.SetRefAsync(NamespaceId, BucketId, RefId, Object);
+			Object.Data = Object.Data.ToArray(); // Duplicated since its lifetime persists beyond this function. See remarks in CreateObjectAsync.
 
 			foreach (NodeInfo Node in Nodes)
 			{
@@ -674,12 +692,24 @@ namespace EpicGames.Horde.Bundles
 			RootCacheKey = new object();
 		}
 
-		async Task WriteObjectAsync(IReadOnlyList<NodeInfo> Nodes, DateTime UtcNow)
+		ReadOnlyMemory<byte> EncodeObject(BundleObject Object)
 		{
-			BundleObject Object = await CreateObjectAsync(Nodes, UtcNow);
-			IoHash Hash = await StorageClient.WriteObjectAsync(NamespaceId, Object);
+			SharedWriter.Clear();
+			CbSerializer.Serialize(SharedWriter, Object);
+
+			int Size = SharedWriter.GetSize();
+			CreateFreeSpace(ref SharedBlobBuffer, 0, Math.Max(Size, Options.MaxBlobSize));
 
-			BlobInfo Blob = new BlobInfo(Hash, Object.CreationTimeUtc, Object.Exports.Sum(x => x.Cost));
+			SharedWriter.CopyTo(SharedBlobBuffer.AsSpan(0, Size));
+			return SharedBlobBuffer.AsMemory(0, Size);
+		}
+
+		async Task WriteObjectAsync(IReadOnlyList<NodeInfo> Nodes)
+		{
+			BundleObject Object = await CreateObjectAsync(Nodes);
+			IoHash Hash = await StorageClient.WriteBlobFromMemoryAsync(NamespaceId, EncodeObject(Object));
+
+			BlobInfo Blob = new BlobInfo(Hash, Object.Exports.Sum(x => x.Length));
 			foreach (NodeInfo Node in Nodes)
 			{
 				Node.Blob = Blob;
@@ -687,8 +717,17 @@ namespace EpicGames.Horde.Bundles
 			}
 		}
 
-		async Task<BundleObject> CreateObjectAsync(IReadOnlyList<NodeInfo> Nodes, DateTime UtcNow)
+		/// <summary>
+		/// Creates a BundleObject instance. 
+		///
+		/// WARNING: The <see cref="BundleObject.Data"/> member will be set to the value of <see cref="EncodedBuffer"/> on return, which will be
+		/// reused on a subsequent call. If the object needs to be persisted across the lifetime of other objects, this field must be duplicated.
+		/// </summary>
+		async Task<BundleObject> CreateObjectAsync(IReadOnlyList<NodeInfo> Nodes)
 		{
+			// Preallocate data in the encoded buffer to reduce fragmentation if we have to resize
+			CreateFreeSpace(ref EncodedBuffer, 0, Options.MaxBlobSize);
+
 			// Find node indices for all the export
 			Dictionary<NodeInfo, int> NodeToIndex = new Dictionary<NodeInfo, int>();
 			foreach (NodeInfo Node in Nodes)
@@ -698,7 +737,6 @@ namespace EpicGames.Horde.Bundles
 
 			// Create the new object
 			BundleObject Object = new BundleObject();
-			Object.CreationTimeUtc = UtcNow;
 
 			// Find all the imported nodes
 			HashSet<NodeInfo> ImportedNodes = new HashSet<NodeInfo>();
@@ -719,8 +757,8 @@ namespace EpicGames.Horde.Bundles
 					NodeToIndex[ImportNode] = NodeToIndex.Count;
 				}
 
-				List<BundleImport> Imports = ImportGroup.Select(x => new BundleImport(x.Hash, x.Rank, x.Cost)).ToList();
-				BundleImportObject ImportObject = new BundleImportObject(Blob.CreationTimeUtc, new CbObjectAttachment(Blob.Hash), Blob.TotalCost, Imports);
+				List<BundleImport> Imports = ImportGroup.Select(x => new BundleImport(x.Hash, x.Rank, x.Length)).ToList();
+				BundleImportObject ImportObject = new BundleImportObject(new CbObjectAttachment(Blob.Hash), Blob.TotalCost, Imports);
 				Object.ImportObjects.Add(ImportObject);
 			}
 
@@ -743,7 +781,7 @@ namespace EpicGames.Horde.Bundles
 
 				// Create the export for this node
 				int[] References = Node.References.Select(x => NodeToIndex[x]).ToArray();
-				Node.Export = new BundleExport(Node.Hash, Node.Rank, Node.Cost, Packet, Packet.DecodedLength, (int)NodeData.Length, References);
+				Node.Export = new BundleExport(Node.Hash, Node.Rank, Packet, Packet.DecodedLength, Node.Length, References);
 				Object.Exports.Add(Node.Export);
 
 				// Write out the new block
@@ -767,8 +805,7 @@ namespace EpicGames.Horde.Bundles
 			FlushPacket(BlockBuffer.AsMemory(0, BlockSize), Packet);
 
 			// Flush the data
-			Object.Data = EncodedBuffer.AsSpan(0, Packet.Offset + Packet.EncodedLength).ToArray();
-
+			Object.Data = EncodedBuffer.AsMemory(0, Packet.Offset + Packet.EncodedLength);
 			return Object;
 		}
 
diff --git a/Engine/Source/Programs/Shared/EpicGames.Horde/Bundles/BundleNode.cs b/Engine/Source/Programs/Shared/EpicGames.Horde/Bundles/BundleNode.cs
--- a/Engine/Source/Programs/Shared/EpicGames.Horde/Bundles/BundleNode.cs
+++ b/Engine/Source/Programs/Shared/EpicGames.Horde/Bundles/BundleNode.cs
@@ -16,6 +16,21 @@ namespace EpicGames.Horde.Bundles.Nodes
 	/// </summary>
 	public abstract class BundleNode
 	{
+		class DefaultSerializer<T> where T : BundleNode
+		{
+			public static readonly BundleNodeDeserializer<T> Instance = CreateInstance();
+
+			static BundleNodeDeserializer<T> CreateInstance()
+			{
+				BundleNodeDeserializerAttribute? Attribute = typeof(T).GetCustomAttribute<BundleNodeDeserializerAttribute>();
+				if (Attribute == null)
+				{
+					throw new InvalidOperationException($"No serializer is defined for {typeof(T).Name}");
+				}
+				return (BundleNodeDeserializer<T>)Activator.CreateInstance(Attribute.Type)!;
+			}
+		}
+
 		/// <summary>
 		/// Cached incoming reference to the owner of this node. When the contents of this node are flushed to storage, this is modified to become a weak reference
 		/// and hash value. When modified, it becomes a strong reference and zero hash.
@@ -33,6 +48,22 @@ namespace EpicGames.Horde.Bundles.Nodes
 		/// <returns>Serialized data for this node</returns>
 		public abstract ReadOnlySequence<byte> Serialize();
 
+		/// <summary>
+		/// Deserialize a node of a particular type, using the default serializer defined by an <see cref="BundleNodeDeserializerAttribute"/>.
+		/// </summary>
+		/// <typeparam name="T">The type to deserialize</typeparam>
+		/// <param name="Memory">Data to deserialize from</param>
+		/// <returns>New instance of the node</returns>
+		public static T Deserialize<T>(ReadOnlyMemory<byte> Memory) where T : BundleNode => DefaultSerializer<T>.Instance.Deserialize(Memory);
+
+		/// <summary>
+		/// Deserialize a node of a particular type, using the default serializer defined by an <see cref="BundleNodeDeserializerAttribute"/>.
+		/// </summary>
+		/// <typeparam name="T">The type to deserialize</typeparam>
+		/// <param name="Sequence">Data to deserialize from</param>
+		/// <returns>New instance of the node</returns>
+		public static T Deserialize<T>(ReadOnlySequence<byte> Sequence) where T : BundleNode => DefaultSerializer<T>.Instance.Deserialize(Sequence);
+
 		/// <summary>
 		/// Enumerates all the child references from this node
 		/// </summary>
@@ -45,11 +76,6 @@ namespace EpicGames.Horde.Bundles.Nodes
 	/// </summary>
 	public class BundleNodeRef
 	{
-		/// <summary>
-		/// The bundle that owns this reference. Used to realize hashes into full nodes, and set when we do a full scan of the tree from the bundle.
-		/// </summary>
-		internal Bundle? Bundle;
-
 		/// <summary>
 		/// Cached reference to the parent node
 		/// </summary>
@@ -58,12 +84,12 @@ namespace EpicGames.Horde.Bundles.Nodes
 		/// <summary>
 		/// Strong reference to the current node
 		/// </summary>
-		protected BundleNode? StrongRef;
+		internal BundleNode? StrongRef;
 
 		/// <summary>
 		/// Weak reference to the child node. Maintained after the object has been flushed.
 		/// </summary>
-		protected WeakReference<BundleNode>? WeakRef { get; set; }
+		private WeakReference<BundleNode>? WeakRef { get; set; }
 
 		/// <summary>
 		/// Last time that the node was modified
@@ -87,11 +113,9 @@ namespace EpicGames.Horde.Bundles.Nodes
 		/// <summary>
 		/// Creates a reference to a node with the given hash
 		/// </summary>
-		/// <param name="Bundle">The bundle that contains the reference</param>
 		/// <param name="Hash">Hash of the referenced node</param>
-		public BundleNodeRef(Bundle Bundle, IoHash Hash)
+		public BundleNodeRef(IoHash Hash)
 		{
-			this.Bundle = Bundle;
 			this.Hash = Hash;
 		}
 
@@ -114,6 +138,25 @@ namespace EpicGames.Horde.Bundles.Nodes
 			ParentRef = null;
 		}
 
+		/// <summary>
+		/// Converts the node in this reference from a weak to strong reference.
+		/// </summary>
+		/// <returns>True if the ref contains a valid node on return</returns>
+		internal bool MakeStrongRef()
+		{
+			if (StrongRef != null)
+			{
+				return true;
+			}
+			if (WeakRef != null && WeakRef.TryGetTarget(out StrongRef))
+			{
+				WeakRef = null;
+				return true;
+			}
+
+			return false;
+		}
+
 		/// <summary>
 		/// Marks this reference as dirty
 		/// </summary>
@@ -143,7 +186,7 @@ namespace EpicGames.Horde.Bundles.Nodes
 
 			if (StrongRef != null)
 			{
-				Collapse();
+				OnCollapse();
 
 				this.WeakRef = new WeakReference<BundleNode>(StrongRef);
 				this.StrongRef = null;
@@ -155,7 +198,7 @@ namespace EpicGames.Horde.Bundles.Nodes
 		/// <summary>
 		/// Callback for when a reference is collapsed to a hash value. At the time of calling, both the hash and node data will be valid.
 		/// </summary>
-		protected virtual void Collapse()
+		protected virtual void OnCollapse()
 		{
 		}
 	}
@@ -166,18 +209,6 @@ namespace EpicGames.Horde.Bundles.Nodes
 	/// <typeparam name="T">Type of the node</typeparam>
 	public class BundleNodeRef<T> : BundleNodeRef where T : BundleNode
 	{
-		static readonly BundleNodeDeserializer<T> Deserializer = CreateDeserializer();
-
-		static BundleNodeDeserializer<T> CreateDeserializer()
-		{
-			BundleNodeDeserializerAttribute? Attribute = typeof(T).GetCustomAttribute<BundleNodeDeserializerAttribute>();
-			if (Attribute == null)
-			{
-				throw new InvalidOperationException($"No serializer is defined for {typeof(T).Name}");
-			}
-			return (BundleNodeDeserializer<T>)Activator.CreateInstance(Attribute.Type)!;
-		}
-
 		/// <inheritdoc cref="BundleNodeRef.Node"/>
 		public new T? Node
 		{
@@ -188,9 +219,8 @@ namespace EpicGames.Horde.Bundles.Nodes
 		/// <summary>
 		/// Constructor
 		/// </summary>
-		/// <param name="Bundle">The bundle that contains the reference</param>
 		/// <param name="Hash">Hash of the referenced node</param>
-		public BundleNodeRef(Bundle Bundle, IoHash Hash) : base(Bundle, Hash)
+		public BundleNodeRef(IoHash Hash) : base(Hash)
 		{
 		}
 
@@ -201,26 +231,6 @@ namespace EpicGames.Horde.Bundles.Nodes
 		public BundleNodeRef(T Node) : base(Node)
 		{
 		}
-
-		/// <summary>
-		/// Gets the referenced node, loading it from storage if necessary
-		/// </summary>
-		/// <returns>The parsed node</returns>
-		public async ValueTask<T> GetAsync()
-		{
-			if (StrongRef == null)
-			{
-				if (WeakRef != null && WeakRef.TryGetTarget(out StrongRef))
-				{
-					WeakRef = null;
-				}
-				else
-				{
-					StrongRef = Deserializer.Deserialize(Bundle!, await Bundle!.GetDataAsync(Hash));
-				}
-			}
-			return (T)StrongRef;
-		}
 	}
 
 	/// <summary>
@@ -248,18 +258,16 @@ namespace EpicGames.Horde.Bundles.Nodes
 		/// <summary>
 		/// Deserializes data from the given bundle
 		/// </summary>
-		/// <param name="Bundle">The bundle containing the data</param>
 		/// <param name="Data">Data to deserialize</param>
 		/// <returns>New node parsed from the data</returns>
-		public abstract T Deserialize(Bundle Bundle, ReadOnlyMemory<byte> Data);
+		public abstract T Deserialize(ReadOnlyMemory<byte> Data);
 
 		/// <summary>
 		/// Deserializes data from the given bundle
 		/// </summary>
-		/// <param name="Bundle">The bundle containing the data</param>
 		/// <param name="Data">Data to deserialize</param>
 		/// <returns>New node parsed from the data</returns>
-		public T Deserialize(Bundle Bundle, ReadOnlySequence<byte> Data)
+		public T Deserialize(ReadOnlySequence<byte> Data)
 		{
 			ReadOnlyMemory<byte> Memory;
 			if (Data.IsSingleSegment)
@@ -270,7 +278,7 @@ namespace EpicGames.Horde.Bundles.Nodes
 			{
 				Memory = Data.ToArray();
 			}
-			return Deserialize(Bundle, Memory);
+			return Deserialize(Memory);
 		}
 	}
 }
diff --git a/Engine/Source/Programs/Shared/EpicGames.Horde/Bundles/BundleObject.cs b/Engine/Source/Programs/Shared/EpicGames.Horde/Bundles/BundleObject.cs
--- a/Engine/Source/Programs/Shared/EpicGames.Horde/Bundles/BundleObject.cs
+++ b/Engine/Source/Programs/Shared/EpicGames.Horde/Bundles/BundleObject.cs
@@ -37,12 +37,6 @@ namespace EpicGames.Horde.Bundles
 		[CbField("schema")]
 		public int Schema { get; set; } = CurrentSchema;
 
-		/// <summary>
-		/// Time that this object was minted. Used to evaluate how likely it will be that a client will already have it.
-		/// </summary>
-		[CbField("time")]
-		public DateTime CreationTimeUtc { get; set; }
-
 		/// <summary>
 		/// Other objects that are referenced.
 		/// </summary>
@@ -69,12 +63,6 @@ namespace EpicGames.Horde.Bundles
 	[DebuggerDisplay("{Object}")]
 	public class BundleImportObject
 	{
-		/// <summary>
-		/// Time that the blob was created.
-		/// </summary>
-		[CbField("time")]
-		public DateTime CreationTimeUtc { get; set; }
-
 		/// <summary>
 		/// Imported object.
 		/// </summary>
@@ -104,9 +92,8 @@ namespace EpicGames.Horde.Bundles
 		/// <summary>
 		/// Constructor
 		/// </summary>
-		public BundleImportObject(DateTime CreationTimeUtc, CbObjectAttachment Object, int TotalCost, List<BundleImport> Imports)
+		public BundleImportObject(CbObjectAttachment Object, int TotalCost, List<BundleImport> Imports)
 		{
-			this.CreationTimeUtc = CreationTimeUtc;
 			this.Object = Object;
 			this.TotalCost = TotalCost;
 			this.Imports = Imports;
@@ -131,16 +118,21 @@ namespace EpicGames.Horde.Bundles
 		/// <summary>
 		/// Decompressed size of this node
 		/// </summary>
-		public int Cost { get; }
+		public int Length { get; }
 
 		/// <summary>
 		/// Constructor
 		/// </summary>
-		public BundleImport(IoHash Hash, int Rank, int Cost)
+		public BundleImport(IoHash Hash, int Rank, int Length)
 		{
+			if (Length <= 0)
+			{
+				throw new ArgumentException("Length must be greater than zero", nameof(Length));
+			}
+
 			this.Hash = Hash;
 			this.Rank = Rank;
-			this.Cost = Cost;
+			this.Length = Length;
 		}
 	}
 
@@ -171,7 +163,7 @@ namespace EpicGames.Horde.Bundles
 
 		static byte[] Serialize(List<BundleImport> Imports)
 		{
-			byte[] Data = new byte[Imports.Sum(x => IoHash.NumBytes + VarInt.Measure(x.Rank) + VarInt.Measure(x.Cost))];
+			byte[] Data = new byte[Imports.Sum(x => IoHash.NumBytes + VarInt.Measure(x.Rank) + VarInt.Measure(x.Length))];
 
 			Span<byte> Span = Data;
 			foreach (BundleImport Import in Imports)
@@ -182,7 +174,7 @@ namespace EpicGames.Horde.Bundles
 				int RankBytes = VarInt.Write(Span, Import.Rank);
 				Span = Span.Slice(RankBytes);
 
-				int CostBytes = VarInt.Write(Span, Import.Cost);
+				int CostBytes = VarInt.Write(Span, Import.Length);
 				Span = Span.Slice(CostBytes);
 			}
 
@@ -235,11 +227,6 @@ namespace EpicGames.Horde.Bundles
 		/// </summary>
 		public int Offset { get; }
 
-		/// <summary>
-		/// Length of the node within the decompressed packet
-		/// </summary>
-		public int Length { get; }
-
 		/// <summary>
 		/// References to other nodes.
 		/// </summary>
@@ -248,12 +235,11 @@ namespace EpicGames.Horde.Bundles
 		/// <summary>
 		/// Constructor
 		/// </summary>
-		public BundleExport(IoHash Hash, int Rank, int Cost, BundleCompressionPacket Packet, int Offset, int Length, int[] References)
-			: base(Hash, Rank, Cost)
+		public BundleExport(IoHash Hash, int Rank, BundleCompressionPacket Packet, int Offset, int Length, int[] References)
+			: base(Hash, Rank, Length)
 		{
 			this.Packet = Packet;
 			this.Offset = Offset;
-			this.Length = Length;
 			this.References = References;
 		}
 	}
@@ -290,9 +276,6 @@ namespace EpicGames.Horde.Bundles
 					int Rank = (int)VarInt.Read(Span, out int RankBytes);
 					Span = Span.Slice(RankBytes);
 
-					int Cost = (int)VarInt.Read(Span, out int CostBytes);
-					Span = Span.Slice(CostBytes);
-
 					int Length = (int)VarInt.Read(Span, out int LengthBytes);
 					Span = Span.Slice(LengthBytes);
 
@@ -306,7 +289,7 @@ namespace EpicGames.Horde.Bundles
 						Span = Span.Slice(ReferenceBytes);
 					}
 
-					Exports.Add(new BundleExport(Hash, Rank, Cost, Packet, Offset, Length, References));
+					Exports.Add(new BundleExport(Hash, Rank, Packet, Offset, Length, References));
 					Offset += Length;
 				}
 			}
@@ -324,6 +307,8 @@ namespace EpicGames.Horde.Bundles
 			BundleCompressionPacket? PrevPacket = null;
 			foreach (BundleExport Export in Exports)
 			{
+				Debug.Assert(Export.Length > 0);
+
 				BundleCompressionPacket Packet = Export.Packet;
 				if (PrevPacket == null || Packet.Offset != PrevPacket.Offset)
 				{
@@ -344,9 +329,6 @@ namespace EpicGames.Horde.Bundles
 				int RankBytes = VarInt.Write(Span, Export.Rank);
 				Span = Span.Slice(RankBytes);
 
-				int CostBytes = VarInt.Write(Span, Export.Cost);
-				Span = Span.Slice(CostBytes);
-
 				int LengthBytes = VarInt.Write(Span, Export.Length);
 				Span = Span.Slice(LengthBytes);
 
@@ -378,7 +360,8 @@ namespace EpicGames.Horde.Bundles
 					Length += VarInt.Measure(Packet.EncodedLength) + VarInt.Measure(Packet.DecodedLength);
 					PrevPacket = Packet;
 				}
-				Length += IoHash.NumBytes + VarInt.Measure(Export.Rank) + VarInt.Measure(Export.Cost) + VarInt.Measure(Export.Length) + VarInt.Measure(Export.References.Length) + Export.References.Sum(x => VarInt.Measure(x));
+
+				Length += IoHash.NumBytes + VarInt.Measure(Export.Rank) + VarInt.Measure(Export.Length) + VarInt.Measure(Export.References.Length) + Export.References.Sum(x => VarInt.Measure(x));
 			}
 
 			return Length;
diff --git a/Engine/Source/Programs/Shared/EpicGames.Horde/Bundles/BundleTools.cs b/Engine/Source/Programs/Shared/EpicGames.Horde/Bundles/BundleTools.cs
--- a/Engine/Source/Programs/Shared/EpicGames.Horde/Bundles/BundleTools.cs
+++ b/Engine/Source/Programs/Shared/EpicGames.Horde/Bundles/BundleTools.cs
@@ -185,7 +185,7 @@ namespace EpicGames.Horde.Bundles
 				NodeInfo? Node;
 				if (!Blob.Nodes.TryGetValue(Export.Hash, out Node))
 				{
-					Node = new NodeInfo(Export.Hash, Export.Rank, Export.Cost, Blob);
+					Node = new NodeInfo(Export.Hash, Export.Rank, Export.Length, Blob);
 					Blob.Nodes.Add(Export.Hash, Node);
 				}
 				Blob.References.Add(Node);
@@ -206,7 +206,7 @@ namespace EpicGames.Horde.Bundles
 					NodeInfo? Node;
 					if (!ImportBlob.Nodes.TryGetValue(Import.Hash, out Node))
 					{
-						Node = new NodeInfo(Import.Hash, Import.Rank, Import.Cost, ImportBlob);
+						Node = new NodeInfo(Import.Hash, Import.Rank, Import.Length, ImportBlob);
 						ImportBlob.Nodes.Add(Import.Hash, Node);
 					}
 					Blob.References.Add(Node);
@@ -264,7 +264,6 @@ namespace EpicGames.Horde.Bundles
 			using StreamWriter Writer = new StreamWriter(FileReference.Open(File, FileMode.Create));
 
 			Writer.WriteLine("schema: {0}", Object.Schema);
-			Writer.WriteLine("time: \"{0}\"", Object.CreationTimeUtc.ToString("u"));
 
 			List<IoHash> References = new List<IoHash>(Object.Exports.Select(x => x.Hash));
 			if (Object.ImportObjects.Count > 0)
@@ -278,7 +277,7 @@ namespace EpicGames.Horde.Bundles
 					{
 						Writer.WriteLine($"  - hash: \"{Import.Hash}\"");
 						Writer.WriteLine($"    rank: \"{Import.Rank}\"");
-						Writer.WriteLine($"    cost: \"{Import.Cost}\"");
+						Writer.WriteLine($"    cost: \"{Import.Length}\"");
 						Writer.WriteLine("");
 						References.Add(Import.Hash);
 					}
diff --git a/Engine/Source/Programs/Shared/EpicGames.Horde/Bundles/Nodes/DirectoryNode.cs b/Engine/Source/Programs/Shared/EpicGames.Horde/Bundles/Nodes/DirectoryNode.cs
--- a/Engine/Source/Programs/Shared/EpicGames.Horde/Bundles/Nodes/DirectoryNode.cs
+++ b/Engine/Source/Programs/Shared/EpicGames.Horde/Bundles/Nodes/DirectoryNode.cs
@@ -84,8 +84,8 @@ namespace EpicGames.Horde.Bundles.Nodes
 		/// <summary>
 		/// Constructor
 		/// </summary>
-		public FileEntry(Bundle Bundle, Utf8String Name, FileEntryFlags Flags, long Length, IoHash Hash)
-			: base(Bundle, Hash)
+		public FileEntry(Utf8String Name, FileEntryFlags Flags, long Length, IoHash Hash)
+			: base(Hash)
 		{
 			this.Name = Name;
 			this.Flags = Flags;
@@ -93,7 +93,7 @@ namespace EpicGames.Horde.Bundles.Nodes
 		}
 
 		/// <inheritdoc/>
-		protected override void Collapse()
+		protected override void OnCollapse()
 		{
 			CachedLength = Node!.Length;
 		}
@@ -131,15 +131,15 @@ namespace EpicGames.Horde.Bundles.Nodes
 		/// <summary>
 		/// Constructor
 		/// </summary>
-		public DirectoryEntry(Bundle Bundle, Utf8String Name, long Length, IoHash Hash)
-			: base(Bundle, Hash)
+		public DirectoryEntry(Utf8String Name, long Length, IoHash Hash)
+			: base(Hash)
 		{
 			this.Name = Name;
 			this.CachedLength = Length;
 		}
 
 		/// <inheritdoc/>
-		protected override void Collapse()
+		protected override void OnCollapse()
 		{
 			CachedLength = Node!.Length;
 		}
@@ -199,7 +199,7 @@ namespace EpicGames.Horde.Bundles.Nodes
 		/// <param name="Name">Name of the new directory</param>
 		/// <param name="Flags">Flags for the new file</param>
 		/// <returns>The new directory object</returns>
-		public FileNode CreateFile(Utf8String Name, FileEntryFlags Flags)
+		public FileNode AddFile(Utf8String Name, FileEntryFlags Flags)
 		{
 			if(TryGetDirectoryEntry(Name, out _))
 			{
@@ -218,29 +218,52 @@ namespace EpicGames.Horde.Bundles.Nodes
 		/// <summary>
 		/// Finds or adds a file with the given path
 		/// </summary>
+		/// <param name="Bundle">Bundle that contains the node</param>
 		/// <param name="Path">Path to the file</param>
 		/// <param name="Flags">Flags for the new file</param>
 		/// <returns>The new directory object</returns>
-		public async ValueTask<FileNode> CreateFileByPathAsync(Utf8String Path, FileEntryFlags Flags)
+		public async ValueTask<FileNode> AddFileByPathAsync(Bundle Bundle, Utf8String Path, FileEntryFlags Flags)
 		{
 			DirectoryNode Directory = this;
 
 			Utf8String RemainingPath = Path;
+			if (RemainingPath[0] == '/' || RemainingPath[0] == '\\')
+			{
+				RemainingPath = RemainingPath.Substring(1);
+			}
+
 			for (; ; )
 			{
-				int Length = RemainingPath.IndexOf('/');
-				if (Length == -1)
+				int Length = 0;
+				for (; ; Length++)
 				{
-					return Directory.CreateFile(RemainingPath, Flags);
+					if (Length == RemainingPath.Length)
+					{
+						return Directory.AddFile(RemainingPath, Flags);
+					}
+
+					byte Character = RemainingPath[Length];
+					if (Character == '\\' || Character == '/')
+					{
+						break;
+					}
 				}
+
 				if (Length > 0)
 				{
-					Directory = await Directory.FindOrAddDirectoryAsync(RemainingPath.Slice(0, Length));
+					Directory = await Bundle.FindOrAddDirectoryAsync(Directory, RemainingPath.Slice(0, Length));
 				}
 				RemainingPath = RemainingPath.Slice(Length + 1);
 			}
 		}
 
+		/// <summary>
+		/// Attempts to get a file entry with the given name
+		/// </summary>
+		/// <param name="Name">Name of the file</param>
+		/// <returns>Entry for the given name</returns>
+		public FileEntry GetFileEntry(Utf8String Name) => NameToFileEntry[Name];
+
 		/// <summary>
 		/// Attempts to get a file entry with the given name
 		/// </summary>
@@ -267,9 +290,10 @@ namespace EpicGames.Horde.Bundles.Nodes
 		/// <summary>
 		/// 
 		/// </summary>
+		/// <param name="Bundle">Bundle that contains the node</param>
 		/// <param name="Path"></param>
 		/// <returns></returns>
-		public async ValueTask<bool> DeleteFileByPathAsync(Utf8String Path)
+		public async ValueTask<bool> DeleteFileByPathAsync(Bundle Bundle, Utf8String Path)
 		{
 			DirectoryNode Directory = this;
 
@@ -283,7 +307,7 @@ namespace EpicGames.Horde.Bundles.Nodes
 				}
 				if (Length > 0)
 				{
-					Directory = await Directory.FindOrAddDirectoryAsync(RemainingPath.Slice(0, Length));
+					Directory = await Bundle.FindOrAddDirectoryAsync(Directory, RemainingPath.Slice(0, Length));
 				}
 				RemainingPath = RemainingPath.Slice(Length + 1);
 			}
@@ -314,6 +338,13 @@ namespace EpicGames.Horde.Bundles.Nodes
 			return Node;
 		}
 
+		/// <summary>
+		/// Get a directory entry with the given name
+		/// </summary>
+		/// <param name="Name">Name of the directory</param>
+		/// <returns>The entry with the given name</returns>
+		public DirectoryEntry GetDirectoryEntry(Utf8String Name) => NameToDirectoryEntry[Name];
+
 		/// <summary>
 		/// Attempts to get a directory entry with the given name
 		/// </summary>
@@ -322,34 +353,6 @@ namespace EpicGames.Horde.Bundles.Nodes
 		/// <returns>True if the directory was found</returns>
 		public bool TryGetDirectoryEntry(Utf8String Name, [NotNullWhen(true)] out DirectoryEntry? Entry) => NameToDirectoryEntry.TryGetValue(Name, out Entry);
 
-		/// <summary>
-		/// Tries to get a directory with the given name
-		/// </summary>
-		/// <param name="Name">Name of the new directory</param>
-		/// <returns>The new directory object</returns>
-		public async ValueTask<DirectoryNode?> FindDirectoryAsync(Utf8String Name)
-		{
-			if (TryGetDirectoryEntry(Name, out DirectoryEntry? Entry))
-			{
-				return await Entry.GetAsync();
-			}
-			else
-			{
-				return null;
-			}
-		}
-
-		/// <summary>
-		/// Tries to get a directory with the given name, or adds it if not present
-		/// </summary>
-		/// <param name="Name">Name of the new directory</param>
-		/// <returns>The new directory object</returns>
-		public async ValueTask<DirectoryNode> FindOrAddDirectoryAsync(Utf8String Name)
-		{
-			DirectoryNode? Node = await FindDirectoryAsync(Name);
-			return Node ?? AddDirectory(Name);
-		}
-
 		/// <summary>
 		/// Deletes the file entry with the given name
 		/// </summary>
@@ -427,8 +430,14 @@ namespace EpicGames.Horde.Bundles.Nodes
 			return new ReadOnlySequence<byte>(Data);
 		}
 
-		internal static DirectoryNode Deserialize(Bundle Bundle, ReadOnlySpan<byte> Span)
+		/// <summary>
+		/// Deserialize a directory node from data
+		/// </summary>
+		/// <param name="Data">The data to read from</param>
+		/// <returns>The deserialized directory node</returns>
+		public static DirectoryNode Deserialize(ReadOnlyMemory<byte> Data)
 		{
+			ReadOnlySpan<byte> Span = Data.Span;
 			if(Span[0] != TypeId)
 			{
 				throw new InvalidOperationException("Invalid signature byte for directory");
@@ -457,7 +466,7 @@ namespace EpicGames.Horde.Bundles.Nodes
 				IoHash Hash = new IoHash(Span);
 				Span = Span.Slice(IoHash.NumBytes);
 
-				Node.NameToFileEntry[Name] = new FileEntry(Bundle, Name, Flags, Length, Hash);
+				Node.NameToFileEntry[Name] = new FileEntry(Name, Flags, Length, Hash);
 			}
 
 			int DirectoryCount = (int)VarInt.Read(Span, out int DirectoryCountBytes);
@@ -476,9 +485,10 @@ namespace EpicGames.Horde.Bundles.Nodes
 				IoHash Hash = new IoHash(Span);
 				Span = Span.Slice(IoHash.NumBytes);
 
-				Node.NameToDirectoryEntry[Name] = new DirectoryEntry(Bundle, Name, Length, Hash);
+				Node.NameToDirectoryEntry[Name] = new DirectoryEntry(Name, Length, Hash);
 			}
 
+			Debug.Assert(Span.Length == 0);
 			return Node;
 		}
 
@@ -500,7 +510,7 @@ namespace EpicGames.Horde.Bundles.Nodes
 			foreach (FileInfo FileInfo in DirectoryInfo.EnumerateFiles())
 			{
 				Logger.LogInformation("Adding {File}", FileInfo.FullName);
-				FileNode FileNode = CreateFile(FileInfo.Name, 0);
+				FileNode FileNode = AddFile(FileInfo.Name, 0);
  				await FileNode.CopyFromFileAsync(FileInfo, Options);
 			}
 		}
@@ -508,10 +518,11 @@ namespace EpicGames.Horde.Bundles.Nodes
 		/// <summary>
 		/// Utility function to allow extracting a packed directory to disk
 		/// </summary>
+		/// <param name="Bundle">Bundle that contains the node</param>
 		/// <param name="DirectoryInfo"></param>
 		/// <param name="Logger"></param>
 		/// <returns></returns>
-		public async Task CopyToDirectoryAsync(DirectoryInfo DirectoryInfo, ILogger Logger)
+		public async Task CopyToDirectoryAsync(Bundle Bundle, DirectoryInfo DirectoryInfo, ILogger Logger)
 		{
 			DirectoryInfo.Create();
 
@@ -519,15 +530,15 @@ namespace EpicGames.Horde.Bundles.Nodes
 			foreach (FileEntry FileEntry in Files)
 			{
 				FileInfo FileInfo = new FileInfo(Path.Combine(DirectoryInfo.FullName, FileEntry.Name.ToString()));
-				FileNode FileNode = await FileEntry.GetAsync();
-				await FileNode.CopyToFileAsync(FileInfo);
+				FileNode FileNode = await Bundle.GetAsync(FileEntry);
+				await FileNode.CopyToFileAsync(Bundle, FileInfo);
 //				Tasks.Add(Task.Run(() => FileNode.CopyToFileAsync(FileInfo)));
 			}
 			foreach (DirectoryEntry DirectoryEntry in Directories)
 			{
 				DirectoryInfo SubDirectoryInfo = DirectoryInfo.CreateSubdirectory(DirectoryEntry.Name.ToString());
-				DirectoryNode SubDirectoryNode = await DirectoryEntry.GetAsync();
-				await SubDirectoryNode.CopyToDirectoryAsync(SubDirectoryInfo, Logger);
+				DirectoryNode SubDirectoryNode = await Bundle.GetAsync(DirectoryEntry);
+				await SubDirectoryNode.CopyToDirectoryAsync(Bundle, SubDirectoryInfo, Logger);
 //				Tasks.Add(Task.Run(() => SubDirectoryNode.CopyToDirectoryAsync(SubDirectoryInfo, Logger)));
 			}
 
@@ -541,9 +552,44 @@ namespace EpicGames.Horde.Bundles.Nodes
 	public class DirectoryNodeDeserializer : BundleNodeDeserializer<DirectoryNode>
 	{
 		/// <inheritdoc/>
-		public override DirectoryNode Deserialize(Bundle Bundle, ReadOnlyMemory<byte> Data)
+		public override DirectoryNode Deserialize(ReadOnlyMemory<byte> Data) => DirectoryNode.Deserialize(Data);
+	}
+
+	/// <summary>
+	/// Extension methods for <see cref="DirectoryNode"/>
+	/// </summary>
+	public static class DirectoryNodeExtensions
+	{
+		/// <summary>
+		/// Tries to get a directory with the given name
+		/// </summary>
+		/// <param name="Bundle">Bundle that contains the node</param>
+		/// <param name="ParentNode">Parent directory to search</param>
+		/// <param name="Name">Name of the new directory</param>
+		/// <returns>The new directory object</returns>
+		public static async ValueTask<DirectoryNode?> FindDirectoryAsync(this Bundle Bundle, DirectoryNode ParentNode, Utf8String Name)
+		{
+			if (ParentNode.TryGetDirectoryEntry(Name, out DirectoryEntry? Entry))
+			{
+				return await Bundle.GetAsync(Entry);
+			}
+			else
+			{
+				return null;
+			}
+		}
+
+		/// <summary>
+		/// Tries to get a directory with the given name, or adds it if not present
+		/// </summary>
+		/// <param name="Bundle">Bundle that contains the node</param>
+		/// <param name="ParentNode">Parent directory to search</param>
+		/// <param name="Name">Name of the new directory</param>
+		/// <returns>The new directory object</returns>
+		public static async ValueTask<DirectoryNode> FindOrAddDirectoryAsync(this Bundle Bundle, DirectoryNode ParentNode, Utf8String Name)
 		{
-			return DirectoryNode.Deserialize(Bundle, Data.Span);
+			DirectoryNode? Node = await FindDirectoryAsync(Bundle, ParentNode, Name);
+			return Node ?? ParentNode.AddDirectory(Name);
 		}
 	}
 }
diff --git a/Engine/Source/Programs/Shared/EpicGames.Horde/Bundles/Nodes/FileNode.cs b/Engine/Source/Programs/Shared/EpicGames.Horde/Bundles/Nodes/FileNode.cs
--- a/Engine/Source/Programs/Shared/EpicGames.Horde/Bundles/Nodes/FileNode.cs
+++ b/Engine/Source/Programs/Shared/EpicGames.Horde/Bundles/Nodes/FileNode.cs
@@ -79,6 +79,8 @@ namespace EpicGames.Horde.Bundles.Nodes
 	{
 		const byte TypeId = (byte)'c';
 
+		static readonly ReadOnlyMemory<byte> DefaultSegment = CreateFirstSegmentData(0, Array.Empty<byte>());
+
 		class DataSegment : ReadOnlySequenceSegment<byte>
 		{
 			public DataSegment(long RunningIndex, ReadOnlyMemory<byte> Data)
@@ -95,7 +97,7 @@ namespace EpicGames.Horde.Bundles.Nodes
 
 		internal int Depth { get; private set; }
 		internal ReadOnlySequence<byte> Payload { get; private set; }
-		ReadOnlySequence<byte> Data; // Full serialized data, including type id and header fields.
+		ReadOnlySequence<byte> Data = new ReadOnlySequence<byte>(DefaultSegment); // Full serialized data, including type id and header fields.
 		IoHash? Hash; // Hash of the serialized data buffer.
 
 		// In-memory state
@@ -133,7 +135,7 @@ namespace EpicGames.Horde.Bundles.Nodes
 		/// <summary>
 		/// Create a file node from deserialized data
 		/// </summary>
-		public static FileNode Deserialize(Bundle Bundle, ReadOnlyMemory<byte> Data)
+		public static FileNode Deserialize(ReadOnlyMemory<byte> Data)
 		{
 			ReadOnlySpan<byte> Span = Data.Span;
 			if (Span[0] != TypeId)
@@ -165,7 +167,7 @@ namespace EpicGames.Horde.Bundles.Nodes
 				while (Span.Length > 0)
 				{
 					IoHash ChildHash = new IoHash(Span);
-					Node.ChildNodeRefs.Add(new BundleNodeRef<FileNode>(Bundle, ChildHash));
+					Node.ChildNodeRefs.Add(new BundleNodeRef<FileNode>(ChildHash));
 				}
 			}
 			return Node;
@@ -175,11 +177,11 @@ namespace EpicGames.Horde.Bundles.Nodes
 		/// Serialize this node and its children into a byte array
 		/// </summary>
 		/// <returns>Array of data stored by the tree</returns>
-		public async Task<byte[]> ToByteArrayAsync()
+		public async Task<byte[]> ToByteArrayAsync(Bundle Bundle)
 		{
 			using (MemoryStream Stream = new MemoryStream())
 			{
-				await CopyToStreamAsync(Stream);
+				await CopyToStreamAsync(Bundle, Stream);
 				return Stream.ToArray();
 			}
 		}
@@ -221,15 +223,16 @@ namespace EpicGames.Horde.Bundles.Nodes
 		/// <summary>
 		/// Copies the contents of this node and its children to the given output stream
 		/// </summary>
+		/// <param name="Bundle">Bundle that can provide node data</param>
 		/// <param name="OutputStream">The output stream to receive the data</param>
-		public async Task CopyToStreamAsync(Stream OutputStream)
+		public async Task CopyToStreamAsync(Bundle Bundle, Stream OutputStream)
 		{
 			if (ChildNodeRefs != null)
 			{
 				foreach (BundleNodeRef<FileNode> ChildNodeRef in ChildNodeRefs)
 				{
-					FileNode ChildNode = await ChildNodeRef.GetAsync();
-					await ChildNode.CopyToStreamAsync(OutputStream);
+					FileNode ChildNode = ChildNodeRef.Node ?? await Bundle.GetAsync(ChildNodeRef);
+					await ChildNode.CopyToStreamAsync(Bundle, OutputStream);
 				}
 			}
 			else
@@ -244,13 +247,14 @@ namespace EpicGames.Horde.Bundles.Nodes
 		/// <summary>
 		/// Extracts the contents of this node to a file
 		/// </summary>
+		/// <param name="Bundle">Bundle that can provide node data</param>
 		/// <param name="File">File to write with the contents of this node</param>
 		/// <returns></returns>
-		public async Task CopyToFileAsync(FileInfo File)
+		public async Task CopyToFileAsync(Bundle Bundle, FileInfo File)
 		{
 			using (FileStream Stream = File.OpenWrite())
 			{
-				await CopyToStreamAsync(Stream);
+				await CopyToStreamAsync(Bundle, Stream);
 			}
 		}
 
@@ -428,13 +432,9 @@ namespace EpicGames.Horde.Bundles.Nodes
 
 			if (LastSegment == null)
 			{
-				byte[] Buffer = new byte[1 + DepthBytes + LeafData.Length];
-				Buffer[0] = TypeId;
-
-				VarInt.Write(Buffer.AsSpan(1, DepthBytes), Depth);
-				LeafData.CopyTo(Buffer.AsSpan(1 + DepthBytes));
-
-				FirstSegment = LastSegment = new DataSegment(0, Buffer);
+				byte[] Buffer = CreateFirstSegmentData(Depth, LeafData);
+				FirstSegment = new DataSegment(0, Buffer);
+				LastSegment = FirstSegment;
 			}
 			else
 			{
@@ -443,12 +443,26 @@ namespace EpicGames.Horde.Bundles.Nodes
 				LastSegment = NewSegment;
 			}
 
+			int HeaderSize = (int)(Data.Length - Payload.Length);
 			Data = new ReadOnlySequence<byte>(FirstSegment!, 0, LastSegment, LastSegment.Memory.Length);
-			Payload = Data.Slice(1 + DepthBytes);
+			Payload = Data.Slice(HeaderSize);
 
 			Length += LeafData.Length;
 		}
 
+		private static byte[] CreateFirstSegmentData(int Depth, ReadOnlySpan<byte> LeafData)
+		{
+			int DepthBytes = VarInt.Measure(Depth);
+
+			byte[] Buffer = new byte[1 + DepthBytes + LeafData.Length];
+			Buffer[0] = TypeId;
+
+			VarInt.Write(Buffer.AsSpan(1, DepthBytes), Depth);
+			LeafData.CopyTo(Buffer.AsSpan(1 + DepthBytes));
+
+			return Buffer;
+		}
+
 		private void MarkComplete()
 		{
 			if (Hash == null)
@@ -538,6 +552,6 @@ namespace EpicGames.Horde.Bundles.Nodes
 	public class FileNodeDeserializer : BundleNodeDeserializer<FileNode>
 	{
 		/// <inheritdoc/>
-		public override FileNode Deserialize(Bundle Bundle, ReadOnlyMemory<byte> Data) => FileNode.Deserialize(Bundle, Data);
+		public override FileNode Deserialize(ReadOnlyMemory<byte> Data) => FileNode.Deserialize(Data);
 	}
 }
diff --git a/Engine/Source/Programs/Shared/EpicGames.Horde/Storage/IStorageClient.cs b/Engine/Source/Programs/Shared/EpicGames.Horde/Storage/IStorageClient.cs
--- a/Engine/Source/Programs/Shared/EpicGames.Horde/Storage/IStorageClient.cs
+++ b/Engine/Source/Programs/Shared/EpicGames.Horde/Storage/IStorageClient.cs
@@ -354,7 +354,23 @@ namespace EpicGames.Horde.Storage
 		/// <returns></returns>
 		public static Task<IoHash> WriteObjectAsync<T>(this IStorageClient StorageClient, NamespaceId NamespaceId, T Object)
 		{
-			return WriteObjectAsync(StorageClient, NamespaceId, CbSerializer.Serialize<T>(Object));
+			CbWriter Writer = new CbWriter();
+			CbSerializer.Serialize<T>(Writer, Object);
+			return WriteObjectAsync(StorageClient, NamespaceId, Writer);
+		}
+
+		/// <summary>
+		/// Gets a blob as a byte array
+		/// </summary>
+		/// <param name="StorageClient">The storage interface</param>
+		/// <param name="NamespaceId">Namespace containing the blob</param>
+		/// <param name="Writer">The object to be written</param>
+		/// <returns></returns>
+		public static async Task<IoHash> WriteObjectAsync<T>(this IStorageClient StorageClient, NamespaceId NamespaceId, CbWriter Writer)
+		{
+			IoHash Hash = Writer.ComputeHash();
+			await StorageClient.WriteBlobAsync(NamespaceId, Hash, Writer.AsStream());
+			return Hash;
 		}
 
 		/// <summary>
diff --git a/Engine/Source/Programs/Shared/EpicGames.Serialization.Tests/SerializerTests.cs b/Engine/Source/Programs/Shared/EpicGames.Serialization.Tests/SerializerTests.cs
new file mode 100644
--- /dev/null
+++ b/Engine/Source/Programs/Shared/EpicGames.Serialization.Tests/SerializerTests.cs
@@ -0,0 +1,94 @@
+// Copyright Epic Games, Inc. All Rights Reserved.
+
+using EpicGames.Core;
+using Microsoft.VisualStudio.TestTools.UnitTesting;
+using System;
+using System.Collections.Generic;
+using System.IO;
+using System.Linq;
+using System.Text;
+using System.Threading.Tasks;
+
+namespace EpicGames.Serialization.Tests
+{
+	[TestClass]
+	public class SerializerTests
+	{
+		class TestObject : IEquatable<TestObject>
+		{
+			[CbField]
+			public int A { get; set; }
+
+			[CbField]
+			public string B { get; set; } = String.Empty;
+
+			[CbField]
+			public IoHash C { get; set; }
+
+			[CbField]
+			public List<TestObject> Children { get; set; } = new List<TestObject>();
+
+			public override bool Equals(object? obj) => obj is TestObject Other && Equals(Other);
+
+			public bool Equals(TestObject? Other) => Other != null && A == Other.A && B == Other.B && C == Other.C && Enumerable.SequenceEqual(Children, Other.Children);
+
+			public override int GetHashCode() => HashCode.Combine(A, B, C);
+		}
+
+		[TestMethod]
+		public void RoundTrip()
+		{
+			TestObject Obj1 = CreateTestObject(new Random(123));
+			CbObject CbObj1 = CbSerializer.Serialize(Obj1);
+			ReadOnlyMemory<byte> Mem1 = CbObj1.GetView();
+
+			TestObject Obj2 = CbSerializer.Deserialize<TestObject>(Mem1);
+			CbObject CbObj2 = CbSerializer.Serialize(Obj2);
+			ReadOnlyMemory<byte> Mem2 = CbObj1.GetView();
+
+			Assert.AreEqual(Obj1, Obj2);
+			Assert.AreEqual(Mem1, Mem2);
+		}
+
+		[TestMethod]
+		public void StreamTest()
+		{
+			TestObject Obj1 = CreateTestObject(new Random(123));
+
+			CbWriter Writer = new CbWriter();
+			CbSerializer.Serialize(Writer, Obj1);
+
+			byte[] Data1 = Writer.ToByteArray();
+			byte[] Data2;
+			using (Stream Stream = Writer.AsStream())
+			{
+				Data2 = new byte[Stream.Length];
+
+				int ReadLength = Stream.Read(Data2, 0, Data2.Length);
+				Assert.AreEqual(Data2.Length, ReadLength);
+
+				ReadLength = Stream.Read(new byte[10], 0, 10);
+				Assert.AreEqual(0, ReadLength);
+			}
+			Assert.IsTrue(Data1.AsSpan().SequenceEqual(Data2));
+
+			TestObject Obj2 = CbSerializer.Deserialize<TestObject>(Data1);
+			Assert.AreEqual(Obj1, Obj2);
+		}
+
+		static TestObject CreateTestObject(Random Random)
+		{
+			TestObject Obj = new TestObject();
+			Obj.A = Random.Next();
+			Obj.B = Random.NextInt64().ToString();
+			Obj.C = IoHash.Compute(Encoding.UTF8.GetBytes(Obj.B));
+
+			for (int Idx = 0; Idx < Random.Next(0, 3); Idx++)
+			{
+				Obj.Children.Add(CreateTestObject(Random));
+			}
+
+			return Obj;
+		}
+	}
+}
diff --git a/Engine/Source/Programs/Shared/EpicGames.Serialization/CbWriter.cs b/Engine/Source/Programs/Shared/EpicGames.Serialization/CbWriter.cs
--- a/Engine/Source/Programs/Shared/EpicGames.Serialization/CbWriter.cs
+++ b/Engine/Source/Programs/Shared/EpicGames.Serialization/CbWriter.cs
@@ -4,6 +4,7 @@ using EpicGames.Core;
 using System;
 using System.Buffers.Binary;
 using System.Collections.Generic;
+using System.IO;
 using System.Linq;
 using System.Text;
 using System.Text.Json;
@@ -837,6 +838,22 @@ namespace EpicGames.Serialization
 			}
 		}
 
+		/// <summary>
+		/// Computes the hash for this object
+		/// </summary>
+		/// <returns>Hash for the object</returns>
+		public IoHash ComputeHash()
+		{
+			using (Blake3.Hasher Hasher = Blake3.Hasher.New())
+			{
+				foreach (ReadOnlyMemory<byte> Segment in EnumerateSegments())
+				{
+					Hasher.Update(Segment.Span);
+				}
+				return IoHash.FromBlake3(Hasher);
+			}
+		}
+
 		/// <summary>
 		/// Gets the size of the serialized data
 		/// </summary>
@@ -851,6 +868,37 @@ namespace EpicGames.Serialization
 			return CurrentOffset + ComputeSizeOfChildHeaders(CurrentScope);
 		}
 
+		/// <summary>
+		/// Gets the contents of this writer as a stream
+		/// </summary>
+		/// <returns>New stream for the contents of this object</returns>
+		public Stream AsStream() => new ReadStream(EnumerateSegments().GetEnumerator(), GetSize());
+
+		private IEnumerable<ReadOnlyMemory<byte>> EnumerateSegments()
+		{
+			byte[] ScopeHeader = new byte[64];
+
+			int SourceOffset = 0;
+			foreach (Chunk Chunk in Chunks)
+			{
+				foreach (Scope Scope in Chunk.Scopes)
+				{
+					ReadOnlyMemory<byte> SourceData = Chunk.Data.AsMemory(SourceOffset - Chunk.Offset, Scope.Offset - SourceOffset);
+					yield return SourceData;
+
+					SourceOffset += SourceData.Length;
+
+					int HeaderLength = WriteScopeHeader(ScopeHeader, Scope);
+					yield return ScopeHeader.AsMemory(0, HeaderLength);
+				}
+
+				ReadOnlyMemory<byte> LastSourceData = Chunk.Data.AsMemory(SourceOffset - Chunk.Offset, (Chunk.Offset + Chunk.Length) - SourceOffset);
+				yield return LastSourceData;
+
+				SourceOffset += LastSourceData.Length;
+			}
+		}
+
 		/// <summary>
 		/// Copy the data from this writer to a buffer
 		/// </summary>
@@ -880,6 +928,79 @@ namespace EpicGames.Serialization
 			}
 		}
 
+		class ReadStream : Stream
+		{
+			IEnumerator<ReadOnlyMemory<byte>> Enumerator;
+			ReadOnlyMemory<byte> Segment;
+			long PositionInternal;
+
+			public ReadStream(IEnumerator<ReadOnlyMemory<byte>> Enumerator, long Length)
+			{
+				this.Enumerator = Enumerator;
+				this.Length = Length;
+			}
+
+			/// <inheritdoc/>
+			public override bool CanRead => true;
+
+			/// <inheritdoc/>
+			public override bool CanSeek => false;
+
+			/// <inheritdoc/>
+			public override bool CanWrite => false;
+
+			/// <inheritdoc/>
+			public override long Length { get; }
+
+			/// <inheritdoc/>
+			public override long Position
+			{
+				get => PositionInternal;
+				set => throw new NotSupportedException();
+			}
+
+			/// <inheritdoc/>
+			public override void Flush() { }
+
+			/// <inheritdoc/>
+			public override int Read(Span<byte> Buffer)
+			{
+				int ReadLength = 0;
+				while (ReadLength < Buffer.Length)
+				{
+					while (Segment.Length == 0)
+					{
+						if (!Enumerator.MoveNext())
+						{
+							return ReadLength;
+						}
+						Segment = Enumerator.Current;
+					}
+
+					int CopyLength = Math.Min(Segment.Length, Buffer.Length);
+
+					Segment.Span.Slice(0, CopyLength).CopyTo(Buffer.Slice(ReadLength));
+					Segment = Segment.Slice(CopyLength);
+
+					PositionInternal += CopyLength;
+					ReadLength += CopyLength;
+				}
+				return ReadLength;
+			}
+
+			/// <inheritdoc/>
+			public override int Read(byte[] Buffer, int Offset, int Count) => Read(Buffer.AsSpan(Offset, Count));
+
+			/// <inheritdoc/>
+			public override long Seek(long Offset, SeekOrigin Origin) => throw new NotSupportedException();
+
+			/// <inheritdoc/>
+			public override void SetLength(long Value) => throw new NotSupportedException();
+
+			/// <inheritdoc/>
+			public override void Write(byte[] Buffer, int Offset, int Count) => throw new NotSupportedException();
+		}
+
 		/// <summary>
 		/// Convert the data into a compact binary object
 		/// </summary>
