Horde: Missing file.

diff --git a/Engine/Source/Programs/Shared/EpicGames.Perforce/Records/SyncRecord.cs b/Engine/Source/Programs/Shared/EpicGames.Perforce/Records/SyncRecord.cs
--- a/Engine/Source/Programs/Shared/EpicGames.Perforce/Records/SyncRecord.cs
+++ b/Engine/Source/Programs/Shared/EpicGames.Perforce/Records/SyncRecord.cs
@@ -1,5 +1,6 @@
 // Copyright Epic Games, Inc. All Rights Reserved.
 
+using EpicGames.Core;
 using System;
 using System.Collections.Generic;
 using System.Diagnostics;
@@ -19,13 +20,13 @@ namespace EpicGames.Perforce
 		/// Path to the file in the depot
 		/// </summary>
 		[PerforceTag("depotFile")]
-		public string DepotFile { get; set; } = String.Empty;
+		public Utf8String DepotFile { get; set; } = Utf8String.Empty;
 
 		/// <summary>
 		/// Path to the file in the workspace. Note: despite being a property called 'clientFile', this is not in client syntax.
 		/// </summary>
 		[PerforceTag("clientFile")]
-		public string Path { get; set; } = String.Empty;
+		public Utf8String Path { get; set; } = Utf8String.Empty;
 
 		/// <summary>
 		/// The revision number of the file that was synced
