More UHT documentation updates

#rnx

diff --git a/Engine/Source/Programs/Shared/EpicGames.UHT/Parsers/UhtFunctionParser.cs b/Engine/Source/Programs/Shared/EpicGames.UHT/Parsers/UhtFunctionParser.cs
--- a/Engine/Source/Programs/Shared/EpicGames.UHT/Parsers/UhtFunctionParser.cs
+++ b/Engine/Source/Programs/Shared/EpicGames.UHT/Parsers/UhtFunctionParser.cs
@@ -326,8 +326,6 @@ namespace EpicGames.UHT.Parsers
 						Function.SparseDelegateName = TopScope.TokenReader.GetIdentifier("delegate name").Value.ToString();
 					}
 
-					Function.FunctionLineNumber = TopScope.TokenReader.InputLine;
-
 					// Get parameter list.
 					if (FoundParamIndex >= 0)
 					{
diff --git a/Engine/Source/Programs/Shared/EpicGames.UHT/Types/Properties/UhtInterfaceProperty.cs b/Engine/Source/Programs/Shared/EpicGames.UHT/Types/Properties/UhtInterfaceProperty.cs
--- a/Engine/Source/Programs/Shared/EpicGames.UHT/Types/Properties/UhtInterfaceProperty.cs
+++ b/Engine/Source/Programs/Shared/EpicGames.UHT/Types/Properties/UhtInterfaceProperty.cs
@@ -11,6 +11,9 @@ using System.Text.Json.Serialization;
 
 namespace EpicGames.UHT.Types
 {
+	/// <summary>
+	/// FInterfaceProperty
+	/// </summary>
 	[UnrealHeaderTool]
 	[UhtEngineClass(Name = "InterfaceProperty", IsProperty = true)]
 	public class UhtInterfaceProperty : UhtProperty
@@ -27,9 +30,17 @@ namespace EpicGames.UHT.Types
 		/// <inheritdoc/>
 		protected override UhtPGetArgumentType PGetTypeArgument { get => UhtPGetArgumentType.TypeText; }
 
+		/// <summary>
+		/// Referenced interface class
+		/// </summary>
 		[JsonConverter(typeof(UhtTypeSourceNameJsonConverter<UhtClass>))]
 		public UhtClass InterfaceClass { get; set; }
 
+		/// <summary>
+		/// Create a new property
+		/// </summary>
+		/// <param name="PropertySettings">Property settings</param>
+		/// <param name="InterfaceClass">Referenced interface</param>
 		public UhtInterfaceProperty(UhtPropertySettings PropertySettings, UhtClass InterfaceClass) : base(PropertySettings)
 		{
 			this.InterfaceClass = InterfaceClass;
@@ -168,13 +179,13 @@ namespace EpicGames.UHT.Types
 
 		#region Keywords
 		[UhtPropertyType(Keyword = "FScriptInterface", Options = UhtPropertyTypeOptions.Simple)] // This can't be immediate due to the reference to UInterface
-		public static UhtProperty? FScriptInterfaceProperty(UhtPropertyResolvePhase ResolvePhase, UhtPropertySettings PropertySettings, IUhtTokenReader TokenReader, UhtToken MatchedToken)
+		private static UhtProperty? FScriptInterfaceProperty(UhtPropertyResolvePhase ResolvePhase, UhtPropertySettings PropertySettings, IUhtTokenReader TokenReader, UhtToken MatchedToken)
 		{
 			return new UhtInterfaceProperty(PropertySettings, PropertySettings.Outer.Session.IInterface);
 		}
 
 		[UhtPropertyType(Keyword = "TScriptInterface")]
-		public static UhtProperty? TScriptInterfaceProperty(UhtPropertyResolvePhase ResolvePhase, UhtPropertySettings PropertySettings, IUhtTokenReader TokenReader, UhtToken MatchedToken)
+		private static UhtProperty? TScriptInterfaceProperty(UhtPropertyResolvePhase ResolvePhase, UhtPropertySettings PropertySettings, IUhtTokenReader TokenReader, UhtToken MatchedToken)
 		{
 			UhtClass? PropertyClass = UhtObjectPropertyBase.ParseTemplateObject(PropertySettings, TokenReader, MatchedToken, false);
 			if (PropertyClass == null)
@@ -186,7 +197,7 @@ namespace EpicGames.UHT.Types
 			{
 				return new UhtInterfaceProperty(PropertySettings, PropertyClass);
 			}
-			return new UhtObjectProperty(PropertySettings, PropertyClass, EPropertyFlags.UObjectWrapper);
+			return new UhtObjectProperty(PropertySettings, PropertyClass, null, EPropertyFlags.UObjectWrapper);
 		}
 		#endregion
 	}
diff --git a/Engine/Source/Programs/Shared/EpicGames.UHT/Types/Properties/UhtMapProperty.cs b/Engine/Source/Programs/Shared/EpicGames.UHT/Types/Properties/UhtMapProperty.cs
--- a/Engine/Source/Programs/Shared/EpicGames.UHT/Types/Properties/UhtMapProperty.cs
+++ b/Engine/Source/Programs/Shared/EpicGames.UHT/Types/Properties/UhtMapProperty.cs
@@ -10,6 +10,10 @@ using System.Text;
 
 namespace EpicGames.UHT.Types
 {
+
+	/// <summary>
+	/// FMapProperty
+	/// </summary>
 	[UnrealHeaderTool]
 	[UhtEngineClass(Name = "MapProperty", IsProperty = true)]
 	public class UhtMapProperty : UhtContainerBaseProperty
@@ -26,8 +30,17 @@ namespace EpicGames.UHT.Types
 		/// <inheritdoc/>
 		protected override UhtPGetArgumentType PGetTypeArgument { get => UhtPGetArgumentType.TypeText; }
 
+		/// <summary>
+		/// Key property
+		/// </summary>
 		public UhtProperty KeyProperty { get; set; }
 
+		/// <summary>
+		/// Construct new property
+		/// </summary>
+		/// <param name="PropertySettings">Property settings</param>
+		/// <param name="Key">Key property</param>
+		/// <param name="Value">Value property</param>
 		public UhtMapProperty(UhtPropertySettings PropertySettings, UhtProperty Key, UhtProperty Value) : base(PropertySettings, Value)
 		{
 			// If the creation of the value property set more flags, then copy those flags to ourselves
@@ -234,9 +247,9 @@ namespace EpicGames.UHT.Types
 			return null;
 		}
 
-#region Keyword
+		#region Keyword
 		[UhtPropertyType(Keyword = "TMap")]
-		public static UhtProperty? MapProperty(UhtPropertyResolvePhase ResolvePhase, UhtPropertySettings PropertySettings, IUhtTokenReader TokenReader, UhtToken MatchedToken)
+		private static UhtProperty? MapProperty(UhtPropertyResolvePhase ResolvePhase, UhtPropertySettings PropertySettings, IUhtTokenReader TokenReader, UhtToken MatchedToken)
 		{
 			using (var TokenContext = new UhtMessageContext(TokenReader, "TMap"))
 			{
@@ -296,6 +309,6 @@ namespace EpicGames.UHT.Types
 				return new UhtMapProperty(PropertySettings, Key, Value);
 			}
 		}
-#endregion
+		#endregion
 	}
 }
diff --git a/Engine/Source/Programs/Shared/EpicGames.UHT/Types/Properties/UhtMulticastDelegateProperty.cs b/Engine/Source/Programs/Shared/EpicGames.UHT/Types/Properties/UhtMulticastDelegateProperty.cs
--- a/Engine/Source/Programs/Shared/EpicGames.UHT/Types/Properties/UhtMulticastDelegateProperty.cs
+++ b/Engine/Source/Programs/Shared/EpicGames.UHT/Types/Properties/UhtMulticastDelegateProperty.cs
@@ -9,6 +9,9 @@ using System.Text.Json.Serialization;
 
 namespace EpicGames.UHT.Types
 {
+	/// <summary>
+	/// FMulticastDelegateProperty
+	/// </summary>
 	[UhtEngineClass(Name = "MulticastDelegateProperty", IsProperty = true)]
 	public abstract class UhtMulticastDelegateProperty : UhtProperty
 	{
@@ -24,9 +27,17 @@ namespace EpicGames.UHT.Types
 		/// <inheritdoc/>
 		protected override UhtPGetArgumentType PGetTypeArgument { get => UhtPGetArgumentType.TypeText; }
 
+		/// <summary>
+		/// Referenced function
+		/// </summary>
 		[JsonConverter(typeof(UhtTypeSourceNameJsonConverter<UhtFunction>))]
 		public UhtFunction Function { get; set; }
 
+		/// <summary>
+		/// Construct property
+		/// </summary>
+		/// <param name="PropertySettings">Property settings</param>
+		/// <param name="Function">Referenced function</param>
 		public UhtMulticastDelegateProperty(UhtPropertySettings PropertySettings, UhtFunction Function) : base(PropertySettings)
 		{
 			this.Function = Function;
diff --git a/Engine/Source/Programs/Shared/EpicGames.UHT/Types/Properties/UhtMulticastInlineDelegateProperty.cs b/Engine/Source/Programs/Shared/EpicGames.UHT/Types/Properties/UhtMulticastInlineDelegateProperty.cs
--- a/Engine/Source/Programs/Shared/EpicGames.UHT/Types/Properties/UhtMulticastInlineDelegateProperty.cs
+++ b/Engine/Source/Programs/Shared/EpicGames.UHT/Types/Properties/UhtMulticastInlineDelegateProperty.cs
@@ -5,16 +5,25 @@ using System.Text;
 
 namespace EpicGames.UHT.Types
 {
+	/// <summary>
+	/// FMulticastInlineDelegateProperty
+	/// </summary>
 	[UhtEngineClass(Name = "MulticastInlineDelegateProperty", IsProperty = true)]
 	public class UhtMulticastInlineDelegateProperty : UhtMulticastDelegateProperty
 	{
 		/// <inheritdoc/>
 		public override string EngineClassName { get => "MulticastInlineDelegateProperty"; }
 
+		/// <summary>
+		/// Construct new property
+		/// </summary>
+		/// <param name="PropertySettings">Property settings</param>
+		/// <param name="Function">Referenced function</param>
 		public UhtMulticastInlineDelegateProperty(UhtPropertySettings PropertySettings, UhtFunction Function) : base(PropertySettings, Function)
 		{
 		}
 
+		/// <inheritdoc/>
 		public override StringBuilder AppendText(StringBuilder Builder, UhtPropertyTextType TextType, bool bIsTemplateArgument)
 		{
 			switch (TextType)
diff --git a/Engine/Source/Programs/Shared/EpicGames.UHT/Types/Properties/UhtMulticastSparseDelegateProperty.cs b/Engine/Source/Programs/Shared/EpicGames.UHT/Types/Properties/UhtMulticastSparseDelegateProperty.cs
--- a/Engine/Source/Programs/Shared/EpicGames.UHT/Types/Properties/UhtMulticastSparseDelegateProperty.cs
+++ b/Engine/Source/Programs/Shared/EpicGames.UHT/Types/Properties/UhtMulticastSparseDelegateProperty.cs
@@ -5,16 +5,26 @@ using System.Text;
 
 namespace EpicGames.UHT.Types
 {
+
+	/// <summary>
+	/// FMulticastSparseDelegateProperty
+	/// </summary>
 	[UhtEngineClass(Name = "MulticastSparseDelegateProperty", IsProperty = true)]
 	public class UhtMulticastSparseDelegateProperty : UhtMulticastDelegateProperty
 	{
 		/// <inheritdoc/>
 		public override string EngineClassName { get => "MulticastSparseDelegateProperty"; }
 
+		/// <summary>
+		/// Construct a new property
+		/// </summary>
+		/// <param name="PropertySettings">Property settings</param>
+		/// <param name="Function">Referenced function</param>
 		public UhtMulticastSparseDelegateProperty(UhtPropertySettings PropertySettings, UhtFunction Function) : base(PropertySettings, Function)
 		{
 		}
 
+		/// <inheritdoc/>
 		public override StringBuilder AppendText(StringBuilder Builder, UhtPropertyTextType TextType, bool bIsTemplateArgument)
 		{
 			switch (TextType)
diff --git a/Engine/Source/Programs/Shared/EpicGames.UHT/Types/Properties/UhtNameProperty.cs b/Engine/Source/Programs/Shared/EpicGames.UHT/Types/Properties/UhtNameProperty.cs
--- a/Engine/Source/Programs/Shared/EpicGames.UHT/Types/Properties/UhtNameProperty.cs
+++ b/Engine/Source/Programs/Shared/EpicGames.UHT/Types/Properties/UhtNameProperty.cs
@@ -8,6 +8,10 @@ using System.Text;
 
 namespace EpicGames.UHT.Types
 {
+
+	/// <summary>
+	/// FNameProperty
+	/// </summary>
 	[UnrealHeaderTool]
 	[UhtEngineClass(Name = "NameProperty", IsProperty = true)]
 	public class UhtNameProperty : UhtProperty
@@ -24,6 +28,10 @@ namespace EpicGames.UHT.Types
 		/// <inheritdoc/>
 		protected override UhtPGetArgumentType PGetTypeArgument { get => UhtPGetArgumentType.EngineClass; }
 
+		/// <summary>
+		/// Construct a new property
+		/// </summary>
+		/// <param name="PropertySettings">Property settings</param>
 		public UhtNameProperty(UhtPropertySettings PropertySettings) : base(PropertySettings)
 		{
 			this.PropertyCaps |= UhtPropertyCaps.CanExposeOnSpawn | UhtPropertyCaps.IsParameterSupportedByBlueprint | UhtPropertyCaps.IsMemberSupportedByBlueprint;
@@ -86,7 +94,7 @@ namespace EpicGames.UHT.Types
 
 		#region Keyword
 		[UhtPropertyType(Keyword = "FName", Options = UhtPropertyTypeOptions.Simple | UhtPropertyTypeOptions.Immediate)]
-		public static UhtProperty? NameProperty(UhtPropertyResolvePhase ResolvePhase, UhtPropertySettings PropertySettings, IUhtTokenReader TokenReader, UhtToken MatchedToken)
+		private static UhtProperty? NameProperty(UhtPropertyResolvePhase ResolvePhase, UhtPropertySettings PropertySettings, IUhtTokenReader TokenReader, UhtToken MatchedToken)
 		{
 			return new UhtNameProperty(PropertySettings);
 		}
diff --git a/Engine/Source/Programs/Shared/EpicGames.UHT/Types/Properties/UhtNumericProperty.cs b/Engine/Source/Programs/Shared/EpicGames.UHT/Types/Properties/UhtNumericProperty.cs
--- a/Engine/Source/Programs/Shared/EpicGames.UHT/Types/Properties/UhtNumericProperty.cs
+++ b/Engine/Source/Programs/Shared/EpicGames.UHT/Types/Properties/UhtNumericProperty.cs
@@ -4,6 +4,10 @@ using System.Text;
 
 namespace EpicGames.UHT.Types
 {
+
+	/// <summary>
+	/// Base class for all numeric properties
+	/// </summary>
 	public abstract class UhtNumericProperty : UhtProperty
 	{
 		/// <inheritdoc/>
@@ -17,6 +21,11 @@ namespace EpicGames.UHT.Types
 		/// </summary>
 		public readonly UhtPropertyIntType IntType;
 
+		/// <summary>
+		/// Construct new property
+		/// </summary>
+		/// <param name="PropertySettings">Property settings</param>
+		/// <param name="IntType">Type of integer</param>
 		protected UhtNumericProperty(UhtPropertySettings PropertySettings, UhtPropertyIntType IntType) : base(PropertySettings)
 		{
 			this.IntType = IntType;
diff --git a/Engine/Source/Programs/Shared/EpicGames.UHT/Types/Properties/UhtObjectProperty.cs b/Engine/Source/Programs/Shared/EpicGames.UHT/Types/Properties/UhtObjectProperty.cs
--- a/Engine/Source/Programs/Shared/EpicGames.UHT/Types/Properties/UhtObjectProperty.cs
+++ b/Engine/Source/Programs/Shared/EpicGames.UHT/Types/Properties/UhtObjectProperty.cs
@@ -7,6 +7,9 @@ using System.Text;
 
 namespace EpicGames.UHT.Types
 {
+	/// <summary>
+	/// FObjectProperty
+	/// </summary>
 	[UhtEngineClass(Name = "ObjectProperty", IsProperty = true)]
 	public class UhtObjectProperty : UhtObjectPropertyBase
 	{
@@ -16,16 +19,15 @@ namespace EpicGames.UHT.Types
 		/// <inheritdoc/>
 		protected override UhtPGetArgumentType PGetTypeArgument { get => UhtPGetArgumentType.TypeText; }
 
-		public UhtObjectProperty(UhtPropertySettings PropertySettings, UhtClass PropertyClass, EPropertyFlags ExtraFlags = EPropertyFlags.None)
-			: base(PropertySettings, PropertyClass, null)
-		{
-			this.PropertyFlags |= ExtraFlags;
-			this.PropertyCaps |= UhtPropertyCaps.RequiresNullConstructorArg | UhtPropertyCaps.CanBeInstanced | UhtPropertyCaps.CanExposeOnSpawn |
-				UhtPropertyCaps.IsParameterSupportedByBlueprint | UhtPropertyCaps.IsMemberSupportedByBlueprint;
-		}
-
-		protected UhtObjectProperty(UhtPropertySettings PropertySettings, UhtClass PropertyClass, UhtClass MetaClass, EPropertyFlags ExtraFlags = EPropertyFlags.None)
-			: base(PropertySettings, PropertyClass, MetaClass)
+		/// <summary>
+		/// Construct a new property
+		/// </summary>
+		/// <param name="PropertySettings">Property settings</param>
+		/// <param name="Class">Referenced class</param>
+		/// <param name="MetaClass">Optional reference class (used by class properties)</param>
+		/// <param name="ExtraFlags">Extra flags to add to the property</param>
+		public UhtObjectProperty(UhtPropertySettings PropertySettings, UhtClass Class, UhtClass? MetaClass = null, EPropertyFlags ExtraFlags = EPropertyFlags.None)
+			: base(PropertySettings, Class, MetaClass)
 		{
 			this.PropertyFlags |= ExtraFlags;
 			this.PropertyCaps |= UhtPropertyCaps.RequiresNullConstructorArg | UhtPropertyCaps.CanBeInstanced | UhtPropertyCaps.CanExposeOnSpawn |
diff --git a/Engine/Source/Programs/Shared/EpicGames.UHT/Types/Properties/UhtObjectPropertyBase.cs b/Engine/Source/Programs/Shared/EpicGames.UHT/Types/Properties/UhtObjectPropertyBase.cs
--- a/Engine/Source/Programs/Shared/EpicGames.UHT/Types/Properties/UhtObjectPropertyBase.cs
+++ b/Engine/Source/Programs/Shared/EpicGames.UHT/Types/Properties/UhtObjectPropertyBase.cs
@@ -11,6 +11,9 @@ using System.Text.Json.Serialization;
 
 namespace EpicGames.UHT.Types
 {
+	/// <summary>
+	/// FObjectPropertyBase
+	/// </summary>
 	[UhtEngineClass(Name = "ObjectPropertyBase", IsProperty = true)]
 	public abstract class UhtObjectPropertyBase : UhtProperty
 	{
@@ -23,15 +26,27 @@ namespace EpicGames.UHT.Types
 		/// <inheritdoc/>
 		protected override string PGetMacroText { get => "OBJECT"; }
 
+		/// <summary>
+		/// Referenced UCLASS
+		/// </summary>
 		[JsonConverter(typeof(UhtTypeSourceNameJsonConverter<UhtClass>))]
 		public UhtClass Class { get; set; }
 
+		/// <summary>
+		/// Referenced UCLASS for class properties
+		/// </summary>
 		[JsonConverter(typeof(UhtNullableTypeSourceNameJsonConverter<UhtClass>))]
 		public UhtClass? MetaClass { get; set; }
 
-		public UhtObjectPropertyBase(UhtPropertySettings PropertySettings, UhtClass PropertyClass, UhtClass? MetaClass) : base(PropertySettings)
+		/// <summary>
+		/// Construct a property
+		/// </summary>
+		/// <param name="PropertySettings">Property settings</param>
+		/// <param name="Class">Referenced UCLASS</param>
+		/// <param name="MetaClass">Referenced UCLASS used by class properties</param>
+		public UhtObjectPropertyBase(UhtPropertySettings PropertySettings, UhtClass Class, UhtClass? MetaClass) : base(PropertySettings)
 		{
-			this.Class = PropertyClass;
+			this.Class = Class;
 			this.MetaClass = MetaClass;
 
 			// This applies to EVERYTHING including raw pointer
@@ -121,6 +136,14 @@ namespace EpicGames.UHT.Types
 		}
 
 		#region Parsing support methods
+		/// <summary>
+		/// Parse a template type
+		/// </summary>
+		/// <param name="PropertySettings">Property settings</param>
+		/// <param name="TokenReader">Token reader</param>
+		/// <param name="MatchedToken">Token matched for type</param>
+		/// <param name="bReturnUInterface">If true, return the UInterface instead of the type listed</param>
+		/// <returns>Referenced class</returns>
 		public static UhtClass? ParseTemplateObject(UhtPropertySettings PropertySettings, IUhtTokenReader TokenReader, UhtToken MatchedToken, bool bReturnUInterface)
 		{
 			if (TokenReader.TryOptional("const"))
@@ -154,6 +177,13 @@ namespace EpicGames.UHT.Types
 			return Return;
 		}
 
+		/// <summary>
+		/// Parse a template type
+		/// </summary>
+		/// <param name="PropertySettings">Property settings</param>
+		/// <param name="TokenReader">Token reader</param>
+		/// <param name="MatchedToken">Token matched for type</param>
+		/// <returns>Referenced class</returns>
 		public static UhtClass? ParseTemplateClass(UhtPropertySettings PropertySettings, IUhtTokenReader TokenReader, UhtToken MatchedToken)
 		{
 			UhtToken Identifier = new UhtToken();
@@ -185,6 +215,17 @@ namespace EpicGames.UHT.Types
 			return Return;
 		}
 
+		/// <summary>
+		/// Log point usage warning/error
+		/// </summary>
+		/// <param name="PropertySettings">Property settings</param>
+		/// <param name="EngineBehavior">Expected behavior for engine types</param>
+		/// <param name="NonEngineBehavior">Expected behavior for non-engine types</param>
+		/// <param name="PointerTypeDesc">Description of the pointer type</param>
+		/// <param name="TokenReader">Token reader for type being parsed</param>
+		/// <param name="TypeStartPos">Starting character position of the type</param>
+		/// <param name="AlternativeTypeDesc">Suggested alternate declaration</param>
+		/// <exception cref="UhtIceException">Thrown if the behavior type is unexpected</exception>
 		public static void ConditionalLogPointerUsage(UhtPropertySettings PropertySettings, UhtPointerMemberBehavior EngineBehavior, UhtPointerMemberBehavior NonEngineBehavior,
 			string PointerTypeDesc, IUhtTokenReader TokenReader, int TypeStartPos, string? AlternativeTypeDesc)
 		{
diff --git a/Engine/Source/Programs/Shared/EpicGames.UHT/Types/Properties/UhtObjectPtrProperty.cs b/Engine/Source/Programs/Shared/EpicGames.UHT/Types/Properties/UhtObjectPtrProperty.cs
--- a/Engine/Source/Programs/Shared/EpicGames.UHT/Types/Properties/UhtObjectPtrProperty.cs
+++ b/Engine/Source/Programs/Shared/EpicGames.UHT/Types/Properties/UhtObjectPtrProperty.cs
@@ -8,6 +8,9 @@ using System.Text;
 
 namespace EpicGames.UHT.Types
 {
+	/// <summary>
+	/// FObjectPtrProperty
+	/// </summary>
 	[UnrealHeaderTool]
 	[UhtEngineClass(Name = "ObjectPtrProperty", IsProperty = true)]
 	public class UhtObjectPtrProperty : UhtObjectProperty
@@ -24,8 +27,14 @@ namespace EpicGames.UHT.Types
 		/// <inheritdoc/>
 		protected override UhtPGetArgumentType PGetTypeArgument { get => UhtPGetArgumentType.TypeText; }
 
-		public UhtObjectPtrProperty(UhtPropertySettings PropertySettings, UhtClass PropertyClass, EPropertyFlags ExtraFlags = EPropertyFlags.None)
-			: base(PropertySettings, PropertyClass)
+		/// <summary>
+		/// Construct a new property
+		/// </summary>
+		/// <param name="PropertySettings">Property settings</param>
+		/// <param name="Class">Referenced class</param>
+		/// <param name="ExtraFlags">Extra property flags to apply</param>
+		public UhtObjectPtrProperty(UhtPropertySettings PropertySettings, UhtClass Class, EPropertyFlags ExtraFlags = EPropertyFlags.None)
+			: base(PropertySettings, Class)
 		{
 			this.PropertyFlags |= ExtraFlags | EPropertyFlags.UObjectWrapper;
 			this.PropertyCaps |= UhtPropertyCaps.PassCppArgsByRef;
@@ -91,7 +100,7 @@ namespace EpicGames.UHT.Types
 
 		#region Keyword
 		[UhtPropertyType(Keyword = "TObjectPtr")]
-		public static UhtProperty? ObjectPtrProperty(UhtPropertyResolvePhase ResolvePhase, UhtPropertySettings PropertySettings, IUhtTokenReader TokenReader, UhtToken MatchedToken)
+		private static UhtProperty? ObjectPtrProperty(UhtPropertyResolvePhase ResolvePhase, UhtPropertySettings PropertySettings, IUhtTokenReader TokenReader, UhtToken MatchedToken)
 		{
 			int TypeStartPos = TokenReader.PeekToken().InputStartPos;
 
diff --git a/Engine/Source/Programs/Shared/EpicGames.UHT/Types/Properties/UhtSetProperty.cs b/Engine/Source/Programs/Shared/EpicGames.UHT/Types/Properties/UhtSetProperty.cs
--- a/Engine/Source/Programs/Shared/EpicGames.UHT/Types/Properties/UhtSetProperty.cs
+++ b/Engine/Source/Programs/Shared/EpicGames.UHT/Types/Properties/UhtSetProperty.cs
@@ -10,6 +10,9 @@ using System.Text;
 
 namespace EpicGames.UHT.Types
 {
+	/// <summary>
+	/// FSetProperty
+	/// </summary>
 	[UnrealHeaderTool]
 	[UhtEngineClass(Name = "SetProperty", IsProperty = true)]
 	public class UhtSetProperty : UhtContainerBaseProperty
@@ -26,6 +29,11 @@ namespace EpicGames.UHT.Types
 		/// <inheritdoc/>
 		protected override UhtPGetArgumentType PGetTypeArgument { get => UhtPGetArgumentType.TypeText; }
 
+		/// <summary>
+		/// Construct a new property
+		/// </summary>
+		/// <param name="PropertySettings">Property settings</param>
+		/// <param name="Value">Property key</param>
 		public UhtSetProperty(UhtPropertySettings PropertySettings, UhtProperty Value) : base(PropertySettings, Value)
 		{
 			// If the creation of the value property set more flags, then copy those flags to ourselves
@@ -202,9 +210,9 @@ namespace EpicGames.UHT.Types
 			return null;
 		}
 
-#region Keyword
+		#region Keyword
 		[UhtPropertyType(Keyword = "TSet")]
-		public static UhtProperty? SetProperty(UhtPropertyResolvePhase ResolvePhase, UhtPropertySettings PropertySettings, IUhtTokenReader TokenReader, UhtToken MatchedToken)
+		private static UhtProperty? SetProperty(UhtPropertyResolvePhase ResolvePhase, UhtPropertySettings PropertySettings, IUhtTokenReader TokenReader, UhtToken MatchedToken)
 		{
 			using (var TokenContext = new UhtMessageContext(TokenReader, "TSet"))
 			{
@@ -245,6 +253,6 @@ namespace EpicGames.UHT.Types
 				return new UhtSetProperty(PropertySettings, Value);
 			}
 		}
-#endregion
+		#endregion
 	}
 }
diff --git a/Engine/Source/Programs/Shared/EpicGames.UHT/Types/Properties/UhtSoftClassProperty.cs b/Engine/Source/Programs/Shared/EpicGames.UHT/Types/Properties/UhtSoftClassProperty.cs
--- a/Engine/Source/Programs/Shared/EpicGames.UHT/Types/Properties/UhtSoftClassProperty.cs
+++ b/Engine/Source/Programs/Shared/EpicGames.UHT/Types/Properties/UhtSoftClassProperty.cs
@@ -7,6 +7,10 @@ using System.Text;
 
 namespace EpicGames.UHT.Types
 {
+
+	/// <summary>
+	/// FSoftClassProperty
+	/// </summary>
 	[UnrealHeaderTool]
 	[UhtEngineClass(Name = "SoftClassProperty", IsProperty = true)]
 	public class UhtSoftClassProperty : UhtSoftObjectProperty
@@ -23,8 +27,14 @@ namespace EpicGames.UHT.Types
 		/// <inheritdoc/>
 		protected override UhtPGetArgumentType PGetTypeArgument { get => UhtPGetArgumentType.TypeText; }
 
-		public UhtSoftClassProperty(UhtPropertySettings PropertySettings, UhtClass PropertyClass, UhtClass MetaClass)
-			: base(PropertySettings, PropertyClass, MetaClass)
+		/// <summary>
+		/// Construct a new class property
+		/// </summary>
+		/// <param name="PropertySettings">Property setting</param>
+		/// <param name="Class">Referenced class (Always UClass)</param>
+		/// <param name="MetaClass">Referenced meta class</param>
+		public UhtSoftClassProperty(UhtPropertySettings PropertySettings, UhtClass Class, UhtClass MetaClass)
+			: base(PropertySettings, Class, MetaClass)
 		{
 		}
 
@@ -75,7 +85,7 @@ namespace EpicGames.UHT.Types
 
 		#region Keyword
 		[UhtPropertyType(Keyword = "TSoftClassPtr")]
-		public static UhtProperty? SoftClassPtrProperty(UhtPropertyResolvePhase ResolvePhase, UhtPropertySettings PropertySettings, IUhtTokenReader TokenReader, UhtToken MatchedToken)
+		private static UhtProperty? SoftClassPtrProperty(UhtPropertyResolvePhase ResolvePhase, UhtPropertySettings PropertySettings, IUhtTokenReader TokenReader, UhtToken MatchedToken)
 		{
 			UhtClass? MetaClass = ParseTemplateClass(PropertySettings, TokenReader, MatchedToken);
 			if (MetaClass == null)
diff --git a/Engine/Source/Programs/Shared/EpicGames.UHT/Types/Properties/UhtSoftObjectProperty.cs b/Engine/Source/Programs/Shared/EpicGames.UHT/Types/Properties/UhtSoftObjectProperty.cs
--- a/Engine/Source/Programs/Shared/EpicGames.UHT/Types/Properties/UhtSoftObjectProperty.cs
+++ b/Engine/Source/Programs/Shared/EpicGames.UHT/Types/Properties/UhtSoftObjectProperty.cs
@@ -8,6 +8,9 @@ using System.Text;
 
 namespace EpicGames.UHT.Types
 {
+	/// <summary>
+	/// FSoftObjectProperty
+	/// </summary>
 	[UnrealHeaderTool]
 	[UhtEngineClass(Name = "SoftObjectProperty", IsProperty = true)]
 	public class UhtSoftObjectProperty : UhtObjectPropertyBase
@@ -24,16 +27,14 @@ namespace EpicGames.UHT.Types
 		/// <inheritdoc/>
 		protected override UhtPGetArgumentType PGetTypeArgument { get => UhtPGetArgumentType.TypeText; }
 
-		public UhtSoftObjectProperty(UhtPropertySettings PropertySettings, UhtClass PropertyClass)
-			: base(PropertySettings, PropertyClass, null)
-		{
-			this.PropertyFlags |= EPropertyFlags.UObjectWrapper;
-			this.PropertyCaps |= UhtPropertyCaps.PassCppArgsByRef | UhtPropertyCaps.RequiresNullConstructorArg | UhtPropertyCaps.CanExposeOnSpawn |
-				UhtPropertyCaps.CanHaveConfig | UhtPropertyCaps.IsParameterSupportedByBlueprint | UhtPropertyCaps.IsMemberSupportedByBlueprint;
-		}
-
-		protected UhtSoftObjectProperty(UhtPropertySettings PropertySettings, UhtClass PropertyClass, UhtClass MetaClass)
-			: base(PropertySettings, PropertyClass, MetaClass)
+		/// <summary>
+		/// Construct a new property
+		/// </summary>
+		/// <param name="PropertySettings">Property settings</param>
+		/// <param name="Class">UCLASS being referenced</param>
+		/// <param name="MetaClass">Optional meta class (used by SoftClassProperty)</param>
+		public UhtSoftObjectProperty(UhtPropertySettings PropertySettings, UhtClass Class, UhtClass? MetaClass = null)
+			: base(PropertySettings, Class, MetaClass)
 		{
 			this.PropertyFlags |= EPropertyFlags.UObjectWrapper;
 			this.PropertyCaps |= UhtPropertyCaps.PassCppArgsByRef | UhtPropertyCaps.RequiresNullConstructorArg | UhtPropertyCaps.CanExposeOnSpawn |
@@ -92,7 +93,7 @@ namespace EpicGames.UHT.Types
 
 		#region Keyword
 		[UhtPropertyType(Keyword = "TSoftObjectPtr")]
-		public static UhtProperty? SoftObjectPtrProperty(UhtPropertyResolvePhase ResolvePhase, UhtPropertySettings PropertySettings, IUhtTokenReader TokenReader, UhtToken MatchedToken)
+		private static UhtProperty? SoftObjectPtrProperty(UhtPropertyResolvePhase ResolvePhase, UhtPropertySettings PropertySettings, IUhtTokenReader TokenReader, UhtToken MatchedToken)
 		{
 			UhtClass? PropertyClass = UhtObjectPropertyBase.ParseTemplateObject(PropertySettings, TokenReader, MatchedToken, true);
 			if (PropertyClass == null)
diff --git a/Engine/Source/Programs/Shared/EpicGames.UHT/Types/Properties/UhtStrProperty.cs b/Engine/Source/Programs/Shared/EpicGames.UHT/Types/Properties/UhtStrProperty.cs
--- a/Engine/Source/Programs/Shared/EpicGames.UHT/Types/Properties/UhtStrProperty.cs
+++ b/Engine/Source/Programs/Shared/EpicGames.UHT/Types/Properties/UhtStrProperty.cs
@@ -8,6 +8,10 @@ using System.Text;
 
 namespace EpicGames.UHT.Types
 {
+
+	/// <summary>
+	/// FStrProperty
+	/// </summary>
 	[UnrealHeaderTool]
 	[UhtEngineClass(Name = "StrProperty", IsProperty = true)]
 	public class UhtStrProperty : UhtProperty
@@ -24,6 +28,10 @@ namespace EpicGames.UHT.Types
 		/// <inheritdoc/>
 		protected override UhtPGetArgumentType PGetTypeArgument { get => UhtPGetArgumentType.EngineClass; }
 
+		/// <summary>
+		/// Construct a new property
+		/// </summary>
+		/// <param name="PropertySettings">Property settings</param>
 		public UhtStrProperty(UhtPropertySettings PropertySettings) : base(PropertySettings)
 		{
 			this.PropertyCaps |= UhtPropertyCaps.PassCppArgsByRef | UhtPropertyCaps.CanExposeOnSpawn | UhtPropertyCaps.IsParameterSupportedByBlueprint | UhtPropertyCaps.IsMemberSupportedByBlueprint;
@@ -99,7 +107,7 @@ namespace EpicGames.UHT.Types
 
 		[UhtPropertyType(Keyword = "FString")]
 		[UhtPropertyType(Keyword = "FMemoryImageString", Options = UhtPropertyTypeOptions.Immediate)]
-		public static UhtProperty? StrProperty(UhtPropertyResolvePhase ResolvePhase, UhtPropertySettings PropertySettings, IUhtTokenReader TokenReader, UhtToken MatchedToken)
+		private static UhtProperty? StrProperty(UhtPropertyResolvePhase ResolvePhase, UhtPropertySettings PropertySettings, IUhtTokenReader TokenReader, UhtToken MatchedToken)
 		{
 			if (!TokenReader.SkipExpectedType(MatchedToken.Value, PropertySettings.PropertyCategory == UhtPropertyCategory.Member))
 			{
diff --git a/Engine/Source/Programs/Shared/EpicGames.UHT/Types/Properties/UhtStructProperty.cs b/Engine/Source/Programs/Shared/EpicGames.UHT/Types/Properties/UhtStructProperty.cs
--- a/Engine/Source/Programs/Shared/EpicGames.UHT/Types/Properties/UhtStructProperty.cs
+++ b/Engine/Source/Programs/Shared/EpicGames.UHT/Types/Properties/UhtStructProperty.cs
@@ -10,6 +10,9 @@ using System.Text.Json.Serialization;
 
 namespace EpicGames.UHT.Types
 {
+	/// <summary>
+	/// FStructProperty
+	/// </summary>
 	[UnrealHeaderTool]
 	[UhtEngineClass(Name = "StructProperty", IsProperty = true)]
 	public class UhtStructProperty : UhtProperty
@@ -26,9 +29,17 @@ namespace EpicGames.UHT.Types
 		/// <inheritdoc/>
 		protected override UhtPGetArgumentType PGetTypeArgument { get => UhtPGetArgumentType.TypeText; }
 
+		/// <summary>
+		/// USTRUCT referenced by the property
+		/// </summary>
 		[JsonConverter(typeof(UhtTypeSourceNameJsonConverter<UhtScriptStruct>))]
 		public UhtScriptStruct ScriptStruct { get; set; }
 
+		/// <summary>
+		/// Construct property
+		/// </summary>
+		/// <param name="PropertySettings">Property settings</param>
+		/// <param name="ScriptStruct">USTRUCT being referenced</param>
 		public UhtStructProperty(UhtPropertySettings PropertySettings, UhtScriptStruct ScriptStruct) : base(PropertySettings)
 		{
 			this.ScriptStruct = ScriptStruct;
@@ -201,7 +212,7 @@ namespace EpicGames.UHT.Types
 
 		#region Structure default value sanitizers
 		[UhtStructDefaultValue(Name = "FVector")]
-		public static bool VectorStructDefaultValue(UhtStructProperty Property, IUhtTokenReader DefaultValueReader, StringBuilder InnerDefaultValue)
+		private static bool VectorStructDefaultValue(UhtStructProperty Property, IUhtTokenReader DefaultValueReader, StringBuilder InnerDefaultValue)
 		{
 			string Format = "{0:F6},{1:F6},{2:F6}";
 
@@ -241,7 +252,7 @@ namespace EpicGames.UHT.Types
 		}
 
 		[UhtStructDefaultValue(Name = "FRotator")]
-		public static bool RotatorStructDefaultValue(UhtStructProperty Property, IUhtTokenReader DefaultValueReader, StringBuilder InnerDefaultValue)
+		private static bool RotatorStructDefaultValue(UhtStructProperty Property, IUhtTokenReader DefaultValueReader, StringBuilder InnerDefaultValue)
 		{
 			DefaultValueReader.Require("FRotator");
 			if (DefaultValueReader.TryOptional("::"))
@@ -273,7 +284,7 @@ namespace EpicGames.UHT.Types
 		}
 
 		[UhtStructDefaultValue(Name = "FVector2D")]
-		public static bool Vector2DStructDefaultValue(UhtStructProperty Property, IUhtTokenReader DefaultValueReader, StringBuilder InnerDefaultValue)
+		private static bool Vector2DStructDefaultValue(UhtStructProperty Property, IUhtTokenReader DefaultValueReader, StringBuilder InnerDefaultValue)
 		{
 			string Format = "(X={0:F3},Y={1:F3})";
 
@@ -306,7 +317,7 @@ namespace EpicGames.UHT.Types
 		}
 
 		[UhtStructDefaultValue(Name = "FLinearColor")]
-		public static bool LinearColorStructDefaultValue(UhtStructProperty Property, IUhtTokenReader DefaultValueReader, StringBuilder InnerDefaultValue)
+		private static bool LinearColorStructDefaultValue(UhtStructProperty Property, IUhtTokenReader DefaultValueReader, StringBuilder InnerDefaultValue)
 		{
 			string Format = "(R={0:F6},G={1:F6},B={2:F6},A={3:F6})";
 
@@ -352,7 +363,7 @@ namespace EpicGames.UHT.Types
 		}
 
 		[UhtStructDefaultValue(Name = "FColor")]
-		public static bool ColorStructDefaultValue(UhtStructProperty Property, IUhtTokenReader DefaultValueReader, StringBuilder InnerDefaultValue)
+		private static bool ColorStructDefaultValue(UhtStructProperty Property, IUhtTokenReader DefaultValueReader, StringBuilder InnerDefaultValue)
 		{
 			string Format = "(R={0},G={1},B={2},A={3})";
 
@@ -398,7 +409,7 @@ namespace EpicGames.UHT.Types
 		}
 
 		[UhtStructDefaultValue(Options = UhtStructDefaultValueOptions.Default)]
-		public static bool DefaultStructDefaultValue(UhtStructProperty Property, IUhtTokenReader DefaultValueReader, StringBuilder InnerDefaultValue)
+		private static bool DefaultStructDefaultValue(UhtStructProperty Property, IUhtTokenReader DefaultValueReader, StringBuilder InnerDefaultValue)
 		{
 			DefaultValueReader
 				.Require(Property.ScriptStruct.SourceName)
diff --git a/Engine/Source/Programs/Shared/EpicGames.UHT/Types/Properties/UhtTextProperty.cs b/Engine/Source/Programs/Shared/EpicGames.UHT/Types/Properties/UhtTextProperty.cs
--- a/Engine/Source/Programs/Shared/EpicGames.UHT/Types/Properties/UhtTextProperty.cs
+++ b/Engine/Source/Programs/Shared/EpicGames.UHT/Types/Properties/UhtTextProperty.cs
@@ -9,6 +9,9 @@ using System.Text;
 
 namespace EpicGames.UHT.Types
 {
+	/// <summary>
+	/// FTextProperty
+	/// </summary>
 	[UnrealHeaderTool]
 	[UhtEngineClass(Name = "TextProperty", IsProperty = true)]
 	public class UhtTextProperty : UhtProperty
@@ -25,6 +28,10 @@ namespace EpicGames.UHT.Types
 		/// <inheritdoc/>
 		protected override UhtPGetArgumentType PGetTypeArgument { get => UhtPGetArgumentType.EngineClass; }
 
+		/// <summary>
+		/// Construct a new property
+		/// </summary>
+		/// <param name="PropertySettings">Property settings</param>
 		public UhtTextProperty(UhtPropertySettings PropertySettings) : base(PropertySettings)
 		{
 			this.PropertyCaps |= UhtPropertyCaps.PassCppArgsByRef | UhtPropertyCaps.CanExposeOnSpawn | UhtPropertyCaps.IsParameterSupportedByBlueprint | UhtPropertyCaps.IsMemberSupportedByBlueprint;
@@ -120,7 +127,7 @@ namespace EpicGames.UHT.Types
 		}
 
 		[UhtPropertyType(Keyword = "Text", Options = UhtPropertyTypeOptions.Simple | UhtPropertyTypeOptions.Immediate | UhtPropertyTypeOptions.CaseInsensitive)]
-		public static UhtProperty? MissingPrefixTextProperty(UhtPropertyResolvePhase ResolvePhase, UhtPropertySettings PropertySettings, IUhtTokenReader TokenReader, UhtToken MatchedToken)
+		private static UhtProperty? MissingPrefixTextProperty(UhtPropertyResolvePhase ResolvePhase, UhtPropertySettings PropertySettings, IUhtTokenReader TokenReader, UhtToken MatchedToken)
 		{
 			throw new UhtException(TokenReader, "'Text' is missing a prefix, expecting 'FText'");
 		}
diff --git a/Engine/Source/Programs/Shared/EpicGames.UHT/Types/Properties/UhtUInt16Property.cs b/Engine/Source/Programs/Shared/EpicGames.UHT/Types/Properties/UhtUInt16Property.cs
--- a/Engine/Source/Programs/Shared/EpicGames.UHT/Types/Properties/UhtUInt16Property.cs
+++ b/Engine/Source/Programs/Shared/EpicGames.UHT/Types/Properties/UhtUInt16Property.cs
@@ -7,6 +7,10 @@ using System.Text;
 
 namespace EpicGames.UHT.Types
 {
+
+	/// <summary>
+	/// FUInt16Property
+	/// </summary>
 	[UnrealHeaderTool]
 	[UhtEngineClass(Name = "UInt16Property", IsProperty = true)]
 	public class UhtUInt16Property : UhtNumericProperty
@@ -17,6 +21,11 @@ namespace EpicGames.UHT.Types
 		/// <inheritdoc/>
 		protected override string CppTypeText { get => "uint16"; }
 
+		/// <summary>
+		/// Construct a new property
+		/// </summary>
+		/// <param name="PropertySettings">Property settings</param>
+		/// <param name="IntType">Integer type</param>
 		public UhtUInt16Property(UhtPropertySettings PropertySettings, UhtPropertyIntType IntType) : base(PropertySettings, IntType)
 		{
 		}
@@ -35,11 +44,13 @@ namespace EpicGames.UHT.Types
 			return Builder;
 		}
 
+		/// <inheritdoc/>
 		public override bool SanitizeDefaultValue(IUhtTokenReader DefaultValueReader, StringBuilder InnerDefaultValue)
 		{
 			return false;
 		}
 
+		/// <inheritdoc/>
 		public override bool IsSameType(UhtProperty Other)
 		{
 			return Other is UhtUInt16Property;
@@ -47,7 +58,7 @@ namespace EpicGames.UHT.Types
 
 		#region Keyword
 		[UhtPropertyType(Keyword = "uint16", Options = UhtPropertyTypeOptions.Simple | UhtPropertyTypeOptions.Immediate)]
-		public static UhtProperty? UInt16Property(UhtPropertyResolvePhase ResolvePhase, UhtPropertySettings PropertySettings, IUhtTokenReader TokenReader, UhtToken MatchedToken)
+		private static UhtProperty? UInt16Property(UhtPropertyResolvePhase ResolvePhase, UhtPropertySettings PropertySettings, IUhtTokenReader TokenReader, UhtToken MatchedToken)
 		{
 			if (PropertySettings.bIsBitfield)
 			{
diff --git a/Engine/Source/Programs/Shared/EpicGames.UHT/Types/Properties/UhtUInt32Property.cs b/Engine/Source/Programs/Shared/EpicGames.UHT/Types/Properties/UhtUInt32Property.cs
--- a/Engine/Source/Programs/Shared/EpicGames.UHT/Types/Properties/UhtUInt32Property.cs
+++ b/Engine/Source/Programs/Shared/EpicGames.UHT/Types/Properties/UhtUInt32Property.cs
@@ -7,6 +7,10 @@ using System.Text;
 
 namespace EpicGames.UHT.Types
 {
+
+	/// <summary>
+	/// FUInt32Property
+	/// </summary>
 	[UnrealHeaderTool]
 	[UhtEngineClass(Name = "UInt32Property", IsProperty = true)]
 	public class UhtUInt32Property : UhtNumericProperty
@@ -17,6 +21,11 @@ namespace EpicGames.UHT.Types
 		/// <inheritdoc/>
 		protected override string CppTypeText { get => "uint32"; }
 
+		/// <summary>
+		/// Construct a new property
+		/// </summary>
+		/// <param name="PropertySettings">Property settings</param>
+		/// <param name="IntType">Integer type</param>
 		public UhtUInt32Property(UhtPropertySettings PropertySettings, UhtPropertyIntType IntType) : base(PropertySettings, IntType)
 		{
 		}
@@ -38,11 +47,13 @@ namespace EpicGames.UHT.Types
 			return Builder;
 		}
 
+		/// <inheritdoc/>
 		public override bool SanitizeDefaultValue(IUhtTokenReader DefaultValueReader, StringBuilder InnerDefaultValue)
 		{
 			return false;
 		}
 
+		/// <inheritdoc/>
 		public override bool IsSameType(UhtProperty Other)
 		{
 			return Other is UhtUInt32Property;
@@ -50,7 +61,7 @@ namespace EpicGames.UHT.Types
 
 		#region Keyword
 		[UhtPropertyType(Keyword = "uint32", Options = UhtPropertyTypeOptions.Simple | UhtPropertyTypeOptions.Immediate)]
-		public static UhtProperty? UInt32Property(UhtPropertyResolvePhase ResolvePhase, UhtPropertySettings PropertySettings, IUhtTokenReader TokenReader, UhtToken MatchedToken)
+		private static UhtProperty? UInt32Property(UhtPropertyResolvePhase ResolvePhase, UhtPropertySettings PropertySettings, IUhtTokenReader TokenReader, UhtToken MatchedToken)
 		{
 			if (PropertySettings.bIsBitfield)
 			{
@@ -63,7 +74,7 @@ namespace EpicGames.UHT.Types
 		}
 
 		[UhtPropertyType(Keyword = "unsigned", Options = UhtPropertyTypeOptions.Immediate)]
-		public static UhtProperty? UnsignedProperty(UhtPropertyResolvePhase ResolvePhase, UhtPropertySettings PropertySettings, IUhtTokenReader TokenReader, UhtToken MatchedToken)
+		private static UhtProperty? UnsignedProperty(UhtPropertyResolvePhase ResolvePhase, UhtPropertySettings PropertySettings, IUhtTokenReader TokenReader, UhtToken MatchedToken)
 		{
 			TokenReader
 				.Require("unsigned")
diff --git a/Engine/Source/Programs/Shared/EpicGames.UHT/Types/Properties/UhtUInt64Property.cs b/Engine/Source/Programs/Shared/EpicGames.UHT/Types/Properties/UhtUInt64Property.cs
--- a/Engine/Source/Programs/Shared/EpicGames.UHT/Types/Properties/UhtUInt64Property.cs
+++ b/Engine/Source/Programs/Shared/EpicGames.UHT/Types/Properties/UhtUInt64Property.cs
@@ -7,6 +7,10 @@ using System.Text;
 
 namespace EpicGames.UHT.Types
 {
+
+	/// <summary>
+	/// FUInt64Property
+	/// </summary>
 	[UnrealHeaderTool]
 	[UhtEngineClass(Name = "UInt64Property", IsProperty = true)]
 	public class UhtUInt64Property : UhtNumericProperty
@@ -17,6 +21,11 @@ namespace EpicGames.UHT.Types
 		/// <inheritdoc/>
 		protected override string CppTypeText { get => "uint64"; }
 
+		/// <summary>
+		/// Construct a new property
+		/// </summary>
+		/// <param name="PropertySettings">Property settings</param>
+		/// <param name="IntType">Integer type</param>
 		public UhtUInt64Property(UhtPropertySettings PropertySettings, UhtPropertyIntType IntType) : base(PropertySettings, IntType)
 		{
 		}
@@ -35,11 +44,13 @@ namespace EpicGames.UHT.Types
 			return Builder;
 		}
 
+		/// <inheritdoc/>
 		public override bool SanitizeDefaultValue(IUhtTokenReader DefaultValueReader, StringBuilder InnerDefaultValue)
 		{
 			return false;
 		}
 
+		/// <inheritdoc/>
 		public override bool IsSameType(UhtProperty Other)
 		{
 			return Other is UhtUInt64Property;
@@ -47,7 +58,7 @@ namespace EpicGames.UHT.Types
 
 		#region Keyword
 		[UhtPropertyType(Keyword = "uint64", Options = UhtPropertyTypeOptions.Simple | UhtPropertyTypeOptions.Immediate)]
-		public static UhtProperty? UInt64Property(UhtPropertyResolvePhase ResolvePhase, UhtPropertySettings PropertySettings, IUhtTokenReader TokenReader, UhtToken MatchedToken)
+		private static UhtProperty? UInt64Property(UhtPropertyResolvePhase ResolvePhase, UhtPropertySettings PropertySettings, IUhtTokenReader TokenReader, UhtToken MatchedToken)
 		{
 			if (PropertySettings.bIsBitfield)
 			{
diff --git a/Engine/Source/Programs/Shared/EpicGames.UHT/Types/Properties/UhtVoidProperty.cs b/Engine/Source/Programs/Shared/EpicGames.UHT/Types/Properties/UhtVoidProperty.cs
--- a/Engine/Source/Programs/Shared/EpicGames.UHT/Types/Properties/UhtVoidProperty.cs
+++ b/Engine/Source/Programs/Shared/EpicGames.UHT/Types/Properties/UhtVoidProperty.cs
@@ -8,6 +8,10 @@ using System.Text;
 
 namespace EpicGames.UHT.Types
 {
+
+	/// <summary>
+	/// Internal only void property for void return types
+	/// </summary>
 	[UnrealHeaderTool]
 	public class UhtVoidProperty : UhtProperty
 	{
@@ -20,6 +24,10 @@ namespace EpicGames.UHT.Types
 		/// <inheritdoc/>
 		protected override string PGetMacroText { get => "invalid"; }
 
+		/// <summary>
+		/// Construct a new void property
+		/// </summary>
+		/// <param name="PropertySettings">Property settings</param>
 		public UhtVoidProperty(UhtPropertySettings PropertySettings) : base(PropertySettings)
 		{
 		}
@@ -68,7 +76,7 @@ namespace EpicGames.UHT.Types
 
 		#region Keyword
 		[UhtPropertyType(Keyword = "void", Options = UhtPropertyTypeOptions.Simple | UhtPropertyTypeOptions.Immediate)]
-		public static UhtProperty? VoidProperty(UhtPropertyResolvePhase ResolvePhase, UhtPropertySettings PropertySettings, IUhtTokenReader TokenReader, UhtToken MatchedToken)
+		private static UhtProperty? VoidProperty(UhtPropertyResolvePhase ResolvePhase, UhtPropertySettings PropertySettings, IUhtTokenReader TokenReader, UhtToken MatchedToken)
 		{
 			if (PropertySettings.PropertyCategory != UhtPropertyCategory.Return)
 			{
diff --git a/Engine/Source/Programs/Shared/EpicGames.UHT/Types/Properties/UhtWeakObjectPtrProperty.cs b/Engine/Source/Programs/Shared/EpicGames.UHT/Types/Properties/UhtWeakObjectPtrProperty.cs
--- a/Engine/Source/Programs/Shared/EpicGames.UHT/Types/Properties/UhtWeakObjectPtrProperty.cs
+++ b/Engine/Source/Programs/Shared/EpicGames.UHT/Types/Properties/UhtWeakObjectPtrProperty.cs
@@ -8,6 +8,10 @@ using System.Text;
 
 namespace EpicGames.UHT.Types
 {
+
+	/// <summary>
+	/// Represents a FWeakObjectProperty
+	/// </summary>
 	[UnrealHeaderTool]
 	[UhtEngineClass(Name = "WeakObjectProperty", IsProperty = true)]
 	public class UhtWeakObjectPtrProperty : UhtObjectPropertyBase
@@ -24,6 +28,12 @@ namespace EpicGames.UHT.Types
 		/// <inheritdoc/>
 		protected override UhtPGetArgumentType PGetTypeArgument { get => UhtPGetArgumentType.TypeText; }
 
+		/// <summary>
+		/// Construct a new property
+		/// </summary>
+		/// <param name="PropertySettings">Property settings</param>
+		/// <param name="PropertyClass">Class being referenced</param>
+		/// <param name="ExtraFlags">Extra property flags to add to the definition</param>
 		public UhtWeakObjectPtrProperty(UhtPropertySettings PropertySettings, UhtClass PropertyClass, EPropertyFlags ExtraFlags = EPropertyFlags.None)
 			: base(PropertySettings, PropertyClass, null)
 		{
@@ -97,22 +107,22 @@ namespace EpicGames.UHT.Types
 		}
 
 		#region Keywords
-		protected static UhtProperty CreateWeakProperty(UhtPropertySettings PropertySettings, IUhtTokenReader TokenReader, UhtClass PropertyClass, EPropertyFlags ExtraFlags = EPropertyFlags.None)
+		private static UhtProperty CreateWeakProperty(UhtPropertySettings PropertySettings, IUhtTokenReader TokenReader, UhtClass Class, EPropertyFlags ExtraFlags = EPropertyFlags.None)
 		{
-			if (PropertyClass.IsChildOf(PropertyClass.Session.UClass))
+			if (Class.IsChildOf(Class.Session.UClass))
 			{
 				TokenReader.LogError("Class variables cannot be weak, they are always strong.");
 			}
 
 			if (PropertySettings.DisallowPropertyFlags.HasAnyFlags(EPropertyFlags.AutoWeak))
 			{
-				return new UhtObjectProperty(PropertySettings, PropertyClass, ExtraFlags | EPropertyFlags.UObjectWrapper);
+				return new UhtObjectProperty(PropertySettings, Class, null, ExtraFlags | EPropertyFlags.UObjectWrapper);
 			}
-			return new UhtWeakObjectPtrProperty(PropertySettings, PropertyClass, ExtraFlags);
+			return new UhtWeakObjectPtrProperty(PropertySettings, Class, ExtraFlags);
 		}
 
 		[UhtPropertyType(Keyword = "TWeakObjectPtr")]
-		public static UhtProperty? WeakObjectProperty(UhtPropertyResolvePhase ResolvePhase, UhtPropertySettings PropertySettings, IUhtTokenReader TokenReader, UhtToken MatchedToken)
+		private static UhtProperty? WeakObjectProperty(UhtPropertyResolvePhase ResolvePhase, UhtPropertySettings PropertySettings, IUhtTokenReader TokenReader, UhtToken MatchedToken)
 		{
 			UhtClass? PropertyClass = ParseTemplateObject(PropertySettings, TokenReader, MatchedToken, true);
 			if (PropertyClass == null)
@@ -124,7 +134,7 @@ namespace EpicGames.UHT.Types
 		}
 
 		[UhtPropertyType(Keyword = "TAutoWeakObjectPtr")]
-		public static UhtProperty? AutoWeakObjectProperty(UhtPropertyResolvePhase ResolvePhase, UhtPropertySettings PropertySettings, IUhtTokenReader TokenReader, UhtToken MatchedToken)
+		private static UhtProperty? AutoWeakObjectProperty(UhtPropertyResolvePhase ResolvePhase, UhtPropertySettings PropertySettings, IUhtTokenReader TokenReader, UhtToken MatchedToken)
 		{
 			UhtClass? PropertyClass = ParseTemplateObject(PropertySettings, TokenReader, MatchedToken, true);
 			if (PropertyClass == null)
diff --git a/Engine/Source/Programs/Shared/EpicGames.UHT/Types/UhtEnum.cs b/Engine/Source/Programs/Shared/EpicGames.UHT/Types/UhtEnum.cs
--- a/Engine/Source/Programs/Shared/EpicGames.UHT/Types/UhtEnum.cs
+++ b/Engine/Source/Programs/Shared/EpicGames.UHT/Types/UhtEnum.cs
@@ -9,50 +9,138 @@ using System.Text.Json.Serialization;
 
 namespace EpicGames.UHT.Types
 {
+	
+	/// <summary>
+	/// How the enumeration was declared
+	/// </summary>
 	public enum UhtEnumCppForm
 	{
+		/// <summary>
+		/// enum Name {...}
+		/// </summary>
 		Regular,
+
+		/// <summary>
+		/// namespace Name { enum Type { ... } }
+		/// </summary>
 		Namespaced,
+
+		/// <summary>
+		/// enum class Name {...}
+		/// </summary>
 		EnumClass
 	}
 
+	/// <summary>
+	/// Underlying type of the enumeration
+	/// </summary>
 	public enum UhtEnumUnderlyingType
 	{
+
+		/// <summary>
+		/// Not specified
+		/// </summary>
 		Unspecified,
+
+		/// <summary>
+		/// uint8
+		/// </summary>
 		uint8,
+
+		/// <summary>
+		/// uint16
+		/// </summary>
 		uint16,
+
+		/// <summary>
+		/// uint32
+		/// </summary>
 		uint32,
+
+		/// <summary>
+		/// uint64
+		/// </summary>
 		uint64,
+
+		/// <summary>
+		/// int8
+		/// </summary>
 		int8,
+
+		/// <summary>
+		/// int16
+		/// </summary>
 		int16,
+
+		/// <summary>
+		/// int32
+		/// </summary>
 		int32,
+
+		/// <summary>
+		/// int64
+		/// </summary>
 		int64
 	}
 
+	/// <summary>
+	/// Represents an enumeration value
+	/// </summary>
 	public struct UhtEnumValue
 	{
-		public string Name { get; internal set; }
-		public long Value { get; internal set; }
+		/// <summary>
+		/// Name of the enumeration value
+		/// </summary>
+		public string Name { get; set; }
+
+		/// <summary>
+		/// Value of the enumeration or -1 if not parsed.
+		/// </summary>
+		public long Value { get; set; }
 	}
 
+	/// <summary>
+	/// Represents a UEnum
+	/// </summary>
 	[UhtEngineClass(Name = "Enum")]
 	public class UhtEnum : UhtField, IUhtMetaDataKeyConversion
 	{
 		private static UhtSpecifierValidatorTable EnumSpecifierValidatorTable = UhtSpecifierValidatorTables.Instance.Get(UhtTableNames.Enum);
 
+		/// <summary>
+		/// Engine enumeration flags
+		/// </summary>
 		[JsonConverter(typeof(JsonStringEnumConverter))]
-		public EEnumFlags EnumFlags { get; internal set; } = EEnumFlags.None;
+		public EEnumFlags EnumFlags { get; set; } = EEnumFlags.None;
 
+		/// <summary>
+		/// C++ form of the enumeration
+		/// </summary>
 		[JsonConverter(typeof(JsonStringEnumConverter))]
-		public UhtEnumCppForm CppForm { get; internal set; } = UhtEnumCppForm.Regular;
+		public UhtEnumCppForm CppForm { get; set; } = UhtEnumCppForm.Regular;
 
+		/// <summary>
+		/// Underlying integer enumeration type
+		/// </summary>
 		[JsonConverter(typeof(JsonStringEnumConverter))]
-		public UhtEnumUnderlyingType UnderlyingType { get; internal set; } = UhtEnumUnderlyingType.uint8;
+		public UhtEnumUnderlyingType UnderlyingType { get; set; } = UhtEnumUnderlyingType.uint8;
 
-		public string CppType { get; internal set; } = String.Empty;
-		public bool bIsEditorOnly { get; internal set; } = false;
+		/// <summary>
+		/// Full enumeration type.  For namespace enumerations, this includes the namespace name and the enum type name
+		/// </summary>
+		public string CppType { get; set; } = String.Empty;
+
+		/// <summary>
+		/// True if the enum was declared in an WITH_EDITORONLY_DATA block
+		/// </summary>
+		public bool bIsEditorOnly { get; set; } = false;
+
+		/// <summary>
+		/// Collection of enumeration values
+		/// </summary>
 		public List<UhtEnumValue> EnumValues { get; set; }
 
+		/// <inheritdoc/>
 		[JsonIgnore]
 		public override UhtEngineType EngineType => UhtEngineType.Enum;
 
@@ -63,12 +151,22 @@ namespace EpicGames.UHT.Types
 		[JsonIgnore]
 		protected override UhtSpecifierValidatorTable? SpecifierValidatorTable { get => UhtEnum.EnumSpecifierValidatorTable; }
 
+		/// <summary>
+		/// Construct a new enumeration
+		/// </summary>
+		/// <param name="Outer">Outer type</param>
+		/// <param name="LineNumber">Line number of declaration</param>
 		public UhtEnum(UhtType Outer, int LineNumber) : base(Outer, LineNumber)
 		{
 			this.MetaData.KeyConversion = this;
 			this.EnumValues = new List<UhtEnumValue>();
 		}
 
+		/// <summary>
+		/// Test to see if the value is a known enum value
+		/// </summary>
+		/// <param name="Value">Value in question</param>
+		/// <returns>True if the value is known</returns>
 		public bool IsValidEnumValue(long Value)
 		{
 			foreach (UhtEnumValue EnumValue in EnumValues)
@@ -81,6 +179,11 @@ namespace EpicGames.UHT.Types
 			return false;
 		}
 
+		/// <summary>
+		/// Return the index of the given enumeration value name
+		/// </summary>
+		/// <param name="Name">Value name in question</param>
+		/// <returns>Index of the value or -1 if not found.</returns>
 		public int GetIndexByName(string Name)
 		{
 			Name = CleanEnumValueName(Name);
@@ -94,6 +197,12 @@ namespace EpicGames.UHT.Types
 			return -1;
 		}
 
+		/// <summary>
+		/// Converts meta data name and index to a full meta data key name
+		/// </summary>
+		/// <param name="Name">Meta data key name</param>
+		/// <param name="NameIndex">Meta data key index</param>
+		/// <returns>Meta data name with the enum value name</returns>
 		public string GetMetaDataKey(string Name, int NameIndex)
 		{
 			string EnumName = this.EnumValues[NameIndex].Name;
@@ -108,6 +217,12 @@ namespace EpicGames.UHT.Types
 			return $"{EnumName}.{Name}";
 		}
 
+		/// <summary>
+		/// Given an enumeration value name, return the full enumeration name
+		/// </summary>
+		/// <param name="ShortEnumName">Enum value name</param>
+		/// <returns>If required, enum type name combined with value name.  Otherwise just the value name.</returns>
+		/// <exception cref="UhtIceException">Unexpected enum form</exception>
 		public string GetFullEnumName(string ShortEnumName)
 		{
 			switch (this.CppForm)
@@ -120,10 +235,16 @@ namespace EpicGames.UHT.Types
 					return ShortEnumName;
 
 				default:
-					throw new UhtIceException("Unexpected EEnumCppForm value");
+					throw new UhtIceException("Unexpected UhtEnumCppForm value");
 			}
 		}
 
+		/// <summary>
+		/// Add a new enum value.
+		/// </summary>
+		/// <param name="ShortEnumName">Name of the enum value.</param>
+		/// <param name="Value">Enumeration value or -1 if the value can't be determined.</param>
+		/// <returns></returns>
 		public int AddEnumValue(string ShortEnumName, long Value)
 		{
 			int EnumIndex = this.EnumValues.Count;
@@ -131,6 +252,12 @@ namespace EpicGames.UHT.Types
 			return EnumIndex;
 		}
 
+		/// <summary>
+		/// Reconstruct the full enum name.  Any existing enumeration name will be stripped and replaced 
+		/// with this enumeration name.
+		/// </summary>
+		/// <param name="Name">Name to reconstruct.</param>
+		/// <returns>Reconstructed enum name</returns>
 		private string CleanEnumValueName(string Name)
 		{
 			int LastColons = Name.LastIndexOf("::");
@@ -138,6 +265,7 @@ namespace EpicGames.UHT.Types
 		}
 
 		#region Validation support
+		///<inheritdoc/>
 		protected override void ValidateDocumentationPolicy(UhtDocumentationPolicy Policy)
 		{
 			if (Policy.bClassOrStructCommentRequired)
diff --git a/Engine/Source/Programs/Shared/EpicGames.UHT/Types/UhtField.cs b/Engine/Source/Programs/Shared/EpicGames.UHT/Types/UhtField.cs
--- a/Engine/Source/Programs/Shared/EpicGames.UHT/Types/UhtField.cs
+++ b/Engine/Source/Programs/Shared/EpicGames.UHT/Types/UhtField.cs
@@ -2,11 +2,20 @@
 
 namespace EpicGames.UHT.Types
 {
+
+	/// <summary>
+	/// Represents a UField
+	/// </summary>
 	public abstract class UhtField : UhtObject
 	{
 		/// <inheritdoc/>
 		public override string EngineClassName { get => "Field"; }
 
+		/// <summary>
+		/// Construct a new field
+		/// </summary>
+		/// <param name="Outer">Outer object</param>
+		/// <param name="LineNumber">Line number of declaration</param>
 		public UhtField(UhtType Outer, int LineNumber) : base(Outer, LineNumber)
 		{
 		}
diff --git a/Engine/Source/Programs/Shared/EpicGames.UHT/Types/UhtFunction.cs b/Engine/Source/Programs/Shared/EpicGames.UHT/Types/UhtFunction.cs
--- a/Engine/Source/Programs/Shared/EpicGames.UHT/Types/UhtFunction.cs
+++ b/Engine/Source/Programs/Shared/EpicGames.UHT/Types/UhtFunction.cs
@@ -16,6 +16,10 @@ namespace EpicGames.UHT.Types
 	[Flags]
 	public enum UhtFunctionExportFlags : UInt32
 	{
+
+		/// <summary>
+		/// No export flags
+		/// </summary>
 		None = 0,
 
 		/// <summary>
@@ -143,44 +147,97 @@ namespace EpicGames.UHT.Types
 		}
 	}
 
+	/// <summary>
+	/// Type of function
+	/// </summary>
 	public enum UhtFunctionType
 	{
+
+		/// <summary>
+		/// UFUNCTION
+		/// </summary>
 		Function,
+
+		/// <summary>
+		/// UDELEGATE/DECLARE_DYNAMIC_...
+		/// </summary>
 		Delegate,
+
+		/// <summary>
+		/// UDELEGATE/DECLARE_DYNAMIC_...SPARSE_...
+		/// </summary>
 		SparseDelegate,
 	}
 
+	/// <summary>
+	/// Extension 
+	/// </summary>
 	public static class UhtFunctionTypeExtensions
 	{
+
+		/// <summary>
+		/// Test to see if the function type is a delegate type
+		/// </summary>
+		/// <param name="FunctionType">Function type being tested</param>
+		/// <returns>True if the type is a delegate</returns>
 		public static bool IsDelegate(this UhtFunctionType FunctionType)
 		{
 			return FunctionType == UhtFunctionType.Delegate || FunctionType == UhtFunctionType.SparseDelegate;
 		}
 	}
 
+	/// <summary>
+	/// Represents a UFUNCTION/delegate
+	/// </summary>
 	public class UhtFunction : UhtStruct
 	{
 		private static UhtSpecifierValidatorTable FunctionSpecifierValidatorTable = UhtSpecifierValidatorTables.Instance.Get(UhtTableNames.Function);
 		private string? StrippedFunctionNameInternal = null;
 
+		/// <summary>
+		/// Suffix added to delegate engine names 
+		/// </summary>
 		public static string GeneratedDelegateSignatureSuffix = "__DelegateSignature";
 
+		/// <summary>
+		/// Engine function flags
+		/// </summary>
 		[JsonConverter(typeof(JsonStringEnumConverter))]
 		public EFunctionFlags FunctionFlags { get; set; } = EFunctionFlags.None;
 
+		/// <summary>
+		/// UHT specific function flags
+		/// </summary>
 		[JsonConverter(typeof(JsonStringEnumConverter))]
 		public UhtFunctionExportFlags FunctionExportFlags { get; set; } = UhtFunctionExportFlags.None;
 
+		/// <summary>
+		/// The type of function
+		/// </summary>
 		public UhtFunctionType FunctionType = UhtFunctionType.Function;
 
-		public int FunctionLineNumber { get; set; } = 1;
+		/// <summary>
+		/// The line number for the macro.
+		/// </summary>
 		public int MacroLineNumber { get; set; } = 1;
+
+		/// <summary>
+		/// Owning class name for sparse functions
+		/// </summary>
 		public string? SparseOwningClassName { get; set; } = null;
+
+		/// <summary>
+		/// Sparse delegate name
+		/// </summary>
 		public string? SparseDelegateName { get; set; } = null;
 
+		/// <summary>
+		/// The super function.  Currently unused
+		/// </summary>
 		[JsonConverter(typeof(UhtNullableTypeSourceNameJsonConverter<UhtFunction>))]
 		public UhtFunction? SuperFunction => (UhtFunction?)this.Super;
 
+		/// <inheritdoc/>
 		[JsonIgnore]
 		public override UhtEngineType EngineType => this.FunctionFlags.HasAnyFlags(EFunctionFlags.Delegate) ? UhtEngineType.Delegate : UhtEngineType.Function;
 
@@ -203,6 +260,9 @@ namespace EpicGames.UHT.Types
 			}
 		}
 
+		/// <summary>
+		/// Stripped function name without the generated delegate suffix
+		/// </summary>
 		[JsonIgnore]
 		public string StrippedFunctionName
 		{
@@ -336,6 +396,7 @@ namespace EpicGames.UHT.Types
 		}
 
 		#region Resolution support
+		/// <inheritdoc/>
 		protected override bool ResolveSelf(UhtResolvePhase Phase)
 		{
 			bool bResults = base.ResolveSelf(Phase);
@@ -424,6 +485,7 @@ namespace EpicGames.UHT.Types
 		#endregion
 
 		#region Validation support
+		/// <inheritdoc/>
 		protected override UhtValidationOptions Validate(UhtValidationOptions Options)
 		{
 			Options = base.Validate(Options);
@@ -794,6 +856,7 @@ namespace EpicGames.UHT.Types
 			}
 		}
 
+		/// <inheritdoc/>
 		protected override void ValidateDocumentationPolicy(UhtDocumentationPolicy Policy)
 		{
 			if (this.FunctionType != UhtFunctionType.Function)
diff --git a/Engine/Source/Programs/Shared/EpicGames.UHT/Types/UhtHeaderFile.cs b/Engine/Source/Programs/Shared/EpicGames.UHT/Types/UhtHeaderFile.cs
--- a/Engine/Source/Programs/Shared/EpicGames.UHT/Types/UhtHeaderFile.cs
+++ b/Engine/Source/Programs/Shared/EpicGames.UHT/Types/UhtHeaderFile.cs
@@ -10,11 +10,31 @@ using System.Text.Json.Serialization;
 
 namespace EpicGames.UHT.Types
 {
+
+	/// <summary>
+	/// Represents the different directories where headers can appear
+	/// </summary>
 	public enum UhtHeaderFileType
 	{
+
+		/// <summary>
+		/// Classes folder
+		/// </summary>
 		Classes,
+
+		/// <summary>
+		/// Public folder
+		/// </summary>
 		Public,
+
+		/// <summary>
+		/// Internal folder
+		/// </summary>
 		Internal,
+
+		/// <summary>
+		/// Private folder
+		/// </summary>
 		Private,
 	}
 
@@ -24,6 +44,9 @@ namespace EpicGames.UHT.Types
 	[Flags]
 	public enum UhtHeaderFileExportFlags : UInt32
 	{
+		/// <summary>
+		/// No export flags
+		/// </summary>
 		None = 0,
 
 		/// <summary>
@@ -74,32 +97,81 @@ namespace EpicGames.UHT.Types
 		}
 	}
 
-	public delegate bool UhtReferencedHeaderDelegate(UhtHeaderFile HeaderFile);
-
+	/// <summary>
+	/// Represents a header file.  Unlike the engine, UhtHeader files will appear as a child 
+	/// of a UhtPackage and the outer of all global types found in that header.
+	/// </summary>
 	public class UhtHeaderFile : UhtType
 	{
+		private UhtSimpleMessageSite MessageSite;
+		private UhtSourceFile SourceFile;
+		private List<UhtHeaderFile> ReferencedHeaders = new List<UhtHeaderFile>();
+
+		/// <summary>
+		/// Contents of the header
+		/// </summary>
 		[JsonIgnore]
 		public StringView Data { get => this.SourceFile.Data; }
 
+		/// <summary>
+		/// Path of the header
+		/// </summary>
 		[JsonIgnore]
 		public string FilePath { get => this.SourceFile.FilePath; }
 
+		/// <summary>
+		/// File name without the extension
+		/// </summary>
 		public readonly string FileNameWithoutExtension;
+
+		/// <summary>
+		/// Required name for the generated.h file name.  Used to validate parsed code
+		/// </summary>
 		public readonly string GeneratedHeaderFileName;
+
+		/// <summary>
+		/// True if this header is NoExportTypes.h
+		/// </summary>
 		public readonly bool bIsNoExportTypes;
 
+		/// <summary>
+		/// The file path of the header relative to the module location
+		/// </summary>
 		public string ModuleRelativeFilePath = string.Empty;
+
+		/// <summary>
+		/// Include file path added as meta data to the types
+		/// </summary>
 		public string IncludeFilePath = string.Empty;
+
+		/// <summary>
+		/// Location where the header file was found
+		/// </summary>
 		public UhtHeaderFileType HeaderFileType = UhtHeaderFileType.Private;
+
+		/// <summary>
+		/// Unique index of the header file
+		/// </summary>
 		public readonly int HeaderFileTypeIndex;
 
+		/// <summary>
+		/// UHT flags for the header
+		/// </summary>
 		[JsonConverter(typeof(JsonStringEnumConverter))]
 		public UhtHeaderFileExportFlags HeaderFileExportFlags { get; set; } = UhtHeaderFileExportFlags.None;
 
+		/// <summary>
+		/// If true, the header file should be exported
+		/// </summary>
 		[JsonIgnore]
 		public bool bShouldExport => this.HeaderFileExportFlags.HasAnyFlags(UhtHeaderFileExportFlags.Referenced) || this.Children.Count > 0;
+
+		/// <summary>
+		/// Resource collector for the header file
+		/// </summary>
 		public UhtReferenceCollector References = new UhtReferenceCollector();
 
+		/// <inheritdoc/>
 		[JsonIgnore]
 		public override UhtPackage Package
 		{
@@ -113,32 +185,39 @@ namespace EpicGames.UHT.Types
 			}
 		}
 
+		/// <inheritdoc/>
 		[JsonIgnore]
 		public override UhtHeaderFile HeaderFile { get => this; }
 
+		/// <inheritdoc/>
 		[JsonIgnore]
 		public override UhtEngineType EngineType => UhtEngineType.Header;
 
 		/// <inheritdoc/>
 		public override string EngineClassName { get => "UhtHeaderFile"; }
 
-		private List<UhtHeaderFile> ReferencedHeaders = new List<UhtHeaderFile>();
-
+		/// <summary>
+		/// Collection of headers directly included by this header
+		/// </summary>
 		[JsonIgnore]
 		public List<UhtHeaderFile> IncludedHeaders { get; set; } = new List<UhtHeaderFile>();
 
-		private UhtSimpleMessageSite MessageSite;
-
 		#region IUHTMessageSite implementation
+
+		/// <inheritdoc/>
 		[JsonIgnore]
 		public override IUhtMessageSession MessageSession => this.MessageSite.MessageSession;
 
+		/// <inheritdoc/>
 		[JsonIgnore]
 		public override IUhtMessageSource? MessageSource => this.MessageSite.MessageSource;
 		#endregion
 
-		private UhtSourceFile SourceFile;
-
+		/// <summary>
+		/// Construct a new header file
+		/// </summary>
+		/// <param name="Package">Owning package</param>
+		/// <param name="Path">Path to the header file</param>
 		public UhtHeaderFile(UhtPackage Package, string Path) : base(Package, 1)
 		{
 			this.HeaderFileTypeIndex = this.Session.GetNextHeaderFileTypeIndex();
@@ -151,11 +230,19 @@ namespace EpicGames.UHT.Types
 			this.bIsNoExportTypes = string.Compare(this.SourceFile.FileName, "NoExportTypes", true) == 0;
 		}
 
+		/// <summary>
+		/// Read the contents of the header
+		/// </summary>
 		public void Read()
 		{
 			this.SourceFile.Read();
 		}
 
+		/// <summary>
+		/// Add a reference to the given header
+		/// </summary>
+		/// <param name="Id">Path of the header</param>
+		/// <param name="bIsIncludedFile">True if this is a directly included file</param>
 		public void AddReferencedHeader(string Id, bool bIsIncludedFile)
 		{
 			UhtHeaderFile? HeaderFile = this.Session.FindHeaderFile(Path.GetFileName(Id));
@@ -165,11 +252,20 @@ namespace EpicGames.UHT.Types
 			}
 		}
 
+		/// <summary>
+		/// Add a reference to the header that defines the given type
+		/// </summary>
+		/// <param name="Type">Type in question</param>
 		public void AddReferencedHeader(UhtType Type)
 		{
 			AddReferencedHeader(Type.HeaderFile, false);
 		}
 
+		/// <summary>
+		/// Add a reference to the given header file
+		/// </summary>
+		/// <param name="HeaderFile">Header file in question</param>
+		/// <param name="bIsIncludedFile">True if this is a directly included file</param>
 		public void AddReferencedHeader(UhtHeaderFile HeaderFile, bool bIsIncludedFile)
 		{
 			lock (this.ReferencedHeaders)
@@ -227,6 +323,7 @@ namespace EpicGames.UHT.Types
 			}
 		}
 
+		/// <inheritdoc/>
 		public override void GetPathName(StringBuilder Builder, UhtType? StopOuter = null)
 		{
 			// Headers do not contribute to path names
@@ -236,6 +333,7 @@ namespace EpicGames.UHT.Types
 			}
 		}
 
+		/// <inheritdoc/>
 		protected override UhtValidationOptions Validate(UhtValidationOptions Options)
 		{
 			Options = base.Validate(Options | UhtValidationOptions.Shadowing);
diff --git a/Engine/Source/Programs/Shared/EpicGames.UHT/Types/UhtManifestFile.cs b/Engine/Source/Programs/Shared/EpicGames.UHT/Types/UhtManifestFile.cs
--- a/Engine/Source/Programs/Shared/EpicGames.UHT/Types/UhtManifestFile.cs
+++ b/Engine/Source/Programs/Shared/EpicGames.UHT/Types/UhtManifestFile.cs
@@ -6,23 +6,45 @@ using System.Text.Json;
 
 namespace EpicGames.UHT.Types
 {
+
+	/// <summary>
+	/// Represents the UHT manifest file
+	/// </summary>
 	public class UhtManifestFile : IUhtMessageSite
 	{
+
+		/// <summary>
+		/// Loaded manifest from the json manifest file
+		/// </summary>
 		public UHTManifest? Manifest = null;
 
 		private UhtSourceFile SourceFile;
 
 		#region IUHTMessageSite implementation
+
+		/// <inheritdoc/>
 		public IUhtMessageSession MessageSession => this.SourceFile.MessageSession;
+
+		/// <inheritdoc/>
 		public IUhtMessageSource? MessageSource => this.SourceFile.MessageSource;
+
+		/// <inheritdoc/>
 		public IUhtMessageLineNumber? MessageLineNumber => null;
 		#endregion
 
+		/// <summary>
+		/// Construct a new manifest file
+		/// </summary>
+		/// <param name="Session">Current session</param>
+		/// <param name="FilePath">Path of the file</param>
 		public UhtManifestFile(UhtSession Session, string FilePath)
 		{
 			this.SourceFile = new UhtSourceFile(Session, FilePath);
 		}
 
+		/// <summary>
+		/// Read the contents of the file
+		/// </summary>
 		public void Read() 
 		{
 			this.SourceFile.Read();
diff --git a/Engine/Source/Programs/Shared/EpicGames.UHT/Types/UhtProperty.cs b/Engine/Source/Programs/Shared/EpicGames.UHT/Types/UhtProperty.cs
--- a/Engine/Source/Programs/Shared/EpicGames.UHT/Types/UhtProperty.cs
+++ b/Engine/Source/Programs/Shared/EpicGames.UHT/Types/UhtProperty.cs
@@ -1962,7 +1962,7 @@ namespace EpicGames.UHT.Types
 		}
 
 		/// <summary>
-		/// Append the full declaration including such things as 'const', '*', and '&'
+		/// Append the full declaration including such things as const, *, and &amp;
 		/// </summary>
 		/// <param name="Builder">Destination builder</param>
 		/// <param name="Property">Property in question</param>
