More UHT documentation updates

#rnx

diff --git a/Engine/Source/Programs/Shared/EpicGames.UHT/Types/Properties/UhtByteProperty.cs b/Engine/Source/Programs/Shared/EpicGames.UHT/Types/Properties/UhtByteProperty.cs
--- a/Engine/Source/Programs/Shared/EpicGames.UHT/Types/Properties/UhtByteProperty.cs
+++ b/Engine/Source/Programs/Shared/EpicGames.UHT/Types/Properties/UhtByteProperty.cs
@@ -10,10 +10,16 @@ using System.Text.Json.Serialization;
 
 namespace EpicGames.UHT.Types
 {
+	/// <summary>
+	/// FByteProperty
+	/// </summary>
 	[UnrealHeaderTool]
 	[UhtEngineClass(Name = "ByteProperty", IsProperty = true)]
 	public class UhtByteProperty : UhtNumericProperty
 	{
+		/// <summary>
+		/// Referenced enumeration (TEnumAsByte)
+		/// </summary>
 		[JsonConverter(typeof(UhtNullableTypeSourceNameJsonConverter<UhtEnum>))]
 		public UhtEnum? Enum { get; set; }
 
@@ -29,6 +35,12 @@ namespace EpicGames.UHT.Types
 		/// <inheritdoc/>
 		protected override UhtPGetArgumentType PGetTypeArgument { get => this.Enum == null || this.Enum.CppForm != UhtEnumCppForm.EnumClass ? UhtPGetArgumentType.EngineClass : UhtPGetArgumentType.TypeText; }
 
+		/// <summary>
+		/// Construct a new property
+		/// </summary>
+		/// <param name="PropertySettings">Property settings</param>
+		/// <param name="IntType">Integer type</param>
+		/// <param name="Enum">Optional referenced enum</param>
 		public UhtByteProperty(UhtPropertySettings PropertySettings, UhtPropertyIntType IntType, UhtEnum? Enum = null) : base(PropertySettings, IntType)
 		{
 			this.Enum = Enum;
diff --git a/Engine/Source/Programs/Shared/EpicGames.UHT/Types/Properties/UhtClassProperty.cs b/Engine/Source/Programs/Shared/EpicGames.UHT/Types/Properties/UhtClassProperty.cs
--- a/Engine/Source/Programs/Shared/EpicGames.UHT/Types/Properties/UhtClassProperty.cs
+++ b/Engine/Source/Programs/Shared/EpicGames.UHT/Types/Properties/UhtClassProperty.cs
@@ -8,6 +8,9 @@ using System.Text;
 
 namespace EpicGames.UHT.Types
 {
+	/// <summary>
+	/// FClassProperty
+	/// </summary>
 	[UnrealHeaderTool]
 	[UhtEngineClass(Name = "ClassProperty", IsProperty = true)]
 	public class UhtClassProperty : UhtObjectProperty
@@ -21,8 +24,15 @@ namespace EpicGames.UHT.Types
 		/// <inheritdoc/>
 		protected override UhtPGetArgumentType PGetTypeArgument { get => UhtPGetArgumentType.TypeText; }
 
-		public UhtClassProperty(UhtPropertySettings PropertySettings, UhtClass PropertyClass, UhtClass MetaClass, EPropertyFlags ExtraFlags = EPropertyFlags.None)
-			: base(PropertySettings, PropertyClass, MetaClass)
+		/// <summary>
+		/// Construct a new property
+		/// </summary>
+		/// <param name="PropertySettings">Property settings</param>
+		/// <param name="Class">Referenced class</param>
+		/// <param name="MetaClass">Reference meta class</param>
+		/// <param name="ExtraFlags">Extra flags to apply to the property.</param>
+		public UhtClassProperty(UhtPropertySettings PropertySettings, UhtClass Class, UhtClass MetaClass, EPropertyFlags ExtraFlags = EPropertyFlags.None)
+			: base(PropertySettings, Class, MetaClass)
 		{
 			this.PropertyFlags |= ExtraFlags;
 			this.PropertyCaps |= UhtPropertyCaps.CanHaveConfig;
@@ -147,7 +157,7 @@ namespace EpicGames.UHT.Types
 
 		#region Keyword
 		[UhtPropertyType(Keyword = "TSubclassOf")]
-		public static UhtProperty? SubclassOfProperty(UhtPropertyResolvePhase ResolvePhase, UhtPropertySettings PropertySettings, IUhtTokenReader TokenReader, UhtToken MatchedToken)
+		private static UhtProperty? SubclassOfProperty(UhtPropertyResolvePhase ResolvePhase, UhtPropertySettings PropertySettings, IUhtTokenReader TokenReader, UhtToken MatchedToken)
 		{
 			UhtClass? MetaClass = UhtObjectPropertyBase.ParseTemplateClass(PropertySettings, TokenReader, MatchedToken);
 			if (MetaClass == null)
diff --git a/Engine/Source/Programs/Shared/EpicGames.UHT/Types/Properties/UhtClassPtrProperties.cs b/Engine/Source/Programs/Shared/EpicGames.UHT/Types/Properties/UhtClassPtrProperties.cs
--- a/Engine/Source/Programs/Shared/EpicGames.UHT/Types/Properties/UhtClassPtrProperties.cs
+++ b/Engine/Source/Programs/Shared/EpicGames.UHT/Types/Properties/UhtClassPtrProperties.cs
@@ -7,6 +7,10 @@ using System.Text;
 
 namespace EpicGames.UHT.Types
 {
+
+	/// <summary>
+	/// FClassPtrProperty
+	/// </summary>
 	[UhtEngineClass(Name = "ClassPtrProperty", IsProperty = true)]
 	public class UhtClassPtrProperty : UhtClassProperty
 	{
@@ -22,12 +26,20 @@ namespace EpicGames.UHT.Types
 		/// <inheritdoc/>
 		protected override UhtPGetArgumentType PGetTypeArgument { get => UhtPGetArgumentType.TypeText; }
 
-		public UhtClassPtrProperty(UhtPropertySettings PropertySettings, UhtClass PropertyClass, UhtClass MetaClass, EPropertyFlags ExtraFlags = EPropertyFlags.None)
-			: base(PropertySettings, PropertyClass, MetaClass)
+		/// <summary>
+		/// Construct a new property
+		/// </summary>
+		/// <param name="PropertySettings"></param>
+		/// <param name="Class">Referenced class</param>
+		/// <param name="MetaClass">Meta data class</param>
+		/// <param name="ExtraFlags">Extra property flags to apply to the property</param>
+		public UhtClassPtrProperty(UhtPropertySettings PropertySettings, UhtClass Class, UhtClass MetaClass, EPropertyFlags ExtraFlags = EPropertyFlags.None)
+			: base(PropertySettings, Class, MetaClass)
 		{
 			this.PropertyFlags |= ExtraFlags | EPropertyFlags.UObjectWrapper;
 		}
 
+		/// <inheritdoc/>
 		public override StringBuilder AppendText(StringBuilder Builder, UhtPropertyTextType TextType, bool bIsTemplateArgument)
 		{
 			switch (TextType)
diff --git a/Engine/Source/Programs/Shared/EpicGames.UHT/Types/Properties/UhtContainerProperty.cs b/Engine/Source/Programs/Shared/EpicGames.UHT/Types/Properties/UhtContainerProperty.cs
--- a/Engine/Source/Programs/Shared/EpicGames.UHT/Types/Properties/UhtContainerProperty.cs
+++ b/Engine/Source/Programs/Shared/EpicGames.UHT/Types/Properties/UhtContainerProperty.cs
@@ -4,12 +4,21 @@ using EpicGames.Core;
 
 namespace EpicGames.UHT.Types
 {
+
+	/// <summary>
+	/// Common base class for containers with a value
+	/// </summary>
 	public abstract class UhtContainerBaseProperty : UhtProperty
 	{
 
 		/// <inheritdoc/>
 		public UhtProperty ValueProperty { get; set; }
 
+		/// <summary>
+		/// Construct a new property
+		/// </summary>
+		/// <param name="PropertySettings">Property settings</param>
+		/// <param name="Value">Value property</param>
 		public UhtContainerBaseProperty(UhtPropertySettings PropertySettings, UhtProperty Value) : base(PropertySettings)
 		{
 			this.ValueProperty = Value;
@@ -23,6 +32,12 @@ namespace EpicGames.UHT.Types
 			this.ValueProperty.Validate(OuterStruct, OutermostProperty, Options | UhtValidationOptions.IsValue);
 		}
 
+		/// <summary>
+		/// Propagate flags and meta data to/from child properties
+		/// </summary>
+		/// <param name="Container">Container property</param>
+		/// <param name="MetaData">Meta data</param>
+		/// <param name="Inner">Inner property</param>
 		protected static void PropagateFlagsFromInnerAndHandlePersistentInstanceMetadata(UhtProperty Container, UhtMetaData? MetaData, UhtProperty Inner)
 		{
 			// Copy some of the property flags to the container property.
@@ -38,6 +53,12 @@ namespace EpicGames.UHT.Types
 			}
 		}
 
+		/// <summary>
+		/// Resolve the child and return any new flags
+		/// </summary>
+		/// <param name="Child">Child to resolve</param>
+		/// <param name="Phase">Resolve phase</param>
+		/// <returns>And new flags</returns>
 		protected static EPropertyFlags ResolveAndReturnNewFlags(UhtProperty Child, UhtResolvePhase Phase)
 		{
 			EPropertyFlags OldFlags = Child.PropertyFlags;
diff --git a/Engine/Source/Programs/Shared/EpicGames.UHT/Types/Properties/UhtDelegateProperty.cs b/Engine/Source/Programs/Shared/EpicGames.UHT/Types/Properties/UhtDelegateProperty.cs
--- a/Engine/Source/Programs/Shared/EpicGames.UHT/Types/Properties/UhtDelegateProperty.cs
+++ b/Engine/Source/Programs/Shared/EpicGames.UHT/Types/Properties/UhtDelegateProperty.cs
@@ -9,6 +9,9 @@ using System.Text.Json.Serialization;
 
 namespace EpicGames.UHT.Types
 {
+	/// <summary>
+	/// FDelegatePropertyDelegateProperty
+	/// </summary>
 	[UhtEngineClass(Name = "DelegateProperty", IsProperty = true)]
 	public class UhtDelegateProperty : UhtProperty
 	{
@@ -24,9 +27,17 @@ namespace EpicGames.UHT.Types
 		/// <inheritdoc/>
 		protected override UhtPGetArgumentType PGetTypeArgument { get => UhtPGetArgumentType.EngineClass; }
 
+		/// <summary>
+		/// Referenced function
+		/// </summary>
 		[JsonConverter(typeof(UhtTypeSourceNameJsonConverter<UhtFunction>))]
 		public UhtFunction Function { get; set; }
 
+		/// <summary>
+		/// Construct new property
+		/// </summary>
+		/// <param name="PropertySettings">Property settings</param>
+		/// <param name="Function">Referenced function</param>
 		public UhtDelegateProperty(UhtPropertySettings PropertySettings, UhtFunction Function) : base(PropertySettings)
 		{
 			this.Function = Function;
diff --git a/Engine/Source/Programs/Shared/EpicGames.UHT/Types/Properties/UhtDoubleProperty.cs b/Engine/Source/Programs/Shared/EpicGames.UHT/Types/Properties/UhtDoubleProperty.cs
--- a/Engine/Source/Programs/Shared/EpicGames.UHT/Types/Properties/UhtDoubleProperty.cs
+++ b/Engine/Source/Programs/Shared/EpicGames.UHT/Types/Properties/UhtDoubleProperty.cs
@@ -7,6 +7,10 @@ using System.Text;
 
 namespace EpicGames.UHT.Types
 {
+
+	/// <summary>
+	/// FDoubleProperty
+	/// </summary>
 	[UnrealHeaderTool]
 	[UhtEngineClass(Name = "DoubleProperty", IsProperty = true)]
 	public class UhtDoubleProperty : UhtNumericProperty
@@ -17,6 +21,10 @@ namespace EpicGames.UHT.Types
 		/// <inheritdoc/>
 		protected override string CppTypeText { get => "double"; }
 
+		/// <summary>
+		/// Create new property
+		/// </summary>
+		/// <param name="PropertySettings">Property settings</param>
 		public UhtDoubleProperty(UhtPropertySettings PropertySettings) : base(PropertySettings, UhtPropertyIntType.None)
 		{
 			this.PropertyCaps |= UhtPropertyCaps.IsParameterSupportedByBlueprint | UhtPropertyCaps.IsMemberSupportedByBlueprint;
@@ -36,12 +44,14 @@ namespace EpicGames.UHT.Types
 			return Builder;
 		}
 
+		/// <inheritdoc/>
 		public override bool SanitizeDefaultValue(IUhtTokenReader DefaultValueReader, StringBuilder InnerDefaultValue)
 		{
 			InnerDefaultValue.AppendFormat("{0:F6}", DefaultValueReader.GetConstFloatExpression());
 			return true;
 		}
 
+		/// <inheritdoc/>
 		public override bool IsSameType(UhtProperty Other)
 		{
 			return Other is UhtDoubleProperty;
@@ -49,7 +59,7 @@ namespace EpicGames.UHT.Types
 
 		#region Keyword
 		[UhtPropertyType(Keyword = "double", Options = UhtPropertyTypeOptions.Simple | UhtPropertyTypeOptions.Immediate)]
-		public static UhtProperty? DoubleProperty(UhtPropertyResolvePhase ResolvePhase, UhtPropertySettings PropertySettings, IUhtTokenReader TokenReader, UhtToken MatchedToken)
+		private static UhtProperty? DoubleProperty(UhtPropertyResolvePhase ResolvePhase, UhtPropertySettings PropertySettings, IUhtTokenReader TokenReader, UhtToken MatchedToken)
 		{
 			return new UhtDoubleProperty(PropertySettings);
 		}
diff --git a/Engine/Source/Programs/Shared/EpicGames.UHT/Types/Properties/UhtEnumProperty.cs b/Engine/Source/Programs/Shared/EpicGames.UHT/Types/Properties/UhtEnumProperty.cs
--- a/Engine/Source/Programs/Shared/EpicGames.UHT/Types/Properties/UhtEnumProperty.cs
+++ b/Engine/Source/Programs/Shared/EpicGames.UHT/Types/Properties/UhtEnumProperty.cs
@@ -10,6 +10,10 @@ using System.Text.Json.Serialization;
 
 namespace EpicGames.UHT.Types
 {
+
+	/// <summary>
+	/// FEnumProperty
+	/// </summary>
 	[UnrealHeaderTool]
 	[UhtEngineClass(Name = "EnumProperty", IsProperty = true)]
 	public class UhtEnumProperty : UhtProperty
@@ -26,14 +30,32 @@ namespace EpicGames.UHT.Types
 		/// <inheritdoc/>
 		protected override UhtPGetArgumentType PGetTypeArgument { get => this.Enum.CppForm != UhtEnumCppForm.EnumClass ? UhtPGetArgumentType.EngineClass : UhtPGetArgumentType.TypeText; }
 
+		/// <summary>
+		/// Referenced enum
+		/// </summary>
 		[JsonConverter(typeof(UhtTypeSourceNameJsonConverter<UhtEnum>))]
 		public UhtEnum Enum { get; set; }
 
+		/// <summary>
+		/// Underlying property set when the enum has an underlying interger type
+		/// </summary>
 		public UhtProperty? UnderlyingProperty { get; set; }
 
+		/// <summary>
+		/// Underlying type which defaults to int32 if the referenced enum doesn't have an underlying type
+		/// </summary>
 		public UhtEnumUnderlyingType UnderlyingType => this.Enum.UnderlyingType != UhtEnumUnderlyingType.Unspecified ? this.Enum.UnderlyingType : UhtEnumUnderlyingType.int32;
+
+		/// <summary>
+		/// Underlying type size.  Defaults to unsized if the referenced enum doesn't have an underlying type
+		/// </summary>
 		public UhtPropertyIntType UnderlyingTypeSize => this.Enum.UnderlyingType != UhtEnumUnderlyingType.Unspecified ? UhtPropertyIntType.Sized : UhtPropertyIntType.Unsized;
 
+		/// <summary>
+		/// Construct property
+		/// </summary>
+		/// <param name="PropertySettings">Property settings</param>
+		/// <param name="Enum">Referenced enum</param>
 		public UhtEnumProperty(UhtPropertySettings PropertySettings, UhtEnum Enum) : base(PropertySettings)
 		{
 			this.PropertyCaps |= UhtPropertyCaps.RequiresNullConstructorArg | UhtPropertyCaps.CanExposeOnSpawn | UhtPropertyCaps.IsParameterSupportedByBlueprint | UhtPropertyCaps.IsMemberSupportedByBlueprint;
@@ -71,6 +93,15 @@ namespace EpicGames.UHT.Types
 			yield return this.Enum;
 		}
 
+		/// <summary>
+		/// Append enum text
+		/// </summary>
+		/// <param name="Builder">Destination builder</param>
+		/// <param name="Property">Property in question</param>
+		/// <param name="Enum">Referenced enum</param>
+		/// <param name="TextType">Type of text to append</param>
+		/// <param name="bIsTemplateArgument">If true, this property is a template argument</param>
+		/// <returns></returns>
 		public static StringBuilder AppendEnumText(StringBuilder Builder, UhtProperty Property, UhtEnum Enum, UhtPropertyTextType TextType, bool bIsTemplateArgument)
 		{
 			switch (TextType)
@@ -110,7 +141,8 @@ namespace EpicGames.UHT.Types
 		/// Append the text for a function thunk call argument
 		/// </summary>
 		/// <param name="Builder">Output builder</param>
-		/// 
+		/// <param name="Property">Property in question</param>
+		/// <param name="Enum">Referenced enum</param>
 		/// <returns>Output builder</returns>
 		public static StringBuilder AppendEnumFunctionThunkParameterArg(StringBuilder Builder, UhtProperty Property, UhtEnum Enum)
 		{
@@ -129,12 +161,47 @@ namespace EpicGames.UHT.Types
 			return Builder;
 		}
 
+		/// <summary>
+		/// Get the RigVM type for an enumeration
+		/// </summary>
+		/// <param name="Property">Property in question</param>
+		/// <param name="Enum">Referenced enumeration</param>
+		/// <param name="ParameterFlags">Parameter flags that can be updated</param>
+		/// <returns>Type string</returns>
 		public static string? GetEnumRigVMType(UhtProperty Property, UhtEnum Enum, ref UhtRigVMParameterFlags ParameterFlags)
 		{
 			ParameterFlags |= UhtRigVMParameterFlags.IsEnum;
 			return Enum.CppType.ToString();
 		}
 
+		/// <summary>
+		/// Sanitize the default value for an enumeration
+		/// </summary>
+		/// <param name="Property">Property in question</param>
+		/// <param name="Enum">Referenced enumeration</param>
+		/// <param name="DefaultValueReader">Default value</param>
+		/// <param name="InnerDefaultValue">Destination builder</param>
+		/// <returns>True if the default value was parsed.</returns>
+		public static bool SanitizeEnumDefaultValue(UhtProperty Property, UhtEnum Enum, IUhtTokenReader DefaultValueReader, StringBuilder InnerDefaultValue)
+		{
+			UhtTokenList CppIdentifier = DefaultValueReader.GetCppIdentifier(UhtCppIdentifierOptions.None);
+			UhtTokenList StartingPoint = CppIdentifier.Next != null ? CppIdentifier.Next : CppIdentifier;
+			StartingPoint.Join(InnerDefaultValue, "::");
+			UhtTokenListCache.Return(CppIdentifier);
+
+			int EntryIndex = Enum.GetIndexByName(InnerDefaultValue.ToString());
+			if (EntryIndex == -1)
+			{
+				return false;
+			}
+			if (Enum.MetaData.ContainsKey(UhtNames.Hidden, EntryIndex))
+			{
+				Property.LogError($"Hidden enum entries cannot be used as default values: '{Property.SourceName}' '{InnerDefaultValue}'");
+			}
+			return true;
+		}
+
+		/// <inheritdoc/>
 		public override StringBuilder AppendText(StringBuilder Builder, UhtPropertyTextType TextType, bool bIsTemplateArgument)
 		{
 			return AppendEnumText(Builder, this, this.Enum, TextType, bIsTemplateArgument);
@@ -230,23 +297,10 @@ namespace EpicGames.UHT.Types
 			return false;
 		}
 
-		public static bool SanitizeEnumDefaultValue(UhtProperty Property, UhtEnum Enum, IUhtTokenReader DefaultValueReader, StringBuilder InnerDefaultValue)
+		/// <inheritdoc/>
+		public override string? GetRigVMType(ref UhtRigVMParameterFlags ParameterFlags)
 		{
-			UhtTokenList CppIdentifier = DefaultValueReader.GetCppIdentifier(UhtCppIdentifierOptions.None);
-			UhtTokenList StartingPoint = CppIdentifier.Next != null ? CppIdentifier.Next : CppIdentifier;
-			StartingPoint.Join(InnerDefaultValue, "::");
-			UhtTokenListCache.Return(CppIdentifier);
-
-			int EntryIndex = Enum.GetIndexByName(InnerDefaultValue.ToString());
-			if (EntryIndex == -1)
-			{
-				return false;
-			}
-			if (Enum.MetaData.ContainsKey(UhtNames.Hidden, EntryIndex))
-			{
-				Property.LogError($"Hidden enum entries cannot be used as default values: '{Property.SourceName}' '{InnerDefaultValue}'");
-			}
-			return true;
+			return GetEnumRigVMType(this, this.Enum, ref ParameterFlags);
 		}
 
 		private UhtProperty CreateUnderlyingProperty()
@@ -277,15 +331,9 @@ namespace EpicGames.UHT.Types
 			}
 		}
 
-		/// <inheritdoc/>
-		public override string? GetRigVMType(ref UhtRigVMParameterFlags ParameterFlags)
-		{
-			return GetEnumRigVMType(this, this.Enum, ref ParameterFlags);
-		}
-
 		#region Keyword
 		[UhtPropertyType(Keyword = "TEnumAsByte")]
-		public static UhtProperty? EnumProperty(UhtPropertyResolvePhase ResolvePhase, UhtPropertySettings PropertySettings, IUhtTokenReader TokenReader, UhtToken MatchedToken)
+		private static UhtProperty? EnumProperty(UhtPropertyResolvePhase ResolvePhase, UhtPropertySettings PropertySettings, IUhtTokenReader TokenReader, UhtToken MatchedToken)
 		{
 			UhtType Outer = PropertySettings.Outer;
 			UhtEnum? Enum = null;
diff --git a/Engine/Source/Programs/Shared/EpicGames.UHT/Types/Properties/UhtFieldPathProperty.cs b/Engine/Source/Programs/Shared/EpicGames.UHT/Types/Properties/UhtFieldPathProperty.cs
--- a/Engine/Source/Programs/Shared/EpicGames.UHT/Types/Properties/UhtFieldPathProperty.cs
+++ b/Engine/Source/Programs/Shared/EpicGames.UHT/Types/Properties/UhtFieldPathProperty.cs
@@ -8,6 +8,9 @@ using System.Text;
 
 namespace EpicGames.UHT.Types
 {
+	/// <summary>
+	/// FFieldPathProperty
+	/// </summary>
 	[UnrealHeaderTool]
 	[UhtEngineClass(Name = "FieldPathProperty", IsProperty = true)]
 	public class UhtFieldPathProperty : UhtProperty
@@ -29,6 +32,11 @@ namespace EpicGames.UHT.Types
 		/// </summary>
 		public string FieldClassName { get; set; }
 
+		/// <summary>
+		/// Construct new property
+		/// </summary>
+		/// <param name="PropertySettings">Property settings</param>
+		/// <param name="FieldClassName">Field class name</param>
 		public UhtFieldPathProperty(UhtPropertySettings PropertySettings, string FieldClassName) : base(PropertySettings)
 		{
 			this.FieldClassName = FieldClassName;
@@ -101,9 +109,9 @@ namespace EpicGames.UHT.Types
 			return null;
 		}
 
-#region Keyword
+		#region Keyword
 		[UhtPropertyType(Keyword = "TFieldPath", Options = UhtPropertyTypeOptions.Immediate)]
-		public static UhtProperty? FieldPathProperty(UhtPropertyResolvePhase ResolvePhase, UhtPropertySettings PropertySettings, IUhtTokenReader TokenReader, UhtToken MatchedToken)
+		private static UhtProperty? FieldPathProperty(UhtPropertyResolvePhase ResolvePhase, UhtPropertySettings PropertySettings, IUhtTokenReader TokenReader, UhtToken MatchedToken)
 		{
 			using (var TokenContext = new UhtMessageContext(TokenReader, "TFieldPath"))
 			{
@@ -122,6 +130,6 @@ namespace EpicGames.UHT.Types
 				return new UhtFieldPathProperty(PropertySettings, FieldClassName.ToString());
 			}
 		}
-#endregion
+		#endregion
 	}
 }
diff --git a/Engine/Source/Programs/Shared/EpicGames.UHT/Types/Properties/UhtFloatProperty.cs b/Engine/Source/Programs/Shared/EpicGames.UHT/Types/Properties/UhtFloatProperty.cs
--- a/Engine/Source/Programs/Shared/EpicGames.UHT/Types/Properties/UhtFloatProperty.cs
+++ b/Engine/Source/Programs/Shared/EpicGames.UHT/Types/Properties/UhtFloatProperty.cs
@@ -7,6 +7,9 @@ using System.Text;
 
 namespace EpicGames.UHT.Types
 {
+	/// <summary>
+	/// FFloatProperty
+	/// </summary>
 	[UnrealHeaderTool]
 	[UhtEngineClass(Name = "FloatProperty", IsProperty = true)]
 	public class UhtFloatProperty : UhtNumericProperty
@@ -17,6 +20,10 @@ namespace EpicGames.UHT.Types
 		/// <inheritdoc/>
 		protected override string CppTypeText { get => "float"; }
 
+		/// <summary>
+		/// Construct a new property
+		/// </summary>
+		/// <param name="PropertySettings">Property settings</param>
 		public UhtFloatProperty(UhtPropertySettings PropertySettings) : base(PropertySettings, UhtPropertyIntType.None)
 		{
 			this.PropertyCaps |= UhtPropertyCaps.CanExposeOnSpawn | UhtPropertyCaps.IsParameterSupportedByBlueprint | UhtPropertyCaps.IsMemberSupportedByBlueprint;
@@ -51,7 +58,7 @@ namespace EpicGames.UHT.Types
 
 		#region Keywords
 		[UhtPropertyType(Keyword = "float", Options = UhtPropertyTypeOptions.Simple | UhtPropertyTypeOptions.Immediate)]
-		public static UhtProperty? FloatProperty(UhtPropertyResolvePhase ResolvePhase, UhtPropertySettings PropertySettings, IUhtTokenReader TokenReader, UhtToken MatchedToken)
+		private static UhtProperty? FloatProperty(UhtPropertyResolvePhase ResolvePhase, UhtPropertySettings PropertySettings, IUhtTokenReader TokenReader, UhtToken MatchedToken)
 		{
 			return new UhtFloatProperty(PropertySettings);
 		}
diff --git a/Engine/Source/Programs/Shared/EpicGames.UHT/Types/Properties/UhtInt16Property.cs b/Engine/Source/Programs/Shared/EpicGames.UHT/Types/Properties/UhtInt16Property.cs
--- a/Engine/Source/Programs/Shared/EpicGames.UHT/Types/Properties/UhtInt16Property.cs
+++ b/Engine/Source/Programs/Shared/EpicGames.UHT/Types/Properties/UhtInt16Property.cs
@@ -7,6 +7,9 @@ using System.Text;
 
 namespace EpicGames.UHT.Types
 {
+	/// <summary>
+	/// FInt16Property
+	/// </summary>
 	[UnrealHeaderTool]
 	[UhtEngineClass(Name = "Int16Property", IsProperty = true)]
 	public class UhtInt16Property : UhtNumericProperty
@@ -17,6 +20,11 @@ namespace EpicGames.UHT.Types
 		/// <inheritdoc/>
 		protected override string CppTypeText { get => "int16"; }
 
+		/// <summary>
+		/// Construct new property
+		/// </summary>
+		/// <param name="PropertySettings">Property settings</param>
+		/// <param name="IntType">Integer type</param>
 		public UhtInt16Property(UhtPropertySettings PropertySettings, UhtPropertyIntType IntType) : base(PropertySettings, IntType)
 		{
 		}
@@ -35,11 +43,13 @@ namespace EpicGames.UHT.Types
 			return Builder;
 		}
 
+		/// <inheritdoc/>
 		public override bool SanitizeDefaultValue(IUhtTokenReader DefaultValueReader, StringBuilder InnerDefaultValue)
 		{
 			return false;
 		}
 
+		/// <inheritdoc/>
 		public override bool IsSameType(UhtProperty Other)
 		{
 			return Other is UhtInt16Property;
@@ -47,7 +57,7 @@ namespace EpicGames.UHT.Types
 
 		#region Keyword
 		[UhtPropertyType(Keyword = "int16", Options = UhtPropertyTypeOptions.Simple | UhtPropertyTypeOptions.Immediate)]
-		public static UhtProperty? Int16Property(UhtPropertyResolvePhase ResolvePhase, UhtPropertySettings PropertySettings, IUhtTokenReader TokenReader, UhtToken MatchedToken)
+		private static UhtProperty? Int16Property(UhtPropertyResolvePhase ResolvePhase, UhtPropertySettings PropertySettings, IUhtTokenReader TokenReader, UhtToken MatchedToken)
 		{
 			return new UhtInt16Property(PropertySettings, UhtPropertyIntType.Sized);
 		}
diff --git a/Engine/Source/Programs/Shared/EpicGames.UHT/Types/Properties/UhtInt64Property.cs b/Engine/Source/Programs/Shared/EpicGames.UHT/Types/Properties/UhtInt64Property.cs
--- a/Engine/Source/Programs/Shared/EpicGames.UHT/Types/Properties/UhtInt64Property.cs
+++ b/Engine/Source/Programs/Shared/EpicGames.UHT/Types/Properties/UhtInt64Property.cs
@@ -7,6 +7,9 @@ using System.Text;
 
 namespace EpicGames.UHT.Types
 {
+	/// <summary>
+	/// FInt64Property
+	/// </summary>
 	[UnrealHeaderTool]
 	[UhtEngineClass(Name = "Int64Property", IsProperty = true)]
 	public class UhtInt64Property : UhtNumericProperty
@@ -17,6 +20,11 @@ namespace EpicGames.UHT.Types
 		/// <inheritdoc/>
 		protected override string CppTypeText { get => "int64"; }
 
+		/// <summary>
+		/// Construct new property
+		/// </summary>
+		/// <param name="PropertySettings">Property settings</param>
+		/// <param name="IntType">Integer type</param>
 		public UhtInt64Property(UhtPropertySettings PropertySettings, UhtPropertyIntType IntType) : base(PropertySettings, IntType)
 		{
 			this.PropertyCaps |= UhtPropertyCaps.CanExposeOnSpawn | UhtPropertyCaps.IsParameterSupportedByBlueprint | UhtPropertyCaps.IsMemberSupportedByBlueprint;
@@ -51,7 +59,7 @@ namespace EpicGames.UHT.Types
 
 		#region Keyword
 		[UhtPropertyType(Keyword = "int64", Options = UhtPropertyTypeOptions.Simple | UhtPropertyTypeOptions.Immediate)]
-		public static UhtProperty? Int64Property(UhtPropertyResolvePhase ResolvePhase, UhtPropertySettings PropertySettings, IUhtTokenReader TokenReader, UhtToken MatchedToken)
+		private static UhtProperty? Int64Property(UhtPropertyResolvePhase ResolvePhase, UhtPropertySettings PropertySettings, IUhtTokenReader TokenReader, UhtToken MatchedToken)
 		{
 			return new UhtInt64Property(PropertySettings, UhtPropertyIntType.Sized);
 		}
diff --git a/Engine/Source/Programs/Shared/EpicGames.UHT/Types/Properties/UhtInt8Property.cs b/Engine/Source/Programs/Shared/EpicGames.UHT/Types/Properties/UhtInt8Property.cs
--- a/Engine/Source/Programs/Shared/EpicGames.UHT/Types/Properties/UhtInt8Property.cs
+++ b/Engine/Source/Programs/Shared/EpicGames.UHT/Types/Properties/UhtInt8Property.cs
@@ -7,6 +7,9 @@ using System.Text;
 
 namespace EpicGames.UHT.Types
 {
+	/// <summary>
+	/// FInt8Property
+	/// </summary>
 	[UnrealHeaderTool]
 	[UhtEngineClass(Name = "Int8Property", IsProperty = true)]
 	public class UhtInt8Property : UhtNumericProperty
@@ -17,6 +20,11 @@ namespace EpicGames.UHT.Types
 		/// <inheritdoc/>
 		protected override string CppTypeText { get => "int8"; }
 
+		/// <summary>
+		/// Construct a new property
+		/// </summary>
+		/// <param name="PropertySettings">Property settings</param>
+		/// <param name="IntType">Integer type</param>
 		public UhtInt8Property(UhtPropertySettings PropertySettings, UhtPropertyIntType IntType) : base(PropertySettings, IntType)
 		{
 		}
@@ -35,11 +43,13 @@ namespace EpicGames.UHT.Types
 			return Builder;
 		}
 
+		/// <inheritdoc/>
 		public override bool SanitizeDefaultValue(IUhtTokenReader DefaultValueReader, StringBuilder InnerDefaultValue)
 		{
 			return false;
 		}
 
+		/// <inheritdoc/>
 		public override bool IsSameType(UhtProperty Other)
 		{
 			return Other is UhtInt8Property;
@@ -47,7 +57,7 @@ namespace EpicGames.UHT.Types
 
 		#region Keyword
 		[UhtPropertyType(Keyword = "int8", Options = UhtPropertyTypeOptions.Simple | UhtPropertyTypeOptions.Immediate)]
-		public static UhtProperty? Int8Property(UhtPropertyResolvePhase ResolvePhase, UhtPropertySettings PropertySettings, IUhtTokenReader TokenReader, UhtToken MatchedToken)
+		private static UhtProperty? Int8Property(UhtPropertyResolvePhase ResolvePhase, UhtPropertySettings PropertySettings, IUhtTokenReader TokenReader, UhtToken MatchedToken)
 		{
 			return new UhtInt8Property(PropertySettings, UhtPropertyIntType.Sized);
 		}
diff --git a/Engine/Source/Programs/Shared/EpicGames.UHT/Types/Properties/UhtIntProperty.cs b/Engine/Source/Programs/Shared/EpicGames.UHT/Types/Properties/UhtIntProperty.cs
--- a/Engine/Source/Programs/Shared/EpicGames.UHT/Types/Properties/UhtIntProperty.cs
+++ b/Engine/Source/Programs/Shared/EpicGames.UHT/Types/Properties/UhtIntProperty.cs
@@ -7,6 +7,9 @@ using System.Text;
 
 namespace EpicGames.UHT.Types
 {
+	/// <summary>
+	/// FIntProperty
+	/// </summary>
 	[UnrealHeaderTool]
 	[UhtEngineClass(Name = "IntProperty", IsProperty = true)]
 	public class UhtIntProperty : UhtNumericProperty
@@ -17,6 +20,11 @@ namespace EpicGames.UHT.Types
 		/// <inheritdoc/>
 		protected override string CppTypeText { get => "int32"; }
 
+		/// <summary>
+		/// Construct new property
+		/// </summary>
+		/// <param name="PropertySettings">Property settings</param>
+		/// <param name="IntType">Integer type</param>
 		public UhtIntProperty(UhtPropertySettings PropertySettings, UhtPropertyIntType IntType) : base(PropertySettings, IntType)
 		{
 			this.PropertyCaps |= UhtPropertyCaps.CanExposeOnSpawn | UhtPropertyCaps.IsParameterSupportedByBlueprint | UhtPropertyCaps.IsMemberSupportedByBlueprint;
@@ -54,19 +62,19 @@ namespace EpicGames.UHT.Types
 
 		#region Keywords
 		[UhtPropertyType(Keyword = "int32", Options = UhtPropertyTypeOptions.Simple | UhtPropertyTypeOptions.Immediate)]
-		public static UhtProperty? Int32Property(UhtPropertyResolvePhase ResolvePhase, UhtPropertySettings PropertySettings, IUhtTokenReader TokenReader, UhtToken MatchedToken)
+		private static UhtProperty? Int32Property(UhtPropertyResolvePhase ResolvePhase, UhtPropertySettings PropertySettings, IUhtTokenReader TokenReader, UhtToken MatchedToken)
 		{
 			return new UhtIntProperty(PropertySettings, UhtPropertyIntType.Sized);
 		}
 
 		[UhtPropertyType(Keyword = "int", Options = UhtPropertyTypeOptions.Simple | UhtPropertyTypeOptions.Immediate)]
-		public static UhtProperty? IntProperty(UhtPropertyResolvePhase ResolvePhase, UhtPropertySettings PropertySettings, IUhtTokenReader TokenReader, UhtToken MatchedToken)
+		private static UhtProperty? IntProperty(UhtPropertyResolvePhase ResolvePhase, UhtPropertySettings PropertySettings, IUhtTokenReader TokenReader, UhtToken MatchedToken)
 		{
 			return new UhtIntProperty(PropertySettings, UhtPropertyIntType.Unsized);
 		}
 
 		[UhtPropertyType(Keyword = "signed", Options = UhtPropertyTypeOptions.Immediate)]
-		public static UhtProperty? SignedProperty(UhtPropertyResolvePhase ResolvePhase, UhtPropertySettings PropertySettings, IUhtTokenReader TokenReader, UhtToken MatchedToken)
+		private static UhtProperty? SignedProperty(UhtPropertyResolvePhase ResolvePhase, UhtPropertySettings PropertySettings, IUhtTokenReader TokenReader, UhtToken MatchedToken)
 		{
 			TokenReader
 				.Require("signed")
diff --git a/Engine/Source/Programs/Shared/EpicGames.UHT/Types/Properties/UhtLargeWorldCoordinatesRealProperty.cs b/Engine/Source/Programs/Shared/EpicGames.UHT/Types/Properties/UhtLargeWorldCoordinatesRealProperty.cs
--- a/Engine/Source/Programs/Shared/EpicGames.UHT/Types/Properties/UhtLargeWorldCoordinatesRealProperty.cs
+++ b/Engine/Source/Programs/Shared/EpicGames.UHT/Types/Properties/UhtLargeWorldCoordinatesRealProperty.cs
@@ -7,6 +7,9 @@ using System.Text;
 
 namespace EpicGames.UHT.Types
 {
+	/// <summary>
+	/// FLargeWorldCoordinatesRealProperty
+	/// </summary>
 	[UnrealHeaderTool]
 	[UhtEngineClass(Name = "LargeWorldCoordinatesRealProperty", IsProperty = true)]
 	public class UhtLargeWorldCoordinatesRealProperty : UhtNumericProperty
@@ -17,6 +20,10 @@ namespace EpicGames.UHT.Types
 		/// <inheritdoc/>
 		protected override string CppTypeText { get => "double"; }
 
+		/// <summary>
+		/// Construct new property
+		/// </summary>
+		/// <param name="PropertySettings">Property settings</param>
 		public UhtLargeWorldCoordinatesRealProperty(UhtPropertySettings PropertySettings) : base(PropertySettings, UhtPropertyIntType.None)
 		{
 			this.PropertyCaps |= UhtPropertyCaps.IsParameterSupportedByBlueprint | UhtPropertyCaps.IsMemberSupportedByBlueprint;
@@ -53,7 +60,7 @@ namespace EpicGames.UHT.Types
 
 		#region Keyword
 		[UhtPropertyType(Keyword = "FLargeWorldCoordinatesReal", Options = UhtPropertyTypeOptions.Simple | UhtPropertyTypeOptions.Immediate)]
-		public static UhtProperty? LargeWorldCoordinatesRealProperty(UhtPropertyResolvePhase ResolvePhase, UhtPropertySettings PropertySettings, IUhtTokenReader TokenReader, UhtToken MatchedToken)
+		private static UhtProperty? LargeWorldCoordinatesRealProperty(UhtPropertyResolvePhase ResolvePhase, UhtPropertySettings PropertySettings, IUhtTokenReader TokenReader, UhtToken MatchedToken)
 		{
 			if (!PropertySettings.Outer.HeaderFile.bIsNoExportTypes)
 			{
diff --git a/Engine/Source/Programs/Shared/EpicGames.UHT/Types/Properties/UhtLazyObjectPtrProperty.cs b/Engine/Source/Programs/Shared/EpicGames.UHT/Types/Properties/UhtLazyObjectPtrProperty.cs
--- a/Engine/Source/Programs/Shared/EpicGames.UHT/Types/Properties/UhtLazyObjectPtrProperty.cs
+++ b/Engine/Source/Programs/Shared/EpicGames.UHT/Types/Properties/UhtLazyObjectPtrProperty.cs
@@ -8,6 +8,10 @@ using System.Text;
 
 namespace EpicGames.UHT.Types
 {
+
+	/// <summary>
+	/// FLazyObjectProperty
+	/// </summary>
 	[UnrealHeaderTool]
 	[UhtEngineClass(Name = "LazyObjectProperty", IsProperty = true)]
 	public class UhtLazyObjectPtrProperty : UhtObjectPropertyBase
@@ -27,8 +31,13 @@ namespace EpicGames.UHT.Types
 		/// <inheritdoc/>
 		protected override UhtPGetArgumentType PGetTypeArgument { get => UhtPGetArgumentType.TypeText; }
 
-		public UhtLazyObjectPtrProperty(UhtPropertySettings PropertySettings, UhtClass PropertyClass)
-			: base(PropertySettings, PropertyClass, null)
+		/// <summary>
+		/// Construct new property
+		/// </summary>
+		/// <param name="PropertySettings">Property settings</param>
+		/// <param name="Class">Referenced class</param>
+		public UhtLazyObjectPtrProperty(UhtPropertySettings PropertySettings, UhtClass Class)
+			: base(PropertySettings, Class, null)
 		{
 			this.PropertyFlags |= EPropertyFlags.UObjectWrapper;
 			this.PropertyCaps |= UhtPropertyCaps.PassCppArgsByRef | UhtPropertyCaps.RequiresNullConstructorArg;
@@ -103,7 +112,7 @@ namespace EpicGames.UHT.Types
 
 		#region Keyword
 		[UhtPropertyType(Keyword = "TLazyObjectPtr")]
-		public static UhtProperty? LazyObjectPtrProperty(UhtPropertyResolvePhase ResolvePhase, UhtPropertySettings PropertySettings, IUhtTokenReader TokenReader, UhtToken MatchedToken)
+		private static UhtProperty? LazyObjectPtrProperty(UhtPropertyResolvePhase ResolvePhase, UhtPropertySettings PropertySettings, IUhtTokenReader TokenReader, UhtToken MatchedToken)
 		{
 			UhtClass? PropertyClass = ParseTemplateObject(PropertySettings, TokenReader, MatchedToken, true);
 			if (PropertyClass == null)
