Horde: Add IStorageClient.FindMissingBlobsAsync

#rnx

diff --git a/Engine/Source/Programs/Horde/Horde.Build/Storage/Services/BasicStorageClient.cs b/Engine/Source/Programs/Horde/Horde.Build/Storage/Services/BasicStorageClient.cs
--- a/Engine/Source/Programs/Horde/Horde.Build/Storage/Services/BasicStorageClient.cs
+++ b/Engine/Source/Programs/Horde/Horde.Build/Storage/Services/BasicStorageClient.cs
@@ -99,6 +99,14 @@ namespace Horde.Build.Storage.Services
 			return await StorageBackend.ExistsAsync(GetFullBlobPath(NamespaceId, Hash));
 		}
 
+		/// <inheritdoc/>
+		public async Task<HashSet<IoHash>> FindMissingBlobsAsync(NamespaceId NamespaceId, HashSet<IoHash> Hashes, CancellationToken CancellationToken)
+		{
+			var Tasks = Hashes.Select(async Hash => (Hash, await StorageBackend.ExistsAsync(GetFullBlobPath(NamespaceId, Hash))));
+			await Task.WhenAll(Tasks);
+			return Tasks.Where(x => !x.Result.Item2).Select(x => x.Result.Item1).ToHashSet();
+		}
+
 		/// <inheritdoc/>
 		public async Task<IRef> GetRefAsync(NamespaceId NamespaceId, BucketId BucketId, RefId RefId, CancellationToken CancellationToken)
 		{
diff --git a/Engine/Source/Programs/Shared/EpicGames.Horde/Storage/IStorageClient.cs b/Engine/Source/Programs/Shared/EpicGames.Horde/Storage/IStorageClient.cs
--- a/Engine/Source/Programs/Shared/EpicGames.Horde/Storage/IStorageClient.cs
+++ b/Engine/Source/Programs/Shared/EpicGames.Horde/Storage/IStorageClient.cs
@@ -185,6 +185,15 @@ namespace EpicGames.Horde.Storage
 		/// <returns>True if the blob exists, false if it did not exist</returns>
 		Task<bool> HasBlobAsync(NamespaceId NamespaceId, IoHash Hash, CancellationToken CancellationToken = default);
 
+		/// <summary>
+		/// Checks if a list of blobs exist, returning the set of missing blobs
+		/// </summary>
+		/// <param name="NamespaceId">Namespace to operate on</param>
+		/// <param name="Hashes">Set of hashes to check</param>
+		/// <param name="CancellationToken">Cancellation token for the operation</param>
+		/// <returns>A set of the missing hashes</returns>
+		Task<HashSet<IoHash>> FindMissingBlobsAsync(NamespaceId NamespaceId, HashSet<IoHash> Hashes, CancellationToken CancellationToken = default);
+
 		#endregion
 		#region Refs
 
diff --git a/Engine/Source/Programs/Shared/EpicGames.Horde/Storage/Impl/FileStorageClient.cs b/Engine/Source/Programs/Shared/EpicGames.Horde/Storage/Impl/FileStorageClient.cs
--- a/Engine/Source/Programs/Shared/EpicGames.Horde/Storage/Impl/FileStorageClient.cs
+++ b/Engine/Source/Programs/Shared/EpicGames.Horde/Storage/Impl/FileStorageClient.cs
@@ -85,6 +85,12 @@ namespace EpicGames.Horde.Storage.Impl
 			return Task.FromResult(FileReference.Exists(GetBlobFile(NamespaceId, Hash)));
 		}
 
+		/// <inheritdoc/>
+		public Task<HashSet<IoHash>> FindMissingBlobsAsync(NamespaceId NamespaceId, HashSet<IoHash> Hashes, CancellationToken CancellationToken)
+		{
+			return Task.FromResult(Hashes.Where(x => !FileReference.Exists(GetBlobFile(NamespaceId, x))).ToHashSet());
+		}
+
 		/// <inheritdoc/>
 		public Task<bool> DeleteRefAsync(NamespaceId NamespaceId, BucketId BucketId, RefId RefId, CancellationToken CancellationToken = default)
 		{
diff --git a/Engine/Source/Programs/Shared/EpicGames.Horde/Storage/Impl/HttpStorageClient.cs b/Engine/Source/Programs/Shared/EpicGames.Horde/Storage/Impl/HttpStorageClient.cs
--- a/Engine/Source/Programs/Shared/EpicGames.Horde/Storage/Impl/HttpStorageClient.cs
+++ b/Engine/Source/Programs/Shared/EpicGames.Horde/Storage/Impl/HttpStorageClient.cs
@@ -17,6 +17,17 @@ using System.Threading.Tasks;
 
 namespace EpicGames.Horde.Storage.Impl
 {
+	/// <summary>
+	/// Response from checking for missing blobs
+	/// </summary>
+	public class BlobExistsResponse
+	{
+		/// <summary>
+		/// Set of missing hashes
+		/// </summary>
+		public HashSet<IoHash> Needs { get; set; } = new HashSet<IoHash>();
+	}
+
 	/// <summary>
 	/// Response from adding a new object
 	/// </summary>
@@ -138,6 +149,29 @@ namespace EpicGames.Horde.Storage.Impl
 			}
 		}
 
+		/// <inheritdoc/>
+		public async Task<HashSet<IoHash>> FindMissingBlobsAsync(NamespaceId NamespaceId, HashSet<IoHash> Hashes, CancellationToken CancellationToken)
+		{
+			using StringContent Content = new StringContent(JsonSerializer.Serialize(Hashes));
+
+			using HttpRequestMessage Request = new HttpRequestMessage(HttpMethod.Post, $"api/v1/blobs/{NamespaceId}/exist");
+
+			Request.Headers.Accept.Clear();
+			Request.Headers.Accept.Add(new MediaTypeWithQualityHeaderValue("application/json"));
+			Content.Headers.ContentType = new MediaTypeHeaderValue("application/json");
+			Request.Content = Content;
+
+			HttpResponseMessage Response = await HttpClient.SendAsync(Request, CancellationToken);
+			Response.EnsureSuccessStatusCode();
+
+			BlobExistsResponse? ResponseBody = await ReadJsonResponse<BlobExistsResponse>(Response.Content);
+			if (ResponseBody == null)
+			{
+				throw new StorageException("Unable to parse response body", null);
+			}
+			return ResponseBody.Needs;
+		}
+
 		#endregion
 
 		#region Refs
diff --git a/Engine/Source/Programs/Shared/EpicGames.Horde/Storage/Impl/MemoryStorageClient.cs b/Engine/Source/Programs/Shared/EpicGames.Horde/Storage/Impl/MemoryStorageClient.cs
--- a/Engine/Source/Programs/Shared/EpicGames.Horde/Storage/Impl/MemoryStorageClient.cs
+++ b/Engine/Source/Programs/Shared/EpicGames.Horde/Storage/Impl/MemoryStorageClient.cs
@@ -63,6 +63,12 @@ namespace EpicGames.Horde.Storage.Impl
 			return Task.FromResult(Blobs.ContainsKey((NamespaceId, Hash)));
 		}
 
+		/// <inheritdoc/>
+		public Task<HashSet<IoHash>> FindMissingBlobsAsync(NamespaceId NamespaceId, HashSet<IoHash> Hashes, CancellationToken CancellationToken)
+		{
+			return Task.FromResult(Hashes.Where(x => !Blobs.ContainsKey((NamespaceId, x))).ToHashSet());
+		}
+
 		/// <inheritdoc/>
 		public Task<bool> DeleteRefAsync(NamespaceId NamespaceId, BucketId BucketId, RefId RefId, CancellationToken CancellationToken = default)
 		{
diff --git a/Engine/Source/Programs/Shared/EpicGames.Horde/Storage/Impl/TracingStorageClient.cs b/Engine/Source/Programs/Shared/EpicGames.Horde/Storage/Impl/TracingStorageClient.cs
--- a/Engine/Source/Programs/Shared/EpicGames.Horde/Storage/Impl/TracingStorageClient.cs
+++ b/Engine/Source/Programs/Shared/EpicGames.Horde/Storage/Impl/TracingStorageClient.cs
@@ -52,6 +52,12 @@ namespace EpicGames.Horde.Storage.Impl
 			return Inner.HasBlobAsync(NamespaceId, Hash, CancellationToken);
 		}
 
+		/// <inheritdoc/>
+		public Task<HashSet<IoHash>> FindMissingBlobsAsync(NamespaceId NamespaceId, HashSet<IoHash> Hashes, CancellationToken CancellationToken)
+		{
+			return Inner.FindMissingBlobsAsync(NamespaceId, Hashes, CancellationToken);
+		}
+
 		/// <inheritdoc/>
 		public Task<bool> DeleteRefAsync(NamespaceId NamespaceId, BucketId BucketId, RefId RefId, CancellationToken CancellationToken = default)
 		{
