More UHT documentation updates

#rnx

diff --git a/Engine/Source/Programs/Shared/EpicGames.UHT/Tokenizer/UhtToken.cs b/Engine/Source/Programs/Shared/EpicGames.UHT/Tokenizer/UhtToken.cs
--- a/Engine/Source/Programs/Shared/EpicGames.UHT/Tokenizer/UhtToken.cs
+++ b/Engine/Source/Programs/Shared/EpicGames.UHT/Tokenizer/UhtToken.cs
@@ -9,26 +9,85 @@ using System.Text;
 
 namespace EpicGames.UHT.Tokenizer
 {
+
+	/// <summary>
+	/// Type of the token
+	/// </summary>
 	public enum UhtTokenType
 	{
-		EndOfFile,          // End of file token.
-		EndOfDefault,       // End of default value.
-		EndOfType,          // End of type
-		EndOfDeclaration,	// End of declaration
-		Line,				// Line of text (when calling GetLine only)
-		Identifier,			// Alphanumeric identifier.
-		Symbol,				// Symbol.
-		TrueConst,			// True value via identifier
-		FalseConst,			// False value via identifier
-		FloatConst,			// Floating point constant
-		DecimalConst,		// Decimal Integer constant
-		HexConst,			// Hex integer constant
-		CharConst,			// Single character constant
-		StringConst,		// String constant
+
+		/// <summary>
+		/// End of file token.
+		/// </summary>
+		EndOfFile,
+
+		/// <summary>
+		/// End of default value. 
+		/// </summary>
+		EndOfDefault,
+
+		/// <summary>
+		/// End of type
+		/// </summary>
+		EndOfType,
+
+		/// <summary>
+		/// End of declaration
+		/// </summary>
+		EndOfDeclaration,
+
+		/// <summary>
+		/// Line of text (when calling GetLine only)
+		/// </summary>
+		Line,
+
+		/// <summary>
+		/// Alphanumeric identifier.
+		/// </summary>
+		Identifier,
+
+		/// <summary>
+		/// Symbol.
+		/// </summary>
+		Symbol,
+
+		/// <summary>
+		/// Floating point constant
+		/// </summary>
+		FloatConst,
+
+		/// <summary>
+		/// Decimal Integer constant
+		/// </summary>
+		DecimalConst,
+
+		/// <summary>
+		/// Hex integer constant
+		/// </summary>
+		HexConst,
+
+		/// <summary>
+		/// Single character constant
+		/// </summary>
+		CharConst,
+
+		/// <summary>
+		/// String constant
+		/// </summary>
+		StringConst,
 	}
 
+	/// <summary>
+	/// Series of extension methods for the token type
+	/// </summary>
 	public static class UhtTokenTypeExtensions
 	{
+
+		/// <summary>
+		/// Return true if the token type is an end type
+		/// </summary>
+		/// <param name="TokenType">Token type in question</param>
+		/// <returns>True if the token type is an end type</returns>
 		public static bool IsEndType(this UhtTokenType TokenType)
 		{
 			switch (TokenType)
@@ -44,19 +103,61 @@ namespace EpicGames.UHT.Tokenizer
 		}
 	}
 
+	/// <summary>
+	/// Token declaration
+	/// </summary>
 	public struct UhtToken
 	{
+
+		/// <summary>
+		/// Names/Identifiers can not be longer that the following
+		/// </summary>
 		public const int MaxNameLength = 1024;
+
+		/// <summary>
+		/// Strings can not be longer than the following.
+		/// </summary>
 		public const int MaxStringLength = 1024;
 
+		/// <summary>
+		/// Type of the token
+		/// </summary>
 		public UhtTokenType TokenType { get; set; }
+
+		/// <summary>
+		/// Position to restore the reader
+		/// </summary>
 		public int UngetPos { get; set; }
+
+		/// <summary>
+		/// Line to restore the reader
+		/// </summary>
 		public int UngetLine { get; set; }
+
+		/// <summary>
+		/// Starting position of the token value
+		/// </summary>
 		public int InputStartPos { get; set; }
+
+		/// <summary>
+		/// End position of the token value
+		/// </summary>
 		public int InputEndPos { get => this.InputStartPos + this.Value.Span.Length; }
+
+		/// <summary>
+		/// Line containing the token
+		/// </summary>
 		public int InputLine { get; set; }
+
+		/// <summary>
+		/// Token value
+		/// </summary>
 		public StringView Value { get; set; }
 
+		/// <summary>
+		/// Construct a new token
+		/// </summary>
+		/// <param name="TokenType">Type of the token</param>
 		public UhtToken(UhtTokenType TokenType)
 		{
 			this.TokenType = TokenType;
@@ -67,6 +168,15 @@ namespace EpicGames.UHT.Tokenizer
 			this.Value = new StringView();
 		}
 
+		/// <summary>
+		/// Construct a new token
+		/// </summary>
+		/// <param name="TokenType">Type of token</param>
+		/// <param name="UngetPos">Unget position</param>
+		/// <param name="UngetLine">Unget line</param>
+		/// <param name="InputStartPos">Start position of value</param>
+		/// <param name="InputLine">Line of value</param>
+		/// <param name="Value">Token value</param>
 		public UhtToken(UhtTokenType TokenType, int UngetPos, int UngetLine, int InputStartPos, int InputLine, StringView Value)
 		{
 			this.TokenType = TokenType;
@@ -77,82 +187,160 @@ namespace EpicGames.UHT.Tokenizer
 			this.Value = Value;
 		}
 
+		/// <summary>
+		/// True if the token isn't an end token
+		/// </summary>
+		/// <param name="Token">Token in question</param>
 		public static implicit operator bool(UhtToken Token)
 		{
 			return !Token.IsEndType();
 		}
 
+		/// <summary>
+		/// Return true if the token is an end token
+		/// </summary>
+		/// <returns>True if the token is an end token</returns>
 		public bool IsEndType()
 		{
 			return this.TokenType.IsEndType();
 		}
 
+		/// <summary>
+		/// Test to see if the value matches the given character
+		/// </summary>
+		/// <param name="Value">Value to test</param>
+		/// <returns>True if the token value matches the given value</returns>
 		public bool IsValue(char Value)
 		{
 			return this.Value.Span.Length == 1 && this.Value.Span[0] == Value;
 		}
 
+		/// <summary>
+		/// Test to see if the value matches the given string
+		/// </summary>
+		/// <param name="Value">Value to test</param>
+		/// <param name="bIgnoreCase">If true, ignore case</param>
+		/// <returns>True if the value matches</returns>
 		public bool IsValue(string Value, bool bIgnoreCase = false)
 		{
 			return this.Value.Span.Equals(Value, bIgnoreCase ? StringComparison.OrdinalIgnoreCase : StringComparison.Ordinal);
 		}
 
+		/// <summary>
+		/// Test to see if the value matches the given string
+		/// </summary>
+		/// <param name="Value">Value to test</param>
+		/// <param name="bIgnoreCase">If true, ignore case</param>
+		/// <returns>True if the value matches</returns>
 		public bool IsValue(StringView Value, bool bIgnoreCase = false)
 		{
 			return this.Value.Span.Equals(Value.Span, bIgnoreCase ? StringComparison.OrdinalIgnoreCase : StringComparison.Ordinal);
 		}
 
+		/// <summary>
+		/// Test to see if the value starts with the given string
+		/// </summary>
+		/// <param name="Value">Value to test</param>
+		/// <param name="bIgnoreCase">If true, ignore case</param>
+		/// <returns>True is the value starts with the given string</returns>
 		public bool ValueStartsWith(string Value, bool bIgnoreCase = false)
 		{
 			return this.Value.Span.StartsWith(Value, bIgnoreCase ? StringComparison.OrdinalIgnoreCase : StringComparison.Ordinal);
 		}
 
+		/// <summary>
+		/// Return true if the token is an identifier
+		/// </summary>
+		/// <returns>True if the token is an identifier</returns>
 		public bool IsIdentifier()
 		{
 			return this.TokenType == UhtTokenType.Identifier;
 		}
 
+		/// <summary>
+		/// Return true if the identifier matches
+		/// </summary>
+		/// <param name="Identifier">Identifier to test</param>
+		/// <param name="bIgnoreCase">If true, ignore case</param>
+		/// <returns>True if the identifier matches</returns>
 		public bool IsIdentifier(string Identifier, bool bIgnoreCase = false)
 		{
 			return IsIdentifier() && IsValue(Identifier, bIgnoreCase);
 		}
 
+		/// <summary>
+		/// Return true if the identifier matches
+		/// </summary>
+		/// <param name="Identifier">Identifier to test</param>
+		/// <param name="bIgnoreCase">If true, ignore case</param>
+		/// <returns>True if the identifier matches</returns>
 		public bool IsIdentifier(StringView Identifier, bool bIgnoreCase = false)
 		{
 			return IsIdentifier() && IsValue(Identifier, bIgnoreCase);
 		}
 
+		/// <summary>
+		/// Return true if the token is a symbol
+		/// </summary>
+		/// <returns>True if the token is a symbol</returns>
 		public bool IsSymbol()
 		{
 			return this.TokenType == UhtTokenType.Symbol;
 		}
 
+		/// <summary>
+		/// Return true if the symbol matches
+		/// </summary>
+		/// <param name="Symbol">Symbol to test</param>
+		/// <returns>True if the symbol matches</returns>
 		public bool IsSymbol(char Symbol)
 		{
 			return IsSymbol() && IsValue(Symbol);
 		}
 
+		/// <summary>
+		/// Return true if the symbol matches
+		/// </summary>
+		/// <param name="Symbol">Symbol to test</param>
+		/// <returns>True if the symbol matches</returns>
 		public bool IsSymbol(string Symbol)
 		{
 			return IsSymbol() && IsValue(Symbol);
 		}
 
+		/// <summary>
+		/// Return true if the symbol matches
+		/// </summary>
+		/// <param name="Symbol">Symbol to test</param>
+		/// <returns>True if the symbol matches</returns>
 		public bool IsSymbol(StringView Symbol)
 		{
 			return IsSymbol() && IsValue(Symbol);
 		}
 
+		/// <summary>
+		/// Return true if the token is a constant integer
+		/// </summary>
+		/// <returns>True if constant integer</returns>
 		public bool IsConstInt()
 		{
 			return this.TokenType == UhtTokenType.DecimalConst || this.TokenType == UhtTokenType.HexConst;
 		}
 
+		/// <summary>
+		/// Return true if the token is a constant floag
+		/// </summary>
+		/// <returns>True if constant float</returns>
 		public bool IsConstFloat()
 		{
 			return this.TokenType == UhtTokenType.FloatConst;
 		}
 
-		// Get the integer value of the token.  Only supported for decimal, hexidecimal, and floating point values
+		/// <summary>
+		/// Get the integer value of the token.  Only supported for decimal, hexadecimal, and floating point values
+		/// </summary>
+		/// <param name="Value">Resulting value</param>
+		/// <returns>True if the value was set</returns>
 		public bool GetConstInt(out int Value)
 		{
 			switch (TokenType)
@@ -175,7 +363,11 @@ namespace EpicGames.UHT.Tokenizer
 			}
 		}
 
-		// Get the integer value of the token.  Only supported for decimal, hexadecimal, and floating point values
+		/// <summary>
+		/// Get the integer value of the token.  Only supported for decimal, hexadecimal, and floating point values
+		/// </summary>
+		/// <param name="Value">Resulting value</param>
+		/// <returns>True if the value was set</returns>
 		public bool GetConstLong(out long Value)
 		{
 			switch (TokenType)
@@ -198,7 +390,11 @@ namespace EpicGames.UHT.Tokenizer
 			}
 		}
 
-		// Get the float value of the token.  Only supported for decimal, hexadecimal, and floating point values
+		/// <summary>
+		/// Get the float value of the token.  Only supported for decimal, hexadecimal, and floating point values
+		/// </summary>
+		/// <param name="Value">Resulting value</param>
+		/// <returns>True if the value was set</returns>
 		public bool GetConstFloat(out float Value)
 		{
 			switch (TokenType)
@@ -218,7 +414,11 @@ namespace EpicGames.UHT.Tokenizer
 			}
 		}
 
-		// Get the double value of the token.  Only supported for decimal, hexadecimal, and floating point values
+		/// <summary>
+		/// Get the double value of the token.  Only supported for decimal, hexadecimal, and floating point values
+		/// </summary>
+		/// <param name="Value">Resulting value</param>
+		/// <returns>True if the value was set</returns>
 		public bool GetConstDouble(out double Value)
 		{
 			switch (TokenType)
@@ -238,12 +438,23 @@ namespace EpicGames.UHT.Tokenizer
 			}
 		}
 
+		/// <summary>
+		/// Return true if the token is a constant string (or a char constant)
+		/// </summary>
+		/// <returns>True if the token is a string or character constant</returns>
 		public bool IsConstString()
 		{
 			return this.TokenType == UhtTokenType.StringConst || this.TokenType == UhtTokenType.CharConst;
 		}
 
 		// Return the token value for string constants
+
+		/// <summary>
+		/// Return an un-escaped string.  The surrounding quotes will be removed and escaped characters will be converted to the actual values.
+		/// </summary>
+		/// <param name="MessageSite"></param>
+		/// <returns>Resulting string</returns>
+		/// <exception cref="UhtException">Thrown if the token type is not a string or character constant</exception>
 		public StringView GetUnescapedString(IUhtMessageSite MessageSite)
 		{
 			switch (this.TokenType)
@@ -296,6 +507,11 @@ namespace EpicGames.UHT.Tokenizer
 			throw new UhtException(MessageSite, this.InputLine, "Call to GetUnescapedString on token that isn't a string or char constant");
 		}
 
+		/// <summary>
+		/// Return a string representation of the token value.  This will convert numeric values and format them.
+		/// </summary>
+		/// <param name="bRespectQuotes">If true, embedded quotes will be respected</param>
+		/// <returns>Resulting string</returns>
 		public StringView GetConstantValue(bool bRespectQuotes = false)
 		{
 			switch (this.TokenType)
@@ -314,6 +530,11 @@ namespace EpicGames.UHT.Tokenizer
 			}
 		}
 
+		/// <summary>
+		/// Return an un-escaped string.  The surrounding quotes will be removed and escaped characters will be converted to the actual values.
+		/// </summary>
+		/// <param name="bRespectQuotes">If true, respect embedded quotes</param>
+		/// <returns>Resulting string</returns>
 		public StringView GetTokenString(bool bRespectQuotes = false)
 		{
 			StringViewBuilder SVB = new StringViewBuilder();
@@ -382,59 +603,11 @@ namespace EpicGames.UHT.Tokenizer
 			return SVB.ToStringView();
 		}
 
-		float GetFloatValue()
-		{
-			if (Value.Span.Length > 0)
-			{
-				if (UhtFCString.IsFloatMarker(Value.Span[Value.Span.Length - 1]))
-				{
-					return float.Parse(Value.Span.Slice(0, Value.Span.Length - 1));
-				}
-			}
-			return float.Parse(Value.Span);
-		}
-
-		double GetDoubleValue()
-		{
-			if (Value.Span.Length > 0)
-			{
-				if (UhtFCString.IsFloatMarker(Value.Span[Value.Span.Length - 1]))
-				{
-					return double.Parse(Value.Span.Slice(0, Value.Span.Length - 1));
-				}
-			}
-			return double.Parse(Value.Span);
-		}
-
-		long GetDecimalValue()
-		{
-			ReadOnlySpan<char> Span = this.Value.Span;
-			bool bIsUnsigned = false;
-			while (Span.Length > 0)
-			{
-				char C = Span[Span.Length - 1];
-				if (UhtFCString.IsLongMarker(C))
-				{
-					Span = Span.Slice(0, Span.Length - 1);
-				}
-				else if (UhtFCString.IsUnsignedMarker(C))
-				{
-					bIsUnsigned = true;
-					Span = Span.Slice(0, Span.Length - 1);
-				}
-				else
-				{
-					break;
-				}
-			}
-			return bIsUnsigned ? (long)Convert.ToUInt64(Span.ToString(), 10) : Convert.ToInt64(Span.ToString(), 10);
-		}
-
-		long GetHexValue()
-		{
-			return Convert.ToInt64(this.Value.ToString(), 16);
-		}
-
+		/// <summary>
+		/// Join the given tokens into a string
+		/// </summary>
+		/// <param name="Tokens">Tokens to join</param>
+		/// <returns>Joined strings</returns>
 		public static string Join(IEnumerable<UhtToken> Tokens)
 		{
 			StringBuilder Builder = new StringBuilder();
@@ -445,7 +618,13 @@ namespace EpicGames.UHT.Tokenizer
 			return Builder.ToString();
 		}
 
-		public static string Join(char Seperator, IEnumerable<UhtToken> Tokens)
+		/// <summary>
+		/// Join the given tokens into a string
+		/// </summary>
+		/// <param name="Separator">Separator between tokens</param>
+		/// <param name="Tokens">Tokens to join</param>
+		/// <returns>Joined strings</returns>
+		public static string Join(char Separator, IEnumerable<UhtToken> Tokens)
 		{
 			StringBuilder Builder = new StringBuilder();
 			bool bIncludeSeperator = false;
@@ -457,14 +636,20 @@ namespace EpicGames.UHT.Tokenizer
 				}
 				else
 				{
-					Builder.Append(Seperator);
+					Builder.Append(Separator);
 				}
 				Builder.Append(Token.Value.ToString());
 			}
 			return Builder.ToString();
 		}
 
-		public static string Join(string Seperator, IEnumerable<UhtToken> Tokens)
+		/// <summary>
+		/// Join the given tokens into a string
+		/// </summary>
+		/// <param name="Separator">Separator between tokens</param>
+		/// <param name="Tokens">Tokens to join</param>
+		/// <returns>Joined strings</returns>
+		public static string Join(string Separator, IEnumerable<UhtToken> Tokens)
 		{
 			StringBuilder Builder = new StringBuilder();
 			bool bIncludeSeperator = false;
@@ -476,13 +661,17 @@ namespace EpicGames.UHT.Tokenizer
 				}
 				else
 				{
-					Builder.Append(Seperator);
+					Builder.Append(Separator);
 				}
 				Builder.Append(Token.Value.ToString());
 			}
 			return Builder.ToString();
 		}
 
+		/// <summary>
+		/// Convert the token to a string.  This will be the value.
+		/// </summary>
+		/// <returns>Value of the token</returns>
 		public override string ToString()
 		{
 			if (IsEndType())
@@ -494,5 +683,58 @@ namespace EpicGames.UHT.Tokenizer
 				return Value.Span.ToString();
 			}
 		}
+
+		float GetFloatValue()
+		{
+			if (Value.Span.Length > 0)
+			{
+				if (UhtFCString.IsFloatMarker(Value.Span[Value.Span.Length - 1]))
+				{
+					return float.Parse(Value.Span.Slice(0, Value.Span.Length - 1));
+				}
+			}
+			return float.Parse(Value.Span);
+		}
+
+		double GetDoubleValue()
+		{
+			if (Value.Span.Length > 0)
+			{
+				if (UhtFCString.IsFloatMarker(Value.Span[Value.Span.Length - 1]))
+				{
+					return double.Parse(Value.Span.Slice(0, Value.Span.Length - 1));
+				}
+			}
+			return double.Parse(Value.Span);
+		}
+
+		long GetDecimalValue()
+		{
+			ReadOnlySpan<char> Span = this.Value.Span;
+			bool bIsUnsigned = false;
+			while (Span.Length > 0)
+			{
+				char C = Span[Span.Length - 1];
+				if (UhtFCString.IsLongMarker(C))
+				{
+					Span = Span.Slice(0, Span.Length - 1);
+				}
+				else if (UhtFCString.IsUnsignedMarker(C))
+				{
+					bIsUnsigned = true;
+					Span = Span.Slice(0, Span.Length - 1);
+				}
+				else
+				{
+					break;
+				}
+			}
+			return bIsUnsigned ? (long)Convert.ToUInt64(Span.ToString(), 10) : Convert.ToInt64(Span.ToString(), 10);
+		}
+
+		long GetHexValue()
+		{
+			return Convert.ToInt64(this.Value.ToString(), 16);
+		}
 	}
 }
diff --git a/Engine/Source/Programs/Shared/EpicGames.UHT/Tokenizer/UhtTokenBufferReader.cs b/Engine/Source/Programs/Shared/EpicGames.UHT/Tokenizer/UhtTokenBufferReader.cs
--- a/Engine/Source/Programs/Shared/EpicGames.UHT/Tokenizer/UhtTokenBufferReader.cs
+++ b/Engine/Source/Programs/Shared/EpicGames.UHT/Tokenizer/UhtTokenBufferReader.cs
@@ -8,6 +8,10 @@ using System.Diagnostics;
 
 namespace EpicGames.UHT.Tokenizer
 {
+
+	/// <summary>
+	/// Token reader for source buffers
+	/// </summary>
 	public class UhtTokenBufferReader : IUhtTokenReader, IUhtMessageLineNumber
 	{
 		private static StringView[] EmptyComments = new StringView[] { };
@@ -34,6 +38,11 @@ namespace EpicGames.UHT.Tokenizer
 		private bool bRecordTokens = false;
 		private int PreprocessorPendingCommentsCount = 0;
 
+		/// <summary>
+		/// Construct a new token reader
+		/// </summary>
+		/// <param name="MessageSite">Message site for messages</param>
+		/// <param name="Input">Input source</param>
 		public UhtTokenBufferReader(IUhtMessageSite MessageSite, ReadOnlyMemory<char> Input)
 		{
 			this.MessageSiteInternal = MessageSite;
@@ -51,6 +60,7 @@ namespace EpicGames.UHT.Tokenizer
 		#endregion
 
 		#region ITokenReader Implementation
+		/// <inheritdoc/>
 		public bool bIsEOF
 		{
 			get
@@ -66,17 +76,20 @@ namespace EpicGames.UHT.Tokenizer
 			}
 		}
 
+		/// <inheritdoc/>
 		public int InputPos
 		{
 			get => this.bHasToken ? this.PreCurrentTokenInputPos : this.InputPosInternal;
 		}
 
+		/// <inheritdoc/>
 		public int InputLine
 		{
 			get => this.bHasToken ? this.PreCurrentTokenInputLine : this.InputLineInternal;
 			set { ClearToken(); this.InputLineInternal = value; }
 		}
 
+		/// <inheritdoc/>
 		public ReadOnlySpan<StringView> Comments 
 		{
 			get
@@ -92,8 +105,10 @@ namespace EpicGames.UHT.Tokenizer
 			}
 		}
 
+		/// <inheritdoc/>
 		public IUhtTokenPreprocessor? TokenPreprocessor { get => this.TokenPreprocessorInternal; set => this.TokenPreprocessorInternal = value; }
 
+		/// <inheritdoc/>
 		public ref UhtToken PeekToken()
 		{
 			if (!this.bHasToken)
@@ -104,9 +119,7 @@ namespace EpicGames.UHT.Tokenizer
 			return ref this.CurrentToken;
 		}
 
-		/// <summary>
-		/// Skip all leading whitespace and collect any comments
-		/// </summary>
+		/// <inheritdoc/>
 		public void SkipWhitespaceAndComments()
 		{
 			bool bGotInlineComment = false;
diff --git a/Engine/Source/Programs/Shared/EpicGames.UHT/Tokenizer/UhtTokenReader.cs b/Engine/Source/Programs/Shared/EpicGames.UHT/Tokenizer/UhtTokenReader.cs
--- a/Engine/Source/Programs/Shared/EpicGames.UHT/Tokenizer/UhtTokenReader.cs
+++ b/Engine/Source/Programs/Shared/EpicGames.UHT/Tokenizer/UhtTokenReader.cs
@@ -9,11 +9,27 @@ using System.Threading;
 
 namespace EpicGames.UHT.Tokenizer
 {
+
+	/// <summary>
+	/// Options for GetRawString method
+	/// </summary>
 	[Flags]
 	public enum UhtRawStringOptions
 	{
+
+		/// <summary>
+		/// No options
+		/// </summary>
 		None,
+
+		/// <summary>
+		/// Don't consider the terminator while in a quoted string
+		/// </summary>
 		RespectQuotes = 1 << 0,
+
+		/// <summary>
+		/// Don't consume the terminator.  It will be parsed later.
+		/// </summary>
 		DontConsumeTerminator = 1 << 1,
 	}
 
@@ -59,13 +75,36 @@ namespace EpicGames.UHT.Tokenizer
 		}
 	}
 
+	/// <summary>
+	/// Interface invoked when the parser reaches a '#' preprocessor block
+	/// </summary>
 	public interface IUhtTokenPreprocessor
 	{
+
+		/// <summary>
+		/// Parse a preprocessor directive
+		/// </summary>
+		/// <param name="Token">Token starting the directive.  Will be only the '#'</param>
+		/// <param name="bIsBeingIncluded">If true, the directive the source is being included.  Otherwise it is being skipped as part of an #if block</param>
+		/// <param name="bClearComments">If true, comments should be cleared</param>
+		/// <param name="bIllegalContentsCheck">If true, excluded contents should be checked for unparsed UE macros (i.e. UCLASS) </param>
+		/// <returns>True if the source should continue to be included</returns>
 		public bool ParsePreprocessorDirective(ref UhtToken Token, bool bIsBeingIncluded, out bool bClearComments, out bool bIllegalContentsCheck);
+
+		/// <summary>
+		/// Save the current preprocessor state
+		/// </summary>
 		public void SaveState();
+
+		/// <summary>
+		/// Restore the current preprocessor state
+		/// </summary>
 		public void RestoreState();
 	}
 
+	/// <summary>
+	/// Common token reader interfaces for all token reader.  When creating extension methods, use the interface.
+	/// </summary>
 	public interface IUhtTokenReader : IUhtMessageSite
 	{
 		/// <summary>
@@ -287,6 +326,11 @@ namespace EpicGames.UHT.Tokenizer
 	{
 		static private ThreadLocal<UhtTokenList?> Tls = new ThreadLocal<UhtTokenList?>(() => null);
 
+		/// <summary>
+		/// Borrow a token list
+		/// </summary>
+		/// <param name="Token">Starting token</param>
+		/// <returns>Token list</returns>
 		public static UhtTokenList Borrow(UhtToken Token)
 		{
 			UhtTokenList? Identifier = Tls.Value;
@@ -303,6 +347,10 @@ namespace EpicGames.UHT.Tokenizer
 			return Identifier;
 		}
 
+		/// <summary>
+		/// Return a token list to the cache
+		/// </summary>
+		/// <param name="Identifier"></param>
 		public static void Return(UhtTokenList Identifier)
 		{
 			UhtTokenList? Tail = Tls.Value;
@@ -319,35 +367,81 @@ namespace EpicGames.UHT.Tokenizer
 		}
 	}
 
+	/// <summary>
+	/// Delegate for when a token is parsed
+	/// </summary>
+	/// <param name="Token">The token in question</param>
 	public delegate void UhtTokenDelegate(ref UhtToken Token);
+
+	/// <summary>
+	/// Delegate for when a token is parsed in an until block
+	/// </summary>
+	/// <param name="Token">The token in question</param>
+	/// <returns>True if parsing should continue</returns>
 	public delegate bool UhtTokensUntilDelegate(ref UhtToken Token);
-	public delegate void UhtTokensDelegate(IEnumerable<UhtToken> Token);
+
+	/// <summary>
+	/// Delegate for an enumeration of tokens
+	/// </summary>
+	/// <param name="Tokens">Parsed tokens</param>
+	public delegate void UhtTokensDelegate(IEnumerable<UhtToken> Tokens);
+
+	/// <summary>
+	/// Delegate for cached token list
+	/// </summary>
+	/// <param name="TokenList">Token list that can be cached</param>
 	public delegate void UhtTokenListDelegate(UhtTokenList TokenList);
-	public delegate void UhtTokenPositionDelegate(int Position);
+
+	/// <summary>
+	/// Delegate for a constant float
+	/// </summary>
+	/// <param name="Value">Value in question</param>
 	public delegate void UhtTokenConstFloatDelegate(float Value);
+
+	/// <summary>
+	/// Delegate for a constant double
+	/// </summary>
+	/// <param name="Value">Value in question</param>
 	public delegate void UhtTokenConstDoubleDelegate(double Value);
 
+	/// <summary>
+	/// Helper struct to disable comment parsing.  Should be used in a using block
+	/// </summary>
 	public struct UhtTokenDisableComments : IDisposable
 	{
 		private readonly IUhtTokenReader TokenReader;
 
+		/// <summary>
+		/// Construct instance
+		/// </summary>
+		/// <param name="TokenReader">Token reader to disable</param>
 		public UhtTokenDisableComments(IUhtTokenReader TokenReader)
 		{
 			this.TokenReader = TokenReader;
 			this.TokenReader.DisableComments();
 		}
 
+		/// <summary>
+		/// Enable comments
+		/// </summary>
 		public void Dispose()
 		{
 			this.TokenReader.EnableComments();
 		}
 	}
 
+	/// <summary>
+	/// Helper struct to save token reader state
+	/// </summary>
 	public struct UhtTokenSaveState : IDisposable
 	{
 		private readonly IUhtTokenReader TokenReader;
 		private bool bHandled;
 
+		/// <summary>
+		/// Construct instance
+		/// </summary>
+		/// <param name="TokenReader">Token reader</param>
 		public UhtTokenSaveState(IUhtTokenReader TokenReader)
 		{
 			this.TokenReader = TokenReader;
@@ -355,14 +449,21 @@ namespace EpicGames.UHT.Tokenizer
 			this.TokenReader.SaveState();
 		}
 
+		/// <summary>
+		/// Restore the token reader state
+		/// </summary>
 		public void Dispose()
 		{
-			if (!bHandled)
+			if (!this.bHandled)
 			{
 				RestoreState();
 			}
 		}
 
+		/// <summary>
+		/// Restore the token reader state
+		/// </summary>
+		/// <exception cref="UhtIceException">Thrown if state has already been restored or aborted</exception>
 		public void RestoreState()
 		{
 			if (this.bHandled)
@@ -373,6 +474,10 @@ namespace EpicGames.UHT.Tokenizer
 			this.bHandled = true;
 		}
 
+		/// <summary>
+		/// Abandon the saved state
+		/// </summary>
+		/// <exception cref="UhtIceException">Thrown if state has already been restored or aborted</exception>
 		public void AbandonState()
 		{
 			if (this.bHandled)
diff --git a/Engine/Source/Programs/Shared/EpicGames.UHT/Tokenizer/UhtTokenReaderFloatExtensions.cs b/Engine/Source/Programs/Shared/EpicGames.UHT/Tokenizer/UhtTokenReaderFloatExtensions.cs
--- a/Engine/Source/Programs/Shared/EpicGames.UHT/Tokenizer/UhtTokenReaderFloatExtensions.cs
+++ b/Engine/Source/Programs/Shared/EpicGames.UHT/Tokenizer/UhtTokenReaderFloatExtensions.cs
@@ -4,13 +4,17 @@ using EpicGames.UHT.Utils;
 
 namespace EpicGames.UHT.Tokenizer
 {
+
+	/// <summary>
+	/// Collection of token reader extensions for float values
+	/// </summary>
 	public static class UhtTokenReaderFloatExtensions
 	{
 		/// <summary>
 		/// Get the next token as a float.  If the next token is not a float, no token is consumed.
 		/// </summary>
+		/// <param name="TokenReader">Token reader</param>
 		/// <param name="Value">The float value of the token</param>
-		/// <param name="ExceptionContext">If not null, an exception will be thrown with the given text as part of the message.</param>
 		/// <returns>True if the next token was an float, false if not.</returns>
 		public static bool TryOptionalConstFloat(this IUhtTokenReader TokenReader, out float Value)
 		{
@@ -27,8 +31,7 @@ namespace EpicGames.UHT.Tokenizer
 		/// <summary>
 		/// Get the next token as a float.  If the next token is not a float, no token is consumed.
 		/// </summary>
-		/// <param name="Value">The float value of the token</param>
-		/// <param name="ExceptionContext">If not null, an exception will be thrown with the given text as part of the message.</param>
+		/// <param name="TokenReader"></param>
 		/// <returns>The token reader</returns>
 		public static IUhtTokenReader OptionalConstFloat(this IUhtTokenReader TokenReader)
 		{
@@ -40,8 +43,8 @@ namespace EpicGames.UHT.Tokenizer
 		/// <summary>
 		/// Get the next token as a float.  If the next token is not a float, no token is consumed.
 		/// </summary>
-		/// <param name="Value">The float value of the token</param>
-		/// <param name="ExceptionContext">If not null, an exception will be thrown with the given text as part of the message.</param>
+		/// <param name="TokenReader">Token reader</param>
+		/// <param name="Delegate">Delegate to invoke with the float value</param>
 		/// <returns>The token reader</returns>
 		public static IUhtTokenReader OptionalConstFloat(this IUhtTokenReader TokenReader, UhtTokenConstFloatDelegate Delegate)
 		{
@@ -56,6 +59,7 @@ namespace EpicGames.UHT.Tokenizer
 		/// <summary>
 		/// Get the next token as a float.  If the next token is not a float, an exception is thrown
 		/// </summary>
+		/// <param name="TokenReader">Token reader</param>
 		/// <param name="ExceptionContext">If not null, an exception will be thrown with the given text as part of the message.</param>
 		/// <returns>True if the next token was an float, false if not.</returns>
 		public static IUhtTokenReader RequireConstFloat(this IUhtTokenReader TokenReader, object? ExceptionContext = null)
@@ -71,6 +75,7 @@ namespace EpicGames.UHT.Tokenizer
 		/// <summary>
 		/// Get the next token as a float.  If the next token is not a float, an exception is thrown
 		/// </summary>
+		/// <param name="TokenReader">Token reader</param>
 		/// <param name="ExceptionContext">If not null, an exception will be thrown with the given text as part of the message.</param>
 		/// <returns>The floating point value of the token</returns>
 		public static float GetConstFloat(this IUhtTokenReader TokenReader, object? ExceptionContext = null)
diff --git a/Engine/Source/Programs/Shared/EpicGames.UHT/Tokenizer/UhtTokenReaderFlowControlExtensions.cs b/Engine/Source/Programs/Shared/EpicGames.UHT/Tokenizer/UhtTokenReaderFlowControlExtensions.cs
--- a/Engine/Source/Programs/Shared/EpicGames.UHT/Tokenizer/UhtTokenReaderFlowControlExtensions.cs
+++ b/Engine/Source/Programs/Shared/EpicGames.UHT/Tokenizer/UhtTokenReaderFlowControlExtensions.cs
@@ -6,8 +6,23 @@ using System.Collections.Generic;
 
 namespace EpicGames.UHT.Tokenizer
 {
+
+	/// <summary>
+	/// Collection of token reader extensions to help with flow control
+	/// </summary>
 	public static class UhtTokenReaderFlowControlExtensions
 	{
+
+		/// <summary>
+		/// Parse an optional list
+		/// </summary>
+		/// <param name="TokenReader">Token reader</param>
+		/// <param name="Initiator">Initiating symbol</param>
+		/// <param name="Terminator">Terminating symbol</param>
+		/// <param name="Separator">Separator symbol</param>
+		/// <param name="bAllowTrailingSeparator">If true, allow trailing separators</param>
+		/// <param name="Action">Action to be invoked for each list element.</param>
+		/// <returns>True if a list was read</returns>
 		public static bool TryOptionalList(this IUhtTokenReader TokenReader, char Initiator, char Terminator, char Separator, bool bAllowTrailingSeparator, Action Action)
 		{
 			if (TokenReader.TryOptional(Initiator))
@@ -18,22 +33,62 @@ namespace EpicGames.UHT.Tokenizer
 			return false;
 		}
 
+		/// <summary>
+		/// Parse a required list
+		/// </summary>
+		/// <param name="TokenReader">Token reader</param>
+		/// <param name="Initiator">Initiating symbol</param>
+		/// <param name="Terminator">Terminating symbol</param>
+		/// <param name="Separator">Separator symbol</param>
+		/// <param name="bAllowTrailingSeparator">If true, allow trailing separators</param>
+		/// <param name="Action">Action to be invoked for each list element.</param>
+		/// <returns>Token reader</returns>
 		public static IUhtTokenReader RequireList(this IUhtTokenReader TokenReader, char Initiator, char Terminator, char Separator, bool bAllowTrailingSeparator, Action Action)
 		{
 			return TokenReader.RequireList(Initiator, Terminator, Separator, bAllowTrailingSeparator, null, Action);
 		}
 
+		/// <summary>
+		/// Parse a required list
+		/// </summary>
+		/// <param name="TokenReader">Token reader</param>
+		/// <param name="Initiator">Initiating symbol</param>
+		/// <param name="Terminator">Terminating symbol</param>
+		/// <param name="Separator">Separator symbol</param>
+		/// <param name="bAllowTrailingSeparator">If true, allow trailing separators</param>
+		/// <param name="ExceptionContext">Extra context for error messages</param>
+		/// <param name="Action">Action to be invoked for each list element.</param>
+		/// <returns>Token reader</returns>
 		public static IUhtTokenReader RequireList(this IUhtTokenReader TokenReader, char Initiator, char Terminator, char Separator, bool bAllowTrailingSeparator, object? ExceptionContext, Action Action)
 		{
 			TokenReader.Require(Initiator);
 			return TokenReader.RequireList(Terminator, Separator, bAllowTrailingSeparator, ExceptionContext, Action);
 		}
 
+		/// <summary>
+		/// Parse a required list.  Initiating token must have already been parsed
+		/// </summary>
+		/// <param name="TokenReader">Token reader</param>
+		/// <param name="Terminator">Terminating symbol</param>
+		/// <param name="Separator">Separator symbol</param>
+		/// <param name="bAllowTrailingSeparator">If true, allow trailing separators</param>
+		/// <param name="Action">Action to be invoked for each list element.</param>
+		/// <returns>Token reader</returns>
 		public static IUhtTokenReader RequireList(this IUhtTokenReader TokenReader, char Terminator, char Separator, bool bAllowTrailingSeparator, Action Action)
 		{
 			return TokenReader.RequireList(Terminator, Separator, bAllowTrailingSeparator, null, Action);
 		}
 
+		/// <summary>
+		/// Parse a required list.  Initiating token must have already been parsed
+		/// </summary>
+		/// <param name="TokenReader">Token reader</param>
+		/// <param name="Terminator">Terminating symbol</param>
+		/// <param name="Separator">Separator symbol</param>
+		/// <param name="bAllowTrailingSeparator">If true, allow trailing separators</param>
+		/// <param name="ExceptionContext">Extra context for error messages</param>
+		/// <param name="Action">Action to be invoked for each list element.</param>
+		/// <returns>Token reader</returns>
 		public static IUhtTokenReader RequireList(this IUhtTokenReader TokenReader, char Terminator, char Separator, bool bAllowTrailingSeparator, object? ExceptionContext, Action Action)
 		{
 			// Check for an empty list
@@ -74,6 +129,16 @@ namespace EpicGames.UHT.Tokenizer
 			}
 		}
 
+		/// <summary>
+		/// Parse an optional list
+		/// </summary>
+		/// <param name="TokenReader">Token reader</param>
+		/// <param name="Initiator">Initiating symbol</param>
+		/// <param name="Terminator">Terminating symbol</param>
+		/// <param name="Separator">Separator symbol</param>
+		/// <param name="bAllowTrailingSeparator">If true, allow trailing separators</param>
+		/// <param name="Delegate">Delegate to invoke with the parsed token list</param>
+		/// <returns>True if a list was read</returns>
 		public static bool TryOptionalList(this IUhtTokenReader TokenReader, char Initiator, char Terminator, char Separator, bool bAllowTrailingSeparator, UhtTokensDelegate Delegate)
 		{
 			if (TokenReader.TryOptional(Initiator))
@@ -84,22 +149,62 @@ namespace EpicGames.UHT.Tokenizer
 			return false;
 		}
 
+		/// <summary>
+		/// Parse a required list
+		/// </summary>
+		/// <param name="TokenReader">Token reader</param>
+		/// <param name="Initiator">Initiating symbol</param>
+		/// <param name="Terminator">Terminating symbol</param>
+		/// <param name="Separator">Separator symbol</param>
+		/// <param name="bAllowTrailingSeparator">If true, allow trailing separators</param>
+		/// <param name="Delegate">Delegate to invoke with the parsed token list</param>
+		/// <returns>Token reader</returns>
 		public static IUhtTokenReader RequireList(this IUhtTokenReader TokenReader, char Initiator, char Terminator, char Separator, bool bAllowTrailingSeparator, UhtTokensDelegate Delegate)
 		{
 			return TokenReader.RequireList(Initiator, Terminator, Separator, bAllowTrailingSeparator, null, Delegate);
 		}
 
+		/// <summary>
+		/// Parse a required list
+		/// </summary>
+		/// <param name="TokenReader">Token reader</param>
+		/// <param name="Initiator">Initiating symbol</param>
+		/// <param name="Terminator">Terminating symbol</param>
+		/// <param name="Separator">Separator symbol</param>
+		/// <param name="bAllowTrailingSeparator">If true, allow trailing separators</param>
+		/// <param name="ExceptionContext">Extra context for error messages</param>
+		/// <param name="Delegate">Delegate to invoke with the parsed token list</param>
+		/// <returns>Token reader</returns>
 		public static IUhtTokenReader RequireList(this IUhtTokenReader TokenReader, char Initiator, char Terminator, char Separator, bool bAllowTrailingSeparator, object? ExceptionContext, UhtTokensDelegate Delegate)
 		{
 			TokenReader.Require(Initiator);
 			return TokenReader.RequireList(Terminator, Separator, bAllowTrailingSeparator, ExceptionContext, Delegate);
 		}
 
+		/// <summary>
+		/// Parse a required list.  Initiating token must have already been parsed
+		/// </summary>
+		/// <param name="TokenReader">Token reader</param>
+		/// <param name="Terminator">Terminating symbol</param>
+		/// <param name="Separator">Separator symbol</param>
+		/// <param name="bAllowTrailingSeparator">If true, allow trailing separators</param>
+		/// <param name="Delegate">Delegate to invoke with the parsed token list</param>
+		/// <returns>Token reader</returns>
 		public static IUhtTokenReader RequireList(this IUhtTokenReader TokenReader, char Terminator, char Separator, bool bAllowTrailingSeparator, UhtTokensDelegate Delegate)
 		{
 			return TokenReader.RequireList(Terminator, Separator, bAllowTrailingSeparator, null, Delegate);
 		}
 
+		/// <summary>
+		/// Parse a required list.  Initiating token must have already been parsed
+		/// </summary>
+		/// <param name="TokenReader">Token reader</param>
+		/// <param name="Terminator">Terminating symbol</param>
+		/// <param name="Separator">Separator symbol</param>
+		/// <param name="bAllowTrailingSeparator">If true, allow trailing separators</param>
+		/// <param name="ExceptionContext">Extra context for error messages</param>
+		/// <param name="Delegate">Delegate to invoke with the parsed token list</param>
+		/// <returns>Token reader</returns>
 		public static IUhtTokenReader RequireList(this IUhtTokenReader TokenReader, char Terminator, char Separator, bool bAllowTrailingSeparator, object? ExceptionContext, UhtTokensDelegate Delegate)
 		{
 			// Check for an empty list
@@ -183,6 +288,13 @@ namespace EpicGames.UHT.Tokenizer
 			return TokenReader;
 		}
 
+		/// <summary>
+		/// Invoke action while the next token is the given string
+		/// </summary>
+		/// <param name="TokenReader">Token reader</param>
+		/// <param name="Text">Text to match to continue invoking Action</param>
+		/// <param name="Action">Action to invoke if and only if the prior text was parsed.</param>
+		/// <returns>Token reader</returns>
 		public static IUhtTokenReader While(this IUhtTokenReader TokenReader, string Text, Action Action)
 		{
 			while (TokenReader.TryOptional(Text))
@@ -192,6 +304,13 @@ namespace EpicGames.UHT.Tokenizer
 			return TokenReader;
 		}
 
+		/// <summary>
+		/// Invoke action while the next token is the given string
+		/// </summary>
+		/// <param name="TokenReader">Token reader</param>
+		/// <param name="Text">Text to match to continue invoking Action</param>
+		/// <param name="Action">Action to invoke if and only if the prior text was parsed.</param>
+		/// <returns>Token reader</returns>
 		public static IUhtTokenReader While(this IUhtTokenReader TokenReader, char Text, Action Action)
 		{
 			while (TokenReader.TryOptional(Text))
@@ -227,6 +346,12 @@ namespace EpicGames.UHT.Tokenizer
 			}
 		}
 
+		/// <summary>
+		/// Consume tokens until one of the strings are found.  Terminating token will not be consumed.
+		/// </summary>
+		/// <param name="TokenReader">Token reader</param>
+		/// <param name="Terminators">Strings that will terminate processing.</param>
+		/// <returns>Number of tokens consumed</returns>
 		public static int ConsumeUntil(this IUhtTokenReader TokenReader, string[] Terminators)
 		{
 			int ConsumedTokens = 0;
@@ -246,6 +371,12 @@ namespace EpicGames.UHT.Tokenizer
 			return ConsumedTokens;
 		}
 
+		/// <summary>
+		/// Consume until the given terminator is found.  Terminating token will be consumed
+		/// </summary>
+		/// <param name="TokenReader">Token reader</param>
+		/// <param name="Terminator">Terminating symbol</param>
+		/// <returns>Token reader</returns>
 		public static IUhtTokenReader ConsumeUntil(this IUhtTokenReader TokenReader, char Terminator)
 		{
 			while (!TokenReader.bIsEOF)
diff --git a/Engine/Source/Programs/Shared/EpicGames.UHT/Tokenizer/UhtTokenReaderGeneralExtensions.cs b/Engine/Source/Programs/Shared/EpicGames.UHT/Tokenizer/UhtTokenReaderGeneralExtensions.cs
--- a/Engine/Source/Programs/Shared/EpicGames.UHT/Tokenizer/UhtTokenReaderGeneralExtensions.cs
+++ b/Engine/Source/Programs/Shared/EpicGames.UHT/Tokenizer/UhtTokenReaderGeneralExtensions.cs
@@ -6,8 +6,19 @@ using System;
 
 namespace EpicGames.UHT.Tokenizer
 {
+	
+	/// <summary>
+	/// Collection of general token reader extensions
+	/// </summary>
 	public static class UhtTokenReaderGeneralExtensions
 	{
+
+		/// <summary>
+		/// Try to parse the given text
+		/// </summary>
+		/// <param name="TokenReader">Token reader</param>
+		/// <param name="Text">Text to match</param>
+		/// <returns>True if the text matched</returns>
 		public static bool TryOptional(this IUhtTokenReader TokenReader, string Text)
 		{
 			ref UhtToken Token = ref TokenReader.PeekToken();
@@ -19,6 +30,12 @@ namespace EpicGames.UHT.Tokenizer
 			return false;
 		}
 
+		/// <summary>
+		/// Try to parse the given text
+		/// </summary>
+		/// <param name="TokenReader">Token reader</param>
+		/// <param name="Text">Text to match</param>
+		/// <returns>True if the text matched</returns>
 		public static bool TryOptional(this IUhtTokenReader TokenReader, StringView Text)
 		{
 			ref UhtToken Token = ref TokenReader.PeekToken();
@@ -30,6 +47,12 @@ namespace EpicGames.UHT.Tokenizer
 			return false;
 		}
 
+		/// <summary>
+		/// Try to parse the given text.  However, the matching token will not be consumed.
+		/// </summary>
+		/// <param name="TokenReader">Token reader</param>
+		/// <param name="Text">Text to match</param>
+		/// <returns>True if the text matched</returns>
 		public static bool TryPeekOptional(this IUhtTokenReader TokenReader, string Text)
 		{
 			ref UhtToken Token = ref TokenReader.PeekToken();
@@ -56,6 +79,12 @@ namespace EpicGames.UHT.Tokenizer
 			return -1;
 		}
 
+		/// <summary>
+		/// Try to parse the given text
+		/// </summary>
+		/// <param name="TokenReader">Token reader</param>
+		/// <param name="Text">Text to match</param>
+		/// <returns>True if the text matched</returns>
 		public static bool TryOptional(this IUhtTokenReader TokenReader, char Text)
 		{
 			ref UhtToken Token = ref TokenReader.PeekToken();
@@ -67,6 +96,13 @@ namespace EpicGames.UHT.Tokenizer
 			return false;
 		}
 
+		/// <summary>
+		/// Try to parse the given text
+		/// </summary>
+		/// <param name="TokenReader">Token reader</param>
+		/// <param name="Text">Text to match</param>
+		/// <param name="OutToken">Open that was matched</param>
+		/// <returns>True if the text matched</returns>
 		public static bool TryOptional(this IUhtTokenReader TokenReader, char Text, out UhtToken OutToken)
 		{
 			ref UhtToken Token = ref TokenReader.PeekToken();
@@ -80,12 +116,24 @@ namespace EpicGames.UHT.Tokenizer
 			return false;
 		}
 
+		/// <summary>
+		/// Try to parse the given text.  However, the matching token will not be consumed.
+		/// </summary>
+		/// <param name="TokenReader">Token reader</param>
+		/// <param name="Text">Text to match</param>
+		/// <returns>True if the text matched</returns>
 		public static bool TryPeekOptional(this IUhtTokenReader TokenReader, char Text)
 		{
 			ref UhtToken Token = ref TokenReader.PeekToken();
 			return Token.IsSymbol(Text);
 		}
 
+		/// <summary>
+		/// Parse optional text
+		/// </summary>
+		/// <param name="TokenReader">Token reader</param>
+		/// <param name="Text">Text to match</param>
+		/// <returns>Token reader</returns>
 		public static IUhtTokenReader Optional(this IUhtTokenReader TokenReader, string Text)
 		{
 			ref UhtToken Token = ref TokenReader.PeekToken();
@@ -96,6 +144,13 @@ namespace EpicGames.UHT.Tokenizer
 			return TokenReader;
 		}
 
+		/// <summary>
+		/// Parse optional text
+		/// </summary>
+		/// <param name="TokenReader">Token reader</param>
+		/// <param name="Text">Text to match</param>
+		/// <param name="Action">Action to invoke if the text was found</param>
+		/// <returns>Token reader</returns>
 		public static IUhtTokenReader Optional(this IUhtTokenReader TokenReader, string Text, Action Action)
 		{
 			ref UhtToken Token = ref TokenReader.PeekToken();
@@ -107,21 +162,12 @@ namespace EpicGames.UHT.Tokenizer
 			return TokenReader;
 		}
 
-		public static IUhtTokenReader Optional(this IUhtTokenReader TokenReader, string Text, Action YesAction, Action NoAction)
-		{
-			ref UhtToken Token = ref TokenReader.PeekToken();
-			if (Token.IsIdentifier(Text) || Token.IsSymbol(Text))
-			{
-				TokenReader.ConsumeToken();
-				YesAction();
-			}
-			else
-			{
-				NoAction();
-			}
-			return TokenReader;
-		}
-
+		/// <summary>
+		/// Parse optional text
+		/// </summary>
+		/// <param name="TokenReader">Token reader</param>
+		/// <param name="Text">Text to match</param>
+		/// <returns>Token reader</returns>
 		public static IUhtTokenReader Optional(this IUhtTokenReader TokenReader, char Text)
 		{
 			ref UhtToken Token = ref TokenReader.PeekToken();
@@ -132,6 +178,13 @@ namespace EpicGames.UHT.Tokenizer
 			return TokenReader;
 		}
 
+		/// <summary>
+		/// Parse optional text
+		/// </summary>
+		/// <param name="TokenReader">Token reader</param>
+		/// <param name="Text">Text to match</param>
+		/// <param name="Action">Action to invoke if the text was found</param>
+		/// <returns>Token reader</returns>
 		public static IUhtTokenReader Optional(this IUhtTokenReader TokenReader, char Text, Action Action)
 		{
 			ref UhtToken Token = ref TokenReader.PeekToken();
@@ -143,6 +196,12 @@ namespace EpicGames.UHT.Tokenizer
 			return TokenReader;
 		}
 
+		/// <summary>
+		/// Parse optional token that starts with the given text 
+		/// </summary>
+		/// <param name="TokenReader">Token reader</param>
+		/// <param name="Text">Text to match</param>
+		/// <returns>Token reader</returns>
 		public static IUhtTokenReader OptionalStartsWith(this IUhtTokenReader TokenReader, string Text)
 		{
 			ref UhtToken Token = ref TokenReader.PeekToken();
@@ -153,6 +212,13 @@ namespace EpicGames.UHT.Tokenizer
 			return TokenReader;
 		}
 
+		/// <summary>
+		/// Parse optional token that starts with the given text 
+		/// </summary>
+		/// <param name="TokenReader">Token reader</param>
+		/// <param name="Text">Text to match</param>
+		/// <param name="TokenDelegate">Delegate to invoke on a match</param>
+		/// <returns>Token reader</returns>
 		public static IUhtTokenReader OptionalStartsWith(this IUhtTokenReader TokenReader, string Text, UhtTokenDelegate TokenDelegate)
 		{
 			ref UhtToken Token = ref TokenReader.PeekToken();
@@ -165,6 +231,14 @@ namespace EpicGames.UHT.Tokenizer
 			return TokenReader;
 		}
 
+		/// <summary>
+		/// Require the given text
+		/// </summary>
+		/// <param name="TokenReader">Token reader</param>
+		/// <param name="Text">Required text</param>
+		/// <param name="ExceptionContext">Extra exception context</param>
+		/// <returns>Token reader</returns>
+		/// <exception cref="UhtTokenException">Thrown if text is not found</exception>
 		public static IUhtTokenReader Require(this IUhtTokenReader TokenReader, string Text, object? ExceptionContext = null)
 		{
 			ref UhtToken Token = ref TokenReader.PeekToken();
@@ -179,11 +253,28 @@ namespace EpicGames.UHT.Tokenizer
 			return TokenReader;
 		}
 
+		/// <summary>
+		/// Require the given text
+		/// </summary>
+		/// <param name="TokenReader">Token reader</param>
+		/// <param name="Text">Required text</param>
+		/// <param name="TokenDelegate">Delegate to invoke on a match</param>
+		/// <returns>Token reader</returns>
+		/// <exception cref="UhtTokenException">Thrown if text is not found</exception>
 		public static IUhtTokenReader Require(this IUhtTokenReader TokenReader, string Text, UhtTokenDelegate TokenDelegate)
 		{
 			return TokenReader.Require(Text, null, TokenDelegate);
 		}
 
+		/// <summary>
+		/// Require the given text
+		/// </summary>
+		/// <param name="TokenReader">Token reader</param>
+		/// <param name="Text">Required text</param>
+		/// <param name="ExceptionContext">Extra exception context</param>
+		/// <param name="TokenDelegate">Delegate to invoke on a match</param>
+		/// <returns>Token reader</returns>
+		/// <exception cref="UhtTokenException">Thrown if text is not found</exception>
 		public static IUhtTokenReader Require(this IUhtTokenReader TokenReader, string Text, object? ExceptionContext, UhtTokenDelegate TokenDelegate)
 		{
 			ref UhtToken Token = ref TokenReader.PeekToken();
@@ -200,6 +291,14 @@ namespace EpicGames.UHT.Tokenizer
 			return TokenReader;
 		}
 
+		/// <summary>
+		/// Require the given text
+		/// </summary>
+		/// <param name="TokenReader">Token reader</param>
+		/// <param name="Text">Required text</param>
+		/// <param name="ExceptionContext">Extra exception context</param>
+		/// <returns>Token reader</returns>
+		/// <exception cref="UhtTokenException">Thrown if text is not found</exception>
 		public static IUhtTokenReader Require(this IUhtTokenReader TokenReader, char Text, object? ExceptionContext = null)
 		{
 			ref UhtToken Token = ref TokenReader.PeekToken();
@@ -214,11 +313,28 @@ namespace EpicGames.UHT.Tokenizer
 			return TokenReader;
 		}
 
+		/// <summary>
+		/// Require the given text
+		/// </summary>
+		/// <param name="TokenReader">Token reader</param>
+		/// <param name="Text">Required text</param>
+		/// <param name="TokenDelegate">Delegate to invoke on a match</param>
+		/// <returns>Token reader</returns>
+		/// <exception cref="UhtTokenException">Thrown if text is not found</exception>
 		public static IUhtTokenReader Require(this IUhtTokenReader TokenReader, char Text, UhtTokenDelegate TokenDelegate)
 		{
 			return TokenReader.Require(Text, null, TokenDelegate);
 		}
 
+		/// <summary>
+		/// Require the given text
+		/// </summary>
+		/// <param name="TokenReader">Token reader</param>
+		/// <param name="Text">Required text</param>
+		/// <param name="ExceptionContext">Extra exception context</param>
+		/// <param name="TokenDelegate">Delegate to invoke on a match</param>
+		/// <returns>Token reader</returns>
+		/// <exception cref="UhtTokenException">Thrown if text is not found</exception>
 		public static IUhtTokenReader Require(this IUhtTokenReader TokenReader, char Text, object? ExceptionContext, UhtTokenDelegate TokenDelegate)
 		{
 			ref UhtToken Token = ref TokenReader.PeekToken();
diff --git a/Engine/Source/Programs/Shared/EpicGames.UHT/Tokenizer/UhtTokenReaderIdentifierExtensions.cs b/Engine/Source/Programs/Shared/EpicGames.UHT/Tokenizer/UhtTokenReaderIdentifierExtensions.cs
--- a/Engine/Source/Programs/Shared/EpicGames.UHT/Tokenizer/UhtTokenReaderIdentifierExtensions.cs
+++ b/Engine/Source/Programs/Shared/EpicGames.UHT/Tokenizer/UhtTokenReaderIdentifierExtensions.cs
@@ -7,12 +7,23 @@ using System.Diagnostics.CodeAnalysis;
 
 namespace EpicGames.UHT.Tokenizer
 {
+
+	/// <summary>
+	/// Options when parsing identifier
+	/// </summary>
 	[Flags]
 	public enum UhtCppIdentifierOptions
 	{
+
+		/// <summary>
+		/// No options
+		/// </summary>
 		None,
+
+		/// <summary>
+		/// Include template arguments when parsing identifier
+		/// </summary>
 		AllowTemplates = 1 << 0,
-		AllTokens = 1 << 1, // AllowTemplates implies AllTokens
 	}
 
 	/// <summary>
@@ -57,6 +68,9 @@ namespace EpicGames.UHT.Tokenizer
 		}
 	}
 
+	/// <summary>
+	/// Collection of token reader extensions for working with identifiers
+	/// </summary>
 	public static class UhtTokenReaderIdentifierExtensions
 	{
 
@@ -78,6 +92,7 @@ namespace EpicGames.UHT.Tokenizer
 		/// <summary>
 		/// Get the next token and verify that it is an identifier
 		/// </summary>
+		/// <param name="TokenReader">Token reader</param>
 		/// <param name="Identifier">The fetched value of the identifier</param>
 		/// <returns>True if it is an identifier, false if not.</returns>
 		public static bool TryOptionalIdentifier(this IUhtTokenReader TokenReader, out UhtToken Identifier)
@@ -93,6 +108,12 @@ namespace EpicGames.UHT.Tokenizer
 			return false;
 		}
 
+		/// <summary>
+		/// Parse an optional identifier
+		/// </summary>
+		/// <param name="TokenReader">Token reader</param>
+		/// <param name="TokenDelegate">Invoked of an identifier is parsed</param>
+		/// <returns>Token reader</returns>
 		public static IUhtTokenReader OptionalIdentifier(this IUhtTokenReader TokenReader, UhtTokenDelegate TokenDelegate)
 		{
 			ref UhtToken Token = ref TokenReader.PeekToken();
@@ -105,6 +126,13 @@ namespace EpicGames.UHT.Tokenizer
 			return TokenReader;
 		}
 
+		/// <summary>
+		/// Parse a required identifier
+		/// </summary>
+		/// <param name="TokenReader">Token reader</param>
+		/// <param name="ExceptionContext">Extra exception context</param>
+		/// <returns>Token reader</returns>
+		/// <exception cref="UhtTokenException">Thrown if an identifier isn't found</exception>
 		public static IUhtTokenReader RequireIdentifier(this IUhtTokenReader TokenReader, object? ExceptionContext = null)
 		{
 			ref UhtToken Token = ref TokenReader.PeekToken();
@@ -116,11 +144,25 @@ namespace EpicGames.UHT.Tokenizer
 			throw new UhtTokenException(TokenReader, Token, "an identifier", ExceptionContext);
 		}
 
+		/// <summary>
+		/// Parse a required identifier
+		/// </summary>
+		/// <param name="TokenReader">Token reader</param>
+		/// <param name="TokenDelegate">Invoked if an identifier is parsed</param>
+		/// <returns>Token reader</returns>
 		public static IUhtTokenReader RequireIdentifier(this IUhtTokenReader TokenReader, UhtTokenDelegate TokenDelegate)
 		{
 			return TokenReader.RequireIdentifier(null, TokenDelegate);
 		}
 
+		/// <summary>
+		/// Parse a required identifier
+		/// </summary>
+		/// <param name="TokenReader">Token reader</param>
+		/// <param name="ExceptionContext">Extra exception context</param>
+		/// <param name="TokenDelegate">Invoked if an identifier is parsed</param>
+		/// <returns>Token reader</returns>
+		/// <exception cref="UhtTokenException">Thrown if an identifier isn't found</exception>
 		public static IUhtTokenReader RequireIdentifier(this IUhtTokenReader TokenReader, object? ExceptionContext, UhtTokenDelegate TokenDelegate)
 		{
 			ref UhtToken Token = ref TokenReader.PeekToken();
@@ -134,6 +176,13 @@ namespace EpicGames.UHT.Tokenizer
 			throw new UhtTokenException(TokenReader, Token, "an identifier", ExceptionContext);
 		}
 
+		/// <summary>
+		/// Get a required identifier
+		/// </summary>
+		/// <param name="TokenReader">Token reader</param>
+		/// <param name="ExceptionContext">Extra exception context</param>
+		/// <returns>Identifier token</returns>
+		/// <exception cref="UhtTokenException">Thrown if an identifier isn't found</exception>
 		public static UhtToken GetIdentifier(this IUhtTokenReader TokenReader, object? ExceptionContext = null)
 		{
 			ref UhtToken Token = ref TokenReader.PeekToken();
@@ -146,6 +195,13 @@ namespace EpicGames.UHT.Tokenizer
 			throw new UhtTokenException(TokenReader, Token, ExceptionContext != null ? ExceptionContext : "an identifier");
 		}
 
+		/// <summary>
+		/// Parse an optional cpp identifier
+		/// </summary>
+		/// <param name="TokenReader">Token reader</param>
+		/// <param name="Identifier">Enumeration of the identifier tokens</param>
+		/// <param name="Options">Identifier options</param>
+		/// <returns>True if an identifier is parsed</returns>
 		public static bool TryOptionalCppIdentifier(this IUhtTokenReader TokenReader, [NotNullWhen(true)] out IEnumerable<UhtToken>? Identifier, UhtCppIdentifierOptions Options = UhtCppIdentifierOptions.None)
 		{
 			Identifier = null;
@@ -164,13 +220,11 @@ namespace EpicGames.UHT.Tokenizer
 					return false;
 				}
 
-				if (Options.HasAnyFlags(UhtCppIdentifierOptions.AllowTemplates | UhtCppIdentifierOptions.AllTokens))
+				if (Options.HasAnyFlags(UhtCppIdentifierOptions.AllowTemplates))
 				{
-					bool bAllowTemplates = Options.HasAnyFlags(UhtCppIdentifierOptions.AllowTemplates);
-
 					while (true)
 					{
-						if (bAllowTemplates && TokenReader.TryPeekOptional('<'))
+						if (TokenReader.TryPeekOptional('<'))
 						{
 							LocalIdentifier.Add(TokenReader.GetToken());
 							int NestedScopes = 1;
@@ -228,24 +282,50 @@ namespace EpicGames.UHT.Tokenizer
 			}
 		}
 
+		/// <summary>
+		/// Parse a required cpp identifier
+		/// </summary>
+		/// <param name="TokenReader">Token reader</param>
+		/// <param name="Options">Parsing options</param>
+		/// <returns>Token reader</returns>
 		public static IUhtTokenReader RequireCppIdentifier(this IUhtTokenReader TokenReader, UhtCppIdentifierOptions Options = UhtCppIdentifierOptions.None)
 		{
 			TokenReader.GetCppIdentifier(Options);
 			return TokenReader;
 		}
 
+		/// <summary>
+		/// Parse a required cpp identifier
+		/// </summary>
+		/// <param name="TokenReader">Token reader</param>
+		/// <param name="InitialIdentifier">Initial token of the identifier</param>
+		/// <param name="Options">Parsing options</param>
+		/// <returns>Token reader</returns>
 		public static IUhtTokenReader RequireCppIdentifier(this IUhtTokenReader TokenReader, ref UhtToken InitialIdentifier, UhtCppIdentifierOptions Options = UhtCppIdentifierOptions.None)
 		{
 			TokenReader.GetCppIdentifier(ref InitialIdentifier, Options);
 			return TokenReader;
 		}
 
+		/// <summary>
+		/// Parse a required cpp identifier
+		/// </summary>
+		/// <param name="TokenReader">Token reader</param>
+		/// <param name="TokenListDelegate">Invoked when identifier is parsed</param>
+		/// <returns>Token reader</returns>
 		public static IUhtTokenReader RequireCppIdentifier(this IUhtTokenReader TokenReader, UhtTokenListDelegate TokenListDelegate)
 		{
 			TokenReader.RequireCppIdentifier(UhtCppIdentifierOptions.None, TokenListDelegate);
 			return TokenReader;
 		}
 
+		/// <summary>
+		/// Parse a required cpp identifier
+		/// </summary>
+		/// <param name="TokenReader">Token reader</param>
+		/// <param name="Options">Parsing options</param>
+		/// <param name="TokenListDelegate">Invoked when identifier is parsed</param>
+		/// <returns>Token reader</returns>
 		public static IUhtTokenReader RequireCppIdentifier(this IUhtTokenReader TokenReader, UhtCppIdentifierOptions Options, UhtTokenListDelegate TokenListDelegate)
 		{
 			UhtTokenList TokenList = TokenReader.GetCppIdentifier(Options);
@@ -254,12 +334,27 @@ namespace EpicGames.UHT.Tokenizer
 			return TokenReader;
 		}
 
+		/// <summary>
+		/// Parse a required cpp identifier
+		/// </summary>
+		/// <param name="TokenReader">Token reader</param>
+		/// <param name="InitialIdentifier">Initial token of the identifier</param>
+		/// <param name="TokenListDelegate">Invoked when identifier is parsed</param>
+		/// <returns>Token reader</returns>
 		public static IUhtTokenReader RequireCppIdentifier(this IUhtTokenReader TokenReader, ref UhtToken InitialIdentifier, UhtTokenListDelegate TokenListDelegate)
 		{
 			TokenReader.RequireCppIdentifier(ref InitialIdentifier, UhtCppIdentifierOptions.None, TokenListDelegate);
 			return TokenReader;
 		}
 
+		/// <summary>
+		/// Parse a required cpp identifier
+		/// </summary>
+		/// <param name="TokenReader">Token reader</param>
+		/// <param name="InitialIdentifier">Initial token of the identifier</param>
+		/// <param name="Options">Parsing options</param>
+		/// <param name="TokenListDelegate">Invoked when identifier is parsed</param>
+		/// <returns>Token reader</returns>
 		public static IUhtTokenReader RequireCppIdentifier(this IUhtTokenReader TokenReader, ref UhtToken InitialIdentifier, UhtCppIdentifierOptions Options, UhtTokenListDelegate TokenListDelegate)
 		{
 			UhtTokenList TokenList = TokenReader.GetCppIdentifier(ref InitialIdentifier, Options);
@@ -268,24 +363,35 @@ namespace EpicGames.UHT.Tokenizer
 			return TokenReader;
 		}
 
+		/// <summary>
+		/// Get a required cpp identifier
+		/// </summary>
+		/// <param name="TokenReader">Token reader</param>
+		/// <param name="Options">Parsing options</param>
+		/// <returns>Token list</returns>
 		public static UhtTokenList GetCppIdentifier(this IUhtTokenReader TokenReader, UhtCppIdentifierOptions Options = UhtCppIdentifierOptions.None)
 		{
 			UhtToken Token = TokenReader.GetIdentifier();
 			return TokenReader.GetCppIdentifier(ref Token, Options);
 		}
 
+		/// <summary>
+		/// Get a required cpp identifier
+		/// </summary>
+		/// <param name="TokenReader">Token reader</param>
+		/// <param name="InitialIdentifier">Initial token of the identifier</param>
+		/// <param name="Options">Parsing options</param>
+		/// <returns>Token list</returns>
 		public static UhtTokenList GetCppIdentifier(this IUhtTokenReader TokenReader, ref UhtToken InitialIdentifier, UhtCppIdentifierOptions Options = UhtCppIdentifierOptions.None)
 		{
 			UhtTokenList ListHead = UhtTokenListCache.Borrow(InitialIdentifier);
 			UhtTokenList ListTail = ListHead;
 
-			if (Options.HasAnyFlags(UhtCppIdentifierOptions.AllowTemplates | UhtCppIdentifierOptions.AllTokens))
+			if (Options.HasAnyFlags(UhtCppIdentifierOptions.AllowTemplates))
 			{
-				bool bAllowTemplates = Options.HasAnyFlags(UhtCppIdentifierOptions.AllowTemplates);
-
 				while (true)
 				{
-					if (bAllowTemplates && TokenReader.TryPeekOptional('<'))
+					if (TokenReader.TryPeekOptional('<'))
 					{
 						ListTail.Next = UhtTokenListCache.Borrow(TokenReader.GetToken());
 						ListTail = ListTail.Next;
diff --git a/Engine/Source/Programs/Shared/EpicGames.UHT/Tokenizer/UhtTokenReaderIntegerExtensions.cs b/Engine/Source/Programs/Shared/EpicGames.UHT/Tokenizer/UhtTokenReaderIntegerExtensions.cs
--- a/Engine/Source/Programs/Shared/EpicGames.UHT/Tokenizer/UhtTokenReaderIntegerExtensions.cs
+++ b/Engine/Source/Programs/Shared/EpicGames.UHT/Tokenizer/UhtTokenReaderIntegerExtensions.cs
@@ -5,13 +5,23 @@ using EpicGames.UHT.Utils;
 
 namespace EpicGames.UHT.Tokenizer
 {
+
+	/// <summary>
+	/// Notification of signed token being parsed
+	/// </summary>
+	/// <param name="Token">Token in question</param>
+	/// <returns>True if the token value is acceptable</returns>
 	public delegate bool UhtParseMergedSignToken(ref UhtToken Token);
 
+	/// <summary>
+	/// Collection of helper methods to parse integers
+	/// </summary>
 	public static class UhtTokenReaderIntegerExtensions
 	{
 		/// <summary>
 		/// Get the next token as an integer.  If the next token is not an integer, no token is consumed.
 		/// </summary>
+		/// <param name="TokenReader">Token reader</param>
 		/// <param name="Value">The integer value of the token</param>
 		/// <returns>True if the next token was an integer, false if not.</returns>
 		public static bool TryOptionalConstInt(this IUhtTokenReader TokenReader, out int Value)
@@ -29,6 +39,7 @@ namespace EpicGames.UHT.Tokenizer
 		/// <summary>
 		/// Get the next token as an integer.  If the next token is not an integer, an exception is thrown
 		/// </summary>
+		/// <param name="TokenReader">Token reader</param>
 		/// <param name="ExceptionContext">If not null, an exception will be thrown with the given text as part of the message.</param>
 		/// <returns>The value of the constant</returns>
 		public static IUhtTokenReader OptionalConstInt(this IUhtTokenReader TokenReader, object? ExceptionContext = null)
@@ -41,6 +52,7 @@ namespace EpicGames.UHT.Tokenizer
 		/// <summary>
 		/// Get the next token as an integer.  If the next token is not an integer, an exception is thrown
 		/// </summary>
+		/// <param name="TokenReader">Token reader</param>
 		/// <param name="ExceptionContext">If not null, an exception will be thrown with the given text as part of the message.</param>
 		/// <returns>The value of the constant</returns>
 		public static IUhtTokenReader RequireConstInt(this IUhtTokenReader TokenReader, object? ExceptionContext = null)
@@ -56,6 +68,7 @@ namespace EpicGames.UHT.Tokenizer
 		/// <summary>
 		/// Get the next token as an integer.  If the next token is not an integer, an exception is thrown
 		/// </summary>
+		/// <param name="TokenReader">Token reader</param>
 		/// <param name="ExceptionContext">If not null, an exception will be thrown with the given text as part of the message.</param>
 		/// <returns>The value of the constant</returns>
 		public static int GetConstInt(this IUhtTokenReader TokenReader, object? ExceptionContext = null)
@@ -71,6 +84,7 @@ namespace EpicGames.UHT.Tokenizer
 		/// <summary>
 		/// Get the next token as an integer.  If the next token is not an integer, no token is consumed.
 		/// </summary>
+		/// <param name="TokenReader">Token reader</param>
 		/// <param name="Value">The integer value of the token</param>
 		/// <returns>True if the next token was an integer, false if not.</returns>
 		public static bool TryOptionalConstLong(this IUhtTokenReader TokenReader, out long Value)
@@ -88,6 +102,7 @@ namespace EpicGames.UHT.Tokenizer
 		/// <summary>
 		/// Get the next token as an integer.  If the next token is not an integer, an exception is thrown
 		/// </summary>
+		/// <param name="TokenReader">Token reader</param>
 		/// <param name="ExceptionContext">If not null, an exception will be thrown with the given text as part of the message.</param>
 		/// <returns>The value of the constant</returns>
 		public static IUhtTokenReader OptionalConstLong(this IUhtTokenReader TokenReader, object? ExceptionContext = null)
@@ -105,6 +120,7 @@ namespace EpicGames.UHT.Tokenizer
 		/// <summary>
 		/// Get the next token as an integer.  If the next token is not an integer, an exception is thrown
 		/// </summary>
+		/// <param name="TokenReader">Token reader</param>
 		/// <param name="ExceptionContext">If not null, an exception will be thrown with the given text as part of the message.</param>
 		/// <returns>The value of the constant</returns>
 		public static IUhtTokenReader RequireConstLong(this IUhtTokenReader TokenReader, object? ExceptionContext = null)
@@ -120,6 +136,7 @@ namespace EpicGames.UHT.Tokenizer
 		/// <summary>
 		/// Get the next token as an integer.  If the next token is not an integer, an exception is thrown
 		/// </summary>
+		/// <param name="TokenReader">Token reader</param>
 		/// <param name="ExceptionContext">If not null, an exception will be thrown with the given text as part of the message.</param>
 		/// <returns>The value of the constant</returns>
 		public static long GetConstLong(this IUhtTokenReader TokenReader, object? ExceptionContext = null)
diff --git a/Engine/Source/Programs/Shared/EpicGames.UHT/Tokenizer/UhtTokenReaderSkipExtensions.cs b/Engine/Source/Programs/Shared/EpicGames.UHT/Tokenizer/UhtTokenReaderSkipExtensions.cs
--- a/Engine/Source/Programs/Shared/EpicGames.UHT/Tokenizer/UhtTokenReaderSkipExtensions.cs
+++ b/Engine/Source/Programs/Shared/EpicGames.UHT/Tokenizer/UhtTokenReaderSkipExtensions.cs
@@ -9,6 +9,12 @@ namespace EpicGames.UHT.Tokenizer
 	/// </summary>
 	public static class UhtTokenReaderSkipExtensions
 	{
+
+		/// <summary>
+		/// Skip a token regardless of the type.
+		/// </summary>
+		/// <param name="TokenReader">Token reader</param>
+		/// <returns>Token reader</returns>
 		public static IUhtTokenReader SkipOne(this IUhtTokenReader TokenReader)
 		{
 			TokenReader.PeekToken();
@@ -16,6 +22,11 @@ namespace EpicGames.UHT.Tokenizer
 			return TokenReader;
 		}
 
+		/// <summary>
+		/// Skip an alignas expression
+		/// </summary>
+		/// <param name="TokenReader">Token reader</param>
+		/// <returns>Token reader</returns>
 		public static IUhtTokenReader SkipAlignasIfNecessary(this IUhtTokenReader TokenReader)
 		{
 			const string Identifier = "alignas";
@@ -29,6 +40,11 @@ namespace EpicGames.UHT.Tokenizer
 			return TokenReader;
 		}
 
+		/// <summary>
+		/// Skip deprecation macro
+		/// </summary>
+		/// <param name="TokenReader">Token reader</param>
+		/// <returns>Token reader</returns>
 		public static IUhtTokenReader SkipDeprecatedMacroIfNecessary(this IUhtTokenReader TokenReader)
 		{
 			if (TokenReader.TryOptional(new string[] { "DEPRECATED", "UE_DEPRECATED" }) >= 0)
@@ -43,6 +59,11 @@ namespace EpicGames.UHT.Tokenizer
 			return TokenReader;
 		}
 
+		/// <summary>
+		/// Skip alignas and/or deprecation macros
+		/// </summary>
+		/// <param name="TokenReader">Token reader</param>
+		/// <returns>Token reader</returns>
 		public static IUhtTokenReader SkipAlignasAndDeprecatedMacroIfNecessary(this IUhtTokenReader TokenReader)
 		{
 			// alignas() can come before or after the deprecation macro.
@@ -53,6 +74,16 @@ namespace EpicGames.UHT.Tokenizer
 				.SkipAlignasIfNecessary();
 		}
 
+		/// <summary>
+		/// Skip any block of tokens wrapped by the given token symbols
+		/// </summary>
+		/// <param name="TokenReader">Token reader</param>
+		/// <param name="Initiator">Initiating token (i.e. &quot;(&quot;)</param>
+		/// <param name="Terminator">Terminating token (i.e. &quot;)&quot;)</param>
+		/// <param name="InitialNesting">If true, start with an initial nesting count of one (assume we already parsed an initiator)</param>
+		/// <param name="ExceptionContext">Extra context for any errors</param>
+		/// <returns>Token reader</returns>
+		/// <exception cref="UhtTokenException">Throw if end of file is reached</exception>
 		public static IUhtTokenReader SkipBrackets(this IUhtTokenReader TokenReader, char Initiator, char Terminator, int InitialNesting, object? ExceptionContext = null)
 		{
 			int Nesting = InitialNesting;
diff --git a/Engine/Source/Programs/Shared/EpicGames.UHT/Tokenizer/UhtTokenReaderStringExtensions.cs b/Engine/Source/Programs/Shared/EpicGames.UHT/Tokenizer/UhtTokenReaderStringExtensions.cs
--- a/Engine/Source/Programs/Shared/EpicGames.UHT/Tokenizer/UhtTokenReaderStringExtensions.cs
+++ b/Engine/Source/Programs/Shared/EpicGames.UHT/Tokenizer/UhtTokenReaderStringExtensions.cs
@@ -5,11 +5,16 @@ using EpicGames.UHT.Utils;
 
 namespace EpicGames.UHT.Tokenizer
 {
+
+	/// <summary>
+	/// Collection of token reader exceptions for handling strings
+	/// </summary>
 	public static class UhtTokenReaderStringExtensions
 	{
 		/// <summary>
 		/// Get the next token as a string.  If the next token is not a string, no token is consumed.
 		/// </summary>
+		/// <param name="TokenReader">Token reader</param>
 		/// <param name="Value">The string value of the token</param>
 		/// <returns>True if the next token was an string, false if not.</returns>
 		public static bool TryOptionalConstString(this IUhtTokenReader TokenReader, out StringView Value)
@@ -28,6 +33,7 @@ namespace EpicGames.UHT.Tokenizer
 		/// <summary>
 		/// Get the next token as a string.  If the next token is not a string, no token is consumed.
 		/// </summary>
+		/// <param name="TokenReader">Token reader</param>
 		/// <returns>True if the next token was an string, false if not.</returns>
 		public static IUhtTokenReader OptionalConstString(this IUhtTokenReader TokenReader)
 		{
@@ -42,6 +48,7 @@ namespace EpicGames.UHT.Tokenizer
 		/// <summary>
 		/// Verify that the next token is a string.
 		/// </summary>
+		/// <param name="TokenReader">Token reader</param>
 		/// <param name="ExceptionContext">If not null, an exception will be thrown with the given text as part of the message.</param>
 		/// <returns>True if the next token was a string, false if not.</returns>
 		public static IUhtTokenReader RequireConstString(this IUhtTokenReader TokenReader, object? ExceptionContext)
@@ -58,6 +65,7 @@ namespace EpicGames.UHT.Tokenizer
 		/// <summary>
 		/// Get the next token as a string.  If the next token is not a string, an exception is thrown
 		/// </summary>
+		/// <param name="TokenReader">Token reader</param>
 		/// <param name="ExceptionContext">If not null, an exception will be thrown with the given text as part of the message.</param>
 		/// <returns>The value of the string.</returns>
 		public static StringView GetConstString(this IUhtTokenReader TokenReader, object? ExceptionContext = null)
@@ -76,6 +84,7 @@ namespace EpicGames.UHT.Tokenizer
 		/// Get the next token as a quoted string.  If the next token is not a string, an exception is thrown.
 		/// Character constants are not considered strings by this routine.
 		/// </summary>
+		/// <param name="TokenReader">Token reader</param>
 		/// <param name="ExceptionContext">If not null, an exception will be thrown with the given text as part of the message.</param>
 		/// <returns>The value of the string.</returns>
 		public static StringView GetConstQuotedString(this IUhtTokenReader TokenReader, object? ExceptionContext = null)
diff --git a/Engine/Source/Programs/Shared/EpicGames.UHT/Tokenizer/UhtTokenReaderUtilityExtensions.cs b/Engine/Source/Programs/Shared/EpicGames.UHT/Tokenizer/UhtTokenReaderUtilityExtensions.cs
--- a/Engine/Source/Programs/Shared/EpicGames.UHT/Tokenizer/UhtTokenReaderUtilityExtensions.cs
+++ b/Engine/Source/Programs/Shared/EpicGames.UHT/Tokenizer/UhtTokenReaderUtilityExtensions.cs
@@ -7,6 +7,10 @@ using System.Collections.Generic;
 
 namespace EpicGames.UHT.Tokenizer
 {
+
+	/// <summary>
+	/// Collection of assorted utility token reader extensions
+	/// </summary>
 	public static class UhtTokenReaderUtilityExtensions
 	{
 		private static HashSet<StringView> SkipDeclarationWarningStrings = new HashSet<StringView>
@@ -31,7 +35,8 @@ namespace EpicGames.UHT.Tokenizer
 		/// When processing type, make sure that the next token is the expected token
 		/// </summary>
 		/// <param name="TokenReader">Token reader</param>
-		/// <param name="Token">Token that started the process</param>
+		/// <param name="ExpectedIdentifier">Expected identifier</param>
+		/// <param name="bIsMember">If true, log an error if the type begins with const</param>
 		/// <returns>true if there could be more header to process, false if the end was reached.</returns>
 		public static bool SkipExpectedType(this IUhtTokenReader TokenReader, StringView ExpectedIdentifier, bool bIsMember)
 		{
@@ -50,6 +55,12 @@ namespace EpicGames.UHT.Tokenizer
 			return false;
 		}
 
+		/// <summary>
+		/// Try to parse an optional _API macro
+		/// </summary>
+		/// <param name="TokenReader">Token reader</param>
+		/// <param name="APIMacroToken">_API macro parsed</param>
+		/// <returns>True if an _API macro was parsed</returns>
 		public static bool TryOptionalAPIMacro(this IUhtTokenReader TokenReader, out UhtToken APIMacroToken)
 		{
 			ref UhtToken Token = ref TokenReader.PeekToken();
@@ -63,6 +74,12 @@ namespace EpicGames.UHT.Tokenizer
 			return false;
 		}
 
+		/// <summary>
+		/// Parse an optional single inheritance
+		/// </summary>
+		/// <param name="TokenReader">Token reader</param>
+		/// <param name="SuperClassDelegate">Invoked with the inherited type name</param>
+		/// <returns>Token reader</returns>
 		public static IUhtTokenReader OptionalInheritance(this IUhtTokenReader TokenReader, UhtTokenDelegate SuperClassDelegate)
 		{
 			TokenReader.Optional(':', () =>
@@ -74,6 +91,13 @@ namespace EpicGames.UHT.Tokenizer
 			return TokenReader;
 		}
 
+		/// <summary>
+		/// Parse an optional inheritance
+		/// </summary>
+		/// <param name="TokenReader">Token reader</param>
+		/// <param name="SuperClassDelegate">Invoked with the inherited type name</param>
+		/// <param name="BaseClassDelegate">Invoked when other base classes are parsed</param>
+		/// <returns>Token reader</returns>
 		public static IUhtTokenReader OptionalInheritance(this IUhtTokenReader TokenReader, UhtTokenDelegate SuperClassDelegate, UhtTokenListDelegate BaseClassDelegate)
 		{
 			TokenReader.Optional(':', () =>
diff --git a/Engine/Source/Programs/Shared/EpicGames.UHT/Tokenizer/UhtTokenReplayReader.cs b/Engine/Source/Programs/Shared/EpicGames.UHT/Tokenizer/UhtTokenReplayReader.cs
--- a/Engine/Source/Programs/Shared/EpicGames.UHT/Tokenizer/UhtTokenReplayReader.cs
+++ b/Engine/Source/Programs/Shared/EpicGames.UHT/Tokenizer/UhtTokenReplayReader.cs
@@ -8,6 +8,10 @@ using System.Threading;
 
 namespace EpicGames.UHT.Tokenizer
 {
+
+	/// <summary>
+	/// Token reader to replay previously recorded token stream
+	/// </summary>
 	public class UhtTokenReplayReader : IUhtTokenReader, IUhtMessageLineNumber
 	{
 		static private ThreadLocal<UhtTokenReplayReader> Tls = new ThreadLocal<UhtTokenReplayReader>(() => new UhtTokenReplayReader());
@@ -29,6 +33,13 @@ namespace EpicGames.UHT.Tokenizer
 		private int SavedStateCount = 0;
 		private UhtTokenType EndTokenType = UhtTokenType.EndOfFile;
 
+		/// <summary>
+		/// Construct new token reader
+		/// </summary>
+		/// <param name="MessageSite">Message site for generating errors</param>
+		/// <param name="Data">Complete data for the token (i.e. original source)</param>
+		/// <param name="Tokens">Tokens to replay</param>
+		/// <param name="EndTokenType">Token type to return when end of tokens reached</param>
 		public UhtTokenReplayReader(IUhtMessageSite MessageSite, ReadOnlyMemory<char> Data, ReadOnlyMemory<UhtToken> Tokens, UhtTokenType EndTokenType)
 		{
 			this.MessageSiteInternal = MessageSite;
@@ -38,6 +49,9 @@ namespace EpicGames.UHT.Tokenizer
 			this.CurrentToken = new UhtToken(EndTokenType);
 		}
 
+		/// <summary>
+		/// Construct token replay reader intended for caching.  Use Reset method to prepare it for use
+		/// </summary>
 		public UhtTokenReplayReader()
 		{
 			this.MessageSiteInternal = new UhtEmptyMessageSite();
@@ -45,6 +59,14 @@ namespace EpicGames.UHT.Tokenizer
 			this.Data = new char[0].AsMemory();
 		}
 
+		/// <summary>
+		/// Reset a cached replay reader for replaying a new stream of tokens
+		/// </summary>
+		/// <param name="MessageSite">Message site for generating errors</param>
+		/// <param name="Data">Complete data for the token (i.e. original source)</param>
+		/// <param name="Tokens">Tokens to replay</param>
+		/// <param name="EndTokenType">Token type to return when end of tokens reached</param>
+		/// <returns>The replay reader</returns>
 		public UhtTokenReplayReader Reset(IUhtMessageSite MessageSite, ReadOnlyMemory<char> Data, ReadOnlyMemory<UhtToken> Tokens, UhtTokenType EndTokenType)
 		{
 			this.MessageSiteInternal = MessageSite;
@@ -85,6 +107,7 @@ namespace EpicGames.UHT.Tokenizer
 		#endregion
 
 		#region ITokenReader Implementation
+		/// <inheritdoc/>
 		public bool bIsEOF
 		{
 			get
@@ -97,6 +120,7 @@ namespace EpicGames.UHT.Tokenizer
 			}
 		}
 
+		/// <inheritdoc/>
 		public int InputPos 
 		{
 			get
@@ -109,6 +133,7 @@ namespace EpicGames.UHT.Tokenizer
 			}
 		}
 
+		/// <inheritdoc/>
 		public int InputLine 
 		{
 			get
@@ -122,52 +147,64 @@ namespace EpicGames.UHT.Tokenizer
 			set => throw new NotImplementedException(); 
 		}
 
+		/// <inheritdoc/>
 		public IUhtTokenPreprocessor? TokenPreprocessor { get => throw new NotImplementedException(); set => throw new NotImplementedException(); }
 
+		/// <inheritdoc/>
 		public int LookAheadEnableCount { get => throw new NotImplementedException(); set => throw new NotImplementedException(); }
 
+		/// <inheritdoc/>
 		public ReadOnlySpan<StringView> Comments => throw new NotImplementedException();
 
+		/// <inheritdoc/>
 		public void ClearComments()
 		{
 			throw new NotImplementedException();
 		}
 
+		/// <inheritdoc/>
 		public void ConsumeToken()
 		{
 			this.bHasToken = false;
 		}
 
+		/// <inheritdoc/>
 		public void DisableComments()
 		{
 			throw new NotImplementedException();
 		}
 
+		/// <inheritdoc/>
 		public void EnableComments()
 		{
 			throw new NotImplementedException();
 		}
 
+		/// <inheritdoc/>
 		public void CommitPendingComments()
 		{
 			throw new NotImplementedException();
 		}
 
+		/// <inheritdoc/>
 		public UhtToken GetLine()
 		{
 			throw new NotImplementedException();
 		}
 
+		/// <inheritdoc/>
 		public StringView GetRawString(char Terminator, UhtRawStringOptions Options)
 		{
 			throw new NotImplementedException();
 		}
 
+		/// <inheritdoc/>
 		public StringView GetStringView(int StartPos, int Count)
 		{
 			return new StringView(this.Data, StartPos, Count);
 		}
 
+		/// <inheritdoc/>
 		public UhtToken GetToken()
 		{
 			UhtToken Token = PeekToken();
@@ -175,11 +212,13 @@ namespace EpicGames.UHT.Tokenizer
 			return Token;
 		}
 
+		/// <inheritdoc/>
 		public bool IsFirstTokenInLine(ref UhtToken Token)
 		{
 			throw new NotImplementedException();
 		}
 
+		/// <inheritdoc/>
 		public ref UhtToken PeekToken()
 		{
 			if (!this.bHasToken)
@@ -189,11 +228,13 @@ namespace EpicGames.UHT.Tokenizer
 			return ref this.CurrentToken;
 		}
 
+		/// <inheritdoc/>
 		public void SkipWhitespaceAndComments()
 		{
 			throw new NotImplementedException();
 		}
 
+		/// <inheritdoc/>
 		public void SaveState()
 		{
 			if (this.SavedStateCount == MaxSavedStates)
@@ -204,6 +245,7 @@ namespace EpicGames.UHT.Tokenizer
 			++this.SavedStateCount;
 		}
 
+		/// <inheritdoc/>
 		public void RestoreState()
 		{
 			if (this.SavedStateCount == 0)
@@ -216,6 +258,7 @@ namespace EpicGames.UHT.Tokenizer
 			this.bHasToken = this.SavedStates[this.SavedStateCount].bHasToken;
 		}
 
+		/// <inheritdoc/>
 		public void AbandonState()
 		{
 			if (this.SavedStateCount == 0)
