More UHT documentation updates

#rnx

diff --git a/Engine/Source/Programs/Shared/EpicGames.UHT/Specifiers/UhtClassSpecifiers.cs b/Engine/Source/Programs/Shared/EpicGames.UHT/Specifiers/UhtClassSpecifiers.cs
--- a/Engine/Source/Programs/Shared/EpicGames.UHT/Specifiers/UhtClassSpecifiers.cs
+++ b/Engine/Source/Programs/Shared/EpicGames.UHT/Specifiers/UhtClassSpecifiers.cs
@@ -8,6 +8,9 @@ using System.Collections.Generic;
 
 namespace EpicGames.UHT.Parsers
 {
+	/// <summary>
+	/// Collection of UCLASS specifiers
+	/// </summary>
 	[UnrealHeaderTool]
 	public static class UhtClassSpecifiers
 	{
diff --git a/Engine/Source/Programs/Shared/EpicGames.UHT/Specifiers/UhtDefaultSpecifiers.cs b/Engine/Source/Programs/Shared/EpicGames.UHT/Specifiers/UhtDefaultSpecifiers.cs
--- a/Engine/Source/Programs/Shared/EpicGames.UHT/Specifiers/UhtDefaultSpecifiers.cs
+++ b/Engine/Source/Programs/Shared/EpicGames.UHT/Specifiers/UhtDefaultSpecifiers.cs
@@ -9,6 +9,10 @@ using System.Collections.Generic;
 
 namespace EpicGames.UHT.Parsers
 {
+
+	/// <summary>
+	/// Collection of default specifiers that apply to everything
+	/// </summary>
 	[UnrealHeaderTool]
 	public static class UhtDefaultSpecifiers
 	{
diff --git a/Engine/Source/Programs/Shared/EpicGames.UHT/Specifiers/UhtEnumSpecifiers.cs b/Engine/Source/Programs/Shared/EpicGames.UHT/Specifiers/UhtEnumSpecifiers.cs
--- a/Engine/Source/Programs/Shared/EpicGames.UHT/Specifiers/UhtEnumSpecifiers.cs
+++ b/Engine/Source/Programs/Shared/EpicGames.UHT/Specifiers/UhtEnumSpecifiers.cs
@@ -7,6 +7,10 @@ using EpicGames.UHT.Utils;
 
 namespace EpicGames.UHT.Parsers
 {
+
+	/// <summary>
+	/// Collection of UENUM specifiers
+	/// </summary>
 	[UnrealHeaderTool]
 	public static class UhtEnumSpecifiers
 	{
diff --git a/Engine/Source/Programs/Shared/EpicGames.UHT/Specifiers/UhtFunctionSpecifiers.cs b/Engine/Source/Programs/Shared/EpicGames.UHT/Specifiers/UhtFunctionSpecifiers.cs
--- a/Engine/Source/Programs/Shared/EpicGames.UHT/Specifiers/UhtFunctionSpecifiers.cs
+++ b/Engine/Source/Programs/Shared/EpicGames.UHT/Specifiers/UhtFunctionSpecifiers.cs
@@ -10,6 +10,9 @@ using System.Collections.Generic;
 
 namespace EpicGames.UHT.Parsers
 {
+	/// <summary>
+	/// Collection of UFUNCTION/UDELEGATE specifiers
+	/// </summary>
 	[UnrealHeaderTool]
 	public static class UhtFunctionSpecifiers
 	{
diff --git a/Engine/Source/Programs/Shared/EpicGames.UHT/Specifiers/UhtPropertyArgumentSpecifiers.cs b/Engine/Source/Programs/Shared/EpicGames.UHT/Specifiers/UhtPropertyArgumentSpecifiers.cs
--- a/Engine/Source/Programs/Shared/EpicGames.UHT/Specifiers/UhtPropertyArgumentSpecifiers.cs
+++ b/Engine/Source/Programs/Shared/EpicGames.UHT/Specifiers/UhtPropertyArgumentSpecifiers.cs
@@ -7,10 +7,22 @@ using EpicGames.UHT.Utils;
 
 namespace EpicGames.UHT.Parsers
 {
+
+	/// <summary>
+	/// Collection of property argument specifiers
+	/// </summary>
 	[UnrealHeaderTool]
 	public static class UhtPropertyArgumentSpecifiers
 	{
+
+		/// <summary>
+		/// Table for property argument specifiers
+		/// </summary>
 		public static readonly UhtSpecifierTable SpecifierTable = UhtSpecifierTables.Instance.Get(UhtTableNames.PropertyArgument);
+
+		/// <summary>
+		/// Table for property argument specifier validators
+		/// </summary>
 		public static readonly UhtSpecifierValidatorTable SpecifierValidatorTable = UhtSpecifierValidatorTables.Instance.Get(UhtTableNames.PropertyArgument);
 
 		#region Argument Property Specifiers
diff --git a/Engine/Source/Programs/Shared/EpicGames.UHT/Specifiers/UhtPropertyMemberSpecifiers.cs b/Engine/Source/Programs/Shared/EpicGames.UHT/Specifiers/UhtPropertyMemberSpecifiers.cs
--- a/Engine/Source/Programs/Shared/EpicGames.UHT/Specifiers/UhtPropertyMemberSpecifiers.cs
+++ b/Engine/Source/Programs/Shared/EpicGames.UHT/Specifiers/UhtPropertyMemberSpecifiers.cs
@@ -8,10 +8,22 @@ using System;
 
 namespace EpicGames.UHT.Parsers
 {
+
+	/// <summary>
+	/// Collection of property member specifiers
+	/// </summary>
 	[UnrealHeaderTool]
 	public static class UhtPropertyMemberSpecifiers
 	{
+
+		/// <summary>
+		/// Property member specifier table
+		/// </summary>
 		public static UhtSpecifierTable SpecifierTable = UhtSpecifierTables.Instance.Get(UhtTableNames.PropertyMember);
+
+		/// <summary>
+		/// Property member specifier validator table
+		/// </summary>
 		public static UhtSpecifierValidatorTable SpecifierValidatorTable = UhtSpecifierValidatorTables.Instance.Get(UhtTableNames.PropertyMember);
 
 		[UhtSpecifier(Extends = UhtTableNames.PropertyMember, ValueType = UhtSpecifierValueType.Legacy)]
diff --git a/Engine/Source/Programs/Shared/EpicGames.UHT/Specifiers/UhtScriptStructSpecifiers.cs b/Engine/Source/Programs/Shared/EpicGames.UHT/Specifiers/UhtScriptStructSpecifiers.cs
--- a/Engine/Source/Programs/Shared/EpicGames.UHT/Specifiers/UhtScriptStructSpecifiers.cs
+++ b/Engine/Source/Programs/Shared/EpicGames.UHT/Specifiers/UhtScriptStructSpecifiers.cs
@@ -7,6 +7,9 @@ using EpicGames.UHT.Utils;
 
 namespace EpicGames.UHT.Parsers
 {
+	/// <summary>
+	/// Collection of USCRIPT specifiers
+	/// </summary>
 	[UnrealHeaderTool]
 	public static class UhtScriptStructSpecifiers
 	{
diff --git a/Engine/Source/Programs/Shared/EpicGames.UHT/Tables/UhtExporterTable.cs b/Engine/Source/Programs/Shared/EpicGames.UHT/Tables/UhtExporterTable.cs
--- a/Engine/Source/Programs/Shared/EpicGames.UHT/Tables/UhtExporterTable.cs
+++ b/Engine/Source/Programs/Shared/EpicGames.UHT/Tables/UhtExporterTable.cs
@@ -8,11 +8,23 @@ using System.Reflection;
 
 namespace EpicGames.UHT.Tables
 {
+
+	/// <summary>
+	/// Delegate to invoke to run exporter
+	/// </summary>
+	/// <param name="Factory">Factory used to generate export tasks and outputs</param>
 	public delegate void UhtExporterDelegate(IUhtExportFactory Factory);
 
+	/// <summary>
+	/// Export options
+	/// </summary>
 	[Flags]
 	public enum UhtExporterOptions
 	{
+
+		/// <summary>
+		/// No options
+		/// </summary>
 		None = 0,
 
 		/// <summary>
@@ -63,30 +75,95 @@ namespace EpicGames.UHT.Tables
 		}
 	}
 
+	/// <summary>
+	/// Defines an exporter
+	/// </summary>
 	[AttributeUsage(AttributeTargets.Method, AllowMultiple = true)]
 	public class UhtExporterAttribute : Attribute
 	{
+
+		/// <summary>
+		/// Name of the exporter
+		/// </summary>
 		public string Name = string.Empty;
+
+		/// <summary>
+		/// Description of the export.  Used to display help
+		/// </summary>
 		public string Description = string.Empty;
+
+		/// <summary>
+		/// Exporter options
+		/// </summary>
 		public UhtExporterOptions Options = UhtExporterOptions.None;
+
+		/// <summary>
+		/// Collection of filters used to delete old cpp files
+		/// </summary>
 		public string[]? CppFilters = null;
+
+		/// <summary>
+		/// Collection of filters used to delete old h files
+		/// </summary>
 		public string[]? HeaderFilters = null;
+
+		/// <summary>
+		/// Collection of filters for other file types
+		/// </summary>
 		public string[]? OtherFilters = null;
 	}
 
+	/// <summary>
+	/// Defines an exporter in the table
+	/// </summary>
 	public struct UhtExporter
 	{
+
+		/// <summary>
+		/// Name of the exporter
+		/// </summary>
 		public string Name;
+
+		/// <summary>
+		/// Description of the export.  Used to display help
+		/// </summary>
 		public string Description;
+
+		/// <summary>
+		/// Exporter options
+		/// </summary>
 		public UhtExporterOptions Options;
+
+		/// <summary>
+		/// Delegate to invoke to start export
+		/// </summary>
 		public UhtExporterDelegate Delegate;
+
+		/// <summary>
+		/// Collection of filters used to delete old cpp files
+		/// </summary>
 		public string[]? CppFilters;
+
+		/// <summary>
+		/// Collection of filters used to delete old h files
+		/// </summary>
 		public string[]? HeaderFilters;
+
+		/// <summary>
+		/// Collection of filters for other file types
+		/// </summary>
 		public string[]? OtherFilters;
 	}
 
+	/// <summary>
+	/// Exporter table
+	/// </summary>
 	public class UhtExporterTable : IEnumerable<UhtExporter>
 	{
+
+		/// <summary>
+		/// Global instance of the table
+		/// </summary>
 		public static UhtExporterTable Instance = new UhtExporterTable();
 
 		private Dictionary<string, UhtExporter> ExporterValues = new Dictionary<string, UhtExporter>(StringComparer.OrdinalIgnoreCase);
@@ -102,6 +179,13 @@ namespace EpicGames.UHT.Tables
 			return this.ExporterValues.TryGetValue(Name, out Value);
 		}
 
+		/// <summary>
+		/// Handle an exporter attribute
+		/// </summary>
+		/// <param name="Type">Containing type</param>
+		/// <param name="MethodInfo">Method info</param>
+		/// <param name="ExporterAttribute">Defining attribute</param>
+		/// <exception cref="UhtIceException">Thrown if the attribute doesn't properly define an exporter.</exception>
 		public void OnExporterAttribute(Type Type, MethodInfo MethodInfo, UhtExporterAttribute ExporterAttribute)
 		{
 			if (string.IsNullOrEmpty(ExporterAttribute.Name))
@@ -123,6 +207,10 @@ namespace EpicGames.UHT.Tables
 			this.ExporterValues.Add(ExporterAttribute.Name, ExporterValue);
 		}
 
+		/// <summary>
+		/// Return an enumerator for all the defined exporters
+		/// </summary>
+		/// <returns>Enumerator</returns>
 		public IEnumerator<UhtExporter> GetEnumerator()
 		{
 			foreach (KeyValuePair<string, UhtExporter> KVP in this.ExporterValues)
@@ -131,6 +219,10 @@ namespace EpicGames.UHT.Tables
 			}
 		}
 
+		/// <summary>
+		/// Return an enumerator for all the defined exporters
+		/// </summary>
+		/// <returns>Enumerator</returns>
 		IEnumerator IEnumerable.GetEnumerator()
 		{
 			foreach (KeyValuePair<string, UhtExporter> KVP in this.ExporterValues)
diff --git a/Engine/Source/Programs/Shared/EpicGames.UHT/Tables/UhtKeywordTable.cs b/Engine/Source/Programs/Shared/EpicGames.UHT/Tables/UhtKeywordTable.cs
--- a/Engine/Source/Programs/Shared/EpicGames.UHT/Tables/UhtKeywordTable.cs
+++ b/Engine/Source/Programs/Shared/EpicGames.UHT/Tables/UhtKeywordTable.cs
@@ -17,30 +17,100 @@ namespace EpicGames.UHT.Tables
 	[AttributeUsage(AttributeTargets.Method, AllowMultiple = true)]
 	public class UhtKeywordAttribute : Attribute
 	{
+
+		/// <summary>
+		/// Keyword table/scope being extended
+		/// </summary>
 		public string? Extends;
+
+		/// <summary>
+		/// Name of the keyword
+		/// </summary>
 		public string? Keyword = null;
+
+		/// <summary>
+		/// Text to be displayed to the user when referencing this keyword
+		/// </summary>
 		public string? AllowText = null;
+
+		/// <summary>
+		/// If true, this applies to all scopes
+		/// </summary>
 		public bool AllScopes = false;
+
+		/// <summary>
+		/// If true, do not include in usage errors
+		/// </summary>
 		public bool DisableUsageError = false;
 	}
 
+	/// <summary>
+	/// Invoked as a last chance processor for a keyword
+	/// </summary>
 	[AttributeUsage(AttributeTargets.Method, AllowMultiple = true)]
 	public class UhtKeywordCatchAllAttribute : Attribute
 	{
+
+		/// <summary>
+		/// Table/scope to be extended
+		/// </summary>
 		public string? Extends;
 	}
 
+	/// <summary>
+	/// Delegate to notify a keyword was parsed
+	/// </summary>
+	/// <param name="TopScope">Current scope being parsed</param>
+	/// <param name="ActionScope">The scope who's table was matched</param>
+	/// <param name="Token">Matching token</param>
+	/// <returns>Results of the parsing</returns>
 	public delegate UhtParseResult UhtKeywordDelegate(UhtParsingScope TopScope, UhtParsingScope ActionScope, ref UhtToken Token);
+
+	/// <summary>
+	/// Delegate to invoke as a last chance processor for a keyword
+	/// </summary>
+	/// <param name="TopScope">Current scope being parsed</param>
+	/// <param name="Token">Matching token</param>
+	/// <returns>Results of the parsing</returns>
 	public delegate UhtParseResult UhtKeywordCatchAllDelegate(UhtParsingScope TopScope, ref UhtToken Token);
 
+	/// <summary>
+	/// Defines a keyword
+	/// </summary>
 	public struct UhtKeyword
 	{
+
+		/// <summary>
+		/// Name of the keyword
+		/// </summary>
 		public readonly string Name;
+
+		/// <summary>
+		/// Delegate to invoke
+		/// </summary>
 		public readonly UhtKeywordDelegate Delegate;
+
+		/// <summary>
+		/// Text to be displayed to the user when referencing this keyword
+		/// </summary>
 		public readonly string? AllowText;
+
+		/// <summary>
+		/// If true, this applies to all scopes
+		/// </summary>
 		public readonly bool bAllScopes;
+
+		/// <summary>
+		/// If true, do not include in usage errors
+		/// </summary>
 		public readonly bool bDisableUsageError;
 
+		/// <summary>
+		/// Construct a new keyword
+		/// </summary>
+		/// <param name="Name">Name of the keyword</param>
+		/// <param name="Delegate">Delegate to invoke</param>
+		/// <param name="Attribute">Defining attribute</param>
 		public UhtKeyword(string Name, UhtKeywordDelegate Delegate, UhtKeywordAttribute? Attribute)
 		{
 			this.Name = Name;
@@ -60,6 +130,9 @@ namespace EpicGames.UHT.Tables
 		}
 	}
 
+	/// <summary>
+	/// Keyword table for a specific scope
+	/// </summary>
 	public class UhtKeywordTable : UhtLookupTable<UhtKeyword>
 	{
 
@@ -98,7 +171,7 @@ namespace EpicGames.UHT.Tables
 		/// <summary>
 		/// Merge the given keyword table.  Duplicates in the BaseTypeTable will be ignored.
 		/// </summary>
-		/// <param name="BaseTypeTable">Base table being merged</param>
+		/// <param name="BaseTable">Base table being merged</param>
 		public override void Merge(UhtLookupTableBase BaseTable)
 		{
 			base.Merge(BaseTable);
@@ -106,14 +179,31 @@ namespace EpicGames.UHT.Tables
 		}
 	}
 
+	/// <summary>
+	/// Table of all keyword tables
+	/// </summary>
 	public class UhtKeywordTables : UhtLookupTables<UhtKeywordTable>
 	{
+
+		/// <summary>
+		/// Global instance of the tables
+		/// </summary>
 		public static UhtKeywordTables Instance = new UhtKeywordTables();
 
+		/// <summary>
+		/// Construct the keyword tables
+		/// </summary>
 		public UhtKeywordTables() : base("keywords")
 		{
 		}
 
+		/// <summary>
+		/// Handle a keyword attribute
+		/// </summary>
+		/// <param name="Type">Containing type</param>
+		/// <param name="MethodInfo">Method information</param>
+		/// <param name="KeywordCatchAllAttribute">Defining attribute</param>
+		/// <exception cref="UhtIceException">Thrown if the attribute isn't well defined</exception>
 		public void OnKeywordCatchAllAttribute(Type Type, MethodInfo MethodInfo, UhtKeywordCatchAllAttribute KeywordCatchAllAttribute)
 		{
 			if (string.IsNullOrEmpty(KeywordCatchAllAttribute.Extends))
@@ -125,6 +215,13 @@ namespace EpicGames.UHT.Tables
 			Table.AddCatchAll((UhtKeywordCatchAllDelegate)Delegate.CreateDelegate(typeof(UhtKeywordCatchAllDelegate), MethodInfo));
 		}
 
+		/// <summary>
+		/// Handle a keyword attribute
+		/// </summary>
+		/// <param name="Type">Containing type</param>
+		/// <param name="MethodInfo">Method information</param>
+		/// <param name="KeywordAttribute">Defining attribute</param>
+		/// <exception cref="UhtIceException">Thrown if the attribute isn't well defined</exception>
 		public void OnKeywordAttribute(Type Type, MethodInfo MethodInfo, UhtKeywordAttribute KeywordAttribute)
 		{
 			string Name = UhtLookupTableBase.GetSuffixedName(Type, MethodInfo, KeywordAttribute.Keyword, "Keyword");
@@ -138,6 +235,11 @@ namespace EpicGames.UHT.Tables
 			Table.Add(new UhtKeyword(Name, (UhtKeywordDelegate)Delegate.CreateDelegate(typeof(UhtKeywordDelegate), MethodInfo), KeywordAttribute));
 		}
 
+		/// <summary>
+		/// Log an unhandled error
+		/// </summary>
+		/// <param name="MessageSite">Destination message site</param>
+		/// <param name="Token">Keyword</param>
 		public void LogUnhandledError(IUhtMessageSite MessageSite, UhtToken Token)
 		{
 			List<string>? Tables = null;
diff --git a/Engine/Source/Programs/Shared/EpicGames.UHT/Tables/UhtLocTextDefaultValueTable.cs b/Engine/Source/Programs/Shared/EpicGames.UHT/Tables/UhtLocTextDefaultValueTable.cs
--- a/Engine/Source/Programs/Shared/EpicGames.UHT/Tables/UhtLocTextDefaultValueTable.cs
+++ b/Engine/Source/Programs/Shared/EpicGames.UHT/Tables/UhtLocTextDefaultValueTable.cs
@@ -11,21 +11,50 @@ using System.Text;
 
 namespace EpicGames.UHT.Tables
 {
+
+	/// <summary>
+	/// Delegate to invoke to sanitize a loctext default value
+	/// </summary>
+	/// <param name="Property">Property in question</param>
+	/// <param name="DefaultValueReader">The default value</param>
+	/// <param name="MacroToken">Token for the loctext type being parsed</param>
+	/// <param name="InnerDefaultValue">Output sanitized value.</param>
+	/// <returns>True if sanitized, false if not.</returns>
 	public delegate bool UhtLocTextDefaultValueDelegate(UhtTextProperty Property, IUhtTokenReader DefaultValueReader, ref UhtToken MacroToken, StringBuilder InnerDefaultValue);
 
+	/// <summary>
+	/// Attribute defining the loctext sanitizer
+	/// </summary>
 	[AttributeUsage(AttributeTargets.Method, AllowMultiple = true)]
 	public class UhtLocTextDefaultValueAttribute : Attribute
 	{
+
+		/// <summary>
+		/// Name of the sanitizer (i.e. LOCTEXT, NSLOCTEXT, ...)
+		/// </summary>
 		public string? Name;
 	}
 
+	/// <summary>
+	/// Loctext sanitizer
+	/// </summary>
 	public struct UhtLocTextDefaultValue
 	{
+		/// <summary>
+		/// Delegate to invoke
+		/// </summary>
 		public UhtLocTextDefaultValueDelegate Delegate;
 	}
 
+	/// <summary>
+	/// Table of loctext sanitizers
+	/// </summary>
 	public class UhtLocTextDefaultValueTable
 	{
+
+		/// <summary>
+		/// Global instance of the table
+		/// </summary>
 		public static UhtLocTextDefaultValueTable Instance = new UhtLocTextDefaultValueTable();
 
 		private Dictionary<StringView, UhtLocTextDefaultValue> LocTextDefaultValues = new Dictionary<StringView, UhtLocTextDefaultValue>();
@@ -41,6 +70,13 @@ namespace EpicGames.UHT.Tables
 			return this.LocTextDefaultValues.TryGetValue(Name, out LocTextDefaultValue);
 		}
 
+		/// <summary>
+		/// Handle a loctext default value attribute
+		/// </summary>
+		/// <param name="Type">Containing type</param>
+		/// <param name="MethodInfo">Method info</param>
+		/// <param name="LocTextDefaultValueAttribute">Defining attribute</param>
+		/// <exception cref="UhtIceException">Thrown if the attribute isn't properly defined</exception>
 		public void OnLocTextDefaultValueAttribute(Type Type, MethodInfo MethodInfo, UhtLocTextDefaultValueAttribute LocTextDefaultValueAttribute)
 		{
 			if (string.IsNullOrEmpty(LocTextDefaultValueAttribute.Name))
diff --git a/Engine/Source/Programs/Shared/EpicGames.UHT/Tables/UhtPropertyTypeTable.cs b/Engine/Source/Programs/Shared/EpicGames.UHT/Tables/UhtPropertyTypeTable.cs
--- a/Engine/Source/Programs/Shared/EpicGames.UHT/Tables/UhtPropertyTypeTable.cs
+++ b/Engine/Source/Programs/Shared/EpicGames.UHT/Tables/UhtPropertyTypeTable.cs
@@ -10,9 +10,17 @@ using System.Reflection;
 
 namespace EpicGames.UHT.Tables
 {
+	
+	/// <summary>
+	/// Property type options
+	/// </summary>
 	[Flags]
 	public enum UhtPropertyTypeOptions
 	{
+
+		/// <summary>
+		/// No options
+		/// </summary>
 		None = 0,
 
 		/// <summary>
@@ -79,9 +87,20 @@ namespace EpicGames.UHT.Tables
 		}
 	}
 
+	/// <summary>
+	/// The phase of UHT where the property is being resolved
+	/// </summary>
 	public enum UhtPropertyResolvePhase
 	{
+
+		/// <summary>
+		/// Resolved during the source processing phase.  Immediate property types only.
+		/// </summary>
 		Parsing,
+
+		/// <summary>
+		/// Resolved during the resolve phase.  Non-immedite property types only.
+		/// </summary>
 		Resolving,
 	}
 
@@ -95,10 +114,21 @@ namespace EpicGames.UHT.Tables
 	/// <returns></returns>
 	public delegate UhtProperty? UhtResolvePropertyDelegate(UhtPropertyResolvePhase ResolvePhase, UhtPropertySettings PropertySettings, IUhtTokenReader TokenReader, UhtToken MatchedToken);
 
+	/// <summary>
+	/// Property type attribute
+	/// </summary>
 	[AttributeUsage(AttributeTargets.Method, AllowMultiple = true)]
 	public class UhtPropertyTypeAttribute : Attribute
 	{
+
+		/// <summary>
+		/// The expected keyword.  Must be set unless this is the default processor
+		/// </summary>
 		public string? Keyword = null;
+
+		/// <summary>
+		/// Options
+		/// </summary>
 		public UhtPropertyTypeOptions Options = UhtPropertyTypeOptions.None;
 	}
 
@@ -107,18 +137,36 @@ namespace EpicGames.UHT.Tables
 	/// </summary>
 	public struct UhtPropertyType
 	{
+
+		/// <summary>
+		/// Delegate to invoke
+		/// </summary>
 		public UhtResolvePropertyDelegate Delegate;
+
+		/// <summary>
+		/// Options
+		/// </summary>
 		public UhtPropertyTypeOptions Options;
 	}
 
+	/// <summary>
+	/// Property type table
+	/// </summary>
 	public class UhtPropertyTypeTable
 	{
+
+		/// <summary>
+		/// Global instance of the property type table
+		/// </summary>
 		public static UhtPropertyTypeTable Instance = new UhtPropertyTypeTable();
 
 		private Dictionary<StringView, UhtPropertyType> CaseSensitive = new Dictionary<StringView, UhtPropertyType>();
 		private Dictionary<StringView, UhtPropertyType> CaseInsensitive = new Dictionary<StringView, UhtPropertyType>(StringViewComparer.OrdinalIgnoreCase);
 		private UhtPropertyType? DefaultInternal = null;
 
+		/// <summary>
+		/// Return the default processor
+		/// </summary>
 		public UhtPropertyType Default
 		{
 			get
@@ -144,6 +192,13 @@ namespace EpicGames.UHT.Tables
 				this.CaseInsensitive.TryGetValue(Name, out PropertyType);
 		}
 
+		/// <summary>
+		/// Handle a property type attribute
+		/// </summary>
+		/// <param name="Type">Containing type</param>
+		/// <param name="MethodInfo">Method info</param>
+		/// <param name="PropertyTypeAttribute">Attribute</param>
+		/// <exception cref="UhtIceException">Thrown if the property type isn't properly defined.</exception>
 		public void OnPropertyTypeAttribute(Type Type, MethodInfo MethodInfo, UhtPropertyTypeAttribute PropertyTypeAttribute)
 		{
 			if (string.IsNullOrEmpty(PropertyTypeAttribute.Keyword) && !PropertyTypeAttribute.Options.HasAnyFlags(UhtPropertyTypeOptions.Default))
diff --git a/Engine/Source/Programs/Shared/EpicGames.UHT/Tables/UhtSpecifierTable.cs b/Engine/Source/Programs/Shared/EpicGames.UHT/Tables/UhtSpecifierTable.cs
--- a/Engine/Source/Programs/Shared/EpicGames.UHT/Tables/UhtSpecifierTable.cs
+++ b/Engine/Source/Programs/Shared/EpicGames.UHT/Tables/UhtSpecifierTable.cs
@@ -10,6 +10,10 @@ using System.Reflection;
 
 namespace EpicGames.UHT.Tables
 {
+
+	/// <summary>
+	/// Defines the different types specifiers relating to their allowed values
+	/// </summary>
 	public enum UhtSpecifierValueType
 	{
 		/// <summary>
@@ -63,9 +67,20 @@ namespace EpicGames.UHT.Tables
 		Legacy,
 	}
 
+	/// <summary>
+	/// Results from dispatching a specifier
+	/// </summary>
 	public enum UhtSpecifierDispatchResults
 	{
+
+		/// <summary>
+		/// Specifier was known and parsed
+		/// </summary>
 		Known,
+
+		/// <summary>
+		/// Specified was unknown
+		/// </summary>
 		Unknown,
 	}
 
@@ -74,11 +89,34 @@ namespace EpicGames.UHT.Tables
 	/// </summary>
 	public class UhtSpecifierContext
 	{
+
+		/// <summary>
+		/// Current parsing scope (i.e. global, class, ...)
+		/// </summary>
 		public UhtParsingScope Scope;
+
+		/// <summary>
+		/// Message site for messages
+		/// </summary>
 		public IUhtMessageSite MessageSite;
+
+		/// <summary>
+		/// Meta data currently being parsed.
+		/// </summary>
 		public UhtMetaData MetaData;
+
+		/// <summary>
+		/// Make data key index utilized by enumeration values
+		/// </summary>
 		public int MetaNameIndex;
 
+		/// <summary>
+		/// Construct a new specifier context
+		/// </summary>
+		/// <param name="Scope"></param>
+		/// <param name="MessageSite"></param>
+		/// <param name="MetaData"></param>
+		/// <param name="MetaNameIndex"></param>
 		public UhtSpecifierContext(UhtParsingScope Scope, IUhtMessageSite MessageSite, UhtMetaData MetaData, int MetaNameIndex = UhtMetaData.INDEX_NONE)
 		{
 			this.Scope = Scope;
@@ -104,21 +142,56 @@ namespace EpicGames.UHT.Tables
 		Deferred,
 	}
 
+	/// <summary>
+	/// The specifier table contains an instance of UhtSpecifier which is used to dispatch the parsing of
+	/// a specifier to the implementation
+	/// </summary>
 	public abstract class UhtSpecifier
 	{
+
+		/// <summary>
+		/// Name of the specifier
+		/// </summary>
 		public string Name = String.Empty;
+
+		/// <summary>
+		/// Expected value type
+		/// </summary>
 		public UhtSpecifierValueType ValueType;
+
+		/// <summary>
+		/// When is the specifier executed
+		/// </summary>
 		public UhtSpecifierWhen When = UhtSpecifierWhen.Deferred;
 
+		/// <summary>
+		/// Dispatch an instance of the specifier
+		/// </summary>
+		/// <param name="SpecifierContext">Current context</param>
+		/// <param name="Value">Specifier value</param>
+		/// <returns>Results of the dispatch</returns>
 		public abstract UhtSpecifierDispatchResults Dispatch(UhtSpecifierContext SpecifierContext, object? Value);
 	}
 
+	/// <summary>
+	/// Delegate for a specifier with no value
+	/// </summary>
+	/// <param name="SpecifierContext"></param>
 	public delegate void UhtSpecifierNoneDelegate(UhtSpecifierContext SpecifierContext);
 
+	/// <summary>
+	/// Specifier with no value
+	/// </summary>
 	public class UhtSpecifierNone : UhtSpecifier
 	{
 		private UhtSpecifierNoneDelegate Delegate;
 
+		/// <summary>
+		/// Construct the specifier
+		/// </summary>
+		/// <param name="Name">Name of the specifier</param>
+		/// <param name="When">When the specifier is executed</param>
+		/// <param name="Delegate">Delegate to invoke</param>
 		public UhtSpecifierNone(string Name, UhtSpecifierWhen When, UhtSpecifierNoneDelegate Delegate)
 		{
 			this.Name = Name;
@@ -127,6 +200,7 @@ namespace EpicGames.UHT.Tables
 			this.Delegate = Delegate;
 		}
 
+		/// <inheritdoc/>
 		public override UhtSpecifierDispatchResults Dispatch(UhtSpecifierContext SpecifierContext, object? Value)
 		{
 			this.Delegate(SpecifierContext);
@@ -134,12 +208,26 @@ namespace EpicGames.UHT.Tables
 		}
 	}
 
+	/// <summary>
+	/// Specifier delegate with a string value
+	/// </summary>
+	/// <param name="SpecifierContext">Specifier context</param>
+	/// <param name="Value">Specifier value</param>
 	public delegate void UhtSpecifierStringDelegate(UhtSpecifierContext SpecifierContext, StringView Value);
 
+	/// <summary>
+	/// Specifier with a string value
+	/// </summary>
 	public class UhtSpecifierString : UhtSpecifier
 	{
 		private UhtSpecifierStringDelegate Delegate;
 
+		/// <summary>
+		/// Construct the specifier
+		/// </summary>
+		/// <param name="Name">Name of the specifier</param>
+		/// <param name="When">When the specifier is executed</param>
+		/// <param name="Delegate">Delegate to invoke</param>
 		public UhtSpecifierString(string Name, UhtSpecifierWhen When, UhtSpecifierStringDelegate Delegate)
 		{
 			this.Name = Name;
@@ -148,6 +236,7 @@ namespace EpicGames.UHT.Tables
 			this.Delegate = Delegate;
 		}
 
+		/// <inheritdoc/>
 		public override UhtSpecifierDispatchResults Dispatch(UhtSpecifierContext SpecifierContext, object? Value)
 		{
 			if (Value == null)
@@ -159,12 +248,26 @@ namespace EpicGames.UHT.Tables
 		}
 	}
 
+	/// <summary>
+	/// Specifier delegate with an optional string value
+	/// </summary>
+	/// <param name="SpecifierContext">Specifier context</param>
+	/// <param name="Value">Specifier value</param>
 	public delegate void UhtSpecifierOptionalStringDelegate(UhtSpecifierContext SpecifierContext, StringView? Value);
 
+	/// <summary>
+	/// Specifier with an optional string value
+	/// </summary>
 	public class UhtSpecifierOptionalString : UhtSpecifier
 	{
 		private UhtSpecifierOptionalStringDelegate Delegate;
 
+		/// <summary>
+		/// Construct the specifier
+		/// </summary>
+		/// <param name="Name">Name of the specifier</param>
+		/// <param name="When">When the specifier is executed</param>
+		/// <param name="Delegate">Delegate to invoke</param>
 		public UhtSpecifierOptionalString(string Name, UhtSpecifierWhen When, UhtSpecifierOptionalStringDelegate Delegate)
 		{
 			this.Name = Name;
@@ -173,6 +276,7 @@ namespace EpicGames.UHT.Tables
 			this.Delegate = Delegate;
 		}
 
+		/// <inheritdoc/>
 		public override UhtSpecifierDispatchResults Dispatch(UhtSpecifierContext SpecifierContext, object? Value)
 		{
 			this.Delegate(SpecifierContext, (StringView?)Value);
@@ -180,12 +284,26 @@ namespace EpicGames.UHT.Tables
 		}
 	}
 
+	/// <summary>
+	/// Specifier delegate with a string value
+	/// </summary>
+	/// <param name="SpecifierContext">Specifier context</param>
+	/// <param name="Value">Specifier value</param>
 	public delegate void UhtSpecifierSingleStringDelegate(UhtSpecifierContext SpecifierContext, StringView Value);
 
+	/// <summary>
+	/// Specifier with a string value
+	/// </summary>
 	public class UhtSpecifierSingleString : UhtSpecifier
 	{
 		private UhtSpecifierSingleStringDelegate Delegate;
 
+		/// <summary>
+		/// Construct the specifier
+		/// </summary>
+		/// <param name="Name">Name of the specifier</param>
+		/// <param name="When">When the specifier is executed</param>
+		/// <param name="Delegate">Delegate to invoke</param>
 		public UhtSpecifierSingleString(string Name, UhtSpecifierWhen When, UhtSpecifierSingleStringDelegate Delegate)
 		{
 			this.Name = Name;
@@ -194,6 +312,7 @@ namespace EpicGames.UHT.Tables
 			this.Delegate = Delegate;
 		}
 
+		/// <inheritdoc/>
 		public override UhtSpecifierDispatchResults Dispatch(UhtSpecifierContext SpecifierContext, object? Value)
 		{
 			if (Value == null)
@@ -205,12 +324,27 @@ namespace EpicGames.UHT.Tables
 		}
 	}
 
+	/// <summary>
+	/// Specifier delegate with list of string keys and values 
+	/// </summary>
+	/// <param name="SpecifierContext">Specifier context</param>
+	/// <param name="Value">Specifier value</param>
 	public delegate void UhtSpecifierKeyValuePairListDelegate(UhtSpecifierContext SpecifierContext, List<KeyValuePair<StringView, StringView>> Value);
 
+	/// <summary>
+	/// Specifier with list of string keys and values 
+	/// </summary>
 	public class UhtSpecifierKeyValuePairList : UhtSpecifier
 	{
 		private UhtSpecifierKeyValuePairListDelegate Delegate;
 
+		/// <summary>
+		/// Construct the specifier
+		/// </summary>
+		/// <param name="Name">Name of the specifier</param>
+		/// <param name="When">When the specifier is executed</param>
+		/// <param name="bEqualsOptional">If true this has an optional KVP list</param>
+		/// <param name="Delegate">Delegate to invoke</param>
 		public UhtSpecifierKeyValuePairList(string Name, UhtSpecifierWhen When, bool bEqualsOptional, UhtSpecifierKeyValuePairListDelegate Delegate)
 		{
 			this.Name = Name;
@@ -219,6 +353,7 @@ namespace EpicGames.UHT.Tables
 			this.Delegate = Delegate;
 		}
 
+		/// <inheritdoc/>
 		public override UhtSpecifierDispatchResults Dispatch(UhtSpecifierContext SpecifierContext, object? Value)
 		{
 			if (Value == null)
@@ -230,12 +365,25 @@ namespace EpicGames.UHT.Tables
 		}
 	}
 
+	/// <summary>
+	/// Specifier delegate with no value
+	/// </summary>
+	/// <param name="SpecifierContext">Specifier context</param>
 	public delegate void UhtSpecifierLegacyDelegate(UhtSpecifierContext SpecifierContext);
 
+	/// <summary>
+	/// Specifier delegate for legacy UHT specifiers with no value.  Will generate a information/deprecation message
+	/// is a value is supplied
+	/// </summary>
 	public class UhtSpecifierLegacy : UhtSpecifier
 	{
 		private UhtSpecifierLegacyDelegate Delegate;
 
+		/// <summary>
+		/// Construct the specifier
+		/// </summary>
+		/// <param name="Name">Name of the specifier</param>
+		/// <param name="Delegate">Delegate to invoke</param>
 		public UhtSpecifierLegacy(string Name, UhtSpecifierLegacyDelegate Delegate)
 		{
 			this.Name = Name;
@@ -244,6 +392,7 @@ namespace EpicGames.UHT.Tables
 			this.Delegate = Delegate;
 		}
 
+		/// <inheritdoc/>
 		public override UhtSpecifierDispatchResults Dispatch(UhtSpecifierContext SpecifierContext, object? Value)
 		{
 			if (Value != null)
@@ -255,12 +404,26 @@ namespace EpicGames.UHT.Tables
 		}
 	}
 
+	/// <summary>
+	/// Specifier delegate with an optional string list
+	/// </summary>
+	/// <param name="SpecifierContext">Specifier context</param>
+	/// <param name="Value">Specifier value</param>
 	public delegate void UhtSpecifierStringListDelegate(UhtSpecifierContext SpecifierContext, List<StringView>? Value);
 
+	/// <summary>
+	/// Specifier with an optional string list
+	/// </summary>
 	public class UhtSpecifierStringList : UhtSpecifier
 	{
 		private UhtSpecifierStringListDelegate Delegate;
 
+		/// <summary>
+		/// Construct the specifier
+		/// </summary>
+		/// <param name="Name">Name of the specifier</param>
+		/// <param name="When">When the specifier is executed</param>
+		/// <param name="Delegate">Delegate to invoke</param>
 		public UhtSpecifierStringList(string Name, UhtSpecifierWhen When, UhtSpecifierStringListDelegate Delegate)
 		{
 			this.Name = Name;
@@ -269,6 +432,7 @@ namespace EpicGames.UHT.Tables
 			this.Delegate = Delegate;
 		}
 
+		/// <inheritdoc/>
 		public override UhtSpecifierDispatchResults Dispatch(UhtSpecifierContext SpecifierContext, object? Value)
 		{
 			this.Delegate(SpecifierContext, (List<StringView>?)Value);
@@ -276,12 +440,26 @@ namespace EpicGames.UHT.Tables
 		}
 	}
 
+	/// <summary>
+	/// Specifier delegate with a list of string views
+	/// </summary>
+	/// <param name="SpecifierContext">Specifier context</param>
+	/// <param name="Value">Specifier value</param>
 	public delegate void UhtSpecifierNonEmptyStringListDelegate(UhtSpecifierContext SpecifierContext, List<StringView> Value);
 
+	/// <summary>
+	/// Specifier with a list of string views
+	/// </summary>
 	public class UhtSpecifierNonEmptyStringList : UhtSpecifier
 	{
 		private UhtSpecifierNonEmptyStringListDelegate Delegate;
 
+		/// <summary>
+		/// Construct the specifier
+		/// </summary>
+		/// <param name="Name">Name of the specifier</param>
+		/// <param name="When">When the specifier is executed</param>
+		/// <param name="Delegate">Delegate to invoke</param>
 		public UhtSpecifierNonEmptyStringList(string Name, UhtSpecifierWhen When, UhtSpecifierNonEmptyStringListDelegate Delegate)
 		{
 			this.Name = Name;
@@ -290,6 +468,7 @@ namespace EpicGames.UHT.Tables
 			this.Delegate = Delegate;
 		}
 
+		/// <inheritdoc/>
 		public override UhtSpecifierDispatchResults Dispatch(UhtSpecifierContext SpecifierContext, object? Value)
 		{
 			if (Value == null)
@@ -301,15 +480,36 @@ namespace EpicGames.UHT.Tables
 		}
 	}
 
+	/// <summary>
+	/// Defines a specifier method
+	/// </summary>
 	[AttributeUsage(AttributeTargets.Method, AllowMultiple = true)]
 	public class UhtSpecifierAttribute : Attribute
 	{
+		/// <summary>
+		/// Name of the specifier.   If not supplied, the method name must end in "Specifier" and the name will be the method name with "Specifier" stripped.
+		/// </summary>
 		public string? Name;
+
+		/// <summary>
+		/// Name of the table/scope this specifier applies
+		/// </summary>
 		public string? Extends;
+
+		/// <summary>
+		/// Value type of the specifier
+		/// </summary>
 		public UhtSpecifierValueType ValueType = UhtSpecifierValueType.NotSet;
+
+		/// <summary>
+		/// When the specifier is dispatched
+		/// </summary>
 		public UhtSpecifierWhen When = UhtSpecifierWhen.Deferred;
 	}
 
+	/// <summary>
+	/// Collection of specifiers for a given scope
+	/// </summary>
 	public class UhtSpecifierTable : UhtLookupTable<UhtSpecifier>
 	{
 
@@ -323,7 +523,7 @@ namespace EpicGames.UHT.Tables
 		/// <summary>
 		/// Add the given value to the lookup table.  It will throw an exception if it is a duplicate.
 		/// </summary>
-		/// <param name="Value">Value to be added</param>
+		/// <param name="Specifier">The specifier to add</param>
 		public UhtSpecifierTable Add(UhtSpecifier Specifier)
 		{
 			base.Add(Specifier.Name, Specifier);
@@ -331,14 +531,31 @@ namespace EpicGames.UHT.Tables
 		}
 	}
 
+	/// <summary>
+	/// Collection of all specifier tables
+	/// </summary>
 	public class UhtSpecifierTables : UhtLookupTables<UhtSpecifierTable>
 	{
+
+		/// <summary>
+		/// Returns the global instance for all specifiers
+		/// </summary>
 		public static UhtSpecifierTables Instance = new UhtSpecifierTables();
 
+		/// <summary>
+		/// Construct the specifier table
+		/// </summary>
 		public UhtSpecifierTables() : base("specifiers")
 		{
 		}
 
+		/// <summary>
+		/// Invoke for a method that has the specifier attribute
+		/// </summary>
+		/// <param name="Type">Type containing the method</param>
+		/// <param name="MethodInfo">Method info</param>
+		/// <param name="SpecifierAttribute">Specified attributes</param>
+		/// <exception cref="UhtIceException">Throw if the attribute isn't properly defined.</exception>
 		public void OnSpecifierAttribute(Type Type, MethodInfo MethodInfo, UhtSpecifierAttribute SpecifierAttribute)
 		{
 			string Name = UhtLookupTableBase.GetSuffixedName(Type, MethodInfo, SpecifierAttribute.Name, "Specifier");
diff --git a/Engine/Source/Programs/Shared/EpicGames.UHT/Tables/UhtSpecifierValidatorTable.cs b/Engine/Source/Programs/Shared/EpicGames.UHT/Tables/UhtSpecifierValidatorTable.cs
--- a/Engine/Source/Programs/Shared/EpicGames.UHT/Tables/UhtSpecifierValidatorTable.cs
+++ b/Engine/Source/Programs/Shared/EpicGames.UHT/Tables/UhtSpecifierValidatorTable.cs
@@ -8,13 +8,37 @@ using System.Reflection;
 
 namespace EpicGames.UHT.Tables
 {
+
+	/// <summary>
+	/// Delegate used to validate a specifier
+	/// </summary>
+	/// <param name="Type">Containing type</param>
+	/// <param name="MetaData">Containing meta data</param>
+	/// <param name="Key">Key of the meta data entry</param>
+	/// <param name="Value">Value of the meta data entry</param>
 	public delegate void UhtSpecifierValidatorDelegate(UhtType Type, UhtMetaData MetaData, UhtMetaDataKey Key, StringView Value);
 
+	/// <summary>
+	/// Defines a specifier validated created from the attribute
+	/// </summary>
 	public class UhtSpecifierValidator
 	{
+
+		/// <summary>
+		/// Name of the validator
+		/// </summary>
 		public string Name;
+
+		/// <summary>
+		/// Delegate for the validator
+		/// </summary>
 		public UhtSpecifierValidatorDelegate Delegate;
 
+		/// <summary>
+		/// Construct a new instance
+		/// </summary>
+		/// <param name="Name">Name of the validator</param>
+		/// <param name="Delegate">Delegate of the validator</param>
 		public UhtSpecifierValidator(string Name, UhtSpecifierValidatorDelegate Delegate)
 		{
 			this.Name = Name;
@@ -22,13 +46,27 @@ namespace EpicGames.UHT.Tables
 		}
 	}
 
+	/// <summary>
+	/// Attribute used to create a specifier validator
+	/// </summary>
 	[AttributeUsage(AttributeTargets.Method, AllowMultiple = true)]
 	public class UhtSpecifierValidatorAttribute : Attribute
 	{
+
+		/// <summary>
+		/// Name of the validator. If not supplied &quot;SpecifierValidator&quot; will be removed from the end of the method name
+		/// </summary>
 		public string? Name;
+
+		/// <summary>
+		/// Name of the table/scope for the validator
+		/// </summary>
 		public string? Extends;
 	}
 
+	/// <summary>
+	/// A table for validators for a given scope
+	/// </summary>
 	public class UhtSpecifierValidatorTable : UhtLookupTable<UhtSpecifierValidator>
 	{
 
@@ -42,7 +80,7 @@ namespace EpicGames.UHT.Tables
 		/// <summary>
 		/// Add the given value to the lookup table.  It will throw an exception if it is a duplicate.
 		/// </summary>
-		/// <param name="Value">Value to be added</param>
+		/// <param name="Specifier">Validator to add</param>
 		public UhtSpecifierValidatorTable Add(UhtSpecifierValidator Specifier)
 		{
 			base.Add(Specifier.Name, Specifier);
@@ -50,14 +88,31 @@ namespace EpicGames.UHT.Tables
 		}
 	}
 
+	/// <summary>
+	/// Collection of specifier validators
+	/// </summary>
 	public class UhtSpecifierValidatorTables : UhtLookupTables<UhtSpecifierValidatorTable>
 	{
+
+		/// <summary>
+		/// Global instance of the validator tables
+		/// </summary>
 		public static UhtSpecifierValidatorTables Instance = new UhtSpecifierValidatorTables();
 
+		/// <summary>
+		/// Construct the validator tables
+		/// </summary>
 		public UhtSpecifierValidatorTables() : base("specifier validators")
 		{
 		}
 
+		/// <summary>
+		/// Handle the attribute appearing on a method
+		/// </summary>
+		/// <param name="Type">Type containing the method</param>
+		/// <param name="MethodInfo">The method</param>
+		/// <param name="SpecifierValidatorAttribute">Attribute</param>
+		/// <exception cref="UhtIceException">Thrown if the validator isn't properly defined</exception>
 		public void OnSpecifierValidatorAttribute(Type Type, MethodInfo MethodInfo, UhtSpecifierValidatorAttribute SpecifierValidatorAttribute)
 		{
 			string Name = UhtLookupTableBase.GetSuffixedName(Type, MethodInfo, SpecifierValidatorAttribute.Name, "SpecifierValidator");
diff --git a/Engine/Source/Programs/Shared/EpicGames.UHT/Tables/UhtStructDefaultValueTable.cs b/Engine/Source/Programs/Shared/EpicGames.UHT/Tables/UhtStructDefaultValueTable.cs
--- a/Engine/Source/Programs/Shared/EpicGames.UHT/Tables/UhtStructDefaultValueTable.cs
+++ b/Engine/Source/Programs/Shared/EpicGames.UHT/Tables/UhtStructDefaultValueTable.cs
@@ -11,11 +11,26 @@ using System.Text;
 
 namespace EpicGames.UHT.Tables
 {
+
+	/// <summary>
+	/// Delegate for invoking structure default value sanitizer
+	/// </summary>
+	/// <param name="Property"></param>
+	/// <param name="DefaultValueReader"></param>
+	/// <param name="InnerDefaultValue"></param>
+	/// <returns></returns>
 	public delegate bool UhtStructDefaultValueDelegate(UhtStructProperty Property, IUhtTokenReader DefaultValueReader, StringBuilder InnerDefaultValue);
 
+	/// <summary>
+	/// Options for structure default value sanitizer
+	/// </summary>
 	[Flags]
 	public enum UhtStructDefaultValueOptions
 	{
+
+		/// <summary>
+		/// No options
+		/// </summary>
 		None = 0,
 
 		/// <summary>
@@ -66,24 +81,53 @@ namespace EpicGames.UHT.Tables
 		}
 	}
 
+	/// <summary>
+	/// Structure default value sanitizer attribute
+	/// </summary>
 	[AttributeUsage(AttributeTargets.Method, AllowMultiple = true)]
 	public class UhtStructDefaultValueAttribute : Attribute
 	{
+
+		/// <summary>
+		/// Name of the structure.  Not required for default processor.  
+		/// </summary>
 		public string? Name;
+
+		/// <summary>
+		/// Options
+		/// </summary>
 		public UhtStructDefaultValueOptions Options = UhtStructDefaultValueOptions.None;
 	}
 
+	/// <summary>
+	/// Structure default value sanitizer
+	/// </summary>
 	public struct UhtStructDefaultValue
 	{
+
+		/// <summary>
+		/// The delegate to invoke
+		/// </summary>
 		public UhtStructDefaultValueDelegate Delegate;
 	}
 
+	/// <summary>
+	/// Table of all structure default value specifiers
+	/// </summary>
 	public class UhtStructDefaultValueTable
 	{
+
+		/// <summary>
+		/// Global instance of the structure default value sanitizers
+		/// </summary>
 		public static UhtStructDefaultValueTable Instance = new UhtStructDefaultValueTable();
 
 		private Dictionary<StringView, UhtStructDefaultValue> StructDefaultValues = new Dictionary<StringView, UhtStructDefaultValue>();
 		private UhtStructDefaultValue? DefaultInternal = null;
+
+		/// <summary>
+		/// Fetch the default sanitizer
+		/// </summary>
 		public UhtStructDefaultValue Default
 		{
 			get
@@ -107,6 +151,13 @@ namespace EpicGames.UHT.Tables
 			return this.StructDefaultValues.TryGetValue(Name, out StructDefaultValue);
 		}
 
+		/// <summary>
+		/// Handle a structure default value sanitizer attribute
+		/// </summary>
+		/// <param name="Type">Containing type</param>
+		/// <param name="MethodInfo">Method information</param>
+		/// <param name="StructDefaultValueAttribute">Found attribute</param>
+		/// <exception cref="UhtIceException">Thrown if the attribute isn't property defined</exception>
 		public void OnStructDefaultValueAttribute(Type Type, MethodInfo MethodInfo, UhtStructDefaultValueAttribute StructDefaultValueAttribute)
 		{
 			if (string.IsNullOrEmpty(StructDefaultValueAttribute.Name) && !StructDefaultValueAttribute.Options.HasAnyFlags(UhtStructDefaultValueOptions.Default))
