Horde: Fix infinite recursion when trying to upload compact binary encoded objects.

diff --git a/Engine/Source/Programs/Shared/EpicGames.Horde/Storage/IStorageClient.cs b/Engine/Source/Programs/Shared/EpicGames.Horde/Storage/IStorageClient.cs
--- a/Engine/Source/Programs/Shared/EpicGames.Horde/Storage/IStorageClient.cs
+++ b/Engine/Source/Programs/Shared/EpicGames.Horde/Storage/IStorageClient.cs
@@ -302,5 +302,33 @@
 		}
 
 		/// <summary>
+		/// Writes a blob from memory to storage
+		/// </summary>
+		/// <param name="StorageClient">The storage interface</param>
+		/// <param name="NamespaceId">Namespace containing the blob</param>
+		/// <param name="Data">Data to write</param>
+		/// <returns></returns>
+		public static async Task<IoHash> WriteBlobFromMemoryAsync(this IStorageClient StorageClient, NamespaceId NamespaceId, ReadOnlyMemory<byte> Data)
+		{
+			IoHash Hash = IoHash.Compute(Data.Span);
+			await WriteBlobFromMemoryAsync(StorageClient, NamespaceId, Hash, Data);
+			return Hash;
+		}
+
+		/// <summary>
+		/// Writes a blob from memory to storage
+		/// </summary>
+		/// <param name="StorageClient">The storage interface</param>
+		/// <param name="NamespaceId">Namespace containing the blob</param>
+		/// <param name="Hash">Hash of the data</param>
+		/// <param name="Data">The data to be written</param>
+		/// <returns></returns>
+		public static async Task WriteBlobFromMemoryAsync(this IStorageClient StorageClient, NamespaceId NamespaceId, IoHash Hash, ReadOnlyMemory<byte> Data)
+		{
+			using ReadOnlyMemoryStream Stream = new ReadOnlyMemoryStream(Data);
+			await StorageClient.WriteBlobAsync(NamespaceId, Hash, Stream);
+		}
+
+		/// <summary>
 		/// Reads a blob and decodes it as a compact binary object
 		/// </summary>
@@ -329,21 +357,7 @@
 		}
 
 		/// <summary>
-		/// Gets a blob as a byte array
+		/// Writes a serialized compact binary object to storage
 		/// </summary>
 		/// <param name="StorageClient">The storage interface</param>
 		/// <param name="NamespaceId">Namespace containing the blob</param>
-		/// <param name="Data">Data to write</param>
-		/// <returns></returns>
-		public static async Task<IoHash> WriteBlobFromMemoryAsync(this IStorageClient StorageClient, NamespaceId NamespaceId, ReadOnlyMemory<byte> Data)
-		{
-			IoHash Hash = IoHash.Compute(Data.Span);
-			await WriteBlobFromMemoryAsync(StorageClient, NamespaceId, Hash, Data);
-			return Hash;
-		}
-
-		/// <summary>
-		/// Gets a blob as a byte array
-		/// </summary>
-		/// <param name="StorageClient">The storage interface</param>
-		/// <param name="NamespaceId">Namespace containing the blob</param>
@@ -355,47 +369,22 @@
 		}
 
 		/// <summary>
-		/// Gets a blob as a byte array
+		/// Serializes an object to compact binary format and writes it to storage
 		/// </summary>
 		/// <param name="StorageClient">The storage interface</param>
 		/// <param name="NamespaceId">Namespace containing the blob</param>
 		/// <param name="Object">The object to be written</param>
 		/// <returns></returns>
-		public static Task<IoHash> WriteObjectAsync<T>(this IStorageClient StorageClient, NamespaceId NamespaceId, T Object)
+		public static async Task<IoHash> WriteObjectAsync<T>(this IStorageClient StorageClient, NamespaceId NamespaceId, T Object)
 		{
 			CbWriter Writer = new CbWriter();
 			CbSerializer.Serialize<T>(Writer, Object);
-			return WriteObjectAsync(StorageClient, NamespaceId, Writer);
-		}
 
-		/// <summary>
-		/// Gets a blob as a byte array
-		/// </summary>
-		/// <param name="StorageClient">The storage interface</param>
-		/// <param name="NamespaceId">Namespace containing the blob</param>
-		/// <param name="Writer">The object to be written</param>
-		/// <returns></returns>
-		public static async Task<IoHash> WriteObjectAsync<T>(this IStorageClient StorageClient, NamespaceId NamespaceId, CbWriter Writer)
-		{
 			IoHash Hash = Writer.ComputeHash();
 			await StorageClient.WriteBlobAsync(NamespaceId, Hash, Writer.AsStream());
 			return Hash;
 		}
 
 		/// <summary>
-		/// Gets a blob as a byte array
-		/// </summary>
-		/// <param name="StorageClient">The storage interface</param>
-		/// <param name="NamespaceId">Namespace containing the blob</param>
-		/// <param name="Hash">Hash of the data</param>
-		/// <param name="Data">The data to be written</param>
-		/// <returns></returns>
-		public static async Task WriteBlobFromMemoryAsync(this IStorageClient StorageClient, NamespaceId NamespaceId, IoHash Hash, ReadOnlyMemory<byte> Data)
-		{
-			using ReadOnlyMemoryStream Stream = new ReadOnlyMemoryStream(Data);
-			await StorageClient.WriteBlobAsync(NamespaceId, Hash, Stream);
-		}
-
-		/// <summary>
 		/// Reads a reference as a specific type
 		/// </summary>
