Horde: Bump DefaultMaxInMemoryBlobLength

#rnx

diff --git a/Engine/Source/Programs/Shared/EpicGames.Horde/Storage/IStorageClient.cs b/Engine/Source/Programs/Shared/EpicGames.Horde/Storage/IStorageClient.cs
--- a/Engine/Source/Programs/Shared/EpicGames.Horde/Storage/IStorageClient.cs
+++ b/Engine/Source/Programs/Shared/EpicGames.Horde/Storage/IStorageClient.cs
@@ -267,7 +267,7 @@ namespace EpicGames.Horde.Storage
 	/// </summary>
 	public static class StorageClientExtensions
 	{
-		const int DefaultMaxInMemoryBlobLength = 50 * 1024 * 1024;
+		const int DefaultMaxInMemoryBlobLength = 128 * 1024 * 1024;
 
 		/// <summary>
 		/// Gets a blob as a byte array
