More UHT documentation updates

#rnx

diff --git a/Engine/Source/Programs/Shared/EpicGames.UHT/EpicGames.UHT.csproj b/Engine/Source/Programs/Shared/EpicGames.UHT/EpicGames.UHT.csproj
--- a/Engine/Source/Programs/Shared/EpicGames.UHT/EpicGames.UHT.csproj
+++ b/Engine/Source/Programs/Shared/EpicGames.UHT/EpicGames.UHT.csproj
@@ -9,7 +9,7 @@
     <Configurations>Debug;Release;Development</Configurations>
     <Nullable>enable</Nullable>
     <AllowUnsafeBlocks>true</AllowUnsafeBlocks>
-    <GenerateDocumentationFile>False</GenerateDocumentationFile>
+    <GenerateDocumentationFile>True</GenerateDocumentationFile>
     <DebugType>pdbonly</DebugType>
   </PropertyGroup>
   <PropertyGroup Condition=" '$(Configuration)|$(Platform)' == 'Debug|AnyCPU' ">
diff --git a/Engine/Source/Programs/Shared/EpicGames.UHT/Parsers/UhtClassBaseParser.cs b/Engine/Source/Programs/Shared/EpicGames.UHT/Parsers/UhtClassBaseParser.cs
--- a/Engine/Source/Programs/Shared/EpicGames.UHT/Parsers/UhtClassBaseParser.cs
+++ b/Engine/Source/Programs/Shared/EpicGames.UHT/Parsers/UhtClassBaseParser.cs
@@ -7,15 +7,33 @@ using System.Collections.Generic;
 
 namespace EpicGames.UHT.Parsers
 {
+
+	/// <summary>
+	/// Base parser class for all UClass types
+	/// </summary>
 	public class UhtClassBaseParser : UhtClass
 	{
+
+		/// <summary>
+		/// The super class identifier
+		/// </summary>
 		public UhtToken SuperIdentifier;
+
+		/// <summary>
+		/// List of base identifiers
+		/// </summary>
 		public List<UhtToken[]>? BaseIdentifiers = null;
 
+		/// <summary>
+		/// Construct a new base class parser
+		/// </summary>
+		/// <param name="Outer">Outer type</param>
+		/// <param name="LineNumber">Line number of class</param>
 		public UhtClassBaseParser(UhtType Outer, int LineNumber) : base(Outer, LineNumber)
 		{
 		}
 
+		/// <inheritdoc/>
 		protected override void ResolveSuper(UhtResolvePhase ResolvePhase)
 		{
 			base.ResolveSuper(ResolvePhase);
@@ -34,6 +52,7 @@ namespace EpicGames.UHT.Parsers
 			}
 		}
 
+		/// <inheritdoc/>
 		protected override bool ResolveSelf(UhtResolvePhase ResolvePhase)
 		{
 			bool bResult = base.ResolveSelf(ResolvePhase);
diff --git a/Engine/Source/Programs/Shared/EpicGames.UHT/Parsers/UhtClassParser.cs b/Engine/Source/Programs/Shared/EpicGames.UHT/Parsers/UhtClassParser.cs
--- a/Engine/Source/Programs/Shared/EpicGames.UHT/Parsers/UhtClassParser.cs
+++ b/Engine/Source/Programs/Shared/EpicGames.UHT/Parsers/UhtClassParser.cs
@@ -11,41 +11,100 @@ using System.Text;
 
 namespace EpicGames.UHT.Parsers
 {
+	/// <summary>
+	/// UCLASS parser
+	/// </summary>
 	[UnrealHeaderTool]
 	public class UhtClassParser : UhtClassBaseParser
 	{
 		private static UhtKeywordTable KeywordTable = UhtKeywordTables.Instance.Get(UhtTableNames.Class);
 		private static UhtSpecifierTable SpecifierTable = UhtSpecifierTables.Instance.Get(UhtTableNames.Class);
 
+		/// <summary>
+		/// Engine class flags removed
+		/// </summary>
 		public EClassFlags RemovedClassFlags = EClassFlags.None;
+
+		/// <summary>
+		/// Class within identifier
+		/// </summary>
 		public string ClassWithinIdentifier = String.Empty;
 
+		/// <summary>
+		/// Collection of show categories
+		/// </summary>
 		public List<string> ShowCategories = new List<string>();
+
+		/// <summary>
+		/// Collection of hide categories
+		/// </summary>
 		public List<string> HideCategories = new List<string>();
+
+		/// <summary>
+		/// Collection of auto expand categories
+		/// </summary>
 		public List<string> AutoExpandCategories = new List<string>();
+
+		/// <summary>
+		/// Collection of auto collapse categories
+		/// </summary>
 		public List<string> AutoCollapseCategories = new List<string>();
+
+		/// <summary>
+		/// Collection of prioritize categories
+		/// </summary>
 		public List<string> PrioritizeCategories = new List<string>();
+
+		/// <summary>
+		/// Collection of show functions
+		/// </summary>
 		public List<string> ShowFunctions = new List<string>();
+
+		/// <summary>
+		/// Collection of hide functions
+		/// </summary>
 		public List<string> HideFunctions = new List<string>();
+
+		/// <summary>
+		/// Sparse class data types
+		/// </summary>
 		public List<string> SparseClassDataTypes = new List<string>();
+
+		/// <summary>
+		/// Class group names
+		/// </summary>
 		public List<string> ClassGroupNames = new List<string>();
 
+		/// <summary>
+		/// Add the given class flags
+		/// </summary>
+		/// <param name="Flags">Flags to add</param>
 		public void AddClassFlags(EClassFlags Flags)
 		{
 			this.ClassFlags |= Flags;
 			this.RemovedClassFlags &= ~Flags;
 		}
 
+		/// <summary>
+		/// Remove the given class flags
+		/// </summary>
+		/// <param name="Flags">Flags to remove</param>
 		public void RemoveClassFlags(EClassFlags Flags)
 		{
 			this.RemovedClassFlags |= Flags;
 			this.ClassFlags &= ~Flags;
 		}
 
+		/// <summary>
+		/// Construct a new class parser
+		/// </summary>
+		/// <param name="Outer">Outer type</param>
+		/// <param name="LineNumber">Line number</param>
 		public UhtClassParser(UhtType Outer, int LineNumber) : base(Outer, LineNumber)
 		{
 		}
 
+		/// <inheritdoc/>
 		protected override void ResolveSuper(UhtResolvePhase ResolvePhase)
 		{
 			base.ResolveSuper(ResolvePhase);
@@ -342,7 +401,7 @@ namespace EpicGames.UHT.Parsers
 		}
 		#endregion
 
-		public static UhtParseResult ParseUClass(UhtParsingScope ParentScope, ref UhtToken Token)
+		private static UhtParseResult ParseUClass(UhtParsingScope ParentScope, ref UhtToken Token)
 		{
 			UhtClassParser Class = new UhtClassParser(ParentScope.ScopeType, Token.InputLine);
 			using (var TopScope = new UhtParsingScope(ParentScope, Class, UhtClassParser.KeywordTable, UhtAccessSpecifier.Private))
diff --git a/Engine/Source/Programs/Shared/EpicGames.UHT/Parsers/UhtEnumParser.cs b/Engine/Source/Programs/Shared/EpicGames.UHT/Parsers/UhtEnumParser.cs
--- a/Engine/Source/Programs/Shared/EpicGames.UHT/Parsers/UhtEnumParser.cs
+++ b/Engine/Source/Programs/Shared/EpicGames.UHT/Parsers/UhtEnumParser.cs
@@ -8,12 +8,20 @@ using EpicGames.UHT.Utils;
 
 namespace EpicGames.UHT.Parsers
 {
+	/// <summary>
+	/// UENUM parser
+	/// </summary>
 	[UnrealHeaderTool]
 	public class UhtEnumParser : UhtEnum
 	{
 		private static UhtKeywordTable KeywordTable = UhtKeywordTables.Instance.Get(UhtTableNames.Enum);
 		private static UhtSpecifierTable SpecifierTable = UhtSpecifierTables.Instance.Get(UhtTableNames.Enum);
 
+		/// <summary>
+		/// Construct a new enum parser
+		/// </summary>
+		/// <param name="Outer">Outer type</param>
+		/// <param name="LineNumber">Line number</param>
 		public UhtEnumParser(UhtType Outer, int LineNumber) : base(Outer, LineNumber)
 		{
 		}
diff --git a/Engine/Source/Programs/Shared/EpicGames.UHT/Parsers/UhtFunctionParser.cs b/Engine/Source/Programs/Shared/EpicGames.UHT/Parsers/UhtFunctionParser.cs
--- a/Engine/Source/Programs/Shared/EpicGames.UHT/Parsers/UhtFunctionParser.cs
+++ b/Engine/Source/Programs/Shared/EpicGames.UHT/Parsers/UhtFunctionParser.cs
@@ -90,12 +90,20 @@ namespace EpicGames.UHT.Parsers
 		}
 	}
 
+	/// <summary>
+	/// UFUNCTION parser
+	/// </summary>
 	[UnrealHeaderTool]
 	public class UhtFunctionParser : UhtFunction
 	{
 		private static UhtKeywordTable KeywordTable = UhtKeywordTables.Instance.Get(UhtTableNames.Function);
 		private static UhtSpecifierTable SpecifierTable = UhtSpecifierTables.Instance.Get(UhtTableNames.Function);
 
+		/// <summary>
+		/// Construct a new function parser
+		/// </summary>
+		/// <param name="Outer">Outer object</param>
+		/// <param name="LineNumber">Line number</param>
 		public UhtFunctionParser(UhtType Outer, int LineNumber) : base(Outer, LineNumber)
 		{
 		}
@@ -105,6 +113,7 @@ namespace EpicGames.UHT.Parsers
 		/// </summary>
 		public bool bSawPropertyAccessor = false;
 
+		/// <inheritdoc/>
 		protected override bool ResolveSelf(UhtResolvePhase ResolvePhase)
 		{
 			bool bResult = base.ResolveSelf(ResolvePhase);
@@ -150,7 +159,7 @@ namespace EpicGames.UHT.Parsers
 			return bResult;
 		}
 
-		public UhtPropertyParseOptions GetPropertyParseOptions(bool bReturnValue)
+		private UhtPropertyParseOptions GetPropertyParseOptions(bool bReturnValue)
 		{
 			switch (this.FunctionType)
 			{
@@ -171,6 +180,7 @@ namespace EpicGames.UHT.Parsers
 					throw new UhtIceException("Unknown enumeration value");
 			}
 		}
+
 		#region Keywords
 		[UhtKeyword(Extends = UhtTableNames.Global)]
 		[UhtKeyword(Extends = UhtTableNames.Class)]
diff --git a/Engine/Source/Programs/Shared/EpicGames.UHT/Parsers/UhtHeaderFileParser.cs b/Engine/Source/Programs/Shared/EpicGames.UHT/Parsers/UhtHeaderFileParser.cs
--- a/Engine/Source/Programs/Shared/EpicGames.UHT/Parsers/UhtHeaderFileParser.cs
+++ b/Engine/Source/Programs/Shared/EpicGames.UHT/Parsers/UhtHeaderFileParser.cs
@@ -11,45 +11,75 @@ using System.IO;
 
 namespace EpicGames.UHT.Parsers
 {
+
+	/// <summary>
+	/// Keyword parse results
+	/// </summary>
 	public enum UhtParseResult
 	{
+
+		/// <summary>
+		/// Keyword was handled
+		/// </summary>
 		Handled,
+
+		/// <summary>
+		/// Keyword wasn't handled (more attempts will be made to match)
+		/// </summary>
 		Unhandled,
+
+		/// <summary>
+		/// Keyword is invalid
+		/// </summary>
 		Invalid,
 	}
 
+	/// <summary>
+	/// Compiler directives
+	/// </summary>
 	[Flags]
 	public enum UhtCompilerDirective
 	{
+		/// <summary>
+		/// No compile directives
+		/// </summary>
 		None = 0,
+
 		/// <summary>
 		/// This indicates we are in a "#if CPP" block
 		/// </summary>
 		CPPBlock = 1 << 0,
+
 		/// <summary>
 		/// This indicates we are in a "#if !CPP" block
 		/// </summary>
 		NotCPPBlock = 1 << 1,
+
 		/// <summary>
 		/// This indicates we are in a "#if 0" block
 		/// </summary>
 		ZeroBlock = 1 << 2,
+
 		/// <summary>
 		/// This indicates we are in a "#if 1" block
 		/// </summary>
 		OneBlock = 1 << 3,
+
 		/// <summary>
 		/// This indicates we are in a "#if WITH_EDITOR" block
 		/// </summary>
 		WithEditor = 1 << 4,
+
 		/// <summary>
 		/// This indicates we are in a "#if WITH_EDITORONLY_DATA" block
 		/// </summary>
 		WithEditorOnlyData = 1 << 5,
+
 		/// <summary>
 		/// This indicates we are in a "#if WITH_HOT_RELOAD" block
 		/// </summary>
 		WithHotReload = 1 << 6,
+
 		/// <summary>
 		/// This directive is unrecognized and does not change the code generation at all
 		/// </summary>
@@ -98,6 +128,9 @@ namespace EpicGames.UHT.Parsers
 		}
 	}
 
+	/// <summary>
+	/// Specifiers for public, private, and protected
+	/// </summary>
 	[UnrealHeaderTool]
 	public static class UhtAccessSpecifierKeywords
 	{
@@ -132,6 +165,9 @@ namespace EpicGames.UHT.Parsers
 		}
 	}
 
+	/// <summary>
+	/// Header file parser
+	/// </summary>
 	public class UhtHeaderFileParser : IUhtTokenPreprocessor
 	{
 		private static UhtKeywordTable KeywordTable = UhtKeywordTables.Instance.Get(UhtTableNames.Global);
@@ -193,6 +229,11 @@ namespace EpicGames.UHT.Parsers
 		/// </summary>
 		private int StatementsParsed = 0;
 
+		/// <summary>
+		/// Parse the given header file
+		/// </summary>
+		/// <param name="HeaderFile">Header file to parse</param>
+		/// <returns>Parser</returns>
 		public static UhtHeaderFileParser Parse(UhtHeaderFile HeaderFile)
 		{
 			UhtHeaderFileParser HeaderParser = new UhtHeaderFileParser(HeaderFile);
@@ -224,13 +265,18 @@ namespace EpicGames.UHT.Parsers
 			return HeaderParser;
 		}
 
-		public UhtHeaderFileParser(UhtHeaderFile HeaderFile)
+		private UhtHeaderFileParser(UhtHeaderFile HeaderFile)
 		{
 			this.TokenReader = new UhtTokenBufferReader(HeaderFile, HeaderFile.Data.Memory);
 			this.HeaderFile = HeaderFile;
 			this.TokenReader.TokenPreprocessor = this;
 		}
 
+		/// <summary>
+		/// Push a new scope
+		/// </summary>
+		/// <param name="Scope">Scope to push</param>
+		/// <exception cref="UhtIceException">Throw if the new scope isn't parented by the current scope</exception>
 		public void PushScope(UhtParsingScope Scope)
 		{
 			if (Scope.ParentScope != this.TopScope)
@@ -240,6 +286,11 @@ namespace EpicGames.UHT.Parsers
 			this.TopScope = Scope;
 		}
 
+		/// <summary>
+		/// Pop the given scope
+		/// </summary>
+		/// <param name="Scope">Scope to be popped</param>
+		/// <exception cref="UhtIceException">Thrown if the given scope isn't the top scope</exception>
 		public void PopScope(UhtParsingScope Scope)
 		{
 			if (Scope != this.TopScope)
@@ -249,6 +300,13 @@ namespace EpicGames.UHT.Parsers
 			this.TopScope = Scope.ParentScope;
 		}
 
+		/// <summary>
+		/// Get the cached specifier parser
+		/// </summary>
+		/// <param name="SpecifierContext">Specifier context</param>
+		/// <param name="Context">User facing context</param>
+		/// <param name="Table">Specifier table</param>
+		/// <returns>Specifier parser</returns>
 		public UhtSpecifierParser GetSpecifierParser(UhtSpecifierContext SpecifierContext, StringView Context, UhtSpecifierTable Table)
 		{
 			if (this.SpecifierParser == null)
@@ -262,6 +320,11 @@ namespace EpicGames.UHT.Parsers
 			return this.SpecifierParser;
 		}
 
+		/// <summary>
+		/// Get the cached property parser
+		/// </summary>
+		/// <param name="TopScope">Top scope requesting parser</param>
+		/// <returns>Property parser</returns>
 		public UhtPropertyParser GetPropertyParser(UhtParsingScope TopScope)
 		{
 			if (this.PropertyParser == null)
@@ -284,12 +347,17 @@ namespace EpicGames.UHT.Parsers
 			return this.CompilerDirectives.Count > 0 ? this.CompilerDirectives[this.CompilerDirectives.Count - 1].Composite : UhtCompilerDirective.None;
 		}
 
+		/// <summary>
+		/// Get the current compiler directive without any parent scopes merged in
+		/// </summary>
+		/// <returns>Current compiler directive</returns>
 		public UhtCompilerDirective GetCurrentNonCompositeCompilerDirective()
 		{
 			return this.CompilerDirectives.Count > 0 ? this.CompilerDirectives[this.CompilerDirectives.Count - 1].Element : UhtCompilerDirective.None;
 		}
 
 		#region ITokenPreprocessor implementation
+		/// <inheritdoc/>
 		public bool ParsePreprocessorDirective(ref UhtToken Token, bool bIsBeingIncluded, out bool bClearComments, out bool bIllegalContentsCheck)
 		{
 			bClearComments = true;
@@ -302,12 +370,14 @@ namespace EpicGames.UHT.Parsers
 			return IncludeCurrentCompilerDirective();
 		}
 
+		/// <inheritdoc/>
 		public void SaveState()
 		{
 			this.SavedCompilerDirectives.Clear();
 			this.SavedCompilerDirectives.AddRange(this.CompilerDirectives);
 		}
 
+		/// <inheritdoc/>
 		public void RestoreState()
 		{
 			this.CompilerDirectives.Clear();
@@ -316,11 +386,21 @@ namespace EpicGames.UHT.Parsers
 		#endregion
 
 		#region Statement parsing
+
+		/// <summary>
+		/// Parse all statements in the header file
+		/// </summary>
 		public void ParseStatements()
 		{
 			ParseStatements((char)0, (char)0, true);
 		}
 
+		/// <summary>
+		/// Parse the statements between the given symbols
+		/// </summary>
+		/// <param name="Initiator">Starting symbol</param>
+		/// <param name="Terminator">Ending symbol</param>
+		/// <param name="bLogUnhandledKeywords">If true, log any unhandled keywords</param>
 		public void ParseStatements(char Initiator, char Terminator, bool bLogUnhandledKeywords)
 		{
 			if (this.TopScope == null)
@@ -357,7 +437,13 @@ namespace EpicGames.UHT.Parsers
 			}
 		}
 
-
+		/// <summary>
+		/// Parse a statement
+		/// </summary>
+		/// <param name="TopScope">Current top scope</param>
+		/// <param name="Token">Token starting the statement</param>
+		/// <param name="bLogUnhandledKeywords">If true, log unhandled keywords</param>
+		/// <returns>Always returns true ATM</returns>
 		public static bool ParseStatement(UhtParsingScope TopScope, ref UhtToken Token, bool bLogUnhandledKeywords)
 		{
 			UhtParseResult ParseResult = UhtParseResult.Unhandled;
@@ -462,6 +548,7 @@ namespace EpicGames.UHT.Parsers
 		/// <summary>
 		/// Tests if an identifier looks like a macro which doesn't have a following open parenthesis.
 		/// </summary>
+		/// <param name="TokenReader">Token reader</param>
 		/// <param name="Token">The current token that initiated the process</param>
 		/// <returns>Result if matching the token</returns>
 		private static UhtParseResult ProbablyAnUnknownObjectLikeMacro(IUhtTokenReader TokenReader, ref UhtToken Token)
diff --git a/Engine/Source/Programs/Shared/EpicGames.UHT/Parsers/UhtInterfaceClassParser.cs b/Engine/Source/Programs/Shared/EpicGames.UHT/Parsers/UhtInterfaceClassParser.cs
--- a/Engine/Source/Programs/Shared/EpicGames.UHT/Parsers/UhtInterfaceClassParser.cs
+++ b/Engine/Source/Programs/Shared/EpicGames.UHT/Parsers/UhtInterfaceClassParser.cs
@@ -8,16 +8,26 @@ using EpicGames.UHT.Utils;
 
 namespace EpicGames.UHT.Parsers
 {
+
+	/// <summary>
+	/// Interface class parser
+	/// </summary>
 	[UnrealHeaderTool]
 	public class UhtInterfaceClassParser : UhtClassBaseParser
 	{
 		private static UhtKeywordTable KeywordTable = UhtKeywordTables.Instance.Get(UhtTableNames.Interface);
 		private static UhtSpecifierTable SpecifierTable = UhtSpecifierTables.Instance.Get(UhtTableNames.Interface);
 
+		/// <summary>
+		/// Construct a new interface class parser
+		/// </summary>
+		/// <param name="Outer">Outer object</param>
+		/// <param name="LineNumber">Line number</param>
 		public UhtInterfaceClassParser(UhtType Outer, int LineNumber) : base(Outer, LineNumber)
 		{
 		}
 
+		/// <inheritdoc/>
 		protected override void ResolveSuper(UhtResolvePhase ResolvePhase)
 		{
 			base.ResolveSuper( ResolvePhase);
@@ -42,6 +52,7 @@ namespace EpicGames.UHT.Parsers
 			}
 		}
 
+		/// <inheritdoc/>
 		protected override bool ResolveSelf(UhtResolvePhase ResolvePhase)
 		{
 			bool bResult = base.ResolveSelf(ResolvePhase);
@@ -93,7 +104,7 @@ namespace EpicGames.UHT.Parsers
 		}
 		#endregion
 
-		public static UhtParseResult ParseUInterface(UhtParsingScope ParentScope, ref UhtToken Token)
+		private static UhtParseResult ParseUInterface(UhtParsingScope ParentScope, ref UhtToken Token)
 		{
 			UhtInterfaceClassParser Class = new UhtInterfaceClassParser(ParentScope.ScopeType, Token.InputLine);
 			Class.ClassType = UhtClassType.Interface;
diff --git a/Engine/Source/Programs/Shared/EpicGames.UHT/Parsers/UhtNativeInterfaceClassParser.cs b/Engine/Source/Programs/Shared/EpicGames.UHT/Parsers/UhtNativeInterfaceClassParser.cs
--- a/Engine/Source/Programs/Shared/EpicGames.UHT/Parsers/UhtNativeInterfaceClassParser.cs
+++ b/Engine/Source/Programs/Shared/EpicGames.UHT/Parsers/UhtNativeInterfaceClassParser.cs
@@ -8,16 +8,26 @@ using EpicGames.UHT.Utils;
 
 namespace EpicGames.UHT.Parsers
 {
+
+	/// <summary>
+	/// Parser object for native interfaces
+	/// </summary>
 	[UnrealHeaderTool]
 	public class UhtParserNativeInterfaceClass : UhtClassBaseParser
 	{
 		private static UhtKeywordTable KeywordTable = UhtKeywordTables.Instance.Get(UhtTableNames.NativeInterface);
 		private static UhtSpecifierTable SpecifierTable = UhtSpecifierTables.Instance.Get(UhtTableNames.NativeInterface);
 
+		/// <summary>
+		/// Construct a new native interface
+		/// </summary>
+		/// <param name="Outer">Outer object</param>
+		/// <param name="LineNumber">Line number of declaration</param>
 		public UhtParserNativeInterfaceClass(UhtType Outer, int LineNumber) : base(Outer, LineNumber)
 		{
 		}
 
+		/// <inheritdoc/>
 		protected override bool ResolveSelf(UhtResolvePhase ResolvePhase)
 		{
 			bool bResult = base.ResolveSelf(ResolvePhase);
diff --git a/Engine/Source/Programs/Shared/EpicGames.UHT/Parsers/UhtParserHelpers.cs b/Engine/Source/Programs/Shared/EpicGames.UHT/Parsers/UhtParserHelpers.cs
--- a/Engine/Source/Programs/Shared/EpicGames.UHT/Parsers/UhtParserHelpers.cs
+++ b/Engine/Source/Programs/Shared/EpicGames.UHT/Parsers/UhtParserHelpers.cs
@@ -9,8 +9,18 @@ using System.Collections.Generic;
 
 namespace EpicGames.UHT.Parsers
 {
+	/// <summary>
+	/// Collection of helper methods
+	/// </summary>
 	public static class UhtParserHelpers
 	{
+
+		/// <summary>
+		/// Parse the inheritance 
+		/// </summary>
+		/// <param name="TokenReader">Token reader</param>
+		/// <param name="SuperIdentifier">Output super identifier</param>
+		/// <param name="BaseIdentifiers">Output base identifiers</param>
 		public static void ParseInheritance(IUhtTokenReader TokenReader, out UhtToken SuperIdentifier, out List<UhtToken[]>? BaseIdentifiers)
 		{
 			UhtToken SuperIdentifierTemp = new UhtToken();
@@ -33,6 +43,11 @@ namespace EpicGames.UHT.Parsers
 			BaseIdentifiers = BaseIdentifiersTemp;
 		}
 
+		/// <summary>
+		/// Parse compiler version declaration
+		/// </summary>
+		/// <param name="TokenReader">Token reader</param>
+		/// <param name="Struct">Struct being parsed</param>
 		public static void ParseCompileVersionDeclaration(IUhtTokenReader TokenReader, UhtStruct Struct)
 		{
 
diff --git a/Engine/Source/Programs/Shared/EpicGames.UHT/Parsers/UhtParsingScope.cs b/Engine/Source/Programs/Shared/EpicGames.UHT/Parsers/UhtParsingScope.cs
--- a/Engine/Source/Programs/Shared/EpicGames.UHT/Parsers/UhtParsingScope.cs
+++ b/Engine/Source/Programs/Shared/EpicGames.UHT/Parsers/UhtParsingScope.cs
@@ -10,13 +10,41 @@ using System.Text;
 
 namespace EpicGames.UHT.Parsers
 {
+
+	/// <summary>
+	/// Nested structure of scopes being parsed
+	/// </summary>
 	public class UhtParsingScope : IDisposable
 	{
+
+		/// <summary>
+		/// Header file parser
+		/// </summary>
 		public readonly UhtHeaderFileParser HeaderParser;
+
+		/// <summary>
+		/// Token reader
+		/// </summary>
 		public readonly IUhtTokenReader TokenReader;
+
+		/// <summary>
+		/// Parent scope
+		/// </summary>
 		public readonly UhtParsingScope? ParentScope;
+
+		/// <summary>
+		/// Type being parsed
+		/// </summary>
 		public readonly UhtType ScopeType;
+
+		/// <summary>
+		/// Keyword table for the scope
+		/// </summary>
 		public readonly UhtKeywordTable ScopeKeywordTable;
+
+		/// <summary>
+		/// Current access specifier
+		/// </summary>
 		public UhtAccessSpecifier AccessSpecifier = UhtAccessSpecifier.Public;
 
 		/// <summary>
@@ -44,6 +72,12 @@ namespace EpicGames.UHT.Parsers
 		/// </summary>
 		public UhtClass CurrentClass => (UhtClass)this.CurrentClassScope.ScopeType;
 
+		/// <summary>
+		/// Construct a root/global scope
+		/// </summary>
+		/// <param name="HeaderParser">Header parser</param>
+		/// <param name="ScopeType">Type being parsed</param>
+		/// <param name="KeywordTable">Keyword table</param>
 		public UhtParsingScope(UhtHeaderFileParser HeaderParser, UhtType ScopeType, UhtKeywordTable KeywordTable)
 		{
 			this.HeaderParser = HeaderParser;
@@ -54,6 +88,13 @@ namespace EpicGames.UHT.Parsers
 			this.HeaderParser.PushScope(this);
 		}
 
+		/// <summary>
+		/// Construct a scope for a type
+		/// </summary>
+		/// <param name="ParentScope">Parent scope</param>
+		/// <param name="ScopeType">Type being parsed</param>
+		/// <param name="KeywordTable">Keyword table</param>
+		/// <param name="AccessSpecifier">Current access specifier</param>
 		public UhtParsingScope(UhtParsingScope ParentScope, UhtType ScopeType, UhtKeywordTable KeywordTable, UhtAccessSpecifier AccessSpecifier)
 		{
 			this.HeaderParser = ParentScope.HeaderParser;
@@ -65,6 +106,9 @@ namespace EpicGames.UHT.Parsers
 			this.HeaderParser.PushScope(this);
 		}
 
+		/// <summary>
+		/// Dispose the scope
+		/// </summary>
 		public void Dispose()
 		{
 			this.HeaderParser.PopScope(this);
@@ -432,12 +476,20 @@ namespace EpicGames.UHT.Parsers
 		}
 	}
 
+	/// <summary>
+	/// Token recorder
+	/// </summary>
 	public struct UhtTokenRecorder : IDisposable
 	{
 		private readonly UhtCompilerDirective CompilerDirective;
 		private readonly UhtParsingScope Scope;
 		private bool bFlushed;
 
+		/// <summary>
+		/// Construct a new recorder
+		/// </summary>
+		/// <param name="Scope">Scope being parsed</param>
+		/// <param name="InitialToken">Initial toke nto add to the recorder</param>
 		public UhtTokenRecorder(UhtParsingScope Scope, ref UhtToken InitialToken)
 		{
 			this.Scope = Scope;
@@ -445,12 +497,15 @@ namespace EpicGames.UHT.Parsers
 			this.bFlushed = false;
 			if (this.Scope.ScopeType is UhtClass)
 			{
-
 				this.Scope.TokenReader.EnableRecording();
 				this.Scope.TokenReader.RecordToken(ref InitialToken);
 			}
 		}
 
+		/// <summary>
+		/// Create a new recorder
+		/// </summary>
+		/// <param name="Scope">Scope being parsed</param>
 		public UhtTokenRecorder(UhtParsingScope Scope)
 		{
 			this.Scope = Scope;
@@ -462,11 +517,18 @@ namespace EpicGames.UHT.Parsers
 			}
 		}
 
+		/// <summary>
+		/// Stop the recording
+		/// </summary>
 		public void Dispose()
 		{
 			Stop();
 		}
 
+		/// <summary>
+		/// Stop the recording
+		/// </summary>
+		/// <returns>True if the recorded content was added to a class</returns>
 		public bool Stop()
 		{
 			if (!bFlushed)
diff --git a/Engine/Source/Programs/Shared/EpicGames.UHT/Parsers/UhtPropertyParser.cs b/Engine/Source/Programs/Shared/EpicGames.UHT/Parsers/UhtPropertyParser.cs
--- a/Engine/Source/Programs/Shared/EpicGames.UHT/Parsers/UhtPropertyParser.cs
+++ b/Engine/Source/Programs/Shared/EpicGames.UHT/Parsers/UhtPropertyParser.cs
@@ -19,6 +19,10 @@ namespace EpicGames.UHT.Parsers
 	[Flags]
 	public enum UhtPropertyParseOptions
 	{
+
+		/// <summary>
+		/// No options
+		/// </summary>
 		None = 0,
 
 		/// <summary>
@@ -121,20 +125,64 @@ namespace EpicGames.UHT.Parsers
 		UPROPERTY
 	};
 
+	/// <summary>
+	/// Layout macro type
+	/// </summary>
 	public enum UhtLayoutMacroType
 	{
+
+		/// <summary>
+		/// None found
+		/// </summary>
 		None,
+		
+		/// <summary>
+		/// Array
+		/// </summary>
 		Array,
+
+		/// <summary>
+		/// Editor only array
+		/// </summary>
 		ArrayEditorOnly,
+
+		/// <summary>
+		/// Bit field
+		/// </summary>
 		Bitfield,
+
+		/// <summary>
+		/// Editor only bit field
+		/// </summary>
 		BitfieldEditorOnly,
+
+		/// <summary>
+		/// Field
+		/// </summary>
 		Field,
+
+		/// <summary>
+		/// Editor only field
+		/// </summary>
 		FieldEditorOnly,
+
+		/// <summary>
+		/// Field with initializer
+		/// </summary>
 		FieldInitialized,
 	}
 
+	/// <summary>
+	/// Extensions for working with the layout macro type
+	/// </summary>
 	public static class UhtLayoutMacroTypeExtensions
 	{
+
+		/// <summary>
+		/// Return true if the type is editor only
+		/// </summary>
+		/// <param name="LayoutMacroType">Layout macro type</param>
+		/// <returns>True if editor only</returns>
 		public static bool IsEditorOnly(this UhtLayoutMacroType LayoutMacroType)
 		{
 			switch (LayoutMacroType)
@@ -149,6 +197,11 @@ namespace EpicGames.UHT.Parsers
 			}
 		}
 
+		/// <summary>
+		/// Return true if the type is a bit field
+		/// </summary>
+		/// <param name="LayoutMacroType">Layout macro type</param>
+		/// <returns>True if bit field</returns>
 		public static bool IsBitfield(this UhtLayoutMacroType LayoutMacroType)
 		{
 			switch (LayoutMacroType)
@@ -162,6 +215,11 @@ namespace EpicGames.UHT.Parsers
 			}
 		}
 
+		/// <summary>
+		/// Return true if the type is an array
+		/// </summary>
+		/// <param name="LayoutMacroType">Layout macro type</param>
+		/// <returns>True if array</returns>
 		public static bool IsArray(this UhtLayoutMacroType LayoutMacroType)
 		{
 			switch (LayoutMacroType)
@@ -175,6 +233,11 @@ namespace EpicGames.UHT.Parsers
 			}
 		}
 
+		/// <summary>
+		/// Return true if the type has an initializer
+		/// </summary>
+		/// <param name="LayoutMacroType">Layout macro type</param>
+		/// <returns>True if it has an initializer</returns>
 		public static bool HasInitializer(this UhtLayoutMacroType LayoutMacroType)
 		{
 			switch (LayoutMacroType)
@@ -187,6 +250,12 @@ namespace EpicGames.UHT.Parsers
 			}
 		}
 
+		/// <summary>
+		/// Return the layout macro name
+		/// </summary>
+		/// <param name="LayoutMacroType">Type in question</param>
+		/// <returns>Macro name</returns>
+		/// <exception cref="UhtIceException">Thrown if the macro type is none or invalid</exception>
 		public static StringView MacroName(this UhtLayoutMacroType LayoutMacroType)
 		{
 			switch (LayoutMacroType)
@@ -218,22 +287,61 @@ namespace EpicGames.UHT.Parsers
 			}
 		}
 
+		/// <summary>
+		/// Return the macro name and value
+		/// </summary>
+		/// <param name="LayoutMacroType">Macro name</param>
+		/// <returns>Name and type</returns>
 		public static KeyValuePair<StringView, UhtLayoutMacroType> MacroNameAndValue(this UhtLayoutMacroType LayoutMacroType)
 		{
 			return new KeyValuePair<StringView, UhtLayoutMacroType>(LayoutMacroType.MacroName(), LayoutMacroType);
 		}
 	}
 
+	/// <summary>
+	/// Delegate invoked to handle a parsed property 
+	/// </summary>
+	/// <param name="TopScope">Scope being parsed</param>
+	/// <param name="Property">Property just parsed</param>
+	/// <param name="NameToken">Name of the property</param>
+	/// <param name="LayoutMacroType">Layout macro type</param>
 	public delegate void UhtPropertyDelegate(UhtParsingScope TopScope, UhtProperty Property, ref UhtToken NameToken, UhtLayoutMacroType LayoutMacroType);
 
+	/// <summary>
+	/// Context for property specifier parsing
+	/// </summary>
 	public class UhtPropertySpecifierContext : UhtSpecifierContext
 	{
+		/// <summary>
+		/// The property settings being parsed
+		/// </summary>
 		public UhtPropertySettings PropertySettings;
+
+		/// <summary>
+		/// If true, editor specifier seen
+		/// </summary>
 		public bool bSeenEditSpecifier = false;
+
+		/// <summary>
+		/// If true, blueprint write specifier seen
+		/// </summary>
 		public bool bSeenBlueprintWriteSpecifier = false;
+
+		/// <summary>
+		/// If true, blueprint readonly specifier seen
+		/// </summary>
 		public bool bSeenBlueprintReadOnlySpecifier = false;
+
+		/// <summary>
+		/// If true, blueprint getter specifier seen
+		/// </summary>
 		public bool bSeenBlueprintGetterSpecifier = false;
 
+		/// <summary>
+		/// Construct a new property specifier context
+		/// </summary>
+		/// <param name="Scope">Scope being parsed</param>
+		/// <param name="MessageSite">Message site</param>
 		public UhtPropertySpecifierContext(UhtParsingScope Scope, IUhtMessageSite MessageSite) : base(Scope, MessageSite, UhtMetaData.Empty)
 		{
 			this.PropertySettings = new UhtPropertySettings();
@@ -256,9 +364,21 @@ namespace EpicGames.UHT.Parsers
 		/// <inheritdoc/>
 		protected override string PGetMacroText { get => "invalid"; }
 
+		/// <summary>
+		/// Collection of type tokens 
+		/// </summary>
 		public ReadOnlyMemory<UhtToken> TypeTokens;
+
+		/// <summary>
+		/// Property settings being parsed
+		/// </summary>
 		public UhtPropertySettings PropertySettings;
 
+		/// <summary>
+		/// Construct a new property to be resolved
+		/// </summary>
+		/// <param name="PropertySettings">Property settings</param>
+		/// <param name="TypeTokens">Type tokens</param>
 		public UhtPreResolveProperty(UhtPropertySettings PropertySettings, ReadOnlyMemory<UhtToken> TypeTokens) : base(PropertySettings)
 		{
 			this.TypeTokens = TypeTokens;
@@ -314,6 +434,9 @@ namespace EpicGames.UHT.Parsers
 		}
 	}
 
+	/// <summary>
+	/// Property parser
+	/// </summary>
 	public class UhtPropertyParser : IUhtMessageExtraContext
 	{
 		private UhtPropertySpecifierContext SpecifierContext;
@@ -342,18 +465,37 @@ namespace EpicGames.UHT.Parsers
 			UhtLayoutMacroType.FieldInitialized.MacroNameAndValue(),
 		});
 
+		/// <summary>
+		/// Construct a new property parser
+		/// </summary>
+		/// <param name="Scope">Current scope being parsed</param>
+		/// <param name="MessageSite">Message site</param>
 		public UhtPropertyParser(UhtParsingScope Scope, IUhtMessageSite MessageSite)
 		{
 			this.SpecifierContext = new UhtPropertySpecifierContext(Scope, MessageSite);
 			this.GatherTypeTokensDelegate = GatherTypeTokens;
 		}
 
+		/// <summary>
+		/// Reset the property parser for a new parse
+		/// </summary>
+		/// <param name="Scope">Current scope being parsed</param>
+		/// <param name="MessageSite">Message site</param>
 		public void Reset(UhtParsingScope Scope, IUhtMessageSite MessageSite)
 		{
 			this.SpecifierContext.Scope = Scope;
 			this.SpecifierContext.MessageSite = MessageSite;
 		}
 
+		/// <summary>
+		/// Parse the property
+		/// </summary>
+		/// <param name="DisallowPropertyFlags">Flags to be disallowed</param>
+		/// <param name="Options">Parsing options</param>
+		/// <param name="DeclarationStyle">Style of declaration</param>
+		/// <param name="Category">Property category</param>
+		/// <param name="Delegate">Delegate to be invoked after property has been parsed</param>
+		/// <returns>The property parser</returns>
 		public UhtPropertyParser Parse(EPropertyFlags DisallowPropertyFlags, UhtPropertyParseOptions Options, UhtParsePropertyDeclarationStyle DeclarationStyle, UhtPropertyCategory Category, UhtPropertyDelegate Delegate)
 		{
 			// Reset the context and create the property
@@ -382,6 +524,7 @@ namespace EpicGames.UHT.Parsers
 		}
 
 		#region IMessageExtraContext implementation
+		/// <inheritdoc/>
 		public IEnumerable<object?>? MessageExtraContext
 		{
 			get
@@ -479,6 +622,14 @@ namespace EpicGames.UHT.Parsers
 			}
 		}
 
+		/// <summary>
+		/// Parse a template parameter
+		/// </summary>
+		/// <param name="ResolvePhase">Resolution phase</param>
+		/// <param name="ParentPropertySettings">Parent property (container) settings</param>
+		/// <param name="ParamName">Name of the template parameter</param>
+		/// <param name="TokenReader">Token type</param>
+		/// <returns>Parsed property</returns>
 		public static UhtProperty? ParseTemplateParam(UhtPropertyResolvePhase ResolvePhase, UhtPropertySettings ParentPropertySettings, StringView ParamName, IUhtTokenReader TokenReader)
 		{
 			UhtPropertyTypeTable PropertyTypeTable = UhtPropertyTypeTable.Instance;
@@ -1041,9 +1192,9 @@ namespace EpicGames.UHT.Parsers
 	}
 
 	[UnrealHeaderTool]
-	public static class PropertyKeywords
+	static class PropertyKeywords
 	{
-#region Keywords
+		#region Keywords
 		[UhtKeyword(Extends = UhtTableNames.Class)]
 		[UhtKeyword(Extends = UhtTableNames.ScriptStruct)]
 		private static UhtParseResult UPROPERTYKeyword(UhtParsingScope TopScope, UhtParsingScope ActionScope, ref UhtToken Token)
@@ -1053,7 +1204,7 @@ namespace EpicGames.UHT.Parsers
 			TopScope.TokenReader.Require(';');
 			return UhtParseResult.Handled;
 		}
-#endregion
+		#endregion
 
 		private static UhtPropertyDelegate PropertyDelegate = PropertyParsed;
 
@@ -1075,12 +1226,10 @@ namespace EpicGames.UHT.Parsers
 	}
 
 	[UnrealHeaderTool]
-	public static class UhtDefaultPropertyParser
+	static class UhtDefaultPropertyParser
 	{
 		[UhtPropertyType(Options = UhtPropertyTypeOptions.Default)]
-
-
-		public static UhtProperty? DefaultProperty(UhtPropertyResolvePhase ResolvePhase, UhtPropertySettings PropertySettings, IUhtTokenReader TokenReader, UhtToken MatchedToken)
+		private static UhtProperty? DefaultProperty(UhtPropertyResolvePhase ResolvePhase, UhtPropertySettings PropertySettings, IUhtTokenReader TokenReader, UhtToken MatchedToken)
 		{
 			int TypeStartPos = TokenReader.PeekToken().InputStartPos;
 
diff --git a/Engine/Source/Programs/Shared/EpicGames.UHT/Parsers/UhtScriptStructParser.cs b/Engine/Source/Programs/Shared/EpicGames.UHT/Parsers/UhtScriptStructParser.cs
--- a/Engine/Source/Programs/Shared/EpicGames.UHT/Parsers/UhtScriptStructParser.cs
+++ b/Engine/Source/Programs/Shared/EpicGames.UHT/Parsers/UhtScriptStructParser.cs
@@ -9,19 +9,36 @@ using System.Collections.Generic;
 
 namespace EpicGames.UHT.Parsers
 {
+
+	/// <summary>
+	/// USTRUCT parser object
+	/// </summary>
 	[UnrealHeaderTool]
 	public class UhtScriptStructParser : UhtScriptStruct
 	{
 		private static UhtKeywordTable KeywordTable = UhtKeywordTables.Instance.Get(UhtTableNames.ScriptStruct);
 		private static UhtSpecifierTable SpecifierTable = UhtSpecifierTables.Instance.Get(UhtTableNames.ScriptStruct);
 
+		/// <summary>
+		/// Super identifier
+		/// </summary>
 		public UhtToken SuperIdentifier;
+
+		/// <summary>
+		/// Base identifiers
+		/// </summary>
 		public List<UhtToken[]>? BaseIdentifiers = null;
 
+		/// <summary>
+		/// Construct a new instance
+		/// </summary>
+		/// <param name="Outer">Output type</param>
+		/// <param name="LineNumber">Line number of declaration</param>
 		public UhtScriptStructParser(UhtType Outer, int LineNumber) : base(Outer, LineNumber)
 		{
 		}
 
+		/// <inheritdoc/>
 		protected override void ResolveSuper(UhtResolvePhase ResolvePhase)
 		{
 			base.ResolveSuper(ResolvePhase);
@@ -41,6 +58,7 @@ namespace EpicGames.UHT.Parsers
 			}
 		}
 
+		/// <inheritdoc/>
 		protected override bool ResolveSelf(UhtResolvePhase ResolvePhase)
 		{
 			bool bResult = base.ResolveSelf(ResolvePhase);
diff --git a/Engine/Source/Programs/Shared/EpicGames.UHT/Parsers/UhtSpecifierParser.cs b/Engine/Source/Programs/Shared/EpicGames.UHT/Parsers/UhtSpecifierParser.cs
--- a/Engine/Source/Programs/Shared/EpicGames.UHT/Parsers/UhtSpecifierParser.cs
+++ b/Engine/Source/Programs/Shared/EpicGames.UHT/Parsers/UhtSpecifierParser.cs
@@ -41,6 +41,12 @@ namespace EpicGames.UHT.Parsers
 		private readonly Action ParseKVPValueAction;
 		private readonly Action ParseStringViewListAction;
 
+		/// <summary>
+		/// Construct a new specifier parser
+		/// </summary>
+		/// <param name="SpecifierContext">Specifier context</param>
+		/// <param name="Context">User facing context added to messages</param>
+		/// <param name="Table">Specifier table</param>
 		public UhtSpecifierParser(UhtSpecifierContext SpecifierContext, StringView Context, UhtSpecifierTable Table)
 		{
 			this.SpecifierContext = SpecifierContext;
@@ -68,6 +74,12 @@ namespace EpicGames.UHT.Parsers
 			};
 		}
 
+		/// <summary>
+		/// Reset an existing parser to parse a new specifier block
+		/// </summary>
+		/// <param name="SpecifierContext">Specifier context</param>
+		/// <param name="Context">User facing context added to messages</param>
+		/// <param name="Table">Specifier table</param>
 		public void Reset(UhtSpecifierContext SpecifierContext, StringView Context, UhtSpecifierTable Table)
 		{
 			this.SpecifierContext = SpecifierContext;
@@ -76,6 +88,10 @@ namespace EpicGames.UHT.Parsers
 			this.Table = Table;
 		}
 
+		/// <summary>
+		/// Perform the specify parsing
+		/// </summary>
+		/// <returns>The parser</returns>
 		public UhtSpecifierParser ParseSpecifiers()
 		{
 			this.bIsParsingFieldMetaData = false;
@@ -88,6 +104,11 @@ namespace EpicGames.UHT.Parsers
 			return this;
 		}
 
+		/// <summary>
+		/// Parse field meta data
+		/// </summary>
+		/// <param name="Context">User facing context for error messages</param>
+		/// <returns>Specifier parser</returns>
 		public UhtSpecifierParser ParseFieldMetaData(StringView Context)
 		{
 			this.TokenReader = this.SpecifierContext.Scope.TokenReader;
@@ -108,6 +129,9 @@ namespace EpicGames.UHT.Parsers
 			return this;
 		}
 
+		/// <summary>
+		/// Parse any deffered specifiers
+		/// </summary>
 		public void ParseDeferred()
 		{
 			foreach (var Deferred in this.DeferredSpecifiers)
@@ -118,6 +142,8 @@ namespace EpicGames.UHT.Parsers
 		}
 
 		#region IMessageExtraContext implementation
+
+		/// <inheritdoc/>
 		public IEnumerable<object?>? MessageExtraContext
 		{
 			get
