Horde: More bundle fixes.

- Fixed tracking of parent ownership for child nodes when adding new nodes to the tree.
- Fixed incorrect offsets in exports when multiple items are in a single compression packet.
- Document and enforce invariants when switching nodes between different states. Nodes read from refs are now immediately turned into standalone nodes. Removes unnecessary logic for caching/decompressing nodes in the root ref on demand.

diff --git a/Engine/Source/Programs/Horde/Horde.Build/Commands/Perforce/ReplicateCommand.cs b/Engine/Source/Programs/Horde/Horde.Build/Commands/Perforce/ReplicateCommand.cs
--- a/Engine/Source/Programs/Horde/Horde.Build/Commands/Perforce/ReplicateCommand.cs
+++ b/Engine/Source/Programs/Horde/Horde.Build/Commands/Perforce/ReplicateCommand.cs
@@ -40,6 +40,12 @@ namespace HordeServer.Commands
 		[CommandLine]
 		public bool Content { get; set; }
 
+		[CommandLine]
+		public bool Compact { get; set; } = true;
+
+		[CommandLine]
+		public bool Clean { get; set; }
+
 		[CommandLine]
 		public string Filter { get; set; } = "//...";
 
@@ -81,34 +87,50 @@ namespace HordeServer.Commands
 			Dictionary<IStream, int> StreamToFirstChange = new Dictionary<IStream, int>();
 			StreamToFirstChange[Stream] = Change;
 
-			CommitTree? BaseTree = null; 
+			CommitTree? BaseTree = null;
 			if (BaseChange != 0)
 			{
-				ICommit? Commit = await CommitCollection.GetCommitAsync(Stream.Id, BaseChange);
-				if (Commit == null)
-				{
-					throw new InvalidOperationException($"Unable to find existing commit for {StreamId} @ CL {BaseChange}");
-				}
-				if (Commit.TreeRefId == null)
-				{
-					throw new InvalidOperationException($"No tree for commit {StreamId} @ CL {BaseChange}");
-				}
-				BaseTree = new CommitTree(Commit.StreamId, Commit.Change, Commit.TreeRefId.Value);
+				BaseTree = await GetBaseTreeAsync(CommitCollection, Stream.Id, Change);
 			}
 
 			await foreach (NewCommit NewCommit in CommitService.FindCommitsForClusterAsync(Stream.ClusterName, StreamToFirstChange).Take(Count))
 			{
-				string BriefSummary = NewCommit.Description.Replace('\n', ' ').Substring(0, 50);
-				Logger.LogInformation("Commit {Change} by {AuthorId}: {Summary}", NewCommit.Change, NewCommit.AuthorId, BriefSummary);
+				string BriefSummary = NewCommit.Description.Replace('\n', ' ');
+				Logger.LogInformation("Commit {Change} by {AuthorId}: {Summary}", NewCommit.Change, NewCommit.AuthorId, BriefSummary.Substring(0, Math.Min(50, BriefSummary.Length)));
 				Logger.LogInformation(" - Base path: {BasePath}", NewCommit.BasePath);
 
 				if (Content)
 				{
-					BaseTree = await CommitService.FindCommitTreeAsync(Stream, NewCommit.Change, BaseTree, Filter);
+					if (BaseTree == null && !Clean)
+					{
+						BaseTree = await GetBaseTreeAsync(CommitCollection, Stream.Id, NewCommit.Change);
+					}
+					else
+					{
+						BaseTree = await CommitService.FindCommitTreeAsync(Stream, NewCommit.Change, BaseTree, Filter);
+					}
+
+					NewCommit.TreeRefId = BaseTree.RefId;
 				}
+
+				await CommitCollection.AddOrReplaceAsync(NewCommit);
 			}
 
 			return 0;
 		}
+
+		static async Task<CommitTree> GetBaseTreeAsync(ICommitCollection CommitCollection, StreamId StreamId, int Change)
+		{
+			ICommit? Commit = await CommitCollection.GetCommitAsync(StreamId, Change);
+			if (Commit == null)
+			{
+				throw new InvalidOperationException($"Unable to find existing commit for {StreamId} @ CL {Change}");
+			}
+			if (Commit.TreeRefId == null)
+			{
+				throw new InvalidOperationException($"No tree for commit {StreamId} @ CL {Change}");
+			}
+			return new CommitTree(Commit.StreamId, Commit.Change, Commit.TreeRefId.Value);
+		}
 	}
 }
diff --git a/Engine/Source/Programs/Shared/EpicGames.Horde/Bundles/Bundle.cs b/Engine/Source/Programs/Shared/EpicGames.Horde/Bundles/Bundle.cs
--- a/Engine/Source/Programs/Shared/EpicGames.Horde/Bundles/Bundle.cs
+++ b/Engine/Source/Programs/Shared/EpicGames.Horde/Bundles/Bundle.cs
@@ -62,7 +62,17 @@ namespace EpicGames.Horde.Bundles
 	public class Bundle
 	{
 		/// <summary>
-		/// Information about a node within a bundle.
+		/// Information about a node within a bundle. Nodes may be in the following states:
+		/// 
+		/// * Imported
+		///   Node is known to be in located a particular hashed blob in storage, but we haven't yet read the blob index to figure out exactly where.
+		///   
+		/// * Exported
+		///   Node is known to be located within a particular blob in storage, and the <see cref="NodeInfo.Export"/> member identifies specifically where.
+		///   
+		/// * Standalone
+		///   Node exists in memory, but has not been written to storage yet.
+		/// 
 		/// </summary>
 		[DebuggerDisplay("{Hash}")]
 		sealed class NodeInfo
@@ -71,10 +81,10 @@ namespace EpicGames.Horde.Bundles
 			public int Rank { get; }
 			public int Length { get; }
 
-			public BlobInfo? Blob;
-			public BundleExport? Export;
-			public ReadOnlySequence<byte> Data;
-			public NodeInfo[]? References;
+			public BlobInfo? Blob { get; private set; }
+			public BundleExport? Export { get; private set; }
+			public ReadOnlySequence<byte> Data { get; private set; }
+			public NodeInfo[]? References { get; private set; }
 
 			public NodeInfo(IoHash Hash, int Rank, int Length)
 			{
@@ -103,7 +113,7 @@ namespace EpicGames.Horde.Bundles
 			}
 
 			public void SetImport(BlobInfo Blob) => Set(Blob, null, ReadOnlySequence<byte>.Empty, null);
-			public void SetExport(BlobInfo? Blob, BundleExport Export, ReadOnlySequence<byte> Data, NodeInfo[] References) => Set(Blob, Export, Data, References);
+			public void SetExport(BlobInfo Blob, BundleExport Export, NodeInfo[] References) => Set(Blob, Export, ReadOnlySequence<byte>.Empty, References);
 			public void SetStandalone(ReadOnlySequence<byte> Data, NodeInfo[] References) => Set(null, null, Data, References);
 		}
 
@@ -129,8 +139,6 @@ namespace EpicGames.Horde.Bundles
 		readonly IStorageClient StorageClient;
 		readonly NamespaceId NamespaceId;
 
-		object RootCacheKey = new object();
-
 		/// <summary>
 		/// Reference to the root node in the bundle
 		/// </summary>
@@ -176,7 +184,7 @@ namespace EpicGames.Horde.Bundles
 		protected Bundle(IStorageClient StorageClient, NamespaceId NamespaceId, BundleObject Object, BundleNode Root, BundleOptions Options, IMemoryCache? Cache)
 			: this(StorageClient, NamespaceId, Root, Options, Cache)
 		{
-			RegisterObject(null, Object, Object.Data);
+			RegisterRootObject(Object);
 		}
 
 		/// <summary>
@@ -216,12 +224,12 @@ namespace EpicGames.Horde.Bundles
 			Bundle<T> NewBundle = new Bundle<T>(StorageClient, NamespaceId, null!, Options, Cache);
 
 			BundleObject Object = await StorageClient.GetRefAsync<BundleObject>(NamespaceId, BucketId, RefId);
-			NewBundle.RegisterObject(null, Object, Object.Data);
+			NewBundle.RegisterRootObject(Object);
 
 			IoHash RootHash = Object.Exports[Object.Exports.Count - 1].Hash;
-			BundleNodeRef<T> RootRef = new BundleNodeRef<T>(RootHash);
+			ReadOnlySequence<byte> Data = await NewBundle.GetDataAsync(RootHash);
+			((Bundle)NewBundle).Root = BundleNode.Deserialize<T>(Data);
 
-			((Bundle)NewBundle).Root = await NewBundle.GetAsync(RootRef);
 			return NewBundle;
 		}
 
@@ -232,14 +240,16 @@ namespace EpicGames.Horde.Bundles
 		{
 			// Find all the nodes that need to be flushed, in order
 			List<BundleNodeRef> NodeRefs = new List<BundleNodeRef>();
-			LinkWorkingSet(Root, null, NodeRefs);
+			LinkWorkingSet(Root, NodeRefs);
+
+			// Find the working set of mutable nodes to retain in memory
+			HashSet<BundleNodeRef> WorkingSet = new HashSet<BundleNodeRef>(NodeRefs.Where(x => !x.Node!.IsReadOnly()).OrderByDescending(x => x.LastModifiedTime).Take(IgnoreCount));
+			NodeRefs.RemoveAll(x => WorkingSet.Contains(x));
 			NodeRefs.SortBy(x => x.LastModifiedTime);
 
 			// Write them all to storage
-			for (int Idx = 0; Idx + IgnoreCount < NodeRefs.Count; Idx++)
+			foreach (BundleNodeRef NodeRef in NodeRefs)
 			{
-				BundleNodeRef NodeRef = NodeRefs[Idx];
-
 				BundleNode? Node = NodeRef.Node;
 				Debug.Assert(Node != null);
 
@@ -252,22 +262,17 @@ namespace EpicGames.Horde.Bundles
 		/// Traverse a tree of <see cref="BundleNodeRef"/> objects and update the incoming reference to each node
 		/// </summary>
 		/// <param name="Node">The reference to update</param>
-		/// <param name="IncomingRef">The parent reference</param>
 		/// <param name="NodeRefs">List to accumulate the list of all modified references</param>
-		private void LinkWorkingSet(BundleNode Node, BundleNodeRef? IncomingRef, List<BundleNodeRef> NodeRefs)
+		private void LinkWorkingSet(BundleNode Node, List<BundleNodeRef> NodeRefs)
 		{
-			Node.IncomingRef = IncomingRef;
-
 			foreach (BundleNodeRef ChildRef in Node.GetReferences())
 			{
-				ChildRef.ParentRef ??= IncomingRef;
-
 				if (ChildRef.Node != null && ChildRef.LastModifiedTime != 0)
 				{
-					LinkWorkingSet(ChildRef.Node, ChildRef, NodeRefs);
-					if (IncomingRef != null)
+					LinkWorkingSet(ChildRef.Node, NodeRefs);
+					if (Node.IncomingRef != null)
 					{
-						IncomingRef.LastModifiedTime = Math.Max(IncomingRef.LastModifiedTime, ChildRef.LastModifiedTime + 1);
+						Node.IncomingRef.LastModifiedTime = Math.Max(Node.IncomingRef.LastModifiedTime, ChildRef.LastModifiedTime + 1);
 					}
 					NodeRefs.Add(ChildRef);
 				}
@@ -346,6 +351,7 @@ namespace EpicGames.Horde.Bundles
 			{
 				ReadOnlySequence<byte> Data = await GetDataAsync(Ref.Hash);
 				Ref.Node = BundleNode.Deserialize<T>(Data);
+				Ref.Node.IncomingRef = Ref;
 			}
 			return Ref.Node!;
 		}
@@ -369,21 +375,17 @@ namespace EpicGames.Horde.Bundles
 		{
 			if (Node.Blob == null)
 			{
-				if (Node.Export == null)
-				{
-					return Node.Data;
-				}
-				else
-				{
-					return new ReadOnlySequence<byte>(DecodePacket(RootCacheKey, Node.Data, Node.Export.Packet).Slice(Node.Export.Offset, Node.Export.Length));
-				}
+				return Node.Data;
 			}
 			else
 			{
 				await MountBlobAsync(Node.Blob);
 				Debug.Assert(Node.Export != null);
+
 				BundleObject Object = await GetObjectAsync(Node.Blob.Hash);
-				return new ReadOnlySequence<byte>(DecodePacket(Node.Blob.Hash, new ReadOnlySequence<byte>(Object.Data), Node.Export.Packet).Slice(Node.Export.Offset, Node.Export.Length));
+				ReadOnlySequence<byte> Data = new ReadOnlySequence<byte>(DecodePacket(Node.Blob.Hash, new ReadOnlySequence<byte>(Object.Data), Node.Export.Packet).Slice(Node.Export.Offset, Node.Export.Length));
+
+				return Data;
 			}
 		}
 
@@ -434,8 +436,7 @@ namespace EpicGames.Horde.Bundles
 			if(!Blob.Mounted)
 			{
 				BundleObject Object = await GetObjectAsync(Blob.Hash);
-				RegisterObject(Blob, Object, ReadOnlyMemory<byte>.Empty);
-				Blob.Mounted = true;
+				RegisterBlobObject(Object, Blob);
 			}
 		}
 
@@ -509,7 +510,50 @@ namespace EpicGames.Horde.Bundles
 			return Node;
 		}
 
-		void RegisterObject(BlobInfo? Blob, BundleObject Object, ReadOnlyMemory<byte> Data)
+		/// <summary>
+		/// Creates nodes for a given root object. Data for each node will be decompressed and copied from the object.
+		/// </summary>
+		void RegisterRootObject(BundleObject Object)
+		{
+			List<NodeInfo> Nodes = CreateNodesForObject(Object);
+
+			for (int Idx = 0; Idx < Object.Exports.Count; )
+			{
+				BundleCompressionPacket Packet = Object.Exports[Idx].Packet;
+
+				byte[] DecodedData = new byte[Packet.DecodedLength];
+				LZ4Codec.Decode(Object.Data.Slice(Packet.Offset, Packet.EncodedLength).Span, DecodedData);
+
+				for (; Idx < Object.Exports.Count && Object.Exports[Idx].Packet == Packet; Idx++)
+				{
+					BundleExport Export = Object.Exports[Idx];
+
+					ReadOnlyMemory<byte> NodeData = DecodedData.AsMemory(Export.Offset, Export.Length);
+					if(Export.Length != DecodedData.Length)
+					{
+						NodeData = NodeData.ToArray();
+					}
+
+					Nodes[Idx].SetStandalone(new ReadOnlySequence<byte>(NodeData), Export.References.Select(x => Nodes[x]).ToArray());
+				}
+			}
+		}
+
+		/// <summary>
+		/// Creates nodes for a given blob object.
+		/// </summary>
+		void RegisterBlobObject(BundleObject Object, BlobInfo Blob)
+		{
+			List<NodeInfo> Nodes = CreateNodesForObject(Object);
+			for (int Idx = 0; Idx < Object.Exports.Count; Idx++)
+			{
+				BundleExport Export = Object.Exports[Idx];
+				Nodes[Idx].SetExport(Blob, Export, Export.References.Select(x => Nodes[x]).ToArray());
+			}
+			Blob.Mounted = true;
+		}
+
+		List<NodeInfo> CreateNodesForObject(BundleObject Object)
 		{
 			List<NodeInfo> Nodes = new List<NodeInfo>();
 			foreach (BundleExport Export in Object.Exports)
@@ -528,18 +572,14 @@ namespace EpicGames.Horde.Bundles
 				foreach (BundleImport Import in ImportObject.Imports)
 				{
 					NodeInfo Node = FindOrAddNode(Import.Hash, Import.Rank, Import.Length);
-					if(Node.Blob == null)
+					if (Node.Blob == null)
 					{
 						Node.SetImport(ImportBlob);
 					}
 					Nodes.Add(Node);
 				}
 			}
-			for(int Idx = 0; Idx < Object.Exports.Count; Idx++)
-			{
-				NodeInfo Node = Nodes[Idx];
-				Node.SetExport(Blob, Object.Exports[Idx], new ReadOnlySequence<byte>(Data), Object.Exports[Idx].References.Select(x => Nodes[x]).ToArray());
-			}
+			return Nodes;
 		}
 
 		/// <summary>
@@ -679,17 +719,14 @@ namespace EpicGames.Horde.Bundles
 
 		async Task WriteRefAsync(IReadOnlyList<NodeInfo> Nodes, BucketId BucketId, RefId RefId)
 		{
-			BundleObject Object = await CreateObjectAsync(Nodes);
-			await StorageClient.SetRefAsync(NamespaceId, BucketId, RefId, Object);
-			Object.Data = Object.Data.ToArray(); // Duplicated since its lifetime persists beyond this function. See remarks in CreateObjectAsync.
-
 			foreach (NodeInfo Node in Nodes)
 			{
-				Node.Blob = null;
-				Node.Data = new ReadOnlySequence<byte>(Object.Data);
+				ReadOnlySequence<byte> Data = await GetDataAsync(Node);
+				Node.SetStandalone(Data, Node.References!);
 			}
 
-			RootCacheKey = new object();
+			BundleObject Object = await CreateObjectAsync(Nodes);
+			await StorageClient.SetRefAsync(NamespaceId, BucketId, RefId, Object);
 		}
 
 		ReadOnlyMemory<byte> EncodeObject(BundleObject Object)
@@ -710,11 +747,13 @@ namespace EpicGames.Horde.Bundles
 			IoHash Hash = await StorageClient.WriteBlobFromMemoryAsync(NamespaceId, EncodeObject(Object));
 
 			BlobInfo Blob = new BlobInfo(Hash, Object.Exports.Sum(x => x.Length));
-			foreach (NodeInfo Node in Nodes)
+			for (int Idx = 0; Idx < Object.Exports.Count; Idx++)
 			{
-				Node.Blob = Blob;
-				Node.Data = ReadOnlySequence<byte>.Empty;
+				NodeInfo Node = Nodes[Idx];
+				Node.SetExport(Blob, Object.Exports[Idx], Node.References!);
 			}
+
+			Blob.Mounted = true;
 		}
 
 		/// <summary>
@@ -781,8 +820,8 @@ namespace EpicGames.Horde.Bundles
 
 				// Create the export for this node
 				int[] References = Node.References.Select(x => NodeToIndex[x]).ToArray();
-				Node.Export = new BundleExport(Node.Hash, Node.Rank, Packet, Packet.DecodedLength, Node.Length, References);
-				Object.Exports.Add(Node.Export);
+				BundleExport Export = new BundleExport(Node.Hash, Node.Rank, Packet, BlockSize, Node.Length, References);
+				Object.Exports.Add(Export);
 
 				// Write out the new block
 				int Offset = Packet.EncodedLength;
diff --git a/Engine/Source/Programs/Shared/EpicGames.Horde/Bundles/BundleNode.cs b/Engine/Source/Programs/Shared/EpicGames.Horde/Bundles/BundleNode.cs
--- a/Engine/Source/Programs/Shared/EpicGames.Horde/Bundles/BundleNode.cs
+++ b/Engine/Source/Programs/Shared/EpicGames.Horde/Bundles/BundleNode.cs
@@ -32,11 +32,16 @@ namespace EpicGames.Horde.Bundles.Nodes
 		}
 
 		/// <summary>
-		/// Cached incoming reference to the owner of this node. When the contents of this node are flushed to storage, this is modified to become a weak reference
-		/// and hash value. When modified, it becomes a strong reference and zero hash.
+		/// Cached incoming reference to the owner of this node.
 		/// </summary>
 		internal BundleNodeRef? IncomingRef { get; set; }
 
+		/// <summary>
+		/// Queries if the node in its current state is read-only. Once we know that nodes are no longer going to be modified, they are favored for spilling to persistent storage.
+		/// </summary>
+		/// <returns>True if the node is read-only.</returns>
+		public virtual bool IsReadOnly() => false;
+
 		/// <summary>
 		/// Mark this node as dirty
 		/// </summary>
@@ -79,7 +84,7 @@ namespace EpicGames.Horde.Bundles.Nodes
 		/// <summary>
 		/// Cached reference to the parent node
 		/// </summary>
-		internal BundleNodeRef? ParentRef;
+		internal BundleNode Owner;
 
 		/// <summary>
 		/// Strong reference to the current node
@@ -113,31 +118,44 @@ namespace EpicGames.Horde.Bundles.Nodes
 		/// <summary>
 		/// Creates a reference to a node with the given hash
 		/// </summary>
+		/// <param name="Owner">The node which owns the reference</param>
 		/// <param name="Hash">Hash of the referenced node</param>
-		public BundleNodeRef(IoHash Hash)
+		public BundleNodeRef(BundleNode Owner, IoHash Hash)
 		{
+			this.Owner = Owner;
 			this.Hash = Hash;
 		}
 
 		/// <summary>
 		/// Creates a reference to the given node
 		/// </summary>
+		/// <param name="Owner">The node which owns the reference</param>
 		/// <param name="Node">The referenced node</param>
-		public BundleNodeRef(BundleNode Node)
+		public BundleNodeRef(BundleNode Owner, BundleNode Node)
 		{
+			this.Owner = Owner;
 			this.StrongRef = Node;
 			this.LastModifiedTime = Stopwatch.GetTimestamp();
+
+			Node.IncomingRef = this;
 		}
 
 		/// <summary>
-		/// Detach this reference from its current owner. Used in situations where a node needs to be reparented.
+		/// Reparent this reference to a new owner.
 		/// </summary>
-		public void Detach()
+		public void Reparent(BundleNode NewOwner)
 		{
 			MarkAsDirty();
-			ParentRef = null;
+			Owner = NewOwner;
+			MarkAsDirty();
 		}
 
+		/// <summary>
+		/// Figure out whether this ref is dirty
+		/// </summary>
+		/// <returns></returns>
+		public bool IsDirty() => LastModifiedTime != 0;
+
 		/// <summary>
 		/// Converts the node in this reference from a weak to strong reference.
 		/// </summary>
@@ -173,7 +191,7 @@ namespace EpicGames.Horde.Bundles.Nodes
 					}
 					this.WeakRef = null;
 				}
-				ParentRef?.MarkAsDirty();
+				Owner.IncomingRef?.MarkAsDirty();
 			}
 		}
 
@@ -219,16 +237,18 @@ namespace EpicGames.Horde.Bundles.Nodes
 		/// <summary>
 		/// Constructor
 		/// </summary>
+		/// <param name="Owner">The node which owns the reference</param>
 		/// <param name="Hash">Hash of the referenced node</param>
-		public BundleNodeRef(IoHash Hash) : base(Hash)
+		public BundleNodeRef(BundleNode Owner, IoHash Hash) : base(Owner, Hash)
 		{
 		}
 
 		/// <summary>
 		/// Constructor
 		/// </summary>
+		/// <param name="Owner">The node which owns the reference</param>
 		/// <param name="Node">The referenced node</param>
-		public BundleNodeRef(T Node) : base(Node)
+		public BundleNodeRef(BundleNode Owner, T Node) : base(Owner, Node)
 		{
 		}
 	}
diff --git a/Engine/Source/Programs/Shared/EpicGames.Horde/Bundles/Nodes/DirectoryNode.cs b/Engine/Source/Programs/Shared/EpicGames.Horde/Bundles/Nodes/DirectoryNode.cs
--- a/Engine/Source/Programs/Shared/EpicGames.Horde/Bundles/Nodes/DirectoryNode.cs
+++ b/Engine/Source/Programs/Shared/EpicGames.Horde/Bundles/Nodes/DirectoryNode.cs
@@ -5,6 +5,7 @@ using EpicGames.Serialization;
 using Microsoft.Extensions.Logging;
 using System;
 using System.Buffers;
+using System.Buffers.Binary;
 using System.Collections.Generic;
 using System.Diagnostics;
 using System.Diagnostics.CodeAnalysis;
@@ -74,8 +75,8 @@ namespace EpicGames.Horde.Bundles.Nodes
 		/// <summary>
 		/// Constructor
 		/// </summary>
-		public FileEntry(Utf8String Name, FileEntryFlags Flags, FileNode Node)
-			: base(Node)
+		public FileEntry(DirectoryNode Owner, Utf8String Name, FileEntryFlags Flags, FileNode Node)
+			: base(Owner, Node)
 		{
 			this.Name = Name;
 			this.Flags = Flags;
@@ -84,8 +85,8 @@ namespace EpicGames.Horde.Bundles.Nodes
 		/// <summary>
 		/// Constructor
 		/// </summary>
-		public FileEntry(Utf8String Name, FileEntryFlags Flags, long Length, IoHash Hash)
-			: base(Hash)
+		public FileEntry(DirectoryNode Owner, Utf8String Name, FileEntryFlags Flags, long Length, IoHash Hash)
+			: base(Owner, Hash)
 		{
 			this.Name = Name;
 			this.Flags = Flags;
@@ -97,6 +98,9 @@ namespace EpicGames.Horde.Bundles.Nodes
 		{
 			CachedLength = Node!.Length;
 		}
+
+		/// <inheritdoc/>
+		public override string ToString() => IsDirty() ? $"{Name} (*)" : Name.ToString();
 	}
 
 	/// <summary>
@@ -122,8 +126,8 @@ namespace EpicGames.Horde.Bundles.Nodes
 		/// <summary>
 		/// Constructor
 		/// </summary>
-		public DirectoryEntry(Utf8String Name, DirectoryNode Node)
-			: base(Node)
+		public DirectoryEntry(DirectoryNode Owner, Utf8String Name, DirectoryNode Node)
+			: base(Owner, Node)
 		{
 			this.Name = Name;
 		}
@@ -131,8 +135,8 @@ namespace EpicGames.Horde.Bundles.Nodes
 		/// <summary>
 		/// Constructor
 		/// </summary>
-		public DirectoryEntry(Utf8String Name, long Length, IoHash Hash)
-			: base(Hash)
+		public DirectoryEntry(DirectoryNode Owner, Utf8String Name, long Length, IoHash Hash)
+			: base(Owner, Hash)
 		{
 			this.Name = Name;
 			this.CachedLength = Length;
@@ -143,6 +147,9 @@ namespace EpicGames.Horde.Bundles.Nodes
 		{
 			CachedLength = Node!.Length;
 		}
+
+		/// <inheritdoc/>
+		public override string ToString() => IsDirty() ? $"{Name} (*)" : Name.ToString();
 	}
 
 	/// <summary>
@@ -208,7 +215,7 @@ namespace EpicGames.Horde.Bundles.Nodes
 
 			FileNode NewNode = new FileNode();
 
-			FileEntry Entry = new FileEntry(Name, Flags, NewNode);
+			FileEntry Entry = new FileEntry(this, Name, Flags, NewNode);
 			NameToFileEntry[Name] = Entry;
 			MarkAsDirty();
 
@@ -331,7 +338,7 @@ namespace EpicGames.Horde.Bundles.Nodes
 
 			DirectoryNode Node = new DirectoryNode();
 
-			DirectoryEntry Entry = new DirectoryEntry(Name, Node);
+			DirectoryEntry Entry = new DirectoryEntry(this, Name, Node);
 			NameToDirectoryEntry.Add(Name, Entry);
 			MarkAsDirty();
 
@@ -466,7 +473,7 @@ namespace EpicGames.Horde.Bundles.Nodes
 				IoHash Hash = new IoHash(Span);
 				Span = Span.Slice(IoHash.NumBytes);
 
-				Node.NameToFileEntry[Name] = new FileEntry(Name, Flags, Length, Hash);
+				Node.NameToFileEntry[Name] = new FileEntry(Node, Name, Flags, Length, Hash);
 			}
 
 			int DirectoryCount = (int)VarInt.Read(Span, out int DirectoryCountBytes);
@@ -485,7 +492,7 @@ namespace EpicGames.Horde.Bundles.Nodes
 				IoHash Hash = new IoHash(Span);
 				Span = Span.Slice(IoHash.NumBytes);
 
-				Node.NameToDirectoryEntry[Name] = new DirectoryEntry(Name, Length, Hash);
+				Node.NameToDirectoryEntry[Name] = new DirectoryEntry(Node, Name, Length, Hash);
 			}
 
 			Debug.Assert(Span.Length == 0);
diff --git a/Engine/Source/Programs/Shared/EpicGames.Horde/Bundles/Nodes/FileNode.cs b/Engine/Source/Programs/Shared/EpicGames.Horde/Bundles/Nodes/FileNode.cs
--- a/Engine/Source/Programs/Shared/EpicGames.Horde/Bundles/Nodes/FileNode.cs
+++ b/Engine/Source/Programs/Shared/EpicGames.Horde/Bundles/Nodes/FileNode.cs
@@ -115,7 +115,7 @@ namespace EpicGames.Horde.Bundles.Nodes
 		/// <summary>
 		/// Whether this node is read-only
 		/// </summary>
-		public bool IsComplete() => Hash != null;
+		public override bool IsReadOnly() => Hash != null;
 
 		/// <summary>
 		/// Accessor for the children of this node
@@ -167,7 +167,7 @@ namespace EpicGames.Horde.Bundles.Nodes
 				while (Span.Length > 0)
 				{
 					IoHash ChildHash = new IoHash(Span);
-					Node.ChildNodeRefs.Add(new BundleNodeRef<FileNode>(ChildHash));
+					Node.ChildNodeRefs.Add(new BundleNodeRef<FileNode>(Node, ChildHash));
 				}
 			}
 			return Node;
@@ -299,7 +299,7 @@ namespace EpicGames.Horde.Bundles.Nodes
 				{
 					foreach (BundleNodeRef<FileNode> ChildNodeRef in ChildNodeRefs)
 					{
-						ChildNodeRef.Detach();
+						ChildNodeRef.Reparent(NewNode);
 					}
 					ChildNodeRefs = null;
 				}
@@ -317,13 +317,13 @@ namespace EpicGames.Horde.Bundles.Nodes
 				ChildNodeRefs = new List<BundleNodeRef<FileNode>>();
 
 				// Append the new node as a child
-				ChildNodeRefs.Add(new BundleNodeRef<FileNode>(NewNode));
+				ChildNodeRefs.Add(new BundleNodeRef<FileNode>(this, NewNode));
 			}
 		}
 
 		private ReadOnlyMemory<byte> AppendToNode(ReadOnlyMemory<byte> Data, ChunkingOptions Options)
 		{
-			if (Data.Length == 0 || IsComplete())
+			if (Data.Length == 0 || IsReadOnly())
 			{
 				return Data;
 			}
@@ -515,7 +515,7 @@ namespace EpicGames.Horde.Bundles.Nodes
 						Length += LastNode.Length;
 
 						// If the last node is complete, write it to the buffer
-						if (LastNode.IsComplete())
+						if (LastNode.IsReadOnly())
 						{
 							// Update the hash
 							byte[] HashData = LastNode.Hash!.Value.ToByteArray();
@@ -541,7 +541,7 @@ namespace EpicGames.Horde.Bundles.Nodes
 				// Add a new child node
 				FileNode ChildNode = new FileNode();
 				ChildNode.IsRoot = false;
-				ChildNodeRefs.Add(new BundleNodeRef<FileNode>(ChildNode));
+				ChildNodeRefs.Add(new BundleNodeRef<FileNode>(this, ChildNode));
 			}
 		}
 	}
