Horde: Add stats for writing out new bundles to storage.

diff --git a/Engine/Source/Programs/Horde/Horde.Build/Perforce/Impl/CommitService.cs b/Engine/Source/Programs/Horde/Horde.Build/Perforce/Impl/CommitService.cs
--- a/Engine/Source/Programs/Horde/Horde.Build/Perforce/Impl/CommitService.cs
+++ b/Engine/Source/Programs/Horde/Horde.Build/Perforce/Impl/CommitService.cs
@@ -1002,6 +1002,10 @@ namespace HordeServer.Commits.Impl
 			Logger.LogInformation("Writing ref {RefId} for {StreamId} change {Change}", RefId, Stream.Id, Change);
 			await CommitBundle.WriteAsync(GetStreamBucketId(Stream.Id), RefId, true);
 
+			// Print out stats for the new data
+			BundleStats Stats = CommitBundle.Stats;
+			Logger.LogInformation("Written ref {RefId} for {StreamId} change {Change}. Download: {NewBytes:n0} bytes ({NewRefCount} refs => {NewRefBytes:n0} bytes, {NewBlobCount} blobs => {NewBlobBytes:n0} bytes, {NewNodeCount:n0} nodes => {NewNodeBytes:n0} bytes).", RefId, Stream.Id, Change, Stats.NewRefBytes + Stats.NewBlobBytes, Stats.NewRefCount, Stats.NewRefBytes, Stats.NewBlobCount, Stats.NewBlobBytes, Stats.NewNodeCount, Stats.NewNodeBytes);
+
 			return new CommitTree(Stream.Id, Change, RefId);
 		}
 
diff --git a/Engine/Source/Programs/Shared/EpicGames.Horde/Bundles/Bundle.cs b/Engine/Source/Programs/Shared/EpicGames.Horde/Bundles/Bundle.cs
--- a/Engine/Source/Programs/Shared/EpicGames.Horde/Bundles/Bundle.cs
+++ b/Engine/Source/Programs/Shared/EpicGames.Horde/Bundles/Bundle.cs
@@ -56,6 +56,42 @@ namespace EpicGames.Horde.Bundles
 		public float RepackRatio { get; set; } = 0.6f;
 	}
 
+	/// <summary>
+	/// Statistics for the state of a bundle
+	/// </summary>
+	public class BundleStats
+	{
+		/// <summary>
+		/// Number of blobs that have been written
+		/// </summary>
+		public int NewBlobCount { get; set; }
+
+		/// <summary>
+		/// Total size of the blobs that have been written
+		/// </summary>
+		public long NewBlobBytes { get; set; }
+
+		/// <summary>
+		/// Number of refs that have been written
+		/// </summary>
+		public int NewRefCount { get; set; }
+
+		/// <summary>
+		/// Total size of the refs that have been written
+		/// </summary>
+		public long NewRefBytes { get; set; }
+
+		/// <summary>
+		/// Number of nodes in the tree that have been written
+		/// </summary>
+		public int NewNodeCount { get; set; }
+
+		/// <summary>
+		/// Total size of the nodes that have been serialized
+		/// </summary>
+		public long NewNodeBytes { get; set; }
+	}
+
 	/// <summary>
 	/// Base class for manipulating bundles
 	/// </summary>
@@ -156,6 +192,12 @@ namespace EpicGames.Horde.Bundles
 		/// </summary>
 		public BundleOptions Options { get; }
 
+		/// <summary>
+		/// Tracks statistics for the size of written data
+		/// </summary>
+		public BundleStats Stats { get; private set; } = new BundleStats();
+
+		BundleStats NextStats = new BundleStats();
 		byte[] BlockBuffer = Array.Empty<byte>();
 		byte[] EncodedBuffer = Array.Empty<byte>();
 
@@ -425,6 +467,9 @@ namespace EpicGames.Horde.Bundles
 
 				Node = new NodeInfo(Hash, Rank, (int)Data.Length);
 				HashToNode[Hash] = Node;
+
+				NextStats.NewNodeCount++;
+				NextStats.NewNodeBytes += Data.Length;
 			}
 
 			Node.SetStandalone(Data, NodeReferences);
@@ -682,6 +727,10 @@ namespace EpicGames.Horde.Bundles
 
 			// Write the final ref
 			await WriteRefAsync(WriteNodes.Slice(MinIdx), BucketId, RefId);
+
+			// Copy the stats over
+			Stats = NextStats;
+			NextStats = new BundleStats();
 		}
 
 		private void FindModifiedLiveSet(NodeInfo Node, HashSet<NodeInfo> Nodes)
@@ -721,12 +770,17 @@ namespace EpicGames.Horde.Bundles
 		{
 			foreach (NodeInfo Node in Nodes)
 			{
-				ReadOnlySequence<byte> Data = await GetDataAsync(Node);
-				Node.SetStandalone(Data, Node.References!);
+				ReadOnlySequence<byte> NodeData = await GetDataAsync(Node);
+				Node.SetStandalone(NodeData, Node.References!);
 			}
 
 			BundleObject Object = await CreateObjectAsync(Nodes);
-			await StorageClient.SetRefAsync(NamespaceId, BucketId, RefId, Object);
+
+			ReadOnlyMemory<byte> Data = EncodeObject(Object);
+			NextStats.NewRefCount++;
+			NextStats.NewRefBytes += Data.Length;
+
+			await StorageClient.SetRefAsync(NamespaceId, BucketId, RefId, new CbField(Data));
 		}
 
 		ReadOnlyMemory<byte> EncodeObject(BundleObject Object)
@@ -744,7 +798,12 @@ namespace EpicGames.Horde.Bundles
 		async Task WriteObjectAsync(IReadOnlyList<NodeInfo> Nodes)
 		{
 			BundleObject Object = await CreateObjectAsync(Nodes);
-			IoHash Hash = await StorageClient.WriteBlobFromMemoryAsync(NamespaceId, EncodeObject(Object));
+
+			ReadOnlyMemory<byte> Data = EncodeObject(Object);
+			NextStats.NewBlobCount++;
+			NextStats.NewBlobBytes += Data.Length;
+
+			IoHash Hash = await StorageClient.WriteBlobFromMemoryAsync(NamespaceId, Data);
 
 			BlobInfo Blob = new BlobInfo(Hash, Object.Exports.Sum(x => x.Length));
 			for (int Idx = 0; Idx < Object.Exports.Count; Idx++)
diff --git a/Engine/Source/Programs/Shared/EpicGames.Horde/Storage/Impl/FileStorageClient.cs b/Engine/Source/Programs/Shared/EpicGames.Horde/Storage/Impl/FileStorageClient.cs
--- a/Engine/Source/Programs/Shared/EpicGames.Horde/Storage/Impl/FileStorageClient.cs
+++ b/Engine/Source/Programs/Shared/EpicGames.Horde/Storage/Impl/FileStorageClient.cs
@@ -62,7 +62,7 @@ namespace EpicGames.Horde.Storage.Impl
 		public Task<Stream> ReadBlobAsync(NamespaceId NamespaceId, IoHash Hash, CancellationToken CancellationToken = default)
 		{
 			FileReference File = GetBlobFile(NamespaceId, Hash);
-			Logger.LogInformation("Reading {File} ({Size} bytes)", File, new FileInfo(File.FullName).Length);
+			Logger.LogInformation("Reading {File} ({Size:n0} bytes)", File, new FileInfo(File.FullName).Length);
 			return Task.FromResult<Stream>(FileReference.Open(File, FileMode.Open, FileAccess.Read, FileShare.Read));
 		}
 
@@ -72,7 +72,7 @@ namespace EpicGames.Horde.Storage.Impl
 			FileReference File = GetBlobFile(NamespaceId, Hash);
 			DirectoryReference.CreateDirectory(File.Directory);
 
-			Logger.LogInformation("Writing {File} ({Size} bytes)", File, Stream.Length);
+			Logger.LogInformation("Writing {File} ({Size:n0} bytes)", File, Stream.Length);
 			using (Stream OutputStream = FileReference.Open(File, FileMode.Create, FileAccess.Write, FileShare.Read))
 			{
 				await Stream.CopyToAsync(OutputStream);
@@ -118,7 +118,7 @@ namespace EpicGames.Horde.Storage.Impl
 		public async Task<IRef> GetRefAsync(NamespaceId NamespaceId, BucketId BucketId, RefId RefId, CancellationToken CancellationToken = default)
 		{
 			FileReference File = GetRefFile(NamespaceId, BucketId, RefId);
-			Logger.LogInformation("Reading {File} ({Size} bytes)", File, new FileInfo(File.FullName).Length);
+			Logger.LogInformation("Reading {File} ({Size:n0} bytes)", File, new FileInfo(File.FullName).Length);
 			byte[] Data = await FileReference.ReadAllBytesAsync(File);
 			return new Ref(NamespaceId, BucketId, RefId, new CbObject(Data));
 		}
@@ -141,7 +141,7 @@ namespace EpicGames.Horde.Storage.Impl
 			FileReference File = GetRefFile(NamespaceId, BucketId, RefId);
 			DirectoryReference.CreateDirectory(File.Directory);
 
-			Logger.LogInformation("Writing {File} ({Size} bytes)", File, Value.GetView().Length);
+			Logger.LogInformation("Writing {File} ({Size:n0} bytes)", File, Value.GetView().Length);
 			await FileReference.WriteAllBytesAsync(File, Value.GetView().ToArray(), CancellationToken);
 			return new List<IoHash>();
 		}
