Horde: Add cancellation token arguments to all IStorageClient extension methods.

diff --git a/Engine/Source/Programs/Shared/EpicGames.Horde/Storage/IStorageClient.cs b/Engine/Source/Programs/Shared/EpicGames.Horde/Storage/IStorageClient.cs
--- a/Engine/Source/Programs/Shared/EpicGames.Horde/Storage/IStorageClient.cs
+++ b/Engine/Source/Programs/Shared/EpicGames.Horde/Storage/IStorageClient.cs
@@ -276,10 +276,11 @@ namespace EpicGames.Horde.Storage
 		/// <param name="NamespaceId">Namespace containing the blob</param>
 		/// <param name="Hash">Hash of the blob to read</param>
 		/// <param name="MaxInMemoryBlobLength">Maximum allowed memory allocation to store the blob</param>
+		/// <param name="CancellationToken">Cancellation token for the operation</param>
 		/// <returns>Data for the blob that was read. Throws an exception if the blob was not found.</returns>
-		public static async Task<byte[]> ReadBlobToMemoryAsync(this IStorageClient StorageClient, NamespaceId NamespaceId, IoHash Hash, int MaxInMemoryBlobLength = DefaultMaxInMemoryBlobLength)
+		public static async Task<byte[]> ReadBlobToMemoryAsync(this IStorageClient StorageClient, NamespaceId NamespaceId, IoHash Hash, int MaxInMemoryBlobLength = DefaultMaxInMemoryBlobLength, CancellationToken CancellationToken = default)
 		{
-			using Stream Stream = await StorageClient.ReadBlobAsync(NamespaceId, Hash);
+			using Stream Stream = await StorageClient.ReadBlobAsync(NamespaceId, Hash, CancellationToken);
 
 			long Length = Stream.Length;
 			if (Length > MaxInMemoryBlobLength)
@@ -290,7 +291,7 @@ namespace EpicGames.Horde.Storage
 			byte[] Buffer = new byte[Length];
 			for (int Offset = 0; Offset < Length;)
 			{
-				int Count = await Stream.ReadAsync(Buffer, Offset, (int)Length - Offset);
+				int Count = await Stream.ReadAsync(Buffer, Offset, (int)Length - Offset, CancellationToken);
 				if (Count == 0)
 				{
 					throw new BlobException(NamespaceId, Hash, $"Unexpected end of stream reading blob {Hash}");
@@ -307,11 +308,12 @@ namespace EpicGames.Horde.Storage
 		/// <param name="StorageClient">The storage interface</param>
 		/// <param name="NamespaceId">Namespace containing the blob</param>
 		/// <param name="Data">Data to write</param>
+		/// <param name="CancellationToken">Cancellation token for the operation</param>
 		/// <returns></returns>
-		public static async Task<IoHash> WriteBlobFromMemoryAsync(this IStorageClient StorageClient, NamespaceId NamespaceId, ReadOnlyMemory<byte> Data)
+		public static async Task<IoHash> WriteBlobFromMemoryAsync(this IStorageClient StorageClient, NamespaceId NamespaceId, ReadOnlyMemory<byte> Data, CancellationToken CancellationToken = default)
 		{
 			IoHash Hash = IoHash.Compute(Data.Span);
-			await WriteBlobFromMemoryAsync(StorageClient, NamespaceId, Hash, Data);
+			await WriteBlobFromMemoryAsync(StorageClient, NamespaceId, Hash, Data, CancellationToken);
 			return Hash;
 		}
 
@@ -322,24 +324,12 @@ namespace EpicGames.Horde.Storage
 		/// <param name="NamespaceId">Namespace containing the blob</param>
 		/// <param name="Hash">Hash of the data</param>
 		/// <param name="Data">The data to be written</param>
+		/// <param name="CancellationToken">Cancellation token for the operation</param>
 		/// <returns></returns>
-		public static async Task WriteBlobFromMemoryAsync(this IStorageClient StorageClient, NamespaceId NamespaceId, IoHash Hash, ReadOnlyMemory<byte> Data)
+		public static async Task WriteBlobFromMemoryAsync(this IStorageClient StorageClient, NamespaceId NamespaceId, IoHash Hash, ReadOnlyMemory<byte> Data, CancellationToken CancellationToken = default)
 		{
 			using ReadOnlyMemoryStream Stream = new ReadOnlyMemoryStream(Data);
-			await StorageClient.WriteBlobAsync(NamespaceId, Hash, Stream);
-		}
-
-		/// <summary>
-		/// Reads a blob and decodes it as a compact binary object
-		/// </summary>
-		/// <param name="StorageClient">The storage interface</param>
-		/// <param name="NamespaceId">Namespace containing the blob</param>
-		/// <param name="Hash">Hash of the blob to read</param>
-		/// <param name="MaxInMemoryBlobLength">Maximum allowed memory allocation to store the blob</param>
-		/// <returns>The decoded object</returns>
-		public static async Task<CbObject> ReadObjectAsync(this IStorageClient StorageClient, NamespaceId NamespaceId, IoHash Hash, int MaxInMemoryBlobLength = DefaultMaxInMemoryBlobLength)
-		{
-			return new CbObject(await ReadBlobToMemoryAsync(StorageClient, NamespaceId, Hash, MaxInMemoryBlobLength));
+			await StorageClient.WriteBlobAsync(NamespaceId, Hash, Stream, CancellationToken);
 		}
 
 		/// <summary>
@@ -350,10 +340,12 @@ namespace EpicGames.Horde.Storage
 		/// <param name="NamespaceId">Namespace containing the blob</param>
 		/// <param name="Hash">Hash of the blob to read</param>
 		/// <param name="MaxInMemoryBlobLength">Maximum allowed memory allocation to store the blob</param>
+		/// <param name="CancellationToken">Cancellation token for the operation</param>
 		/// <returns>The decoded object</returns>
-		public static async Task<T> ReadObjectAsync<T>(this IStorageClient StorageClient, NamespaceId NamespaceId, IoHash Hash, int MaxInMemoryBlobLength = DefaultMaxInMemoryBlobLength)
+		public static async Task<T> ReadObjectAsync<T>(this IStorageClient StorageClient, NamespaceId NamespaceId, IoHash Hash, int MaxInMemoryBlobLength = DefaultMaxInMemoryBlobLength, CancellationToken CancellationToken = default)
 		{
-			return CbSerializer.Deserialize<T>(await ReadBlobToMemoryAsync(StorageClient, NamespaceId, Hash, MaxInMemoryBlobLength));
+			ReadOnlyMemory<byte> Data = await ReadBlobToMemoryAsync(StorageClient, NamespaceId, Hash, MaxInMemoryBlobLength, CancellationToken);
+			return CbSerializer.Deserialize<T>(Data);
 		}
 
 		/// <summary>
@@ -362,10 +354,11 @@ namespace EpicGames.Horde.Storage
 		/// <param name="StorageClient">The storage interface</param>
 		/// <param name="NamespaceId">Namespace containing the blob</param>
 		/// <param name="Object">The object to be written</param>
+		/// <param name="CancellationToken">Cancellation token for the operation</param>
 		/// <returns></returns>
-		public static Task<IoHash> WriteObjectAsync(this IStorageClient StorageClient, NamespaceId NamespaceId, CbObject Object)
+		public static Task<IoHash> WriteObjectAsync(this IStorageClient StorageClient, NamespaceId NamespaceId, CbObject Object, CancellationToken CancellationToken = default)
 		{
-			return WriteBlobFromMemoryAsync(StorageClient, NamespaceId, Object.GetView());
+			return WriteBlobFromMemoryAsync(StorageClient, NamespaceId, Object.GetView(), CancellationToken);
 		}
 
 		/// <summary>
@@ -374,14 +367,15 @@ namespace EpicGames.Horde.Storage
 		/// <param name="StorageClient">The storage interface</param>
 		/// <param name="NamespaceId">Namespace containing the blob</param>
 		/// <param name="Object">The object to be written</param>
+		/// <param name="CancellationToken">Cancellation token for the operation</param>
 		/// <returns></returns>
-		public static async Task<IoHash> WriteObjectAsync<T>(this IStorageClient StorageClient, NamespaceId NamespaceId, T Object)
+		public static async Task<IoHash> WriteObjectAsync<T>(this IStorageClient StorageClient, NamespaceId NamespaceId, T Object, CancellationToken CancellationToken = default)
 		{
 			CbWriter Writer = new CbWriter();
 			CbSerializer.Serialize<T>(Writer, Object);
 
 			IoHash Hash = Writer.ComputeHash();
-			await StorageClient.WriteBlobAsync(NamespaceId, Hash, Writer.AsStream());
+			await StorageClient.WriteBlobAsync(NamespaceId, Hash, Writer.AsStream(), CancellationToken);
 			return Hash;
 		}
 
@@ -393,10 +387,11 @@ namespace EpicGames.Horde.Storage
 		/// <param name="NamespaceId">Namespace containing the ref</param>
 		/// <param name="BucketId">Bucket containing the ref</param>
 		/// <param name="RefId">The ref id</param>
+		/// <param name="CancellationToken">Cancellation token for the operation</param>
 		/// <returns>Deserialized object for the given ref</returns>
-		public static async Task<T> GetRefAsync<T>(this IStorageClient StorageClient, NamespaceId NamespaceId, BucketId BucketId, RefId RefId)
+		public static async Task<T> GetRefAsync<T>(this IStorageClient StorageClient, NamespaceId NamespaceId, BucketId BucketId, RefId RefId, CancellationToken CancellationToken = default)
 		{
-			IRef Ref = await StorageClient.GetRefAsync(NamespaceId, BucketId, RefId);
+			IRef Ref = await StorageClient.GetRefAsync(NamespaceId, BucketId, RefId, CancellationToken);
 			return CbSerializer.Deserialize<T>(Ref.Value);
 		}
 
@@ -409,11 +404,12 @@ namespace EpicGames.Horde.Storage
 		/// <param name="BucketId">Bucket containing the ref</param>
 		/// <param name="RefId">The ref id</param>
 		/// <param name="Value">The new object for the ref</param>
+		/// <param name="CancellationToken">Cancellation token for the operation</param>
 		/// <returns>List of missing blob hashes</returns>
-		public static Task<List<IoHash>> TrySetRefAsync<T>(this IStorageClient StorageClient, NamespaceId NamespaceId, BucketId BucketId, RefId RefId, T Value)
+		public static Task<List<IoHash>> TrySetRefAsync<T>(this IStorageClient StorageClient, NamespaceId NamespaceId, BucketId BucketId, RefId RefId, T Value, CancellationToken CancellationToken = default)
 		{
 			CbObject Object = CbSerializer.Serialize<T>(Value);
-			return StorageClient.TrySetRefAsync(NamespaceId, BucketId, RefId, Object);
+			return StorageClient.TrySetRefAsync(NamespaceId, BucketId, RefId, Object, CancellationToken);
 		}
 
 		/// <summary>
@@ -424,9 +420,10 @@ namespace EpicGames.Horde.Storage
 		/// <param name="BucketId">Bucket containing the ref</param>
 		/// <param name="RefId">The ref id</param>
 		/// <param name="Value">The new object for the ref</param>
-		public static async Task SetRefAsync(this IStorageClient StorageClient, NamespaceId NamespaceId, BucketId BucketId, RefId RefId, CbObject Value)
+		/// <param name="CancellationToken">Cancellation token for the operation</param>
+		public static async Task SetRefAsync(this IStorageClient StorageClient, NamespaceId NamespaceId, BucketId BucketId, RefId RefId, CbObject Value, CancellationToken CancellationToken = default)
 		{
-			List<IoHash> MissingHashes = await StorageClient.TrySetRefAsync(NamespaceId, BucketId, RefId, Value);
+			List<IoHash> MissingHashes = await StorageClient.TrySetRefAsync(NamespaceId, BucketId, RefId, Value, CancellationToken);
 			if (MissingHashes.Count > 0)
 			{
 				throw new RefMissingBlobException(NamespaceId, BucketId, RefId, MissingHashes);
@@ -441,10 +438,11 @@ namespace EpicGames.Horde.Storage
 		/// <param name="BucketId">Bucket containing the ref</param>
 		/// <param name="RefId">The ref id</param>
 		/// <param name="Value">The new object for the ref</param>
-		public static Task SetRefAsync<T>(this IStorageClient StorageClient, NamespaceId NamespaceId, BucketId BucketId, RefId RefId, T Value)
+		/// <param name="CancellationToken">Cancellation token for the operation</param>
+		public static Task SetRefAsync<T>(this IStorageClient StorageClient, NamespaceId NamespaceId, BucketId BucketId, RefId RefId, T Value, CancellationToken CancellationToken = default)
 		{
 			CbObject Object = CbSerializer.Serialize<T>(Value);
-			return SetRefAsync(StorageClient, NamespaceId, BucketId, RefId, Object);
+			return SetRefAsync(StorageClient, NamespaceId, BucketId, RefId, Object, CancellationToken);
 		}
 
 		/// <summary>
@@ -455,9 +453,10 @@ namespace EpicGames.Horde.Storage
 		/// <param name="BucketId">Bucket containing the ref</param>
 		/// <param name="RefId">The ref id</param>
 		/// <param name="ValueHash">Hash of the ref value</param>
-		public static async Task FinalizeRefAsync(this IStorageClient StorageClient, NamespaceId NamespaceId, BucketId BucketId, RefId RefId, IoHash ValueHash)
+		/// <param name="CancellationToken">Cancellation token for the operation</param>
+		public static async Task FinalizeRefAsync(this IStorageClient StorageClient, NamespaceId NamespaceId, BucketId BucketId, RefId RefId, IoHash ValueHash, CancellationToken CancellationToken = default)
 		{
-			List<IoHash> MissingHashes = await StorageClient.TryFinalizeRefAsync(NamespaceId, BucketId, RefId, ValueHash);
+			List<IoHash> MissingHashes = await StorageClient.TryFinalizeRefAsync(NamespaceId, BucketId, RefId, ValueHash, CancellationToken);
 			if (MissingHashes.Count > 0)
 			{
 				throw new RefMissingBlobException(NamespaceId, BucketId, RefId, MissingHashes);
