Horde: Rename ReadObject/WriteObject methods on IStorageClient to ReadBlob/WriteBlob for consistency. Objects and blobs are generally treated as identical now.

diff --git a/Engine/Source/Programs/Horde/Horde.Agent/Commands/Compute/ComputeCommand.cs b/Engine/Source/Programs/Horde/Horde.Agent/Commands/Compute/ComputeCommand.cs
--- a/Engine/Source/Programs/Horde/Horde.Agent/Commands/Compute/ComputeCommand.cs
+++ b/Engine/Source/Programs/Horde/Horde.Agent/Commands/Compute/ComputeCommand.cs
@@ -294,7 +294,7 @@ namespace HordeAgent.Commands
 
 		async Task WriteOutputAsync(IStorageClient StorageClient, NamespaceId NamespaceId, IoHash TreeHash, DirectoryReference OutputDir)
 		{
-			DirectoryTree Tree = await StorageClient.ReadObjectAsync<DirectoryTree>(NamespaceId, TreeHash);
+			DirectoryTree Tree = await StorageClient.ReadBlobAsync<DirectoryTree>(NamespaceId, TreeHash);
 
 			List<Task> Tasks = new List<Task>();
 			foreach (FileNode File in Tree.Files)
diff --git a/Engine/Source/Programs/Horde/Horde.Agent/Execution/ComputeTaskExecutor.cs b/Engine/Source/Programs/Horde/Horde.Agent/Execution/ComputeTaskExecutor.cs
--- a/Engine/Source/Programs/Horde/Horde.Agent/Execution/ComputeTaskExecutor.cs
+++ b/Engine/Source/Programs/Horde/Horde.Agent/Execution/ComputeTaskExecutor.cs
@@ -113,7 +113,7 @@ namespace HordeAgent
 			FileUtils.ForceDeleteDirectoryContents(SandboxDir);
 
 			DateTimeOffset InputFetchStart = DateTimeOffset.UtcNow;
-			DirectoryTree InputDirectory = await StorageClient.ReadObjectAsync<DirectoryTree>(NamespaceId, Task.SandboxHash);
+			DirectoryTree InputDirectory = await StorageClient.ReadBlobAsync<DirectoryTree>(NamespaceId, Task.SandboxHash);
 			await SetupSandboxAsync(NamespaceId, InputDirectory, SandboxDir);
 			DateTimeOffset InputFetchCompleted = DateTimeOffset.UtcNow;
 
@@ -196,7 +196,7 @@ namespace HordeAgent
 
 			async Task DownloadDir(DirectoryNode DirectoryNode)
 			{
-				DirectoryTree InputSubDirectory = await StorageClient.ReadObjectAsync<DirectoryTree>(NamespaceId, DirectoryNode.Hash);
+				DirectoryTree InputSubDirectory = await StorageClient.ReadBlobAsync<DirectoryTree>(NamespaceId, DirectoryNode.Hash);
 				DirectoryReference OutputSubDirectory = DirectoryReference.Combine(OutputDir, DirectoryNode.Name.ToString());
 				await SetupSandboxAsync(NamespaceId, InputSubDirectory, OutputSubDirectory);
 			}
@@ -253,7 +253,7 @@ namespace HordeAgent
 			Tree.Files.AddRange(await Task.WhenAll(Files));
 			Tree.Directories.AddRange(await Task.WhenAll(Trees));
 
-			IoHash Hash = await StorageClient.WriteObjectAsync<DirectoryTree>(NamespaceId, Tree);
+			IoHash Hash = await StorageClient.WriteBlobAsync<DirectoryTree>(NamespaceId, Tree);
 			return (Tree, Hash);
 		}
 
diff --git a/Engine/Source/Programs/Horde/Horde.Build/Compute/Impl/ComputeService.cs b/Engine/Source/Programs/Horde/Horde.Build/Compute/Impl/ComputeService.cs
--- a/Engine/Source/Programs/Horde/Horde.Build/Compute/Impl/ComputeService.cs
+++ b/Engine/Source/Programs/Horde/Horde.Build/Compute/Impl/ComputeService.cs
@@ -457,7 +457,7 @@ namespace HordeServer.Compute.Impl
 				NamespaceId NamespaceId = new NamespaceId(ClusterConfig.NamespaceId);
 				try
 				{
-					Requirements = await StorageClient.ReadObjectAsync<Requirements>(NamespaceId, QueueKey.RequirementsHash);
+					Requirements = await StorageClient.ReadBlobAsync<Requirements>(NamespaceId, QueueKey.RequirementsHash);
 				}
 				catch (BlobNotFoundException)
 				{
diff --git a/Engine/Source/Programs/Shared/EpicGames.Horde/Bundles/BundleTools.cs b/Engine/Source/Programs/Shared/EpicGames.Horde/Bundles/BundleTools.cs
--- a/Engine/Source/Programs/Shared/EpicGames.Horde/Bundles/BundleTools.cs
+++ b/Engine/Source/Programs/Shared/EpicGames.Horde/Bundles/BundleTools.cs
@@ -242,7 +242,7 @@ namespace EpicGames.Horde.Bundles
 				IoHash Hash = ImportObject.Object.Hash;
 				if (Hashes.Add(Hash))
 				{
-					BundleObject Object = await StorageClient.ReadObjectAsync<BundleObject>(NamespaceId, Hash);
+					BundleObject Object = await StorageClient.ReadBlobAsync<BundleObject>(NamespaceId, Hash);
 
 					FileReference OutputFile = FileReference.Combine(OutputDir, $"{Hash}.yml");
 					WriteSummaryForObject(Object, OutputFile);
diff --git a/Engine/Source/Programs/Shared/EpicGames.Horde/Storage/IStorageClient.cs b/Engine/Source/Programs/Shared/EpicGames.Horde/Storage/IStorageClient.cs
--- a/Engine/Source/Programs/Shared/EpicGames.Horde/Storage/IStorageClient.cs
+++ b/Engine/Source/Programs/Shared/EpicGames.Horde/Storage/IStorageClient.cs
@@ -342,25 +342,12 @@ namespace EpicGames.Horde.Storage
 		/// <param name="MaxInMemoryBlobLength">Maximum allowed memory allocation to store the blob</param>
 		/// <param name="CancellationToken">Cancellation token for the operation</param>
 		/// <returns>The decoded object</returns>
-		public static async Task<T> ReadObjectAsync<T>(this IStorageClient StorageClient, NamespaceId NamespaceId, IoHash Hash, int MaxInMemoryBlobLength = DefaultMaxInMemoryBlobLength, CancellationToken CancellationToken = default)
+		public static async Task<T> ReadBlobAsync<T>(this IStorageClient StorageClient, NamespaceId NamespaceId, IoHash Hash, int MaxInMemoryBlobLength = DefaultMaxInMemoryBlobLength, CancellationToken CancellationToken = default)
 		{
 			ReadOnlyMemory<byte> Data = await ReadBlobToMemoryAsync(StorageClient, NamespaceId, Hash, MaxInMemoryBlobLength, CancellationToken);
 			return CbSerializer.Deserialize<T>(Data);
 		}
 
-		/// <summary>
-		/// Writes a serialized compact binary object to storage
-		/// </summary>
-		/// <param name="StorageClient">The storage interface</param>
-		/// <param name="NamespaceId">Namespace containing the blob</param>
-		/// <param name="Object">The object to be written</param>
-		/// <param name="CancellationToken">Cancellation token for the operation</param>
-		/// <returns></returns>
-		public static Task<IoHash> WriteObjectAsync(this IStorageClient StorageClient, NamespaceId NamespaceId, CbObject Object, CancellationToken CancellationToken = default)
-		{
-			return WriteBlobFromMemoryAsync(StorageClient, NamespaceId, Object.GetView(), CancellationToken);
-		}
-
 		/// <summary>
 		/// Serializes an object to compact binary format and writes it to storage
 		/// </summary>
@@ -369,7 +356,7 @@ namespace EpicGames.Horde.Storage
 		/// <param name="Object">The object to be written</param>
 		/// <param name="CancellationToken">Cancellation token for the operation</param>
 		/// <returns></returns>
-		public static async Task<IoHash> WriteObjectAsync<T>(this IStorageClient StorageClient, NamespaceId NamespaceId, T Object, CancellationToken CancellationToken = default)
+		public static async Task<IoHash> WriteBlobAsync<T>(this IStorageClient StorageClient, NamespaceId NamespaceId, T Object, CancellationToken CancellationToken = default)
 		{
 			CbWriter Writer = new CbWriter();
 			CbSerializer.Serialize<T>(Writer, Object);
@@ -396,20 +383,18 @@ namespace EpicGames.Horde.Storage
 		}
 
 		/// <summary>
-		/// Attempts to set a ref to a particular value
+		/// Sets a ref to a particular value
 		/// </summary>
-		/// <typeparam name="T"></typeparam>
 		/// <param name="StorageClient">The storage interface</param>
 		/// <param name="NamespaceId">Namespace containing the ref</param>
 		/// <param name="BucketId">Bucket containing the ref</param>
 		/// <param name="RefId">The ref id</param>
 		/// <param name="Value">The new object for the ref</param>
 		/// <param name="CancellationToken">Cancellation token for the operation</param>
-		/// <returns>List of missing blob hashes</returns>
-		public static Task<List<IoHash>> TrySetRefAsync<T>(this IStorageClient StorageClient, NamespaceId NamespaceId, BucketId BucketId, RefId RefId, T Value, CancellationToken CancellationToken = default)
+		public static Task SetRefAsync<T>(this IStorageClient StorageClient, NamespaceId NamespaceId, BucketId BucketId, RefId RefId, T Value, CancellationToken CancellationToken = default)
 		{
 			CbObject Object = CbSerializer.Serialize<T>(Value);
-			return StorageClient.TrySetRefAsync(NamespaceId, BucketId, RefId, Object, CancellationToken);
+			return SetRefAsync(StorageClient, NamespaceId, BucketId, RefId, Object, CancellationToken);
 		}
 
 		/// <summary>
@@ -431,18 +416,20 @@ namespace EpicGames.Horde.Storage
 		}
 
 		/// <summary>
-		/// Sets a ref to a particular value
+		/// Attempts to set a ref to a particular value
 		/// </summary>
+		/// <typeparam name="T"></typeparam>
 		/// <param name="StorageClient">The storage interface</param>
 		/// <param name="NamespaceId">Namespace containing the ref</param>
 		/// <param name="BucketId">Bucket containing the ref</param>
 		/// <param name="RefId">The ref id</param>
 		/// <param name="Value">The new object for the ref</param>
 		/// <param name="CancellationToken">Cancellation token for the operation</param>
-		public static Task SetRefAsync<T>(this IStorageClient StorageClient, NamespaceId NamespaceId, BucketId BucketId, RefId RefId, T Value, CancellationToken CancellationToken = default)
+		/// <returns>List of missing blob hashes</returns>
+		public static Task<List<IoHash>> TrySetRefAsync<T>(this IStorageClient StorageClient, NamespaceId NamespaceId, BucketId BucketId, RefId RefId, T Value, CancellationToken CancellationToken = default)
 		{
 			CbObject Object = CbSerializer.Serialize<T>(Value);
-			return SetRefAsync(StorageClient, NamespaceId, BucketId, RefId, Object, CancellationToken);
+			return StorageClient.TrySetRefAsync(NamespaceId, BucketId, RefId, Object, CancellationToken);
 		}
 
 		/// <summary>
