Horde: Encode settings used to replicate data to Perforce in the ref ids. Settings are now copied into a CommitKey object, which is encoded as compact binary and hashed, and also embedded in its original form in the bundle ref.

diff --git a/Engine/Source/Programs/Horde/Horde.Build/Commands/Bundles/CreateCommand.cs b/Engine/Source/Programs/Horde/Horde.Build/Commands/Bundles/CreateCommand.cs
--- a/Engine/Source/Programs/Horde/Horde.Build/Commands/Bundles/CreateCommand.cs
+++ b/Engine/Source/Programs/Horde/Horde.Build/Commands/Bundles/CreateCommand.cs
@@ -50,7 +50,7 @@ namespace HordeServer.Commands.Bundles
 
 			Bundle<DirectoryNode> NewBundle = Bundle.Create<DirectoryNode>(StorageClient, NamespaceId, new BundleOptions(), null);
 			await NewBundle.Root.CopyFromDirectoryAsync(InputDir.ToDirectoryInfo(), new ChunkingOptions(), Logger);
-			await NewBundle.WriteAsync(BucketId, RefId, false);
+			await NewBundle.WriteAsync(BucketId, RefId, CbObject.Empty, false);
 
 			return 0;
 		}
diff --git a/Engine/Source/Programs/Horde/Horde.Build/Commands/Perforce/ReplicateCommand.cs b/Engine/Source/Programs/Horde/Horde.Build/Commands/Perforce/ReplicateCommand.cs
--- a/Engine/Source/Programs/Horde/Horde.Build/Commands/Perforce/ReplicateCommand.cs
+++ b/Engine/Source/Programs/Horde/Horde.Build/Commands/Perforce/ReplicateCommand.cs
@@ -47,7 +47,7 @@ namespace HordeServer.Commands
 		public bool Clean { get; set; }
 
 		[CommandLine]
-		public string Filter { get; set; } = "//...";
+		public string Filter { get; set; } = "...";
 
 		[CommandLine]
 		public DirectoryReference? OutputDir { get; set; }
diff --git a/Engine/Source/Programs/Shared/EpicGames.Horde/Bundles/Bundle.cs b/Engine/Source/Programs/Shared/EpicGames.Horde/Bundles/Bundle.cs
--- a/Engine/Source/Programs/Shared/EpicGames.Horde/Bundles/Bundle.cs
+++ b/Engine/Source/Programs/Shared/EpicGames.Horde/Bundles/Bundle.cs
@@ -265,12 +265,12 @@ namespace EpicGames.Horde.Bundles
 		{
 			Bundle<T> NewBundle = new Bundle<T>(StorageClient, NamespaceId, null!, Options, Cache);
 
-			BundleObject Object = await StorageClient.GetRefAsync<BundleObject>(NamespaceId, BucketId, RefId);
-			NewBundle.RegisterRootObject(Object);
+			BundleRoot RootObject = await StorageClient.GetRefAsync<BundleRoot>(NamespaceId, BucketId, RefId);
+			NewBundle.RegisterRootObject(RootObject.Object);
 
-			IoHash RootHash = Object.Exports[Object.Exports.Count - 1].Hash;
-			ReadOnlySequence<byte> Data = await NewBundle.GetDataAsync(RootHash);
-			((Bundle)NewBundle).Root = BundleNode.Deserialize<T>(Data);
+			IoHash RootNodeHash = RootObject.Object.Exports[^1].Hash;
+			ReadOnlySequence<byte> RootNodeData = await NewBundle.GetDataAsync(RootNodeHash);
+			((Bundle)NewBundle).Root = BundleNode.Deserialize<T>(RootNodeData);
 
 			return NewBundle;
 		}
@@ -632,8 +632,9 @@ namespace EpicGames.Horde.Bundles
 		/// </summary>
 		/// <param name="BucketId"></param>
 		/// <param name="RefId"></param>
+		/// <param name="Metadata">Metadata for the root object</param>
 		/// <param name="Compact"></param>
-		public virtual async Task WriteAsync(BucketId BucketId, RefId RefId, bool Compact)
+		public virtual async Task WriteAsync(BucketId BucketId, RefId RefId, CbObject Metadata, bool Compact)
 		{
 			// Completely flush the entire working set into NodeInfo objects
 			SerializeWorkingSet(0);
@@ -726,7 +727,7 @@ namespace EpicGames.Horde.Bundles
 			}
 
 			// Write the final ref
-			await WriteRefAsync(WriteNodes.Slice(MinIdx), BucketId, RefId);
+			await WriteRefAsync(WriteNodes.Slice(MinIdx), BucketId, RefId, Metadata);
 
 			// Copy the stats over
 			Stats = NextStats;
@@ -766,7 +767,7 @@ namespace EpicGames.Horde.Bundles
 			}
 		}
 
-		async Task WriteRefAsync(IReadOnlyList<NodeInfo> Nodes, BucketId BucketId, RefId RefId)
+		async Task WriteRefAsync(IReadOnlyList<NodeInfo> Nodes, BucketId BucketId, RefId RefId, CbObject Metadata)
 		{
 			foreach (NodeInfo Node in Nodes)
 			{
@@ -774,19 +775,21 @@ namespace EpicGames.Horde.Bundles
 				Node.SetStandalone(NodeData, Node.References!);
 			}
 
-			BundleObject Object = await CreateObjectAsync(Nodes);
+			BundleRoot Ref = new BundleRoot();
+			Ref.Metadata = Metadata;
+			Ref.Object = await CreateObjectAsync(Nodes);
 
-			ReadOnlyMemory<byte> Data = EncodeObject(Object);
+			ReadOnlyMemory<byte> RefData = EncodeObject(Ref);
 			NextStats.NewRefCount++;
-			NextStats.NewRefBytes += Data.Length;
+			NextStats.NewRefBytes += RefData.Length;
 
-			await StorageClient.SetRefAsync(NamespaceId, BucketId, RefId, new CbField(Data));
+			await StorageClient.SetRefAsync(NamespaceId, BucketId, RefId, new CbField(RefData));
 		}
 
-		ReadOnlyMemory<byte> EncodeObject(BundleObject Object)
+		ReadOnlyMemory<byte> EncodeObject<T>(T Object)
 		{
 			SharedWriter.Clear();
-			CbSerializer.Serialize(SharedWriter, Object);
+			CbSerializer.Serialize<T>(SharedWriter, Object);
 
 			int Size = SharedWriter.GetSize();
 			CreateFreeSpace(ref SharedBlobBuffer, 0, Math.Max(Size, Options.MaxBlobSize));
diff --git a/Engine/Source/Programs/Shared/EpicGames.Horde/Bundles/BundleObject.cs b/Engine/Source/Programs/Shared/EpicGames.Horde/Bundles/BundleObject.cs
--- a/Engine/Source/Programs/Shared/EpicGames.Horde/Bundles/BundleObject.cs
+++ b/Engine/Source/Programs/Shared/EpicGames.Horde/Bundles/BundleObject.cs
@@ -21,6 +21,24 @@ using System.Threading.Tasks;
 
 namespace EpicGames.Horde.Bundles
 {
+	/// <summary>
+	/// Object stored within the ref for a bundle
+	/// </summary>
+	public class BundleRoot
+	{
+		/// <summary>
+		/// Data hashed to produce the ref key; included in unhashed form for debugging purposes.
+		/// </summary>
+		[CbField("metadata")]
+		public CbObject Metadata { get; set; } = CbObject.Empty;
+
+		/// <summary>
+		/// Embedded object data
+		/// </summary>
+		[CbField("object")]
+		public BundleObject Object { get; set; } = new BundleObject();
+	}
+
 	/// <summary>
 	/// Blob within the bundle.
 	/// </summary>
