Horde: Fix tests.

diff --git a/Engine/Source/Programs/Horde/Horde.Agent/Commands/Bundles/CreateCommand.cs b/Engine/Source/Programs/Horde/Horde.Agent/Commands/Bundles/CreateCommand.cs
--- a/Engine/Source/Programs/Horde/Horde.Agent/Commands/Bundles/CreateCommand.cs
+++ b/Engine/Source/Programs/Horde/Horde.Agent/Commands/Bundles/CreateCommand.cs
@@ -72,7 +72,7 @@ namespace HordeAgent.Commands.Bundles
 
 			Bundle<DirectoryNode> NewBundle = Bundle.Create<DirectoryNode>(StorageClient, NamespaceId, new BundleOptions(), null);
 			await NewBundle.Root.CopyFromDirectoryAsync(InputDir.ToDirectoryInfo(), new ChunkingOptions(), Logger);
-			await NewBundle.WriteAsync(BucketId, RefId, false);
+			await NewBundle.WriteAsync(BucketId, RefId, CbObject.Empty, false);
 
 			return 0;
 		}
diff --git a/Engine/Source/Programs/Shared/EpicGames.Horde.Tests/BundleTests.cs b/Engine/Source/Programs/Shared/EpicGames.Horde.Tests/BundleTests.cs
--- a/Engine/Source/Programs/Shared/EpicGames.Horde.Tests/BundleTests.cs
+++ b/Engine/Source/Programs/Shared/EpicGames.Horde.Tests/BundleTests.cs
@@ -149,7 +149,7 @@ namespace HordeServerTests
 			DirectoryNode Node2 = Node.AddDirectory("world");
 
 			RefId RefId = new RefId("testref");
-			await NewBundle.WriteAsync(BucketId, RefId, false);
+			await NewBundle.WriteAsync(BucketId, RefId, CbObject.Empty, false);
 
 			// Should be stored inline
 			Assert.AreEqual(1, StorageClient.Refs.Count);
@@ -196,7 +196,7 @@ namespace HordeServerTests
 			Assert.AreEqual(0, StorageClient.Blobs.Count);
 
 			RefId RefId = new RefId("ref");
-			await Bundle.WriteAsync(BucketId, RefId, false);
+			await Bundle.WriteAsync(BucketId, RefId, CbObject.Empty, false);
 
 			Assert.AreEqual(1, StorageClient.Refs.Count);
 			Assert.AreEqual(1, StorageClient.Blobs.Count);
@@ -217,7 +217,7 @@ namespace HordeServerTests
 			DirectoryNode Node4 = Node3.AddDirectory("node4");
 
 			RefId RefId = new RefId("ref");
-			await InitialBundle.WriteAsync(BucketId, RefId, false);
+			await InitialBundle.WriteAsync(BucketId, RefId, CbObject.Empty, false);
 
 			Assert.AreEqual(1, StorageClient.Refs.Count);
 			Assert.AreEqual(4, StorageClient.Blobs.Count);
@@ -252,7 +252,7 @@ namespace HordeServerTests
 			DirectoryNode Node4 = Bundle.Root.AddDirectory("node4"); // same contents as node 3
 
 			RefId RefId1 = new RefId("ref1");
-			await Bundle.WriteAsync(BucketId, RefId1, false);
+			await Bundle.WriteAsync(BucketId, RefId1, CbObject.Empty, false);
 
 			Assert.AreEqual(1, StorageClient.Refs.Count);
 			Assert.AreEqual(1, StorageClient.Blobs.Count);
@@ -272,7 +272,7 @@ namespace HordeServerTests
 			Bundle.Root.DeleteDirectory("node1");
 
 			RefId RefId2 = new RefId("ref2");
-			await Bundle.WriteAsync(BucketId, RefId2, false);
+			await Bundle.WriteAsync(BucketId, RefId2, CbObject.Empty, false);
 
 			IRef Ref2 = StorageClient.Refs[(NamespaceId, BucketId, RefId2)];
 
@@ -287,7 +287,7 @@ namespace HordeServerTests
 
 			// Repack it and check that we make a new object
 			RefId RefId3 = new RefId("ref3");
-			await Bundle.WriteAsync(BucketId, RefId3, true);
+			await Bundle.WriteAsync(BucketId, RefId3, CbObject.Empty, true);
 
 			IRef Ref3 = StorageClient.Refs[(NamespaceId, BucketId, RefId3)];
 
@@ -434,7 +434,7 @@ namespace HordeServerTests
 			Assert.IsTrue(StorageClient.Refs.Count == 0);
 
 			RefId RefId = new RefId("ref");
-			await Bundle.WriteAsync(BucketId, RefId, true);
+			await Bundle.WriteAsync(BucketId, RefId, CbObject.Empty, true);
 
 			Assert.AreEqual(TotalLength, Root.Length);
 
