Horde: Fix error deserializing file node data from bundles.

diff --git a/Engine/Source/Programs/Shared/EpicGames.Horde/Bundles/Nodes/FileNode.cs b/Engine/Source/Programs/Shared/EpicGames.Horde/Bundles/Nodes/FileNode.cs
--- a/Engine/Source/Programs/Shared/EpicGames.Horde/Bundles/Nodes/FileNode.cs
+++ b/Engine/Source/Programs/Shared/EpicGames.Horde/Bundles/Nodes/FileNode.cs
@@ -149,6 +149,8 @@ namespace EpicGames.Horde.Bundles.Nodes
 			Node.Depth = (int)VarInt.Read(Span.Slice(1), out int DepthBytes);
 
 			int HeaderLength = 1 + DepthBytes;
+			Span = Span.Slice(HeaderLength);
+
 			ReadOnlyMemory<byte> Payload = Data.Slice(HeaderLength);
 			Node.Payload = new ReadOnlySequence<byte>(Payload);
 
@@ -162,12 +164,11 @@ namespace EpicGames.Horde.Bundles.Nodes
 				Span = Span.Slice(LengthBytes);
 
 				Node.ChildNodeRefs = new List<BundleNodeRef<FileNode>>(Payload.Length / IoHash.NumBytes);
-
-				Span = Span.Slice(HeaderLength);
 				while (Span.Length > 0)
 				{
 					IoHash ChildHash = new IoHash(Span);
 					Node.ChildNodeRefs.Add(new BundleNodeRef<FileNode>(Node, ChildHash));
+					Span = Span.Slice(IoHash.NumBytes);
 				}
 			}
 			return Node;
