1) Fix minor validation issue where TObjectPtr was being denied as a function return type.  Also applied to lazy pointers.
2) Fixed TObjectPtr formatting for return value
3) Fixed TMap formatting when template arguments are also templates.
4) Fixed -Test option

#rnx

diff --git a/Engine/Source/Programs/Shared/EpicGames.UHT/Parsers/UhtEnumParser.cs b/Engine/Source/Programs/Shared/EpicGames.UHT/Parsers/UhtEnumParser.cs
--- a/Engine/Source/Programs/Shared/EpicGames.UHT/Parsers/UhtEnumParser.cs
+++ b/Engine/Source/Programs/Shared/EpicGames.UHT/Parsers/UhtEnumParser.cs
@@ -170,6 +170,9 @@ namespace EpicGames.UHT.Parsers
 						int EnumIndex = Enum.AddEnumValue(TagToken.Value.ToString(), 0);
 						TopScope.AddFormattedCommentsAsTooltipMetaData(EnumIndex);
 
+						// Skip any deprecation
+						TopScope.TokenReader.SkipDeprecatedMacroIfNecessary();
+
 						// Try to read an optional explicit enum value specification
 						if (TopScope.TokenReader.TryOptional('='))
 						{
diff --git a/Engine/Source/Programs/Shared/EpicGames.UHT/Types/Properties/UhtClassPtrProperties.cs b/Engine/Source/Programs/Shared/EpicGames.UHT/Types/Properties/UhtClassPtrProperties.cs
--- a/Engine/Source/Programs/Shared/EpicGames.UHT/Types/Properties/UhtClassPtrProperties.cs
+++ b/Engine/Source/Programs/Shared/EpicGames.UHT/Types/Properties/UhtClassPtrProperties.cs
@@ -75,7 +75,7 @@ namespace EpicGames.UHT.Types
 
 			// UFunctions with a smart pointer as input parameter wont compile anyway, because of missing P_GET_... macro.
 			// UFunctions with a smart pointer as return type will crash when called via blueprint, because they are not supported in VM.
-			if (this.PropertyCategory != UhtPropertyCategory.Member)
+			if (this.PropertyCategory == UhtPropertyCategory.RegularParameter || this.PropertyCategory == UhtPropertyCategory.ReplicatedParameter)
 			{
 				OuterStruct.LogError("UFunctions cannot take a TObjectPtr as a parameter.");
 			}
diff --git a/Engine/Source/Programs/Shared/EpicGames.UHT/Types/Properties/UhtLazyObjectPtrProperty.cs b/Engine/Source/Programs/Shared/EpicGames.UHT/Types/Properties/UhtLazyObjectPtrProperty.cs
--- a/Engine/Source/Programs/Shared/EpicGames.UHT/Types/Properties/UhtLazyObjectPtrProperty.cs
+++ b/Engine/Source/Programs/Shared/EpicGames.UHT/Types/Properties/UhtLazyObjectPtrProperty.cs
@@ -98,7 +98,7 @@ namespace EpicGames.UHT.Types
 
 			// UFunctions with a smart pointer as input parameter wont compile anyway, because of missing P_GET_... macro.
 			// UFunctions with a smart pointer as return type will crash when called via blueprint, because they are not supported in VM.
-			if (this.PropertyCategory != UhtPropertyCategory.Member)
+			if (this.PropertyCategory == UhtPropertyCategory.RegularParameter || this.PropertyCategory == UhtPropertyCategory.ReplicatedParameter)
 			{
 				OuterStruct.LogError("UFunctions cannot take a lazy pointer as a parameter.");
 			}
diff --git a/Engine/Source/Programs/Shared/EpicGames.UHT/Types/Properties/UhtMapProperty.cs b/Engine/Source/Programs/Shared/EpicGames.UHT/Types/Properties/UhtMapProperty.cs
--- a/Engine/Source/Programs/Shared/EpicGames.UHT/Types/Properties/UhtMapProperty.cs
+++ b/Engine/Source/Programs/Shared/EpicGames.UHT/Types/Properties/UhtMapProperty.cs
@@ -137,7 +137,19 @@ namespace EpicGames.UHT.Types
 					break;
 
 				default:
-					Builder.Append("TMap<").AppendPropertyText(this.KeyProperty, TextType, true).Append(',').AppendPropertyText(this.ValueProperty, TextType, true).Append('>');
+					Builder.Append("TMap<").AppendPropertyText(this.KeyProperty, TextType, true);
+					if (Builder[Builder.Length - 1] == '>')
+					{
+						// if our internal property type is a template class, add a space between the closing brackets b/c VS.NET cannot parse this correctly
+						Builder.Append(" ");
+					}
+					Builder.Append(',').AppendPropertyText(this.ValueProperty, TextType, true);
+					if (Builder[Builder.Length - 1] == '>')
+					{
+						// if our internal property type is a template class, add a space between the closing brackets b/c VS.NET cannot parse this correctly
+						Builder.Append(" ");
+					}
+					Builder.Append('>');
 					break;
 			}
 			return Builder;
diff --git a/Engine/Source/Programs/Shared/EpicGames.UHT/Types/Properties/UhtObjectPtrProperty.cs b/Engine/Source/Programs/Shared/EpicGames.UHT/Types/Properties/UhtObjectPtrProperty.cs
--- a/Engine/Source/Programs/Shared/EpicGames.UHT/Types/Properties/UhtObjectPtrProperty.cs
+++ b/Engine/Source/Programs/Shared/EpicGames.UHT/Types/Properties/UhtObjectPtrProperty.cs
@@ -54,7 +54,7 @@ namespace EpicGames.UHT.Types
 					{
 						Builder.Append("const ");
 					}
-					Builder.Append(this.Class.SourceName).Append("*");
+					Builder.Append("TObjectPtr<").Append(this.Class.SourceName).Append(">");
 					break;
 
 				default:
@@ -86,7 +86,7 @@ namespace EpicGames.UHT.Types
 
 			// UFunctions with a smart pointer as input parameter wont compile anyway, because of missing P_GET_... macro.
 			// UFunctions with a smart pointer as return type will crash when called via blueprint, because they are not supported in VM.
-			if (this.PropertyCategory != UhtPropertyCategory.Member)
+			if (this.PropertyCategory == UhtPropertyCategory.RegularParameter || this.PropertyCategory == UhtPropertyCategory.ReplicatedParameter)
 			{
 				OuterStruct.LogError("UFunctions cannot take a TObjectPtr as a parameter.");
 			}
diff --git a/Engine/Source/Programs/UnrealBuildTool/Modes/UnrealHeaderToolMode.cs b/Engine/Source/Programs/UnrealBuildTool/Modes/UnrealHeaderToolMode.cs
--- a/Engine/Source/Programs/UnrealBuildTool/Modes/UnrealHeaderToolMode.cs
+++ b/Engine/Source/Programs/UnrealBuildTool/Modes/UnrealHeaderToolMode.cs
@@ -1004,7 +1004,7 @@ namespace UnrealBuildTool.Modes
 						}
 					}
 				}
-				int RequiredArgCount = TargetArgumentIndex >= 0 ? 0 : 2;
+				int RequiredArgCount = TargetArgumentIndex >= 0 || Options.bTest ? 0 : 2;
 				if (Arguments.GetPositionalArgumentCount() != RequiredArgCount || Options.bGetHelp)
 				{
 					PrintUsage();
