Horde: Prevent pipeline bubbles when syncing caused by waiting for all tasks to complete on each iteration. Also increase gap before each save to 5 minutes.

diff --git a/Engine/Source/Programs/Shared/EpicGames.Perforce.Managed/ManagedWorkspace.cs b/Engine/Source/Programs/Shared/EpicGames.Perforce.Managed/ManagedWorkspace.cs
--- a/Engine/Source/Programs/Shared/EpicGames.Perforce.Managed/ManagedWorkspace.cs
+++ b/Engine/Source/Programs/Shared/EpicGames.Perforce.Managed/ManagedWorkspace.cs
@@ -1789,49 +1789,50 @@ namespace EpicGames.Perforce.Managed
 
 					// Spawn some background threads to sync them
 					Dictionary<Task, int> Tasks = new Dictionary<Task, int>();
-					while (Tasks.Count > 0 || NextBatchIdx < Batches.Count)
+					try
 					{
-						// Create new tasks
-						while (Tasks.Count < NumParallelSyncThreads && NextBatchIdx < Batches.Count)
+						while (Tasks.Count > 0 || NextBatchIdx < Batches.Count)
 						{
-							(int BatchBeginIdx, int BatchEndIdx) = Batches[NextBatchIdx];
+							// Create new tasks
+							while (Tasks.Count < NumParallelSyncThreads && NextBatchIdx < Batches.Count)
+							{
+								(int BatchBeginIdx, int BatchEndIdx) = Batches[NextBatchIdx];
 
-							Task Task = Task.Run(() => SyncBatch(Client, FilesToSync, BatchBeginIdx, BatchEndIdx, bFakeSync, CancellationToken));
-							Tasks[Task] = NextBatchIdx++;
-						}
+								Task Task = Task.Run(() => SyncBatch(Client, FilesToSync, BatchBeginIdx, BatchEndIdx, bFakeSync, CancellationToken));
+								Tasks[Task] = NextBatchIdx++;
+							}
 
-						// Wait for anything to complete
-						Task CompleteTask = await Task.WhenAny(Tasks.Keys);
-						try
-						{
+							// Wait for anything to complete
+							Task CompleteTask = await Task.WhenAny(Tasks.Keys);
 							await CompleteTask; // Make sure we re-throw any exceptions from the task that completed
-						}
-						finally
-						{
-							await Task.WhenAll(Tasks.Keys);
-						}
-						int BatchIdx = Tasks[CompleteTask];
-						Tasks.Remove(CompleteTask);
 
-						// Update metadata for the complete batch
-						(int BeginIdx, int EndIdx) = Batches[BatchIdx];
-						await ParallelTask.ForAsync(BeginIdx, EndIdx, Idx => FilesToSync[Idx].WorkspaceFile.UpdateMetadata());
+							int BatchIdx = Tasks[CompleteTask];
+							Tasks.Remove(CompleteTask);
 
-						// Save the current state every minute
-						TimeSpan Elapsed = Timer.Elapsed;
-						if (Elapsed > TimeSpan.FromMinutes(1.0))
-						{
-							await SaveAsync(TransactionState.Dirty, CancellationToken);
-							Logger.LogInformation("Saved workspace state ({Elapsed:0.0}s)", (Timer.Elapsed - Elapsed).TotalSeconds);
-							Timer.Restart();
-						}
+							// Update metadata for the complete batch
+							(int BeginIdx, int EndIdx) = Batches[BatchIdx];
+							await ParallelTask.ForAsync(BeginIdx, EndIdx, Idx => FilesToSync[Idx].WorkspaceFile.UpdateMetadata());
 
-						// Update the status
-						for (int Idx = BeginIdx; Idx < EndIdx; Idx++)
-						{
-							SyncedSize += FilesToSync[Idx].StreamFile.Length;
+							// Save the current state every minute
+							TimeSpan Elapsed = Timer.Elapsed;
+							if (Elapsed > TimeSpan.FromMinutes(5.0))
+							{
+								await SaveAsync(TransactionState.Dirty, CancellationToken);
+								Logger.LogInformation("Saved workspace state ({Elapsed:0.0}s)", (Timer.Elapsed - Elapsed).TotalSeconds);
+								Timer.Restart();
+							}
+
+							// Update the status
+							for (int Idx = BeginIdx; Idx < EndIdx; Idx++)
+							{
+								SyncedSize += FilesToSync[Idx].StreamFile.Length;
+							}
+							Status.Progress = String.Format("{0:n1}% ({1:n1}mb/{2:n1}mb)", SyncedSize * 100.0 / SyncSize, SyncedSize / (1024.0 * 1024.0), SyncSize / (1024.0 * 1024.0));
 						}
-						Status.Progress = String.Format("{0:n1}% ({1:n1}mb/{2:n1}mb)", SyncedSize * 100.0 / SyncSize, SyncedSize / (1024.0 * 1024.0), SyncSize / (1024.0 * 1024.0));
+					}
+					finally
+					{
+						await Task.WhenAll(Tasks.Keys);
 					}
 				}
 			}
