Horde: Support thread safe reads from bundles.

- Node state is now separated into a separate class to allow atomic transitions between imported/exported/standalone states.
- A semaphore is used to limit concurrent reads from underlying storage, so bundles now implement IDisposable.
- DirectoryNode.CopyToDirectoryAsync() spawns multiple tasks to extract data in parallel. The current tree traversal is fairly naive; since it doesn't know which objects are currently in the cache, it can result in a lot of thrashing if the cache is full as objects are purged and re-read to match the traversal order.
- Mutation of the tree is still restricted to a single thread.

diff --git a/Engine/Source/Programs/Horde/Horde.Agent/Commands/Bundles/CreateCommand.cs b/Engine/Source/Programs/Horde/Horde.Agent/Commands/Bundles/CreateCommand.cs
--- a/Engine/Source/Programs/Horde/Horde.Agent/Commands/Bundles/CreateCommand.cs
+++ b/Engine/Source/Programs/Horde/Horde.Agent/Commands/Bundles/CreateCommand.cs
@@ -6,6 +6,7 @@ using EpicGames.Horde.Bundles.Nodes;
 using EpicGames.Horde.Storage;
 using EpicGames.Horde.Storage.Impl;
 using EpicGames.Serialization;
+using Microsoft.Extensions.Caching.Memory;
 using Microsoft.Extensions.Configuration;
 using Microsoft.Extensions.DependencyInjection;
 using Microsoft.Extensions.Logging;
@@ -69,10 +70,14 @@ namespace HordeAgent.Commands.Bundles
 		public override async Task<int> ExecuteAsync(ILogger Logger)
 		{
 			IStorageClient StorageClient = base.CreateStorageClient(Logger);
-
-			Bundle<DirectoryNode> NewBundle = Bundle.Create<DirectoryNode>(StorageClient, NamespaceId, new BundleOptions(), null);
-			await NewBundle.Root.CopyFromDirectoryAsync(InputDir.ToDirectoryInfo(), new ChunkingOptions(), Logger);
-			await NewBundle.WriteAsync(BucketId, RefId, CbObject.Empty, false);
+			using (MemoryCache Cache = new MemoryCache(new MemoryCacheOptions { SizeLimit = 50 * 1024 * 1024 }))
+			{
+				using (Bundle<DirectoryNode> NewBundle = Bundle.Create<DirectoryNode>(StorageClient, NamespaceId, new BundleOptions(), Cache))
+				{
+					await NewBundle.Root.CopyFromDirectoryAsync(InputDir.ToDirectoryInfo(), new ChunkingOptions(), Logger);
+					await NewBundle.WriteAsync(BucketId, RefId, CbObject.Empty, false);
+				}
+			}
 
 			return 0;
 		}
diff --git a/Engine/Source/Programs/Horde/Horde.Agent/Commands/Bundles/ExtractCommand.cs b/Engine/Source/Programs/Horde/Horde.Agent/Commands/Bundles/ExtractCommand.cs
--- a/Engine/Source/Programs/Horde/Horde.Agent/Commands/Bundles/ExtractCommand.cs
+++ b/Engine/Source/Programs/Horde/Horde.Agent/Commands/Bundles/ExtractCommand.cs
@@ -38,7 +38,7 @@ namespace HordeAgent.Commands.Bundles
 			IStorageClient StorageClient = CreateStorageClient(Logger);
 			using (IMemoryCache Cache = new MemoryCache(new MemoryCacheOptions { SizeLimit = 1024 * 1024 * 1000 }))
 			{
-				Bundle<DirectoryNode> NewBundle = await Bundle.ReadAsync<DirectoryNode>(StorageClient, NamespaceId, BucketId, RefId, new BundleOptions(), Cache);
+				using Bundle<DirectoryNode> NewBundle = await Bundle.ReadAsync<DirectoryNode>(StorageClient, NamespaceId, BucketId, RefId, new BundleOptions(), Cache);
 				await NewBundle.Root.CopyToDirectoryAsync(NewBundle, OutputDir.ToDirectoryInfo(), Logger);
 			}
 			return 0;
diff --git a/Engine/Source/Programs/Horde/Horde.Build.Tests/Horde.Build.Tests.csproj b/Engine/Source/Programs/Horde/Horde.Build.Tests/Horde.Build.Tests.csproj
--- a/Engine/Source/Programs/Horde/Horde.Build.Tests/Horde.Build.Tests.csproj
+++ b/Engine/Source/Programs/Horde/Horde.Build.Tests/Horde.Build.Tests.csproj
@@ -22,6 +22,7 @@
     <PackageReference Include="Moq" Version="4.14.5" />
     <PackageReference Include="MSTest.TestAdapter" Version="2.1.0" />
     <PackageReference Include="MSTest.TestFramework" Version="2.1.0" />
+    <PackageReference Include="Microsoft.Extensions.Caching.Memory" Version="6.0.0" />
     <PackageReference Include="coverlet.collector" Version="1.2.0" />
   </ItemGroup>
 
diff --git a/Engine/Source/Programs/Horde/Horde.Build/Commands/Bundles/CreateCommand.cs b/Engine/Source/Programs/Horde/Horde.Build/Commands/Bundles/CreateCommand.cs
--- a/Engine/Source/Programs/Horde/Horde.Build/Commands/Bundles/CreateCommand.cs
+++ b/Engine/Source/Programs/Horde/Horde.Build/Commands/Bundles/CreateCommand.cs
@@ -6,6 +6,7 @@ using EpicGames.Horde.Bundles.Nodes;
 using EpicGames.Horde.Storage;
 using EpicGames.Horde.Storage.Impl;
 using EpicGames.Serialization;
+using Microsoft.Extensions.Caching.Memory;
 using Microsoft.Extensions.Configuration;
 using Microsoft.Extensions.DependencyInjection;
 using Microsoft.Extensions.Logging;
@@ -47,10 +48,14 @@ namespace HordeServer.Commands.Bundles
 		public override async Task<int> ExecuteAsync(ILogger Logger)
 		{
 			IStorageClient StorageClient = base.CreateStorageClient(Logger);
-
-			Bundle<DirectoryNode> NewBundle = Bundle.Create<DirectoryNode>(StorageClient, NamespaceId, new BundleOptions(), null);
-			await NewBundle.Root.CopyFromDirectoryAsync(InputDir.ToDirectoryInfo(), new ChunkingOptions(), Logger);
-			await NewBundle.WriteAsync(BucketId, RefId, CbObject.Empty, false);
+			using (MemoryCache Cache = new MemoryCache(new MemoryCacheOptions { SizeLimit = 50 * 1024 * 1024 }))
+			{
+				using (Bundle<DirectoryNode> NewBundle = Bundle.Create<DirectoryNode>(StorageClient, NamespaceId, new BundleOptions(), Cache))
+				{
+					await NewBundle.Root.CopyFromDirectoryAsync(InputDir.ToDirectoryInfo(), new ChunkingOptions(), Logger);
+					await NewBundle.WriteAsync(BucketId, RefId, CbObject.Empty, false);
+				}
+			}
 
 			return 0;
 		}
diff --git a/Engine/Source/Programs/Horde/Horde.Build/Commands/Bundles/ExtractCommand.cs b/Engine/Source/Programs/Horde/Horde.Build/Commands/Bundles/ExtractCommand.cs
--- a/Engine/Source/Programs/Horde/Horde.Build/Commands/Bundles/ExtractCommand.cs
+++ b/Engine/Source/Programs/Horde/Horde.Build/Commands/Bundles/ExtractCommand.cs
@@ -36,9 +36,9 @@ namespace HordeServer.Commands.Bundles
 		public override async Task<int> ExecuteAsync(ILogger Logger)
 		{
 			IStorageClient StorageClient = CreateStorageClient(Logger);
-			using (IMemoryCache Cache = new MemoryCache(new MemoryCacheOptions { SizeLimit = 1024 * 1024 * 1000 }))
+			using (IMemoryCache Cache = new MemoryCache(new MemoryCacheOptions()))
 			{
-				Bundle<DirectoryNode> NewBundle = await Bundle.ReadAsync<DirectoryNode>(StorageClient, NamespaceId, BucketId, RefId, new BundleOptions(), Cache);
+				using Bundle<DirectoryNode> NewBundle = await Bundle.ReadAsync<DirectoryNode>(StorageClient, NamespaceId, BucketId, RefId, new BundleOptions(), Cache);
 				await NewBundle.Root.CopyToDirectoryAsync(NewBundle, OutputDir.ToDirectoryInfo(), Logger);
 			}
 			return 0;
diff --git a/Engine/Source/Programs/Shared/EpicGames.Horde.Tests/BundleTests.cs b/Engine/Source/Programs/Shared/EpicGames.Horde.Tests/BundleTests.cs
--- a/Engine/Source/Programs/Shared/EpicGames.Horde.Tests/BundleTests.cs
+++ b/Engine/Source/Programs/Shared/EpicGames.Horde.Tests/BundleTests.cs
@@ -6,6 +6,7 @@ using EpicGames.Horde.Bundles.Nodes;
 using EpicGames.Horde.Storage;
 using EpicGames.Horde.Storage.Impl;
 using EpicGames.Serialization;
+using Microsoft.Extensions.Caching.Memory;
 using Microsoft.VisualStudio.TestTools.UnitTesting;
 using System;
 using System.Buffers.Binary;
@@ -18,12 +19,23 @@ using System.Threading.Tasks;
 namespace HordeServerTests
 {
 	[TestClass]
-	public class BundleTests
+	public class BundleTests : IDisposable
 	{
+		MemoryCache Cache;
 		MemoryStorageClient StorageClient = new MemoryStorageClient();
 		NamespaceId NamespaceId = new NamespaceId("namespace");
 		BucketId BucketId = new BucketId("bucket");
 
+		public BundleTests()
+		{
+			Cache = new MemoryCache(new MemoryCacheOptions { SizeLimit = 50 * 1024 * 1024 });
+		}
+
+		public void Dispose()
+		{
+			Cache.Dispose();
+		}
+
 		[TestMethod]
 		public void BuzHashTests()
 		{
@@ -144,7 +156,7 @@ namespace HordeServerTests
 		[TestMethod]
 		public async Task BasicTestDirectory()
 		{
-			Bundle<DirectoryNode> NewBundle = new Bundle<DirectoryNode>(StorageClient, NamespaceId, new DirectoryNode(), new BundleOptions(), null);
+			using Bundle<DirectoryNode> NewBundle = Bundle.Create<DirectoryNode>(StorageClient, NamespaceId, new DirectoryNode(), new BundleOptions(), Cache);
 			DirectoryNode Node = NewBundle.Root.AddDirectory("hello");
 			DirectoryNode Node2 = Node.AddDirectory("world");
 
@@ -156,18 +168,19 @@ namespace HordeServerTests
 			Assert.AreEqual(0, StorageClient.Blobs.Count);
 
 			// Check the ref
-			BundleObject RootObject = await StorageClient.GetRefAsync<BundleObject>(NamespaceId, BucketId, RefId);
+			BundleRoot Root = await StorageClient.GetRefAsync<BundleRoot>(NamespaceId, BucketId, RefId);
+			BundleObject RootObject = Root.Object;
 			Assert.AreEqual(3, RootObject.Exports.Count);
 			Assert.AreEqual(0, RootObject.Exports[0].Rank);
 			Assert.AreEqual(1, RootObject.Exports[1].Rank);
 			Assert.AreEqual(2, RootObject.Exports[2].Rank);
 
 			// Create a new bundle and read it back in again
-			NewBundle = await Bundle.ReadAsync<DirectoryNode>(StorageClient, NamespaceId, BucketId, RefId, new BundleOptions(), null);
+			using Bundle<DirectoryNode> NewBundle2 = await Bundle.ReadAsync<DirectoryNode>(StorageClient, NamespaceId, BucketId, RefId, new BundleOptions(), Cache);
 
-			Assert.AreEqual(0, NewBundle.Root.Files.Count);
-			Assert.AreEqual(1, NewBundle.Root.Directories.Count);
-			DirectoryNode? OutputNode = await NewBundle.FindDirectoryAsync(NewBundle.Root, "hello");
+			Assert.AreEqual(0, NewBundle2.Root.Files.Count);
+			Assert.AreEqual(1, NewBundle2.Root.Directories.Count);
+			DirectoryNode? OutputNode = await NewBundle.FindDirectoryAsync(NewBundle2.Root, "hello");
 			Assert.IsNotNull(OutputNode);
 
 			Assert.AreEqual(0, OutputNode!.Files.Count);
@@ -186,17 +199,17 @@ namespace HordeServerTests
 			Options.MaxBlobSize = 1;
 			Options.MaxInlineBlobSize = 1;
 
-			Bundle<DirectoryNode> Bundle = new Bundle<DirectoryNode>(StorageClient, NamespaceId, new DirectoryNode(), Options, null);
+			using Bundle<DirectoryNode> NewBundle = Bundle.Create<DirectoryNode>(StorageClient, NamespaceId, new DirectoryNode(), Options, Cache);
 
-			Bundle.Root.AddDirectory("node1");
-			Bundle.Root.AddDirectory("node2");
-			Bundle.Root.AddDirectory("node3");
+			NewBundle.Root.AddDirectory("node1");
+			NewBundle.Root.AddDirectory("node2");
+			NewBundle.Root.AddDirectory("node3");
 
 			Assert.AreEqual(0, StorageClient.Refs.Count);
 			Assert.AreEqual(0, StorageClient.Blobs.Count);
 
 			RefId RefId = new RefId("ref");
-			await Bundle.WriteAsync(BucketId, RefId, CbObject.Empty, false);
+			await NewBundle.WriteAsync(BucketId, RefId, CbObject.Empty, false);
 
 			Assert.AreEqual(1, StorageClient.Refs.Count);
 			Assert.AreEqual(1, StorageClient.Blobs.Count);
@@ -209,7 +222,7 @@ namespace HordeServerTests
 			Options.MaxBlobSize = 1;
 			Options.MaxInlineBlobSize = 1;
 
-			Bundle<DirectoryNode> InitialBundle = new Bundle<DirectoryNode>(StorageClient, NamespaceId, new DirectoryNode(), Options, null);
+			using Bundle<DirectoryNode> InitialBundle = Bundle.Create<DirectoryNode>(StorageClient, NamespaceId, new DirectoryNode(), Options, Cache);
 
 			DirectoryNode Node1 = InitialBundle.Root.AddDirectory("node1");
 			DirectoryNode Node2 = Node1.AddDirectory("node2");
@@ -222,7 +235,7 @@ namespace HordeServerTests
 			Assert.AreEqual(1, StorageClient.Refs.Count);
 			Assert.AreEqual(4, StorageClient.Blobs.Count);
 
-			Bundle<DirectoryNode> NewBundle = await Bundle.ReadAsync<DirectoryNode>(StorageClient, NamespaceId, BucketId, RefId, Options, null);
+			using Bundle<DirectoryNode> NewBundle = await Bundle.ReadAsync<DirectoryNode>(StorageClient, NamespaceId, BucketId, RefId, Options, Cache);
 
 			DirectoryNode? NewNode1 = await NewBundle.FindDirectoryAsync(NewBundle.Root, "node1");
 			Assert.IsNotNull(NewNode1);
@@ -244,22 +257,23 @@ namespace HordeServerTests
 			Options.MaxBlobSize = 1024 * 1024;
 			Options.MaxInlineBlobSize = 1;
 
-			Bundle<DirectoryNode> Bundle = new Bundle<DirectoryNode>(StorageClient, NamespaceId, new DirectoryNode(), Options, null);
+			using Bundle<DirectoryNode> NewBundle = Bundle.Create<DirectoryNode>(StorageClient, NamespaceId, new DirectoryNode(), Options, Cache);
 
-			DirectoryNode Node1 = Bundle.Root.AddDirectory("node1");
+			DirectoryNode Node1 = NewBundle.Root.AddDirectory("node1");
 			DirectoryNode Node2 = Node1.AddDirectory("node2");
 			DirectoryNode Node3 = Node2.AddDirectory("node3");
-			DirectoryNode Node4 = Bundle.Root.AddDirectory("node4"); // same contents as node 3
+			DirectoryNode Node4 = NewBundle.Root.AddDirectory("node4"); // same contents as node 3
 
 			RefId RefId1 = new RefId("ref1");
-			await Bundle.WriteAsync(BucketId, RefId1, CbObject.Empty, false);
+			await NewBundle.WriteAsync(BucketId, RefId1, CbObject.Empty, false);
 
 			Assert.AreEqual(1, StorageClient.Refs.Count);
 			Assert.AreEqual(1, StorageClient.Blobs.Count);
 
 			IRef Ref1 = StorageClient.Refs[(NamespaceId, BucketId, RefId1)];
 
-			BundleObject RootObject1 = CbSerializer.Deserialize<BundleObject>(Ref1.Value);
+			BundleRoot Root1 = CbSerializer.Deserialize<BundleRoot>(Ref1.Value);
+			BundleObject RootObject1 = Root1.Object;
 			Assert.AreEqual(1, RootObject1.Exports.Count);
 			Assert.AreEqual(1, RootObject1.ImportObjects.Count);
 
@@ -269,14 +283,15 @@ namespace HordeServerTests
 			Assert.AreEqual(0, LeafObject1.ImportObjects.Count);
 
 			// Remove one of the nodes from the root without compacting. the existing blob should be reused.
-			Bundle.Root.DeleteDirectory("node1");
+			NewBundle.Root.DeleteDirectory("node1");
 
 			RefId RefId2 = new RefId("ref2");
-			await Bundle.WriteAsync(BucketId, RefId2, CbObject.Empty, false);
+			await NewBundle.WriteAsync(BucketId, RefId2, CbObject.Empty, false);
 
 			IRef Ref2 = StorageClient.Refs[(NamespaceId, BucketId, RefId2)];
 
-			BundleObject RootObject2 = CbSerializer.Deserialize<BundleObject>(Ref2.Value);
+			BundleRoot Root2 = CbSerializer.Deserialize<BundleRoot>(Ref2.Value);
+			BundleObject RootObject2 = Root2.Object;
 			Assert.AreEqual(1, RootObject2.Exports.Count);
 			Assert.AreEqual(1, RootObject2.ImportObjects.Count);
 
@@ -287,11 +302,12 @@ namespace HordeServerTests
 
 			// Repack it and check that we make a new object
 			RefId RefId3 = new RefId("ref3");
-			await Bundle.WriteAsync(BucketId, RefId3, CbObject.Empty, true);
+			await NewBundle.WriteAsync(BucketId, RefId3, CbObject.Empty, true);
 
 			IRef Ref3 = StorageClient.Refs[(NamespaceId, BucketId, RefId3)];
 
-			BundleObject RootObject3 = CbSerializer.Deserialize<BundleObject>(Ref3.Value);
+			BundleRoot Root3 = CbSerializer.Deserialize<BundleRoot>(Ref3.Value);
+			BundleObject RootObject3 = Root3.Object;
 			Assert.AreEqual(1, RootObject3.Exports.Count);
 			Assert.AreEqual(1, RootObject3.ImportObjects.Count);
 
@@ -352,8 +368,8 @@ namespace HordeServerTests
 				Data[Idx] = (byte)Idx;
 			}
 
-			Bundle<FileNode> Bundle = new Bundle<FileNode>(StorageClient, NamespaceId, new FileNode(), new BundleOptions(), null);
-			FileNode Root = Bundle.Root;
+			Bundle<FileNode> NewBundle = Bundle.Create<FileNode>(StorageClient, NamespaceId, new FileNode(), new BundleOptions(), Cache);
+			FileNode Root = NewBundle.Root;
 
 			const int NumIterations = 100;
 			for (int Idx = 0; Idx < NumIterations; Idx++)
@@ -370,7 +386,7 @@ namespace HordeServerTests
 				Assert.IsTrue(SpanData.Span.SequenceEqual(Data));
 			}
 
-			await CheckSizes(Bundle, Root, Options, true);
+			await CheckSizes(NewBundle, Root, Options, true);
 		}
 
 		async Task CheckSizes(Bundle Bundle, FileNode Node, ChunkingOptions Options, bool Rightmost)
@@ -400,8 +416,8 @@ namespace HordeServerTests
 			BundleOptions Options = new BundleOptions();
 			Options.MaxBlobSize = 1;
 
-			Bundle<DirectoryNode> Bundle = new Bundle<DirectoryNode>(StorageClient, NamespaceId, new DirectoryNode(), Options, null);
-			DirectoryNode Root = Bundle.Root;
+			using Bundle<DirectoryNode> NewBundle = Bundle.Create<DirectoryNode>(StorageClient, NamespaceId, new DirectoryNode(), Options, Cache);
+			DirectoryNode Root = NewBundle.Root;
 
 			long TotalLength = 0;
 			for (int IdxA = 0; IdxA < 10; IdxA++)
@@ -422,9 +438,9 @@ namespace HordeServerTests
 						}
 					}
 
-					int OldWorkingSetSize = GetWorkingSetSize(Bundle.Root);
-					await Bundle.TrimAsync(20);
-					int NewWorkingSetSize = GetWorkingSetSize(Bundle.Root);
+					int OldWorkingSetSize = GetWorkingSetSize(NewBundle.Root);
+					await NewBundle.TrimAsync(20);
+					int NewWorkingSetSize = GetWorkingSetSize(NewBundle.Root);
 					Assert.IsTrue(NewWorkingSetSize <= OldWorkingSetSize);
 					Assert.IsTrue(NewWorkingSetSize <= 20);
 				}
@@ -434,7 +450,7 @@ namespace HordeServerTests
 			Assert.IsTrue(StorageClient.Refs.Count == 0);
 
 			RefId RefId = new RefId("ref");
-			await Bundle.WriteAsync(BucketId, RefId, CbObject.Empty, true);
+			await NewBundle.WriteAsync(BucketId, RefId, CbObject.Empty, true);
 
 			Assert.AreEqual(TotalLength, Root.Length);
 
diff --git a/Engine/Source/Programs/Shared/EpicGames.Horde.Tests/EpicGames.Horde.Tests.csproj b/Engine/Source/Programs/Shared/EpicGames.Horde.Tests/EpicGames.Horde.Tests.csproj
--- a/Engine/Source/Programs/Shared/EpicGames.Horde.Tests/EpicGames.Horde.Tests.csproj
+++ b/Engine/Source/Programs/Shared/EpicGames.Horde.Tests/EpicGames.Horde.Tests.csproj
@@ -9,6 +9,7 @@
 
   <ItemGroup>
     <PackageReference Include="Microsoft.NET.Test.Sdk" Version="16.5.0" />
+    <PackageReference Include="Microsoft.Extensions.Caching.Memory" Version="6.0.0" />
     <PackageReference Include="MSTest.TestAdapter" Version="2.1.0" />
     <PackageReference Include="MSTest.TestFramework" Version="2.1.0" />
     <PackageReference Include="coverlet.collector" Version="1.2.0" />
diff --git a/Engine/Source/Programs/Shared/EpicGames.Horde/Bundles/Bundle.cs b/Engine/Source/Programs/Shared/EpicGames.Horde/Bundles/Bundle.cs
--- a/Engine/Source/Programs/Shared/EpicGames.Horde/Bundles/Bundle.cs
+++ b/Engine/Source/Programs/Shared/EpicGames.Horde/Bundles/Bundle.cs
@@ -11,6 +11,7 @@ using System;
 using System.Buffers;
 using System.Buffers.Binary;
 using System.Collections;
+using System.Collections.Concurrent;
 using System.Collections.Generic;
 using System.Diagnostics;
 using System.Diagnostics.CodeAnalysis;
@@ -21,6 +22,7 @@ using System.Reflection;
 using System.Security.Cryptography;
 using System.Text;
 using System.Text.Json;
+using System.Threading;
 using System.Threading.Tasks;
 
 namespace EpicGames.Horde.Bundles
@@ -54,6 +56,11 @@ namespace EpicGames.Horde.Bundles
 		/// Proportion of used data under which to trigger a re-pack
 		/// </summary>
 		public float RepackRatio { get; set; } = 0.6f;
+
+		/// <summary>
+		/// Number of reads from storage to allow concurrently. This has an impact on memory usage as well as network throughput.
+		/// </summary>
+		public int MaxConcurrentReads { get; set; } = 5;
 	}
 
 	/// <summary>
@@ -95,20 +102,10 @@ namespace EpicGames.Horde.Bundles
 	/// <summary>
 	/// Base class for manipulating bundles
 	/// </summary>
-	public class Bundle
+	public class Bundle : IDisposable
 	{
 		/// <summary>
-		/// Information about a node within a bundle. Nodes may be in the following states:
-		/// 
-		/// * Imported
-		///   Node is known to be in located a particular hashed blob in storage, but we haven't yet read the blob index to figure out exactly where.
-		///   
-		/// * Exported
-		///   Node is known to be located within a particular blob in storage, and the <see cref="NodeInfo.Export"/> member identifies specifically where.
-		///   
-		/// * Standalone
-		///   Node exists in memory, but has not been written to storage yet.
-		/// 
+		/// Information about a node within a bundle.
 		/// </summary>
 		[DebuggerDisplay("{Hash}")]
 		sealed class NodeInfo
@@ -117,12 +114,15 @@ namespace EpicGames.Horde.Bundles
 			public int Rank { get; }
 			public int Length { get; }
 
-			public BlobInfo? Blob { get; private set; }
-			public BundleExport? Export { get; private set; }
-			public ReadOnlySequence<byte> Data { get; private set; }
-			public NodeInfo[]? References { get; private set; }
+			public NodeState State
+			{
+				get => StatePrivate;
+				set => StatePrivate = value;
+			}
+
+			private NodeState StatePrivate;
 
-			public NodeInfo(IoHash Hash, int Rank, int Length)
+			public NodeInfo(IoHash Hash, int Rank, int Length, NodeState State)
 			{
 				if (Length <= 0)
 				{
@@ -132,15 +132,56 @@ namespace EpicGames.Horde.Bundles
 				this.Hash = Hash;
 				this.Rank = Rank;
 				this.Length = Length;
+				this.StatePrivate = State;
 			}
 
-			public NodeInfo(BlobInfo Blob, BundleImport Import)
-				: this(Import.Hash, Import.Rank, Import.Length)
+			public void TrySetExport(BlobInfo Blob, BundleExport Export, IReadOnlyList<NodeInfo> References)
 			{
-				this.Blob = Blob;
+				NodeState PrevState = State;
+				if (!PrevState.IsStandalone())
+				{
+					NodeState NextState = NodeState.Exported(Blob, Export, References);
+					while (!PrevState.IsStandalone() && !TryUpdate(PrevState, NextState))
+					{
+						PrevState = State;
+					}
+				}
+			}
+
+			public void SetStandalone(ReadOnlySequence<byte> Data, IReadOnlyList<NodeInfo> References)
+			{
+				State = NodeState.Standalone(Data, References);
+			}
+
+			private bool TryUpdate(NodeState PrevState, NodeState NextState)
+			{
+				return Interlocked.CompareExchange(ref StatePrivate, NextState, PrevState) == PrevState;
 			}
+		}
+
+		/// <summary>
+		/// State of each node. Each instance of this class is immutable, in order to allow atomic state transitions on each <see cref="NodeInfo"/> object.
+		/// 
+		/// Nodes in memory may be in the following states.
+		/// 
+		/// * Imported
+		///   Node is known to be in located a particular hashed blob in storage, but we haven't yet read the blob index to figure out exactly where.
+		///   
+		/// * Exported
+		///   Node is known to be located within a particular blob in storage, and the <see cref="NodeState.Export"/> member identifies specifically where.
+		///   
+		/// * Standalone
+		///   Node exists in memory, but has not been written to storage yet.
+		/// 
+		/// </summary>
+		class NodeState
+		{
+			public BlobInfo? Blob { get; }
+			public BundleExport? Export { get; }
+			public ReadOnlySequence<byte> Data { get; }
+			public IReadOnlyList<NodeInfo>? References { get; }
 
-			void Set(BlobInfo? Blob, BundleExport? Export, ReadOnlySequence<byte> Data, NodeInfo[]? References)
+			private NodeState(BlobInfo? Blob, BundleExport? Export, ReadOnlySequence<byte> Data, IReadOnlyList<NodeInfo>? References)
 			{
 				this.Blob = Blob;
 				this.Export = Export;
@@ -148,9 +189,24 @@ namespace EpicGames.Horde.Bundles
 				this.References = References;
 			}
 
-			public void SetImport(BlobInfo Blob) => Set(Blob, null, ReadOnlySequence<byte>.Empty, null);
-			public void SetExport(BlobInfo Blob, BundleExport Export, NodeInfo[] References) => Set(Blob, Export, ReadOnlySequence<byte>.Empty, References);
-			public void SetStandalone(ReadOnlySequence<byte> Data, NodeInfo[] References) => Set(null, null, Data, References);
+			public bool IsImport() => References == null;
+			public bool IsExport() => Export != null;
+			public bool IsStandalone() => Blob == null;
+
+			public static NodeState Imported(BlobInfo Blob)
+			{
+				return new NodeState(Blob, null, ReadOnlySequence<byte>.Empty, null);
+			}
+
+			public static NodeState Exported(BlobInfo Blob, BundleExport Export, IReadOnlyList<NodeInfo> References)
+			{
+				return new NodeState(Blob, Export, ReadOnlySequence<byte>.Empty, References);
+			}
+
+			public static NodeState Standalone(ReadOnlySequence<byte> Data, IReadOnlyList<NodeInfo> References)
+			{
+				return new NodeState(null, null, Data, References);
+			}
 		}
 
 		/// <summary>
@@ -180,13 +236,21 @@ namespace EpicGames.Horde.Bundles
 		/// </summary>
 		public BundleNode Root { get; private set; }
 
-		Dictionary<IoHash, NodeInfo> HashToNode = new Dictionary<IoHash, NodeInfo>();
-		Dictionary<IoHash, BlobInfo> HashToBlob = new Dictionary<IoHash, BlobInfo>();
-		IMemoryCache? Cache;
+		ConcurrentDictionary<IoHash, NodeInfo> HashToNode = new ConcurrentDictionary<IoHash, NodeInfo>();
+		ConcurrentDictionary<IoHash, BlobInfo> HashToBlob = new ConcurrentDictionary<IoHash, BlobInfo>();
+		IMemoryCache Cache;
 
 		CbWriter SharedWriter = new CbWriter();
 		byte[] SharedBlobBuffer = Array.Empty<byte>();
 
+		object LockObject = new object();
+
+		SemaphoreSlim ReadSema;
+
+		// Active read tasks at any moment. If a BundleObject is not available in the cache, we start a read and add an entry to this dictionary
+		// so that other threads can also await it.
+		Dictionary<IoHash, Task<BundleObject>> ReadTasks = new Dictionary<IoHash, Task<BundleObject>>();
+
 		/// <summary>
 		/// Options for the bundle
 		/// </summary>
@@ -204,14 +268,14 @@ namespace EpicGames.Horde.Bundles
 		/// <summary>
 		/// Constructor
 		/// </summary>
-		protected Bundle(IStorageClient StorageClient, NamespaceId NamespaceId, BundleNode Root, BundleOptions Options, IMemoryCache? Cache)
+		protected Bundle(IStorageClient StorageClient, NamespaceId NamespaceId, BundleNode Root, BundleOptions Options, IMemoryCache Cache)
 		{
 			this.StorageClient = StorageClient;
 			this.NamespaceId = NamespaceId;
 			this.Root = Root;
 			this.Options = Options;
 			this.Cache = Cache;
-			this.HashToNode = new Dictionary<IoHash, NodeInfo>();
+			this.ReadSema = new SemaphoreSlim(Options.MaxConcurrentReads);
 		}
 
 		/// <summary>
@@ -223,12 +287,18 @@ namespace EpicGames.Horde.Bundles
 		/// <param name="Root">Reference to the root node for the bundle</param>
 		/// <param name="Options">Options for controlling how things are serialized</param>
 		/// <param name="Cache">Cache for storing decompressed objects, serialized blobs.</param>
-		protected Bundle(IStorageClient StorageClient, NamespaceId NamespaceId, BundleObject Object, BundleNode Root, BundleOptions Options, IMemoryCache? Cache)
+		protected Bundle(IStorageClient StorageClient, NamespaceId NamespaceId, BundleObject Object, BundleNode Root, BundleOptions Options, IMemoryCache Cache)
 			: this(StorageClient, NamespaceId, Root, Options, Cache)
 		{
 			RegisterRootObject(Object);
 		}
 
+		/// <inheritdoc/>
+		public void Dispose()
+		{
+			ReadSema.Dispose();
+		}
+
 		/// <summary>
 		/// Creates a new typed bundle with a default-constructed root node
 		/// </summary>
@@ -238,7 +308,7 @@ namespace EpicGames.Horde.Bundles
 		/// <param name="Options">Options for controlling how things are serialized</param>
 		/// <param name="Cache">Cache for storing decompressed objects, serialized blobs.</param>
 		/// <returns>New bundle instance</returns>
-		public static Bundle<T> Create<T>(IStorageClient StorageClient, NamespaceId NamespaceId, BundleOptions Options, IMemoryCache? Cache) where T : BundleNode, new()
+		public static Bundle<T> Create<T>(IStorageClient StorageClient, NamespaceId NamespaceId, BundleOptions Options, IMemoryCache Cache) where T : BundleNode, new()
 		{
 			return Create(StorageClient, NamespaceId, new T(), Options, Cache);
 		}
@@ -253,7 +323,7 @@ namespace EpicGames.Horde.Bundles
 		/// <param name="Options">Options for controlling how things are serialized</param>
 		/// <param name="Cache">Cache for storing decompressed objects, serialized blobs.</param>
 		/// <returns>New bundle instance</returns>
-		public static Bundle<T> Create<T>(IStorageClient StorageClient, NamespaceId NamespaceId, T Root, BundleOptions Options, IMemoryCache? Cache) where T : BundleNode
+		public static Bundle<T> Create<T>(IStorageClient StorageClient, NamespaceId NamespaceId, T Root, BundleOptions Options, IMemoryCache Cache) where T : BundleNode
 		{
 			return new Bundle<T>(StorageClient, NamespaceId, Root, Options, Cache);
 		}
@@ -261,7 +331,7 @@ namespace EpicGames.Horde.Bundles
 		/// <summary>
 		/// Reads a bundle from storage
 		/// </summary>
-		public static async ValueTask<Bundle<T>> ReadAsync<T>(IStorageClient StorageClient, NamespaceId NamespaceId, BucketId BucketId, RefId RefId, BundleOptions Options, IMemoryCache? Cache) where T : BundleNode
+		public static async ValueTask<Bundle<T>> ReadAsync<T>(IStorageClient StorageClient, NamespaceId NamespaceId, BucketId BucketId, RefId RefId, BundleOptions Options, IMemoryCache Cache) where T : BundleNode
 		{
 			Bundle<T> NewBundle = new Bundle<T>(StorageClient, NamespaceId, null!, Options, Cache);
 
@@ -371,9 +441,9 @@ namespace EpicGames.Horde.Bundles
 
 		void FindNodeInfosToWrite(NodeInfo Node, List<NodeInfo> WriteNodes, HashSet<NodeInfo> WriteNodesSet)
 		{
-			if (Node.Blob == null && WriteNodesSet.Add(Node))
+			if (Node.State.Blob == null && WriteNodesSet.Add(Node))
 			{
-				foreach (NodeInfo ChildNode in Node.References!)
+				foreach (NodeInfo ChildNode in Node.State.References!)
 				{
 					FindNodeInfosToWrite(ChildNode, WriteNodes, WriteNodesSet);
 				}
@@ -415,19 +485,27 @@ namespace EpicGames.Horde.Bundles
 		/// <returns>The node data</returns>
 		async ValueTask<ReadOnlySequence<byte>> GetDataAsync(NodeInfo Node)
 		{
-			if (Node.Blob == null)
+			NodeState NodeState = Node.State;
+			if (NodeState.Blob == null)
 			{
-				return Node.Data;
+				return NodeState.Data;
 			}
 			else
 			{
-				await MountBlobAsync(Node.Blob);
-				Debug.Assert(Node.Export != null);
+				if (!NodeState.IsExport())
+				{
+					await MountBlobAsync(NodeState.Blob);
+					NodeState = Node.State;
+				}
 
-				BundleObject Object = await GetObjectAsync(Node.Blob.Hash);
-				ReadOnlySequence<byte> Data = new ReadOnlySequence<byte>(DecodePacket(Node.Blob.Hash, new ReadOnlySequence<byte>(Object.Data), Node.Export.Packet).Slice(Node.Export.Offset, Node.Export.Length));
+				Debug.Assert(NodeState.Blob != null);
+				Debug.Assert(NodeState.Export != null);
 
-				return Data;
+				BundleObject Object = await GetObjectAsync(NodeState.Blob.Hash);
+				ReadOnlyMemory<byte> PacketData = DecodePacket(NodeState.Blob.Hash, Object.Data, NodeState.Export.Packet);
+				ReadOnlyMemory<byte> NodeData = PacketData.Slice(NodeState.Export.Offset, NodeState.Export.Length);
+
+				return new ReadOnlySequence<byte>(NodeData);
 			}
 		}
 
@@ -456,8 +534,14 @@ namespace EpicGames.Horde.Bundles
 			IoHash Hash = IoHash.Compute(Data);
 			NodeInfo[] NodeReferences = References.Select(x => HashToNode[x]).Distinct().OrderBy(x => x.Hash).ToArray();
 
+			NodeState State = NodeState.Standalone(Data, NodeReferences);
+
 			NodeInfo? Node;
-			if (!HashToNode.TryGetValue(Hash, out Node))
+			if (HashToNode.TryGetValue(Hash, out Node))
+			{
+				Node.State = State;
+			}
+			else
 			{
 				int Rank = 0;
 				foreach (NodeInfo Reference in NodeReferences)
@@ -465,14 +549,12 @@ namespace EpicGames.Horde.Bundles
 					Rank = Math.Max(Rank, Reference.Rank + 1);
 				}
 
-				Node = new NodeInfo(Hash, Rank, (int)Data.Length);
+				Node = new NodeInfo(Hash, Rank, (int)Data.Length, State);
 				HashToNode[Hash] = Node;
 
 				NextStats.NewNodeCount++;
 				NextStats.NewNodeBytes += Data.Length;
 			}
-
-			Node.SetStandalone(Data, NodeReferences);
 			return Node;
 		}
 
@@ -481,7 +563,13 @@ namespace EpicGames.Horde.Bundles
 			if(!Blob.Mounted)
 			{
 				BundleObject Object = await GetObjectAsync(Blob.Hash);
-				RegisterBlobObject(Object, Blob);
+				lock (LockObject)
+				{
+					if (!Blob.Mounted)
+					{
+						RegisterBlobObject(Object, Blob);
+					}
+				}
 			}
 		}
 
@@ -492,20 +580,62 @@ namespace EpicGames.Horde.Bundles
 		/// <returns>The parsed object</returns>
 		async ValueTask<BundleObject> GetObjectAsync(IoHash Hash)
 		{
-			BundleObject Object;
-			if (Cache == null || !Cache.TryGetValue(Hash, out Object))
-			{
-				ReadOnlyMemory<byte> Data = await StorageClient.ReadBlobToMemoryAsync(NamespaceId, Hash);
-				Object = CbSerializer.Deserialize<BundleObject>(Data);
+			string CacheKey = $"object:{Hash}";
 
-				if (Cache != null)
+			BundleObject? Object;
+			if (!Cache.TryGetValue<BundleObject>(CacheKey, out Object))
+			{
+				Task<BundleObject>? ReadTask;
+				lock (LockObject)
 				{
-					using (ICacheEntry Entry = Cache.CreateEntry(Hash))
+					if (!ReadTasks.TryGetValue(Hash, out ReadTask))
 					{
-						Entry.SetValue(Object);
-						Entry.SetSize(Data.Length);
+						ReadTask = Task.Run(() => ReadObjectAsync(CacheKey, Hash));
+						ReadTasks.Add(Hash, ReadTask);
 					}
 				}
+				Object = await ReadTask;
+			}
+			return Object;
+		}
+
+		async Task<BundleObject> ReadObjectAsync(string CacheKey, IoHash Hash)
+		{
+			// Perform another (sequenced) check whether an object has been added to the cache, to counteract the race between a read task being added and a task completing.
+			lock (LockObject)
+			{
+				if (Cache.TryGetValue<BundleObject>(CacheKey, out BundleObject? CachedObject))
+				{
+					return CachedObject;
+				}
+			}
+
+			// Wait for the read semaphore to avoid triggering too many operations at once.
+			await ReadSema.WaitAsync();
+
+			// Read the data from storage
+			ReadOnlyMemory<byte> Data;
+			try
+			{
+				Data = await StorageClient.ReadBlobToMemoryAsync(NamespaceId, Hash);
+			}
+			finally
+			{
+				ReadSema.Release();
+			}
+
+			// Add the object to the cache
+			BundleObject Object = CbSerializer.Deserialize<BundleObject>(Data);
+			using (ICacheEntry Entry = Cache.CreateEntry(CacheKey))
+			{
+				Entry.SetSize(Data.Length);
+				Entry.SetValue(Object);
+			}
+
+			// Remove this object from the list of read tasks
+			lock (LockObject)
+			{
+				ReadTasks.Remove(Hash);
 			}
 			return Object;
 		}
@@ -513,44 +643,35 @@ namespace EpicGames.Horde.Bundles
 		/// <summary>
 		/// Gets a decoded block from the store
 		/// </summary>
-		/// <param name="BlobKey">Key for the blob</param>
+		/// <param name="BlobHash">Key for the blob</param>
 		/// <param name="Data">Raw blob data</param>
 		/// <param name="Packet">The decoded block location and size</param>
 		/// <returns>The decoded data</returns>
-		ReadOnlyMemory<byte> DecodePacket(object BlobKey, ReadOnlySequence<byte> Data, BundleCompressionPacket Packet)
+		ReadOnlyMemory<byte> DecodePacket(IoHash BlobHash, ReadOnlyMemory<byte> Data, BundleCompressionPacket Packet)
 		{
-			object CacheKey = (BlobKey, Packet.Offset);
-
-			ReadOnlyMemory<byte> DecodedData;
-			if (Cache == null || !Cache.TryGetValue(CacheKey, out DecodedData))
+			string CacheKey = $"decode:{BlobHash}/{Packet.Offset}";
+			return Cache.GetOrCreate<ReadOnlyMemory<byte>>(CacheKey, Entry =>
 			{
-				Debug.Assert(Data.IsSingleSegment);
-
-				byte[] DecodedBuffer = new byte[Packet.DecodedLength];
-				LZ4Codec.Decode(Data.FirstSpan.Slice(Packet.Offset, Packet.EncodedLength), DecodedBuffer);
+				ReadOnlyMemory<byte> EncodedPacket = Data.Slice(Packet.Offset, Packet.EncodedLength);
 
-				DecodedData = DecodedBuffer;
+				byte[] DecodedPacket = new byte[Packet.DecodedLength];
+				LZ4Codec.Decode(EncodedPacket.Span, DecodedPacket);
 
-				if (Cache != null)
-				{
-					using (ICacheEntry Entry = Cache.CreateEntry(CacheKey))
-					{
-						Entry.SetValue(DecodedData);
-						Entry.SetSize(DecodedData.Length);
-					}
-				}
-			}
-
-			return DecodedData;
+				Entry.SetSize(Packet.DecodedLength);
+				return DecodedPacket;
+			});
 		}
 
-		NodeInfo FindOrAddNode(IoHash Hash, int Rank, int Cost)
+		/// <summary>
+		/// Find a node or create a new one
+		/// </summary>
+		NodeInfo FindOrAddNodeFromImport(BlobInfo Blob, BundleImport Import)
 		{
 			NodeInfo? Node;
-			if (!HashToNode.TryGetValue(Hash, out Node))
+			if (!HashToNode.TryGetValue(Import.Hash, out Node))
 			{
-				Node = new NodeInfo(Hash, Rank, Cost);
-				HashToNode[Hash] = Node;
+				NodeInfo NewNode = new NodeInfo(Import.Hash, Import.Rank, Import.Length, NodeState.Imported(Blob));
+				Node = HashToNode.GetOrAdd(NewNode.Hash, NewNode);
 			}
 			return Node;
 		}
@@ -560,9 +681,9 @@ namespace EpicGames.Horde.Bundles
 		/// </summary>
 		void RegisterRootObject(BundleObject Object)
 		{
-			List<NodeInfo> Nodes = CreateNodesForObject(Object);
-
-			for (int Idx = 0; Idx < Object.Exports.Count; )
+			// Create all the export nodes without fixing up their dependencies yet
+			List<NodeInfo> Nodes = new List<NodeInfo>();
+			for (int Idx = 0; Idx < Object.Exports.Count;)
 			{
 				BundleCompressionPacket Packet = Object.Exports[Idx].Packet;
 
@@ -574,14 +695,42 @@ namespace EpicGames.Horde.Bundles
 					BundleExport Export = Object.Exports[Idx];
 
 					ReadOnlyMemory<byte> NodeData = DecodedData.AsMemory(Export.Offset, Export.Length);
-					if(Export.Length != DecodedData.Length)
+					if (Export.Length != DecodedData.Length)
 					{
 						NodeData = NodeData.ToArray();
 					}
 
-					Nodes[Idx].SetStandalone(new ReadOnlySequence<byte>(NodeData), Export.References.Select(x => Nodes[x]).ToArray());
+					NodeState State = NodeState.Standalone(new ReadOnlySequence<byte>(NodeData), Array.Empty<NodeInfo>());
+
+					NodeInfo? Node;
+					for (; ; )
+					{
+						if (HashToNode.TryGetValue(Export.Hash, out Node))
+						{
+							Node.State = State;
+							break;
+						}
+
+						NodeInfo NewNode = new NodeInfo(Export.Hash, Export.Rank, Export.Length, State);
+						if (HashToNode.TryAdd(Export.Hash, NewNode))
+						{
+							Node = NewNode;
+							break;
+						}
+					}
+					Nodes.Add(Node);
 				}
 			}
+
+			// Create all the imports
+			RegisterImports(Object.ImportObjects, Nodes);
+
+			// Fixup all the references
+			for (int Idx = 0; Idx < Object.Exports.Count; Idx++)
+			{
+				NodeInfo Node = Nodes[Idx];
+				Node.State = NodeState.Standalone(Node.State.Data, Object.Exports[Idx].References.Select(x => Nodes[x]).ToArray());
+			}
 		}
 
 		/// <summary>
@@ -589,42 +738,39 @@ namespace EpicGames.Horde.Bundles
 		/// </summary>
 		void RegisterBlobObject(BundleObject Object, BlobInfo Blob)
 		{
-			List<NodeInfo> Nodes = CreateNodesForObject(Object);
+			// Create the exports as imports first, since we can't resolve the reference list.
+			List<NodeInfo> Nodes = new List<NodeInfo>();
+			foreach (BundleExport Export in Object.Exports)
+			{
+				NodeInfo Node = FindOrAddNodeFromImport(Blob, Export);
+				Nodes.Add(Node);
+			}
+
+			// Register the regular imports
+			RegisterImports(Object.ImportObjects, Nodes);
+
+			// Go back and resolve all the exports 
 			for (int Idx = 0; Idx < Object.Exports.Count; Idx++)
 			{
 				BundleExport Export = Object.Exports[Idx];
-				Nodes[Idx].SetExport(Blob, Export, Export.References.Select(x => Nodes[x]).ToArray());
+				Nodes[Idx].TrySetExport(Blob, Export, Export.References.Select(x => Nodes[x]).ToArray());
 			}
+
+			// Flag the blob as being mounted
 			Blob.Mounted = true;
 		}
 
-		List<NodeInfo> CreateNodesForObject(BundleObject Object)
+		void RegisterImports(List<BundleImportObject> ImportObjects, List<NodeInfo> Nodes)
 		{
-			List<NodeInfo> Nodes = new List<NodeInfo>();
-			foreach (BundleExport Export in Object.Exports)
-			{
-				NodeInfo Node = FindOrAddNode(Export.Hash, Export.Rank, Export.Length);
-				Nodes.Add(Node);
-			}
-			foreach (BundleImportObject ImportObject in Object.ImportObjects)
+			foreach (BundleImportObject ImportObject in ImportObjects)
 			{
-				BlobInfo? ImportBlob;
-				if (!HashToBlob.TryGetValue(ImportObject.Object.Hash, out ImportBlob))
-				{
-					ImportBlob = new BlobInfo(ImportObject.Object.Hash, ImportObject.TotalCost);
-					HashToBlob[ImportObject.Object.Hash] = ImportBlob;
-				}
+				BlobInfo ImportBlob = HashToBlob.GetOrAdd(ImportObject.Object.Hash, new BlobInfo(ImportObject.Object.Hash, ImportObject.TotalCost));
 				foreach (BundleImport Import in ImportObject.Imports)
 				{
-					NodeInfo Node = FindOrAddNode(Import.Hash, Import.Rank, Import.Length);
-					if (Node.Blob == null)
-					{
-						Node.SetImport(ImportBlob);
-					}
+					NodeInfo Node = FindOrAddNodeFromImport(ImportBlob, Import);
 					Nodes.Add(Node);
 				}
 			}
-			return Nodes;
 		}
 
 		/// <summary>
@@ -653,7 +799,7 @@ namespace EpicGames.Horde.Bundles
 				List<BlobInfo> LiveBlobs = new List<BlobInfo>();
 				foreach (NodeInfo LiveNode in LiveNodes)
 				{
-					BlobInfo? LiveBlob = LiveNode.Blob;
+					BlobInfo? LiveBlob = LiveNode.State.Blob;
 					if (LiveBlob != null)
 					{
 						if (LiveBlob.LiveNodes == null)
@@ -676,7 +822,7 @@ namespace EpicGames.Horde.Bundles
 						foreach (NodeInfo Node in LiveBlob.LiveNodes!)
 						{
 							ReadOnlySequence<byte> Data = await GetDataAsync(Node);
-							Node.SetStandalone(Data, Node.References!);
+							Node.SetStandalone(Data, Node.State.References!);
 						}
 					}
 				}
@@ -685,10 +831,10 @@ namespace EpicGames.Horde.Bundles
 				for (int Idx = LiveNodes.Count - 1; Idx >= 0; Idx--)
 				{
 					NodeInfo LiveNode = LiveNodes[Idx];
-					if (LiveNode.References != null && LiveNode.References.Any(x => x.Blob == null))
+					if (LiveNode.State.References != null && LiveNode.State.References.Any(x => x.State.Blob == null))
 					{
 						ReadOnlySequence<byte> Data = await GetDataAsync(LiveNode);
-						LiveNode.SetStandalone(Data, LiveNode.References!);
+						LiveNode.SetStandalone(Data, LiveNode.State.References!);
 					}
 				}
 
@@ -704,7 +850,7 @@ namespace EpicGames.Horde.Bundles
 			FindModifiedLiveSet(RootNode, NewNodes);
 
 			// Write all the new blobs. Sort by rank and go from end to start, updating the list as we go so that we can figure out imports correctly.
-			NodeInfo[] WriteNodes = NewNodes.Where(x => x.Blob == null).OrderBy(x => x.Rank).ThenBy(x => x.Hash).ToArray();
+			NodeInfo[] WriteNodes = NewNodes.Where(x => x.State.Blob == null).OrderBy(x => x.Rank).ThenBy(x => x.Hash).ToArray();
 
 			int MinIdx = 0;
 			int NextBlobCost = WriteNodes[0].Length;
@@ -738,9 +884,9 @@ namespace EpicGames.Horde.Bundles
 		{
 			if (Nodes.Add(Node) && Node.Rank > 0)
 			{
-				foreach (NodeInfo Reference in Node.References!)
+				foreach (NodeInfo Reference in Node.State.References!)
 				{
-					if (Reference.Blob == null)
+					if (Reference.State.Blob == null)
 					{
 						FindModifiedLiveSet(Reference, Nodes);
 					}
@@ -755,11 +901,11 @@ namespace EpicGames.Horde.Bundles
 				LiveNodes.Add(Node);
 				if (Node.Rank > 0)
 				{
-					if (Node.Blob != null)
+					if (Node.State.Blob != null)
 					{
-						await MountBlobAsync(Node.Blob);
+						await MountBlobAsync(Node.State.Blob);
 					}
-					foreach (NodeInfo Reference in Node.References!)
+					foreach (NodeInfo Reference in Node.State.References!)
 					{
 						await FindLiveNodesAsync(Reference, LiveNodes, LiveNodeSet);
 					}
@@ -772,7 +918,7 @@ namespace EpicGames.Horde.Bundles
 			foreach (NodeInfo Node in Nodes)
 			{
 				ReadOnlySequence<byte> NodeData = await GetDataAsync(Node);
-				Node.SetStandalone(NodeData, Node.References!);
+				Node.SetStandalone(NodeData, Node.State.References!);
 			}
 
 			BundleRoot Ref = new BundleRoot();
@@ -812,7 +958,7 @@ namespace EpicGames.Horde.Bundles
 			for (int Idx = 0; Idx < Object.Exports.Count; Idx++)
 			{
 				NodeInfo Node = Nodes[Idx];
-				Node.SetExport(Blob, Object.Exports[Idx], Node.References!);
+				Node.State = NodeState.Exported(Blob, Object.Exports[Idx], Node.State.References!);
 			}
 
 			Blob.Mounted = true;
@@ -843,12 +989,12 @@ namespace EpicGames.Horde.Bundles
 			HashSet<NodeInfo> ImportedNodes = new HashSet<NodeInfo>();
 			foreach (NodeInfo Node in Nodes)
 			{
-				ImportedNodes.UnionWith(Node.References!);
+				ImportedNodes.UnionWith(Node.State.References!);
 			}
 			ImportedNodes.ExceptWith(Nodes);
 
 			// Find all the imports
-			foreach (IGrouping<BlobInfo, NodeInfo> ImportGroup in ImportedNodes.GroupBy(x => x.Blob!))
+			foreach (IGrouping<BlobInfo, NodeInfo> ImportGroup in ImportedNodes.GroupBy(x => x.State.Blob!))
 			{
 				BlobInfo Blob = ImportGroup.Key;
 
@@ -881,7 +1027,7 @@ namespace EpicGames.Horde.Bundles
 				}
 
 				// Create the export for this node
-				int[] References = Node.References.Select(x => NodeToIndex[x]).ToArray();
+				int[] References = Node.State.References.Select(x => NodeToIndex[x]).ToArray();
 				BundleExport Export = new BundleExport(Node.Hash, Node.Rank, Packet, BlockSize, Node.Length, References);
 				Object.Exports.Add(Export);
 
@@ -952,7 +1098,7 @@ namespace EpicGames.Horde.Bundles
 		/// <summary>
 		/// Constructor
 		/// </summary>
-		internal Bundle(IStorageClient StorageClient, NamespaceId NamespaceId, T Root, BundleOptions Options, IMemoryCache? Cache)
+		internal Bundle(IStorageClient StorageClient, NamespaceId NamespaceId, T Root, BundleOptions Options, IMemoryCache Cache)
 			: base(StorageClient, NamespaceId, Root, Options, Cache)
 		{
 		}
diff --git a/Engine/Source/Programs/Shared/EpicGames.Horde/Bundles/Nodes/DirectoryNode.cs b/Engine/Source/Programs/Shared/EpicGames.Horde/Bundles/Nodes/DirectoryNode.cs
--- a/Engine/Source/Programs/Shared/EpicGames.Horde/Bundles/Nodes/DirectoryNode.cs
+++ b/Engine/Source/Programs/Shared/EpicGames.Horde/Bundles/Nodes/DirectoryNode.cs
@@ -533,23 +533,21 @@ namespace EpicGames.Horde.Bundles.Nodes
 		{
 			DirectoryInfo.Create();
 
-//			List<Task> Tasks = new List<Task>();
+			List<Task> Tasks = new List<Task>();
 			foreach (FileEntry FileEntry in Files)
 			{
 				FileInfo FileInfo = new FileInfo(Path.Combine(DirectoryInfo.FullName, FileEntry.Name.ToString()));
 				FileNode FileNode = await Bundle.GetAsync(FileEntry);
-				await FileNode.CopyToFileAsync(Bundle, FileInfo);
-//				Tasks.Add(Task.Run(() => FileNode.CopyToFileAsync(FileInfo)));
+				Tasks.Add(Task.Run(() => FileNode.CopyToFileAsync(Bundle, FileInfo)));
 			}
 			foreach (DirectoryEntry DirectoryEntry in Directories)
 			{
 				DirectoryInfo SubDirectoryInfo = DirectoryInfo.CreateSubdirectory(DirectoryEntry.Name.ToString());
 				DirectoryNode SubDirectoryNode = await Bundle.GetAsync(DirectoryEntry);
-				await SubDirectoryNode.CopyToDirectoryAsync(Bundle, SubDirectoryInfo, Logger);
-//				Tasks.Add(Task.Run(() => SubDirectoryNode.CopyToDirectoryAsync(SubDirectoryInfo, Logger)));
+				Tasks.Add(Task.Run(() => SubDirectoryNode.CopyToDirectoryAsync(Bundle, SubDirectoryInfo, Logger)));
 			}
 
-//			await Task.WhenAll(Tasks);
+			await Task.WhenAll(Tasks);
 		}
 	}
 
