1) Update parser and code gen with FieldNotify changes
2) Update parser and code gen with RigVM changes
3) Fix issue with UhtException not returning Message
4) Fix issue with deprecated functions inside of non-deprecated classes signaling deprecated property issues.

#rnx

diff --git a/Engine/Source/Programs/Shared/EpicGames.UHT/Exporters/CodeGen/UhtHeaderCodeGenerator.cs b/Engine/Source/Programs/Shared/EpicGames.UHT/Exporters/CodeGen/UhtHeaderCodeGenerator.cs
--- a/Engine/Source/Programs/Shared/EpicGames.UHT/Exporters/CodeGen/UhtHeaderCodeGenerator.cs
+++ b/Engine/Source/Programs/Shared/EpicGames.UHT/Exporters/CodeGen/UhtHeaderCodeGenerator.cs
@@ -25,6 +25,7 @@ namespace EpicGames.UHT.Exporters.CodeGen
 		public const string GeneratedBodyMacroSuffix = "GENERATED_BODY";
 		public const string GeneratedBodyLegacyMacroSuffix = "GENERATED_BODY_LEGACY";
 		public const string GeneratedUInterfaceBodyMacroSuffix = "GENERATED_UINTERFACE_BODY()";
+		public const string FieldNotifyMacroSuffix = "FIELDNOTIFY";
 		public const string InClassMacroSuffix = "INCLASS";
 		public const string InClassNoPureDeclsMacroSuffix = "INCLASS_NO_PURE_DECLS";
 		public const string InClassIInterfaceMacroSuffix = "INCLASS_IINTERFACE";
@@ -468,6 +469,123 @@ namespace EpicGames.UHT.Exporters.CodeGen
 			return $"const {ReturnType}& {Function.SourceName.Substring(1)}";
 		}
 		#endregion
+
+		#region Field notify support
+		protected bool NeedFieldNotifyCodeGen(UhtClass Class)
+		{
+			return 
+				!Class.ClassExportFlags.HasAnyFlags(UhtClassExportFlags.HasCustomFieldNotify) &&
+				Class.ClassExportFlags.HasAnyFlags(UhtClassExportFlags.HasFieldNotify);
+		}
+
+		protected void GetFieldNotifyStats(UhtClass Class, out bool bHasProperties, out bool bHasFunctions, out bool bHasEditorFields, out bool bAllEditorFields)
+		{
+			// Scan the children to see what we have
+			bHasProperties = false;
+			bHasFunctions = false;
+			bHasEditorFields = false;
+			bAllEditorFields = true;
+			foreach (UhtType Type in Class.Children)
+			{
+				if (Type is UhtProperty Property)
+				{
+					if (Property.PropertyExportFlags.HasAnyFlags(UhtPropertyExportFlags.FieldNotify))
+					{
+						bHasProperties = true;
+						bHasEditorFields |= Property.bIsEditorOnlyProperty;
+						bAllEditorFields &= Property.bIsEditorOnlyProperty;
+					}
+				}
+				else if (Type is UhtFunction Function)
+				{
+					if (Function.FunctionExportFlags.HasAnyFlags(UhtFunctionExportFlags.FieldNotify))
+					{
+						bHasFunctions = true;
+						bHasEditorFields |= Function.FunctionFlags.HasAnyFlags(EFunctionFlags.EditorOnly);
+						bAllEditorFields &= Function.FunctionFlags.HasAnyFlags(EFunctionFlags.EditorOnly);
+					}
+				}
+			}
+
+			// If we have no editor fields, then by definition, all fields can't be editor fields
+			bAllEditorFields &= bHasEditorFields;
+		}
+
+		protected StringBuilder AppendFieldNotify(StringBuilder Builder, UhtClass Class,
+			bool bHasProperties, bool bHasFunctions, bool bHasEditorFields, bool bAllEditorFields,
+			bool bIncludeEditorOnlyFields, bool bAppendDefine, Action<StringBuilder, UhtClass, string> AppendAction)
+		{
+			if (bHasProperties && !bAllEditorFields)
+			{
+				foreach (UhtType Child in Class.Children)
+				{
+					if (Child is UhtProperty Property)
+					{
+						if (!Property.bIsEditorOnlyProperty)
+						{
+							AppendAction(Builder, Class, Property.SourceName);
+						}
+					}
+				}
+			}
+
+			if (bHasFunctions && !bAllEditorFields)
+			{
+				foreach (UhtType Child in Class.Children)
+				{
+					if (Child is UhtFunction Function)
+					{
+						if (!Function.FunctionFlags.HasAnyFlags(EFunctionFlags.EditorOnly))
+						{
+							AppendAction(Builder, Class, Function.CppImplName);
+						}
+					}
+				}
+			}
+
+			if (bHasEditorFields && bIncludeEditorOnlyFields)
+			{
+				if (!bAllEditorFields && bAppendDefine)
+				{
+					Builder.Append("#if WITH_EDITORONLY_DATA\r\n");
+				}
+
+				if (bHasProperties)
+				{
+					foreach (UhtType Child in Class.Children)
+					{
+						if (Child is UhtProperty Property)
+						{
+							if (Property.bIsEditorOnlyProperty)
+							{
+								AppendAction(Builder, Class, Property.SourceName);
+							}
+						}
+					}
+				}
+
+				if (bHasFunctions)
+				{
+					foreach (UhtType Child in Class.Children)
+					{
+						if (Child is UhtFunction Function)
+						{
+							if (Function.FunctionFlags.HasAnyFlags(EFunctionFlags.EditorOnly))
+							{
+								AppendAction(Builder, Class, Function.CppImplName);
+							}
+						}
+					}
+				}
+
+				if (!bAllEditorFields && bAppendDefine)
+				{
+					Builder.Append("#endif // WITH_EDITORONLY_DATA\r\n");
+				}
+			}
+			return Builder;
+		}
+		#endregion
 	}
 
 	internal static class UhtHaederCodeGeneratorStringBuilderExtensions
diff --git a/Engine/Source/Programs/Shared/EpicGames.UHT/Exporters/CodeGen/UhtHeaderCodeGeneratorCppFile.cs b/Engine/Source/Programs/Shared/EpicGames.UHT/Exporters/CodeGen/UhtHeaderCodeGeneratorCppFile.cs
--- a/Engine/Source/Programs/Shared/EpicGames.UHT/Exporters/CodeGen/UhtHeaderCodeGeneratorCppFile.cs
+++ b/Engine/Source/Programs/Shared/EpicGames.UHT/Exporters/CodeGen/UhtHeaderCodeGeneratorCppFile.cs
@@ -924,6 +924,8 @@ namespace EpicGames.UHT.Exporters.CodeGen
 				AppendSerializer(Builder, Class, UhtSerializerArchiveType.StructuredArchiveRecord, "IMPLEMENT_FSTRUCTUREDARCHIVE_SERIALIZER");
 			}
 
+			AppendFieldNotify(Builder, Class);
+
 			if (Class.ClassFlags.HasAnyFlags(EClassFlags.Interface))
 			{
 				AppendInterfaceCallFunctions(Builder, Class, CallbackFunctions);
@@ -931,6 +933,48 @@ namespace EpicGames.UHT.Exporters.CodeGen
 			return Builder;
 		}
 
+		private StringBuilder AppendFieldNotify(StringBuilder Builder, UhtClass Class)
+		{
+			if (!NeedFieldNotifyCodeGen(Class))
+			{
+				return Builder;
+			}
+
+			// Scan the children to see what we have
+			bool bHasProperties;
+			bool bHasFunctions;
+			bool bHasEditorFields;
+			bool bAllEditorFields;
+			GetFieldNotifyStats(Class, out bHasProperties, out bHasFunctions, out bHasEditorFields, out bAllEditorFields);
+
+			if (bAllEditorFields)
+			{
+				Builder.Append("#if WITH_EDITORONLY_DATA\r\n");
+			}
+
+			//UE_FIELD_NOTIFICATION_DECLARE_FIELD
+			AppendFieldNotify(Builder, Class, bHasProperties, bHasFunctions, bHasEditorFields, bAllEditorFields, true, true,
+				(StringBuilder Builder, UhtClass Class, string Name) =>
+				{
+					Builder.Append($"\tUE_FIELD_NOTIFICATION_IMPLEMENT_FIELD({Class.SourceName}, {Name})\r\n");
+				});
+
+			//UE_FIELD_NOTIFICATION_DECLARE_ENUM_FIELD
+			Builder.Append($"\tUE_FIELD_NOTIFICATION_IMPLEMENTATION_BEGIN({Class.SourceName})\r\n"); 
+			AppendFieldNotify(Builder, Class, bHasProperties, bHasFunctions, bHasEditorFields, bAllEditorFields, true, true,
+				(StringBuilder Builder, UhtClass Class, string Name) =>
+				{
+					Builder.Append($"\tUE_FIELD_NOTIFICATION_IMPLEMENT_ENUM_FIELD({Class.SourceName}, {Name})\r\n");
+				});
+			Builder.Append($"\tUE_FIELD_NOTIFICATION_IMPLEMENTATION_END({Class.SourceName});\r\n");
+
+			if (bAllEditorFields)
+			{
+				Builder.Append("#endif // WITH_EDITORONLY_DATA\r\n");
+			}
+			return Builder;
+		}
+
 		private StringBuilder AppendPropertyAccessors(StringBuilder Builder, UhtClass Class)
 		{
 			foreach (UhtType Type in Class.Children)
diff --git a/Engine/Source/Programs/Shared/EpicGames.UHT/Exporters/CodeGen/UhtHeaderCodeGeneratorHFile.cs b/Engine/Source/Programs/Shared/EpicGames.UHT/Exporters/CodeGen/UhtHeaderCodeGeneratorHFile.cs
--- a/Engine/Source/Programs/Shared/EpicGames.UHT/Exporters/CodeGen/UhtHeaderCodeGeneratorHFile.cs
+++ b/Engine/Source/Programs/Shared/EpicGames.UHT/Exporters/CodeGen/UhtHeaderCodeGeneratorHFile.cs
@@ -14,7 +14,8 @@ namespace EpicGames.UHT.Exporters.CodeGen
 	internal class UhtHeaderCodeGeneratorHFile
 		: UhtHeaderCodeGenerator
 	{
-		public static string RigVMParameterPrefix = "FRigVMExecuteContext& RigVMExecuteContext";
+		public static string RigVMExecuteContextDeclaration = "FRigVMExtendedExecuteContext& RigVMExecuteContext";
+		public static string RigVMExecuteContextPublicDeclaration = "FRigVMExecuteContext& RigVMExecuteContext";
 
 		/// <summary>
 		/// Construct an instance of this generator object
@@ -117,14 +118,12 @@ namespace EpicGames.UHT.Exporters.CodeGen
 		{
 			if (ScriptStruct.RigVMStructInfo != null)
 			{
-				const string RigVMParameterPrefix = "FRigVMExecuteContext& RigVMExecuteContext";
-
 				Builder.Append("\r\n");
 				foreach (UhtRigVMMethodInfo MethodInfo in ScriptStruct.RigVMStructInfo.Methods)
 				{
 					Builder.Append("#define ").Append(ScriptStruct.SourceName).Append('_').Append(MethodInfo.Name).Append("() \\\r\n");
 					Builder.Append("\t").Append(MethodInfo.ReturnType).Append(' ').Append(ScriptStruct.SourceName).Append("::Static").Append(MethodInfo.Name).Append("( \\\r\n");
-					Builder.Append("\t\t").Append(RigVMParameterPrefix);
+					Builder.Append("\t\t").Append(RigVMExecuteContextPublicDeclaration);
 					Builder.AppendParameterDecls(ScriptStruct.RigVMStructInfo.Members, true, ", \\\r\n\t\t", true, false);
 					Builder.AppendParameterDecls(MethodInfo.Parameters, true, ", \\\r\n\t\t", false, false);
 					Builder.Append(" \\\r\n");
@@ -152,14 +151,14 @@ namespace EpicGames.UHT.Exporters.CodeGen
 						foreach (UhtRigVMMethodInfo MethodInfo in ScriptStruct.RigVMStructInfo.Methods)
 						{
 							Builder.Append("\tstatic ").Append(MethodInfo.ReturnType).Append(" Static").Append(MethodInfo.Name).Append("( \\\r\n");
-							Builder.Append("\t\t").Append(RigVMParameterPrefix);
+							Builder.Append("\t\t").Append(RigVMExecuteContextPublicDeclaration);
 							Builder.AppendParameterDecls(ScriptStruct.RigVMStructInfo.Members, true, ", \\\r\n\t\t", true, false);
 							Builder.AppendParameterDecls(MethodInfo.Parameters, true, ", \\\r\n\t\t", false, false);
 							Builder.Append(" \\\r\n");
 							Builder.Append("\t); \\\r\n");
 
 							Builder.Append("\tFORCEINLINE_DEBUGGABLE static ").Append(MethodInfo.ReturnType).Append(" RigVM").Append(MethodInfo.Name).Append("( \\\r\n");
-							Builder.Append("\t\t").Append(RigVMParameterPrefix).Append(", \\\r\n");
+							Builder.Append("\t\t").Append(RigVMExecuteContextDeclaration).Append(", \\\r\n");
 							Builder.Append("\t\tFRigVMMemoryHandleArray RigVMMemoryHandles \\\r\n");
 							Builder.Append("\t) \\\r\n");
 							Builder.Append("\t{ \\\r\n");
@@ -230,7 +229,7 @@ namespace EpicGames.UHT.Exporters.CodeGen
 							}
 
 							Builder.Append("\t\t").Append(MethodInfo.ReturnPrefix()).Append("Static").Append(MethodInfo.Name).Append("( \\\r\n");
-							Builder.Append("\t\t\tRigVMExecuteContext");
+							Builder.Append("\t\t\tRigVMExecuteContext.PublicData");
 							Builder.AppendParameterNames(ScriptStruct.RigVMStructInfo.Members, true, ", \\\r\n\t\t\t", false);
 							Builder.AppendParameterNames(MethodInfo.Parameters, true, ", \\\r\n\t\t\t");
 							Builder.Append(" \\\r\n");
@@ -406,6 +405,7 @@ namespace EpicGames.UHT.Exporters.CodeGen
 
 				AppendStandardConstructors(Builder, Class, Api);
 				AppendEnhancedConstructors(Builder, Class, Api);
+				AppendFieldNotify(Builder, Class);
 			}
 
 			using (UhtMacroCreator Macro = new UhtMacroCreator(Builder, this, Class.PrologLineNumber, PrologMacroSuffix))
@@ -437,6 +437,79 @@ namespace EpicGames.UHT.Exporters.CodeGen
 			return Builder;
 		}
 
+		private StringBuilder AppendFieldNotify(StringBuilder Builder, UhtClass Class)
+		{
+			if (!NeedFieldNotifyCodeGen(Class))
+			{
+				return Builder;
+			}
+
+			// Scan the children to see what we have
+			bool bHasProperties;
+			bool bHasFunctions;
+			bool bHasEditorFields;
+			bool bAllEditorFields;
+			GetFieldNotifyStats(Class, out bHasProperties, out bHasFunctions, out bHasEditorFields, out bAllEditorFields);
+
+			// If we only have editor fields or no editor fields, then we only emit one block
+			if (bHasEditorFields)
+			{
+				Builder.Append("#if WITH_EDITORONLY_DATA\r\n");
+				using (UhtMacroCreator Macro = new UhtMacroCreator(Builder, this, Class, FieldNotifyMacroSuffix))
+				{
+					AppendFieldNotify(Builder, Class, bHasProperties, bHasFunctions, bHasEditorFields, bAllEditorFields, true);
+				}
+				Builder.Append("#else //WITH_EDITORONLY_DATA\r\n");
+				using (UhtMacroCreator Macro = new UhtMacroCreator(Builder, this, Class, FieldNotifyMacroSuffix))
+				{
+					AppendFieldNotify(Builder, Class, bHasProperties, bHasFunctions, bHasEditorFields, bAllEditorFields, false);
+				}
+				Builder.Append("#endif // WITH_EDITORONLY_DATA\r\n");
+			}
+			else
+			{
+				using (UhtMacroCreator Macro = new UhtMacroCreator(Builder, this, Class, FieldNotifyMacroSuffix))
+				{
+					AppendFieldNotify(Builder, Class, bHasProperties, bHasFunctions, bHasEditorFields, bAllEditorFields, false);
+				}
+			}
+			return Builder;
+		}
+
+		private StringBuilder AppendFieldNotify(StringBuilder Builder, UhtClass Class, 
+			bool bHasProperties, bool bHasFunctions, bool bHasEditorFields, bool bAllEditorFields,
+			bool bIncludeEditorOnlyFields)
+		{
+			Builder.Append("\tUE_FIELD_NOTIFICATION_DECLARE_CLASS_DESCRIPTOR_BEGIN() \\\r\n");
+
+			//UE_FIELD_NOTIFICATION_DECLARE_FIELD
+			AppendFieldNotify(Builder, Class, bHasProperties, bHasFunctions, bHasEditorFields, bAllEditorFields, 
+				bIncludeEditorOnlyFields, false, (StringBuilder Builder, UhtClass Class, string Name) =>
+			{
+				Builder.Append($"\tUE_FIELD_NOTIFICATION_DECLARE_FIELD({Name}) \\\r\n");
+			});
+
+			//UE_FIELD_NOTIFICATION_DECLARE_ENUM_FIELD
+			bool bIsFirst = true;
+			AppendFieldNotify(Builder, Class, bHasProperties, bHasFunctions, bHasEditorFields, bAllEditorFields, 
+				bIncludeEditorOnlyFields, false, (StringBuilder Builder, UhtClass Class, string Name) =>
+			{
+				if (bIsFirst)
+				{
+					bIsFirst = false;
+					Builder.Append($"\tUE_FIELD_NOTIFICATION_DECLARE_ENUM_FIELD_BEGIN({Name}) \\\r\n");
+				}
+				else
+				{
+					Builder.Append($"\tUE_FIELD_NOTIFICATION_DECLARE_ENUM_FIELD({Name}) \\\r\n");
+				}
+			});
+
+			Builder.Append("\tUE_FIELD_NOTIFICATION_DECLARE_ENUM_FIELD_END() \\\r\n");
+			Builder.Append($"\tUE_FIELD_NOTIFICATION_DECLARE_CLASS_DESCRIPTOR_END(); \\\r\n");
+			return Builder;
+		}
+
 		private StringBuilder AppendSparseDeclarations(StringBuilder Builder, UhtClass Class)
 		{
 			// Format the sparse data
@@ -1275,6 +1348,10 @@ namespace EpicGames.UHT.Exporters.CodeGen
 					{
 						Builder.Append('\t').AppendMacroName(this, Class, EnchancedConstructorsMacroSuffix).Append(" \\\r\n");
 					}
+					if (NeedFieldNotifyCodeGen(Class))
+					{
+						Builder.Append('\t').AppendMacroName(this, Class, FieldNotifyMacroSuffix).Append(" \\\r\n");
+					}
 				}
 				if (bIsLegacy)
 				{
diff --git a/Engine/Source/Programs/Shared/EpicGames.UHT/Specifiers/UhtClassSpecifiers.cs b/Engine/Source/Programs/Shared/EpicGames.UHT/Specifiers/UhtClassSpecifiers.cs
--- a/Engine/Source/Programs/Shared/EpicGames.UHT/Specifiers/UhtClassSpecifiers.cs
+++ b/Engine/Source/Programs/Shared/EpicGames.UHT/Specifiers/UhtClassSpecifiers.cs
@@ -340,5 +340,12 @@ namespace EpicGames.UHT.Parsers
 			UhtClassParser Class = (UhtClassParser)SpecifierContext.Scope.ScopeType;
 			Class.AddClassFlags(EClassFlags.Interface);
 		}
+
+		[UhtSpecifier(Extends = UhtTableNames.Class, ValueType = UhtSpecifierValueType.None)]
+		private static void CustomFieldNotifySpecifier(UhtSpecifierContext SpecifierContext)
+		{
+			UhtClassParser Class = (UhtClassParser)SpecifierContext.Scope.ScopeType;
+			Class.ClassExportFlags |= UhtClassExportFlags.HasCustomFieldNotify;
+		}
 	}
 }
diff --git a/Engine/Source/Programs/Shared/EpicGames.UHT/Specifiers/UhtFunctionSpecifiers.cs b/Engine/Source/Programs/Shared/EpicGames.UHT/Specifiers/UhtFunctionSpecifiers.cs
--- a/Engine/Source/Programs/Shared/EpicGames.UHT/Specifiers/UhtFunctionSpecifiers.cs
+++ b/Engine/Source/Programs/Shared/EpicGames.UHT/Specifiers/UhtFunctionSpecifiers.cs
@@ -299,6 +299,13 @@ namespace EpicGames.UHT.Parsers
 			}
 		}
 
+		[UhtSpecifier(Extends = UhtTableNames.Function, ValueType = UhtSpecifierValueType.None)]
+		private static void FieldNotifySpecifier(UhtSpecifierContext SpecifierContext)
+		{
+			UhtFunction Function = (UhtFunction)SpecifierContext.Scope.ScopeType;
+			Function.FunctionExportFlags |= UhtFunctionExportFlags.FieldNotify;
+		}
+
 		[UhtSpecifierValidator(Extends = UhtTableNames.Function)]
 		private static void BlueprintProtectedSpecifierValidator(UhtType Type, UhtMetaData MetaData, UhtMetaDataKey Key, StringView Value)
 		{
diff --git a/Engine/Source/Programs/Shared/EpicGames.UHT/Specifiers/UhtPropertyMemberSpecifiers.cs b/Engine/Source/Programs/Shared/EpicGames.UHT/Specifiers/UhtPropertyMemberSpecifiers.cs
--- a/Engine/Source/Programs/Shared/EpicGames.UHT/Specifiers/UhtPropertyMemberSpecifiers.cs
+++ b/Engine/Source/Programs/Shared/EpicGames.UHT/Specifiers/UhtPropertyMemberSpecifiers.cs
@@ -444,5 +444,12 @@ namespace EpicGames.UHT.Parsers
 				}
 			}
 		}
+
+		[UhtSpecifier(Extends = UhtTableNames.PropertyMember, ValueType = UhtSpecifierValueType.None)]
+		private static void FieldNotifySpecifier(UhtSpecifierContext SpecifierContext)
+		{
+			UhtPropertySpecifierContext Context = (UhtPropertySpecifierContext)SpecifierContext;
+			Context.PropertySettings.PropertyExportFlags |= UhtPropertyExportFlags.FieldNotify;
+		}
 	}
 }
diff --git a/Engine/Source/Programs/Shared/EpicGames.UHT/Types/Properties/UhtClassProperty.cs b/Engine/Source/Programs/Shared/EpicGames.UHT/Types/Properties/UhtClassProperty.cs
--- a/Engine/Source/Programs/Shared/EpicGames.UHT/Types/Properties/UhtClassProperty.cs
+++ b/Engine/Source/Programs/Shared/EpicGames.UHT/Types/Properties/UhtClassProperty.cs
@@ -142,19 +142,6 @@ namespace EpicGames.UHT.Types
 			return Builder;
 		}
 
-		/// <inheritdoc/>
-		public override string? GetRigVMType(ref UhtRigVMParameterFlags ParameterFlags)
-		{
-			if (this.PropertyFlags.HasAnyFlags(EPropertyFlags.UObjectWrapper))
-			{
-				return $"TSubclassOf<{this.MetaClass?.SourceName}> "; //COMPATIBILITY-TODO - Extra space in old UHT
-			}
-			else
-			{
-				return $"{this.Class}*";
-			}
-		}
-
 		#region Keyword
 		[UhtPropertyType(Keyword = "TSubclassOf")]
 		private static UhtProperty? SubclassOfProperty(UhtPropertyResolvePhase ResolvePhase, UhtPropertySettings PropertySettings, IUhtTokenReader TokenReader, UhtToken MatchedToken)
diff --git a/Engine/Source/Programs/Shared/EpicGames.UHT/Types/Properties/UhtClassPtrProperties.cs b/Engine/Source/Programs/Shared/EpicGames.UHT/Types/Properties/UhtClassPtrProperties.cs
--- a/Engine/Source/Programs/Shared/EpicGames.UHT/Types/Properties/UhtClassPtrProperties.cs
+++ b/Engine/Source/Programs/Shared/EpicGames.UHT/Types/Properties/UhtClassPtrProperties.cs
@@ -80,11 +80,5 @@ namespace EpicGames.UHT.Types
 				OuterStruct.LogError("UFunctions cannot take a TObjectPtr as a parameter.");
 			}
 		}
-
-		/// <inheritdoc/>
-		public override string? GetRigVMType(ref UhtRigVMParameterFlags ParameterFlags)
-		{
-			return $"TObjectPtr<{this.MetaClass?.SourceName}>";
-		}
 	}
 }
diff --git a/Engine/Source/Programs/Shared/EpicGames.UHT/Types/Properties/UhtLazyObjectPtrProperty.cs b/Engine/Source/Programs/Shared/EpicGames.UHT/Types/Properties/UhtLazyObjectPtrProperty.cs
--- a/Engine/Source/Programs/Shared/EpicGames.UHT/Types/Properties/UhtLazyObjectPtrProperty.cs
+++ b/Engine/Source/Programs/Shared/EpicGames.UHT/Types/Properties/UhtLazyObjectPtrProperty.cs
@@ -104,12 +104,6 @@ namespace EpicGames.UHT.Types
 			}
 		}
 
-		/// <inheritdoc/>
-		public override string? GetRigVMType(ref UhtRigVMParameterFlags ParameterFlags)
-		{
-			return $"TLazyObjectPtr<{this.Class.SourceName}>";
-		}
-
 		#region Keyword
 		[UhtPropertyType(Keyword = "TLazyObjectPtr")]
 		private static UhtProperty? LazyObjectPtrProperty(UhtPropertyResolvePhase ResolvePhase, UhtPropertySettings PropertySettings, IUhtTokenReader TokenReader, UhtToken MatchedToken)
diff --git a/Engine/Source/Programs/Shared/EpicGames.UHT/Types/Properties/UhtObjectProperty.cs b/Engine/Source/Programs/Shared/EpicGames.UHT/Types/Properties/UhtObjectProperty.cs
--- a/Engine/Source/Programs/Shared/EpicGames.UHT/Types/Properties/UhtObjectProperty.cs
+++ b/Engine/Source/Programs/Shared/EpicGames.UHT/Types/Properties/UhtObjectProperty.cs
@@ -112,11 +112,5 @@ namespace EpicGames.UHT.Types
 			}
 			return false;
 		}
-
-		/// <inheritdoc/>
-		public override string? GetRigVMType(ref UhtRigVMParameterFlags ParameterFlags)
-		{
-			return $"{this.Class}*";
-		}
 	}
 }
diff --git a/Engine/Source/Programs/Shared/EpicGames.UHT/Types/Properties/UhtObjectPropertyBase.cs b/Engine/Source/Programs/Shared/EpicGames.UHT/Types/Properties/UhtObjectPropertyBase.cs
--- a/Engine/Source/Programs/Shared/EpicGames.UHT/Types/Properties/UhtObjectPropertyBase.cs
+++ b/Engine/Source/Programs/Shared/EpicGames.UHT/Types/Properties/UhtObjectPropertyBase.cs
@@ -135,6 +135,12 @@ namespace EpicGames.UHT.Types
 			return this.Class.ClassFlags.HasAnyFlags(EClassFlags.Const);
 		}
 
+		/// <inheritdoc/>
+		public override string? GetRigVMType(ref UhtRigVMParameterFlags ParameterFlags)
+		{
+			return null;
+		}
+
 		#region Parsing support methods
 		/// <summary>
 		/// Parse a template type
diff --git a/Engine/Source/Programs/Shared/EpicGames.UHT/Types/Properties/UhtObjectPtrProperty.cs b/Engine/Source/Programs/Shared/EpicGames.UHT/Types/Properties/UhtObjectPtrProperty.cs
--- a/Engine/Source/Programs/Shared/EpicGames.UHT/Types/Properties/UhtObjectPtrProperty.cs
+++ b/Engine/Source/Programs/Shared/EpicGames.UHT/Types/Properties/UhtObjectPtrProperty.cs
@@ -92,12 +92,6 @@ namespace EpicGames.UHT.Types
 			}
 		}
 
-		/// <inheritdoc/>
-		public override string? GetRigVMType(ref UhtRigVMParameterFlags ParameterFlags)
-		{
-			return $"TObjectPtr<{this.Class.SourceName}>";
-		}
-
 		#region Keyword
 		[UhtPropertyType(Keyword = "TObjectPtr")]
 		private static UhtProperty? ObjectPtrProperty(UhtPropertyResolvePhase ResolvePhase, UhtPropertySettings PropertySettings, IUhtTokenReader TokenReader, UhtToken MatchedToken)
diff --git a/Engine/Source/Programs/Shared/EpicGames.UHT/Types/Properties/UhtSoftClassProperty.cs b/Engine/Source/Programs/Shared/EpicGames.UHT/Types/Properties/UhtSoftClassProperty.cs
--- a/Engine/Source/Programs/Shared/EpicGames.UHT/Types/Properties/UhtSoftClassProperty.cs
+++ b/Engine/Source/Programs/Shared/EpicGames.UHT/Types/Properties/UhtSoftClassProperty.cs
@@ -77,12 +77,6 @@ namespace EpicGames.UHT.Types
 			return Builder;
 		}
 
-		/// <inheritdoc/>
-		public override string? GetRigVMType(ref UhtRigVMParameterFlags ParameterFlags)
-		{
-			return $"TSoftClassPtr<{this.MetaClass?.SourceName}>";
-		}
-
 		#region Keyword
 		[UhtPropertyType(Keyword = "TSoftClassPtr")]
 		private static UhtProperty? SoftClassPtrProperty(UhtPropertyResolvePhase ResolvePhase, UhtPropertySettings PropertySettings, IUhtTokenReader TokenReader, UhtToken MatchedToken)
diff --git a/Engine/Source/Programs/Shared/EpicGames.UHT/Types/Properties/UhtSoftObjectProperty.cs b/Engine/Source/Programs/Shared/EpicGames.UHT/Types/Properties/UhtSoftObjectProperty.cs
--- a/Engine/Source/Programs/Shared/EpicGames.UHT/Types/Properties/UhtSoftObjectProperty.cs
+++ b/Engine/Source/Programs/Shared/EpicGames.UHT/Types/Properties/UhtSoftObjectProperty.cs
@@ -85,12 +85,6 @@ namespace EpicGames.UHT.Types
 			return false;
 		}
 
-		/// <inheritdoc/>
-		public override string? GetRigVMType(ref UhtRigVMParameterFlags ParameterFlags)
-		{
-			return $"TSoftObjectPtr<{this.Class.SourceName}>";
-		}
-
 		#region Keyword
 		[UhtPropertyType(Keyword = "TSoftObjectPtr")]
 		private static UhtProperty? SoftObjectPtrProperty(UhtPropertyResolvePhase ResolvePhase, UhtPropertySettings PropertySettings, IUhtTokenReader TokenReader, UhtToken MatchedToken)
diff --git a/Engine/Source/Programs/Shared/EpicGames.UHT/Types/Properties/UhtWeakObjectPtrProperty.cs b/Engine/Source/Programs/Shared/EpicGames.UHT/Types/Properties/UhtWeakObjectPtrProperty.cs
--- a/Engine/Source/Programs/Shared/EpicGames.UHT/Types/Properties/UhtWeakObjectPtrProperty.cs
+++ b/Engine/Source/Programs/Shared/EpicGames.UHT/Types/Properties/UhtWeakObjectPtrProperty.cs
@@ -93,19 +93,6 @@ namespace EpicGames.UHT.Types
 			return false;
 		}
 
-		/// <inheritdoc/>
-		public override string? GetRigVMType(ref UhtRigVMParameterFlags ParameterFlags)
-		{
-			if (this.PropertyFlags.HasAnyFlags(EPropertyFlags.AutoWeak))
-			{
-				return $"TAutoWeakObjectPtr<{this.Class.SourceName}>";
-			}
-			else
-			{
-				return $"TWeakObjectPtr<{this.Class.SourceName}>";
-			}
-		}
-
 		#region Keywords
 		private static UhtProperty CreateWeakProperty(UhtPropertySettings PropertySettings, IUhtTokenReader TokenReader, UhtClass Class, EPropertyFlags ExtraFlags = EPropertyFlags.None)
 		{
diff --git a/Engine/Source/Programs/Shared/EpicGames.UHT/Types/UhtClass.cs b/Engine/Source/Programs/Shared/EpicGames.UHT/Types/UhtClass.cs
--- a/Engine/Source/Programs/Shared/EpicGames.UHT/Types/UhtClass.cs
+++ b/Engine/Source/Programs/Shared/EpicGames.UHT/Types/UhtClass.cs
@@ -74,6 +74,16 @@ namespace EpicGames.UHT.Types
 		/// Class should not be exported
 		/// </summary>
 		NoExport = 1 << 8,
+
+		/// <summary>
+		/// Class has a custom field notify
+		/// </summary>
+		HasCustomFieldNotify = 1 << 9,
+
+		/// <summary>
+		/// Class has a field notify
+		/// </summary>
+		HasFieldNotify = 1 << 10,
 	}
 
 	/// <summary>
@@ -569,6 +579,11 @@ namespace EpicGames.UHT.Types
 					{
 						if (Child is UhtProperty Property)
 						{
+							if (Property.PropertyExportFlags.HasAnyFlags(UhtPropertyExportFlags.FieldNotify))
+							{
+								this.ClassExportFlags |= UhtClassExportFlags.HasFieldNotify;
+							}
+
 							if (Property.PropertyFlags.HasAnyFlags(EPropertyFlags.Net))
 							{
 								this.ClassExportFlags |= UhtClassExportFlags.SelfHasReplicatedProperties;
@@ -599,6 +614,13 @@ namespace EpicGames.UHT.Types
 								}
 							}
 						}
+						else if (Child is UhtFunction Function)
+						{
+							if (Function.FunctionExportFlags.HasAnyFlags(UhtFunctionExportFlags.FieldNotify))
+							{
+								this.ClassExportFlags |= UhtClassExportFlags.HasFieldNotify;
+							}
+						}
 					}
 
 					if (this.SuperClass != null)
@@ -946,6 +968,18 @@ namespace EpicGames.UHT.Types
 				this.LogError($"'NeedsDeferredDependencyLoading' is set on '{this.SourceName}' but the flag can only be used with classes derived from UClass.");
 			}
 
+			if (this.ClassExportFlags.HasAnyFlags(UhtClassExportFlags.HasFieldNotify))
+			{
+				if (this.Session.INotifyFieldValueChanged == null)
+				{
+					this.LogError($"UClass '{this.SourceName}' has FieldNotify elements but the interface 'INotifyFieldValueChanged' is not defined.");
+				}
+				else if (!ImplementsInterface(this.Session.INotifyFieldValueChanged))
+				{
+					this.LogError($"UClass '{this.SourceName}' need to implement the interface INotifyFieldValueChanged' to support FieldNotify.");
+				}
+			}
+
 			ValidateProperties();
 
 			// Perform any validation specific to the class type
@@ -1363,5 +1397,23 @@ namespace EpicGames.UHT.Types
 					break;
 			}
 		}
+
+		private bool ImplementsInterface(UhtClass Interface)
+		{
+			for (UhtClass? SuperClass = this; SuperClass != null; SuperClass = SuperClass.SuperClass)
+			{
+				if (SuperClass.Bases != null)
+				{
+					foreach (UhtStruct Struct in SuperClass.Bases)
+					{
+						if (Struct.IsChildOf(Interface))
+						{
+							return true;
+						}
+					}
+				}
+			}
+			return false;
+		}
 	}
 }
diff --git a/Engine/Source/Programs/Shared/EpicGames.UHT/Types/UhtFunction.cs b/Engine/Source/Programs/Shared/EpicGames.UHT/Types/UhtFunction.cs
--- a/Engine/Source/Programs/Shared/EpicGames.UHT/Types/UhtFunction.cs
+++ b/Engine/Source/Programs/Shared/EpicGames.UHT/Types/UhtFunction.cs
@@ -103,6 +103,11 @@ namespace EpicGames.UHT.Types
 		/// Final flag was set automatically and should not be considered for validation
 		/// </summary>
 		AutoFinal = 1 << 15,
+
+		/// <summary>
+		/// Generate the entry for the FieldNotificationClassDescriptor
+		/// </summary>
+		FieldNotify = 1 << 16,
 	}
 
 	/// <summary>
@@ -666,6 +671,35 @@ namespace EpicGames.UHT.Types
 							}
 						}
 					}
+
+					// Validate field notify
+					if (this.FunctionExportFlags.HasAnyFlags(UhtFunctionExportFlags.FieldNotify))
+					{
+						if (OuterClass == null)
+						{
+							this.LogError($"FieldNofity function '{this.SourceName}' are only valid as UClass member function.");
+						}
+
+						if (this.ParameterProperties.Length != 0)
+						{
+							this.LogError($"FieldNotify function '{this.SourceName}' must not have parameters.");
+						}
+
+						if (this.ReturnProperty == null)
+						{
+							this.LogError($"FieldNotify function '{this.SourceName}' must return a value.");
+						}
+
+						if (this.FunctionFlags.HasAnyFlags(EFunctionFlags.Event))
+						{
+							this.LogError($"FieldNotify function '{this.SourceName}' cannot be a blueprint event.");
+						}
+
+						if (!this.FunctionFlags.HasAnyFlags(EFunctionFlags.BlueprintPure))
+						{
+							this.LogError($"FieldNotify function '{this.SourceName}' must be pure.");
+						}
+					}
 					break;
 
 				case UhtFunctionType.Delegate:
@@ -729,6 +763,10 @@ namespace EpicGames.UHT.Types
 			{
 				Options |= UhtValidationOptions.Deprecated;
 			}
+			else
+			{
+				Options &= ~UhtValidationOptions.Deprecated;
+			}
 			return Options;
 		}
 
diff --git a/Engine/Source/Programs/Shared/EpicGames.UHT/Types/UhtProperty.cs b/Engine/Source/Programs/Shared/EpicGames.UHT/Types/UhtProperty.cs
--- a/Engine/Source/Programs/Shared/EpicGames.UHT/Types/UhtProperty.cs
+++ b/Engine/Source/Programs/Shared/EpicGames.UHT/Types/UhtProperty.cs
@@ -70,6 +70,11 @@ namespace EpicGames.UHT.Types
 		/// If true, the property has a setter function
 		/// </summary>
 		SetterFound = 0x00000200,
+
+		/// <summary>
+		/// Property is marked as a field notify
+		/// </summary>
+		FieldNotify = 0x00000400,
 	};
 
 	/// <summary>
@@ -1752,6 +1757,21 @@ namespace EpicGames.UHT.Types
 			{
 				this.LogError("ExposeOnSpawn - Property cannot be exposed");
 			}
+
+			if (this.PropertyExportFlags.HasAnyFlags(UhtPropertyExportFlags.FieldNotify))
+			{
+				if (this.Outer is UhtClass Class)
+				{
+					if (Class.ClassType != UhtClassType.Class)
+					{
+						this.LogError("FieldNofity are not valid on UInterface.");
+					}
+				}	
+				else
+				{
+					this.LogError("FieldNofity property are only valid as UClass member variable.");
+				}
+			}
 		}
 
 		/// <summary>
diff --git a/Engine/Source/Programs/Shared/EpicGames.UHT/Utils/UhtException.cs b/Engine/Source/Programs/Shared/EpicGames.UHT/Utils/UhtException.cs
--- a/Engine/Source/Programs/Shared/EpicGames.UHT/Utils/UhtException.cs
+++ b/Engine/Source/Programs/Shared/EpicGames.UHT/Utils/UhtException.cs
@@ -62,6 +62,9 @@ namespace EpicGames.UHT.Utils
 			}
 			this.UhtMessage = UhtMessage.MakeMessage(UhtMessageType.Error, MessageSite.MessageSource, null, MessageSite.GetLineNumber(LineNumber), Message);
 		}
+
+		/// <inheritdoc/>
+		public override string Message { get => this.UhtMessage.Message; }
 	}
 
 	/// <summary>
diff --git a/Engine/Source/Programs/Shared/EpicGames.UHT/Utils/UhtSession.cs b/Engine/Source/Programs/Shared/EpicGames.UHT/Utils/UhtSession.cs
--- a/Engine/Source/Programs/Shared/EpicGames.UHT/Utils/UhtSession.cs
+++ b/Engine/Source/Programs/Shared/EpicGames.UHT/Utils/UhtSession.cs
@@ -952,6 +952,12 @@ namespace EpicGames.UHT.Utils
 		/// </summary>
 		public UhtClass? AActor = null;
 
+		/// <summary>
+		/// After headers are parsed, return the INotifyFieldValueChanged interface.  There is no requirement 
+		/// that this interface be defined.
+		/// </summary>
+		public UhtClass? INotifyFieldValueChanged = null;
+
 		private List<UhtPackage> PackagesInternal = new List<UhtPackage>();
 		private List<UhtHeaderFile> HeaderFilesInternal = new List<UhtHeaderFile>();
 		private List<UhtHeaderFile> SortedHeaderFilesInternal = new List<UhtHeaderFile>();
@@ -2038,6 +2044,7 @@ namespace EpicGames.UHT.Utils
 				this.UInterfaceInternal = (UhtClass?)FindType(null, UhtFindOptions.SourceName | UhtFindOptions.Class, "UInterface");
 				this.IInterfaceInternal = (UhtClass?)FindType(null, UhtFindOptions.SourceName | UhtFindOptions.Class, "IInterface");
 				this.AActor = (UhtClass?)FindType(null, UhtFindOptions.SourceName | UhtFindOptions.Class, "AActor");
+				this.INotifyFieldValueChanged = (UhtClass?)FindType(null, UhtFindOptions.SourceName | UhtFindOptions.Class, "INotifyFieldValueChanged");
 			});
 		}
 
diff --git a/Engine/Source/Programs/Shared/EpicGames.UHT/Utils/UhtTables.cs b/Engine/Source/Programs/Shared/EpicGames.UHT/Utils/UhtTables.cs
--- a/Engine/Source/Programs/Shared/EpicGames.UHT/Utils/UhtTables.cs
+++ b/Engine/Source/Programs/Shared/EpicGames.UHT/Utils/UhtTables.cs
@@ -152,7 +152,7 @@ namespace EpicGames.UHT.Utils
 				}
 				else
 				{
-					throw new UhtIceException($"The '{Suffix}' attribute on the {ClassType.Name}.{MethodInfo.Name} method doesn't have a name specified or the method name doesn't end in '{2}'.");
+					throw new UhtIceException($"The '{Suffix}' attribute on the {ClassType.Name}.{MethodInfo.Name} method doesn't have a name specified or the method name doesn't end in '{Suffix}'.");
 				}
 			}
 			return Name;
