Prepare UBT and automation code type scanning for the moving of the MSBuild code.  This requires a change in how we handle the load failure of Microsoft.build.framework.

#rnx

diff --git a/Engine/Source/Programs/AutomationTool/AutomationUtils/MCPPublic.cs b/Engine/Source/Programs/AutomationTool/AutomationUtils/MCPPublic.cs
--- a/Engine/Source/Programs/AutomationTool/AutomationUtils/MCPPublic.cs
+++ b/Engine/Source/Programs/AutomationTool/AutomationUtils/MCPPublic.cs
@@ -9,6 +9,7 @@ using System.Runtime.Serialization;
 using System.Net;
 using System.Reflection;
 using System.Text.RegularExpressions;
+using UnrealBuildBase;
 using UnrealBuildTool;
 using EpicGames.MCP.Automation;
 
@@ -20,23 +21,6 @@ namespace EpicGames.MCP.Automation
 	using System.ComponentModel;
 	using System.Globalization;
 
-	public static class Extensions
-	{
-		public static Type[] SafeGetLoadedTypes(this Assembly Dll)
-		{
-			Type[] AllTypes;
-			try
-			{
-				AllTypes = Dll.GetTypes();
-			}
-			catch (ReflectionTypeLoadException e)
-			{
-				AllTypes = e.Types.Where(x => x != null).ToArray();
-			}
-			return AllTypes;
-		}
-	}
-
     /// <summary>
     /// Utility class to provide commit/rollback functionality via an RAII-like functionality.
     /// Usage is to provide a rollback action that will be called on Dispose if the Commit() method is not called.
@@ -2472,7 +2456,7 @@ namespace EpicGames.MCP.Automation
 				Assembly[] LoadedAssemblies = AppDomain.CurrentDomain.GetAssemblies();
 				foreach (var Dll in LoadedAssemblies)
 				{
-					Type[] AllTypes = Dll.GetTypes();
+					Type[] AllTypes = Dll.SafeGetLoadedTypes();
 					foreach (var PotentialConfigType in AllTypes)
 					{
 						if (PotentialConfigType != typeof(CatalogServiceBase) && typeof(CatalogServiceBase).IsAssignableFrom(PotentialConfigType))
diff --git a/Engine/Source/Programs/Shared/EpicGames.Build/System/AssemblyExtensions.cs b/Engine/Source/Programs/Shared/EpicGames.Build/System/AssemblyExtensions.cs
new file mode 100644
--- /dev/null
+++ b/Engine/Source/Programs/Shared/EpicGames.Build/System/AssemblyExtensions.cs
@@ -0,0 +1,25 @@
+// Copyright Epic Games, Inc. All Rights Reserved.
+
+using System;
+using System.Linq;
+using System.Reflection;
+
+namespace UnrealBuildBase
+{
+	public static class AssemblyExtensions
+	{
+		public static Type[] SafeGetLoadedTypes(this Assembly Dll)
+		{
+			Type[] AllTypes;
+			try
+			{
+				AllTypes = Dll.GetTypes();
+			}
+			catch (ReflectionTypeLoadException e)
+			{
+				AllTypes = e.Types.Where(x => x != null).ToArray();
+			}
+			return AllTypes;
+		}
+	}
+}
diff --git a/Engine/Source/Programs/Shared/EpicGames.UHT/Utils/UhtAttributeScanner.cs b/Engine/Source/Programs/Shared/EpicGames.UHT/Utils/UhtAttributeScanner.cs
--- a/Engine/Source/Programs/Shared/EpicGames.UHT/Utils/UhtAttributeScanner.cs
+++ b/Engine/Source/Programs/Shared/EpicGames.UHT/Utils/UhtAttributeScanner.cs
@@ -6,6 +6,7 @@ using System;
 using System.Collections.Generic;
 using System.Linq;
 using System.Reflection;
+using UnrealBuildBase;
 
 namespace EpicGames.UHT.Utils
 {
@@ -66,26 +67,12 @@ namespace EpicGames.UHT.Utils
 		/// </summary>
 		private void Initialize()
 		{
-			IEnumerable<Assembly> LoadedAssemblies = AppDomain.CurrentDomain.GetAssemblies()
-				// Exclude the AutomationTool driver assembly - it contains no tasks, and trying to load it in the context of
-				// BuildGraph can result in an exception if Microsoft.Build.Framework is not able to be loaded
-				.Where(A => !String.Equals("AutomationTool", A.GetName().Name));
+			IEnumerable<Assembly> LoadedAssemblies = AppDomain.CurrentDomain.GetAssemblies();
 
 			// Collect all of the types from the loaded assemblies
 			foreach (Assembly LoadedAssembly in LoadedAssemblies)
 			{
-				Type[] Types;
-				try
-				{
-					Types = LoadedAssembly.GetTypes();
-				}
-				catch (ReflectionTypeLoadException ex)
-				{
-					Log.TraceWarning("Exception {0} while trying to get types from assembly {1}. LoaderExceptions: {2}", 
-						ex, LoadedAssembly, string.Join("\n", ex.LoaderExceptions.Select(x => x?.Message)));
-					continue;
-				}
-
+				Type[] Types = LoadedAssembly.SafeGetLoadedTypes();
 				foreach (Type Type in Types)
 				{
 					CheckForAttributes(Type);
