Horde: Use standard coding style for EpicGames.Horde.

diff --git a/Engine/Source/Programs/Shared/EpicGames.Horde.Tests/BundleTests.cs b/Engine/Source/Programs/Shared/EpicGames.Horde.Tests/BundleTests.cs
--- a/Engine/Source/Programs/Shared/EpicGames.Horde.Tests/BundleTests.cs
+++ b/Engine/Source/Programs/Shared/EpicGames.Horde.Tests/BundleTests.cs
@@ -9,10 +9,8 @@ using EpicGames.Serialization;
 using Microsoft.Extensions.Caching.Memory;
 using Microsoft.VisualStudio.TestTools.UnitTesting;
 using System;
-using System.Buffers.Binary;
 using System.Collections.Generic;
 using System.Linq;
-using System.Security.Cryptography;
 using System.Text;
 using System.Threading.Tasks;
 
@@ -21,46 +19,46 @@ namespace EpicGames.Horde.Tests
 	[TestClass]
 	public class BundleTests : IDisposable
 	{
-		MemoryCache Cache;
-		MemoryStorageClient StorageClient = new MemoryStorageClient();
-		NamespaceId NamespaceId = new NamespaceId("namespace");
-		BucketId BucketId = new BucketId("bucket");
+		readonly MemoryCache _cache;
+		readonly MemoryStorageClient _storageClient = new MemoryStorageClient();
+		readonly NamespaceId _namespaceId = new NamespaceId("namespace");
+		readonly BucketId _bucketId = new BucketId("bucket");
 
 		public BundleTests()
 		{
-			Cache = new MemoryCache(new MemoryCacheOptions { SizeLimit = 50 * 1024 * 1024 });
+			_cache = new MemoryCache(new MemoryCacheOptions { SizeLimit = 50 * 1024 * 1024 });
 		}
 
 		public void Dispose()
 		{
-			Cache.Dispose();
+			_cache.Dispose();
 		}
 
 		[TestMethod]
 		public void BuzHashTests()
 		{
-			byte[] Data = new byte[4096];
-			new Random(0).NextBytes(Data);
+			byte[] data = new byte[4096];
+			new Random(0).NextBytes(data);
 
 			const int WindowSize = 128;
 
-			uint RollingHash = 0;
-			for (int MaxIdx = 0; MaxIdx < Data.Length + WindowSize; MaxIdx++)
+			uint rollingHash = 0;
+			for (int maxIdx = 0; maxIdx < data.Length + WindowSize; maxIdx++)
 			{
-				int MinIdx = MaxIdx - WindowSize;
+				int minIdx = maxIdx - WindowSize;
 
-				if (MaxIdx < Data.Length)
+				if (maxIdx < data.Length)
 				{
-					RollingHash = BuzHash.Add(RollingHash, Data[MaxIdx]);
+					rollingHash = BuzHash.Add(rollingHash, data[maxIdx]);
 				}
 
-				int Length = Math.Min(MaxIdx + 1, Data.Length) - Math.Max(MinIdx, 0);
-				uint CleanHash = BuzHash.Add(0, Data.AsSpan(Math.Max(MinIdx, 0), Length));
-				Assert.AreEqual(RollingHash, CleanHash);
+				int length = Math.Min(maxIdx + 1, data.Length) - Math.Max(minIdx, 0);
+				uint cleanHash = BuzHash.Add(0, data.AsSpan(Math.Max(minIdx, 0), length));
+				Assert.AreEqual(rollingHash, cleanHash);
 
-				if (MinIdx >= 0)
+				if (minIdx >= 0)
 				{
-					RollingHash = BuzHash.Sub(RollingHash, Data[MinIdx], Length);
+					rollingHash = BuzHash.Sub(rollingHash, data[minIdx], length);
 				}
 			}
 		}
@@ -68,344 +66,344 @@ namespace EpicGames.Horde.Tests
 		[TestMethod]
 		public void BasicChunkingTests()
 		{
-			ChunkingOptions Options = new ChunkingOptions();
-			Options.LeafOptions = new ChunkingOptionsForNodeType(8, 8, 8);
-
-			FileNode Node = new FileNode();
-			Node.Append(new byte[7], Options);
-			Assert.AreEqual(0, Node.Depth);
-			Assert.AreEqual(7, Node.Payload.Length);
-
-			Node = new FileNode();
-			Node.Append(new byte[8], Options);
-			Assert.AreEqual(0, Node.Depth);
-			Assert.AreEqual(8, Node.Payload.Length);
-
-			Node = new FileNode();
-			Node.Append(new byte[9], Options);
-			Assert.AreEqual(1, Node.Depth);
-			Assert.AreEqual(2, Node.Children.Count);
-
-			FileNode? ChildNode1 = Node.Children[0].Node;
-			Assert.IsNotNull(ChildNode1);
-			Assert.AreEqual(0, ChildNode1!.Depth);
-			Assert.AreEqual(8, ChildNode1!.Payload.Length);
-
-			FileNode? ChildNode2 = Node.Children[1].Node;
-			Assert.IsNotNull(ChildNode2);
-			Assert.AreEqual(0, ChildNode2!.Depth);
-			Assert.AreEqual(1, ChildNode2!.Payload.Length);
+			ChunkingOptions options = new ChunkingOptions();
+			options.LeafOptions = new ChunkingOptionsForNodeType(8, 8, 8);
+
+			FileNode node = new FileNode();
+			node.Append(new byte[7], options);
+			Assert.AreEqual(0, node.Depth);
+			Assert.AreEqual(7, node.Payload.Length);
+
+			node = new FileNode();
+			node.Append(new byte[8], options);
+			Assert.AreEqual(0, node.Depth);
+			Assert.AreEqual(8, node.Payload.Length);
+
+			node = new FileNode();
+			node.Append(new byte[9], options);
+			Assert.AreEqual(1, node.Depth);
+			Assert.AreEqual(2, node.Children.Count);
+
+			FileNode? childNode1 = node.Children[0].Node;
+			Assert.IsNotNull(childNode1);
+			Assert.AreEqual(0, childNode1!.Depth);
+			Assert.AreEqual(8, childNode1!.Payload.Length);
+
+			FileNode? childNode2 = node.Children[1].Node;
+			Assert.IsNotNull(childNode2);
+			Assert.AreEqual(0, childNode2!.Depth);
+			Assert.AreEqual(1, childNode2!.Payload.Length);
 		}
 
 		[TestMethod]
 		public void SerializationTests()
 		{
-			List<BundleImport> Imports = new List<BundleImport>();
-			Imports.Add(new BundleImport(IoHash.Compute(Encoding.UTF8.GetBytes("import1")), 0, 123));
-			Imports.Add(new BundleImport(IoHash.Compute(Encoding.UTF8.GetBytes("import2")), 1, 456));
+			List<BundleImport> imports = new List<BundleImport>();
+			imports.Add(new BundleImport(IoHash.Compute(Encoding.UTF8.GetBytes("import1")), 0, 123));
+			imports.Add(new BundleImport(IoHash.Compute(Encoding.UTF8.GetBytes("import2")), 1, 456));
 
-			BundleObject OldObject = new BundleObject();
-			OldObject.ImportObjects.Add(new BundleImportObject(new CbObjectAttachment(IoHash.Compute(new byte[] { 1, 2, 3 })), 12345, Imports));
-			OldObject.Exports.Add(new BundleExport(IoHash.Compute(Encoding.UTF8.GetBytes("export1")), 2, new BundleCompressionPacket(0) { EncodedLength = 20, DecodedLength = 40 }, 0, 40, new int[] { 1, 2 }));
-			OldObject.Exports.Add(new BundleExport(IoHash.Compute(Encoding.UTF8.GetBytes("export2")), 3, new BundleCompressionPacket(20) { EncodedLength = 10, DecodedLength = 20 }, 0, 20, new int[] { -1 }));
+			BundleObject oldObject = new BundleObject();
+			oldObject.ImportObjects.Add(new BundleImportObject(new CbObjectAttachment(IoHash.Compute(new byte[] { 1, 2, 3 })), 12345, imports));
+			oldObject.Exports.Add(new BundleExport(IoHash.Compute(Encoding.UTF8.GetBytes("export1")), 2, new BundleCompressionPacket(0) { EncodedLength = 20, DecodedLength = 40 }, 0, 40, new int[] { 1, 2 }));
+			oldObject.Exports.Add(new BundleExport(IoHash.Compute(Encoding.UTF8.GetBytes("export2")), 3, new BundleCompressionPacket(20) { EncodedLength = 10, DecodedLength = 20 }, 0, 20, new int[] { -1 }));
 
-			CbObject SerializedObject = CbSerializer.Serialize(OldObject);
-			ReadOnlyMemory<byte> SerializedData = SerializedObject.GetView();
+			CbObject serializedObject = CbSerializer.Serialize(oldObject);
+			ReadOnlyMemory<byte> serializedData = serializedObject.GetView();
 
-			BundleObject NewObject = CbSerializer.Deserialize<BundleObject>(SerializedData);
+			BundleObject newObject = CbSerializer.Deserialize<BundleObject>(serializedData);
 
-			Assert.AreEqual(OldObject.ImportObjects.Count, NewObject.ImportObjects.Count);
-			for (int Idx = 0; Idx < OldObject.ImportObjects.Count; Idx++)
+			Assert.AreEqual(oldObject.ImportObjects.Count, newObject.ImportObjects.Count);
+			for (int idx = 0; idx < oldObject.ImportObjects.Count; idx++)
 			{
-				BundleImportObject OldObjectImport = OldObject.ImportObjects[Idx];
-				BundleImportObject NewObjectImport = NewObject.ImportObjects[Idx];
+				BundleImportObject oldObjectImport = oldObject.ImportObjects[idx];
+				BundleImportObject newObjectImport = newObject.ImportObjects[idx];
 
-				Assert.AreEqual(OldObjectImport.Object, NewObjectImport.Object);
-				Assert.AreEqual(OldObjectImport.TotalCost, NewObjectImport.TotalCost);
-				Assert.AreEqual(OldObjectImport.Imports.Count, NewObjectImport.Imports.Count);
+				Assert.AreEqual(oldObjectImport.Object, newObjectImport.Object);
+				Assert.AreEqual(oldObjectImport.TotalCost, newObjectImport.TotalCost);
+				Assert.AreEqual(oldObjectImport.Imports.Count, newObjectImport.Imports.Count);
 
-				for (int ImportIdx = 0; ImportIdx < OldObjectImport.Imports.Count; ImportIdx++)
+				for (int importIdx = 0; importIdx < oldObjectImport.Imports.Count; importIdx++)
 				{
-					BundleImport OldImport = OldObjectImport.Imports[ImportIdx];
-					BundleImport NewImport = NewObjectImport.Imports[ImportIdx];
+					BundleImport oldImport = oldObjectImport.Imports[importIdx];
+					BundleImport newImport = newObjectImport.Imports[importIdx];
 
-					Assert.AreEqual(OldImport.Hash, NewImport.Hash);
-					Assert.AreEqual(OldImport.Rank, NewImport.Rank);
-					Assert.AreEqual(OldImport.Length, NewImport.Length);
+					Assert.AreEqual(oldImport.Hash, newImport.Hash);
+					Assert.AreEqual(oldImport.Rank, newImport.Rank);
+					Assert.AreEqual(oldImport.Length, newImport.Length);
 				}
 			}
 
-			Assert.AreEqual(OldObject.Exports.Count, NewObject.Exports.Count);
-			for (int Idx = 0; Idx < OldObject.Exports.Count; Idx++)
+			Assert.AreEqual(oldObject.Exports.Count, newObject.Exports.Count);
+			for (int idx = 0; idx < oldObject.Exports.Count; idx++)
 			{
-				BundleExport OldExport = OldObject.Exports[Idx];
-				BundleExport NewExport = NewObject.Exports[Idx];
-
-				Assert.AreEqual(OldExport.Hash, NewExport.Hash);
-				Assert.AreEqual(OldExport.Rank, NewExport.Rank);
-				Assert.AreEqual(OldExport.Length, NewExport.Length);
-				Assert.AreEqual(OldExport.Packet.Offset, NewExport.Packet.Offset);
-				Assert.AreEqual(OldExport.Packet.EncodedLength, NewExport.Packet.EncodedLength);
-				Assert.AreEqual(OldExport.Packet.DecodedLength, NewExport.Packet.DecodedLength);
-				Assert.AreEqual(OldExport.Offset, NewExport.Offset);
-				Assert.AreEqual(OldExport.Length, NewExport.Length);
-				Assert.IsTrue(OldExport.References.AsSpan().SequenceEqual(NewExport.References.AsSpan()));
+				BundleExport oldExport = oldObject.Exports[idx];
+				BundleExport newExport = newObject.Exports[idx];
+
+				Assert.AreEqual(oldExport.Hash, newExport.Hash);
+				Assert.AreEqual(oldExport.Rank, newExport.Rank);
+				Assert.AreEqual(oldExport.Length, newExport.Length);
+				Assert.AreEqual(oldExport.Packet.Offset, newExport.Packet.Offset);
+				Assert.AreEqual(oldExport.Packet.EncodedLength, newExport.Packet.EncodedLength);
+				Assert.AreEqual(oldExport.Packet.DecodedLength, newExport.Packet.DecodedLength);
+				Assert.AreEqual(oldExport.Offset, newExport.Offset);
+				Assert.AreEqual(oldExport.Length, newExport.Length);
+				Assert.IsTrue(oldExport.References.AsSpan().SequenceEqual(newExport.References.AsSpan()));
 			}
 		}
 
 		[TestMethod]
 		public async Task BasicTestDirectory()
 		{
-			using Bundle<DirectoryNode> NewBundle = Bundle.Create<DirectoryNode>(StorageClient, NamespaceId, new DirectoryNode(), new BundleOptions(), Cache);
-			DirectoryNode Node = NewBundle.Root.AddDirectory("hello");
-			DirectoryNode Node2 = Node.AddDirectory("world");
+			using Bundle<DirectoryNode> newBundle = Bundle.Create<DirectoryNode>(_storageClient, _namespaceId, new DirectoryNode(), new BundleOptions(), _cache);
+			DirectoryNode node = newBundle.Root.AddDirectory("hello");
+			DirectoryNode node2 = node.AddDirectory("world");
 
-			RefId RefId = new RefId("testref");
-			await NewBundle.WriteAsync(BucketId, RefId, CbObject.Empty, false);
+			RefId refId = new RefId("testref");
+			await newBundle.WriteAsync(_bucketId, refId, CbObject.Empty, false);
 
 			// Should be stored inline
-			Assert.AreEqual(1, StorageClient.Refs.Count);
-			Assert.AreEqual(0, StorageClient.Blobs.Count);
+			Assert.AreEqual(1, _storageClient.Refs.Count);
+			Assert.AreEqual(0, _storageClient.Blobs.Count);
 
 			// Check the ref
-			BundleRoot Root = await StorageClient.GetRefAsync<BundleRoot>(NamespaceId, BucketId, RefId);
-			BundleObject RootObject = Root.Object;
-			Assert.AreEqual(3, RootObject.Exports.Count);
-			Assert.AreEqual(0, RootObject.Exports[0].Rank);
-			Assert.AreEqual(1, RootObject.Exports[1].Rank);
-			Assert.AreEqual(2, RootObject.Exports[2].Rank);
+			BundleRoot root = await _storageClient.GetRefAsync<BundleRoot>(_namespaceId, _bucketId, refId);
+			BundleObject rootObject = root.Object;
+			Assert.AreEqual(3, rootObject.Exports.Count);
+			Assert.AreEqual(0, rootObject.Exports[0].Rank);
+			Assert.AreEqual(1, rootObject.Exports[1].Rank);
+			Assert.AreEqual(2, rootObject.Exports[2].Rank);
 
 			// Create a new bundle and read it back in again
-			using Bundle<DirectoryNode> NewBundle2 = await Bundle.ReadAsync<DirectoryNode>(StorageClient, NamespaceId, BucketId, RefId, new BundleOptions(), Cache);
+			using Bundle<DirectoryNode> newBundle2 = await Bundle.ReadAsync<DirectoryNode>(_storageClient, _namespaceId, _bucketId, refId, new BundleOptions(), _cache);
 
-			Assert.AreEqual(0, NewBundle2.Root.Files.Count);
-			Assert.AreEqual(1, NewBundle2.Root.Directories.Count);
-			DirectoryNode? OutputNode = await NewBundle.FindDirectoryAsync(NewBundle2.Root, "hello");
-			Assert.IsNotNull(OutputNode);
+			Assert.AreEqual(0, newBundle2.Root.Files.Count);
+			Assert.AreEqual(1, newBundle2.Root.Directories.Count);
+			DirectoryNode? outputNode = await newBundle.FindDirectoryAsync(newBundle2.Root, "hello");
+			Assert.IsNotNull(outputNode);
 
-			Assert.AreEqual(0, OutputNode!.Files.Count);
-			Assert.AreEqual(1, OutputNode!.Directories.Count);
-			DirectoryNode? OutputNode2 = await NewBundle.FindDirectoryAsync(OutputNode, "world");
-			Assert.IsNotNull(OutputNode2);
+			Assert.AreEqual(0, outputNode!.Files.Count);
+			Assert.AreEqual(1, outputNode!.Directories.Count);
+			DirectoryNode? outputNode2 = await newBundle.FindDirectoryAsync(outputNode, "world");
+			Assert.IsNotNull(outputNode2);
 
-			Assert.AreEqual(0, OutputNode2!.Files.Count);
-			Assert.AreEqual(0, OutputNode2!.Directories.Count);
+			Assert.AreEqual(0, outputNode2!.Files.Count);
+			Assert.AreEqual(0, outputNode2!.Directories.Count);
 		}
 
 		[TestMethod]
 		public async Task DedupTests()
 		{
-			BundleOptions Options = new BundleOptions();
-			Options.MaxBlobSize = 1;
-			Options.MaxInlineBlobSize = 1;
+			BundleOptions options = new BundleOptions();
+			options.MaxBlobSize = 1;
+			options.MaxInlineBlobSize = 1;
 
-			using Bundle<DirectoryNode> NewBundle = Bundle.Create<DirectoryNode>(StorageClient, NamespaceId, new DirectoryNode(), Options, Cache);
+			using Bundle<DirectoryNode> newBundle = Bundle.Create<DirectoryNode>(_storageClient, _namespaceId, new DirectoryNode(), options, _cache);
 
-			NewBundle.Root.AddDirectory("node1");
-			NewBundle.Root.AddDirectory("node2");
-			NewBundle.Root.AddDirectory("node3");
+			newBundle.Root.AddDirectory("node1");
+			newBundle.Root.AddDirectory("node2");
+			newBundle.Root.AddDirectory("node3");
 
-			Assert.AreEqual(0, StorageClient.Refs.Count);
-			Assert.AreEqual(0, StorageClient.Blobs.Count);
+			Assert.AreEqual(0, _storageClient.Refs.Count);
+			Assert.AreEqual(0, _storageClient.Blobs.Count);
 
-			RefId RefId = new RefId("ref");
-			await NewBundle.WriteAsync(BucketId, RefId, CbObject.Empty, false);
+			RefId refId = new RefId("ref");
+			await newBundle.WriteAsync(_bucketId, refId, CbObject.Empty, false);
 
-			Assert.AreEqual(1, StorageClient.Refs.Count);
-			Assert.AreEqual(1, StorageClient.Blobs.Count);
+			Assert.AreEqual(1, _storageClient.Refs.Count);
+			Assert.AreEqual(1, _storageClient.Blobs.Count);
 		}
 
 		[TestMethod]
 		public async Task ReloadTests()
 		{
-			BundleOptions Options = new BundleOptions();
-			Options.MaxBlobSize = 1;
-			Options.MaxInlineBlobSize = 1;
+			BundleOptions options = new BundleOptions();
+			options.MaxBlobSize = 1;
+			options.MaxInlineBlobSize = 1;
 
-			using Bundle<DirectoryNode> InitialBundle = Bundle.Create<DirectoryNode>(StorageClient, NamespaceId, new DirectoryNode(), Options, Cache);
+			using Bundle<DirectoryNode> initialBundle = Bundle.Create<DirectoryNode>(_storageClient, _namespaceId, new DirectoryNode(), options, _cache);
 
-			DirectoryNode Node1 = InitialBundle.Root.AddDirectory("node1");
-			DirectoryNode Node2 = Node1.AddDirectory("node2");
-			DirectoryNode Node3 = Node2.AddDirectory("node3");
-			DirectoryNode Node4 = Node3.AddDirectory("node4");
+			DirectoryNode node1 = initialBundle.Root.AddDirectory("node1");
+			DirectoryNode node2 = node1.AddDirectory("node2");
+			DirectoryNode node3 = node2.AddDirectory("node3");
+			DirectoryNode node4 = node3.AddDirectory("node4");
 
-			RefId RefId = new RefId("ref");
-			await InitialBundle.WriteAsync(BucketId, RefId, CbObject.Empty, false);
+			RefId refId = new RefId("ref");
+			await initialBundle.WriteAsync(_bucketId, refId, CbObject.Empty, false);
 
-			Assert.AreEqual(1, StorageClient.Refs.Count);
-			Assert.AreEqual(4, StorageClient.Blobs.Count);
+			Assert.AreEqual(1, _storageClient.Refs.Count);
+			Assert.AreEqual(4, _storageClient.Blobs.Count);
 
-			using Bundle<DirectoryNode> NewBundle = await Bundle.ReadAsync<DirectoryNode>(StorageClient, NamespaceId, BucketId, RefId, Options, Cache);
+			using Bundle<DirectoryNode> newBundle = await Bundle.ReadAsync<DirectoryNode>(_storageClient, _namespaceId, _bucketId, refId, options, _cache);
 
-			DirectoryNode? NewNode1 = await NewBundle.FindDirectoryAsync(NewBundle.Root, "node1");
-			Assert.IsNotNull(NewNode1);
+			DirectoryNode? newNode1 = await newBundle.FindDirectoryAsync(newBundle.Root, "node1");
+			Assert.IsNotNull(newNode1);
 
-			DirectoryNode? NewNode2 = await NewBundle.FindDirectoryAsync(NewNode1!, "node2");
-			Assert.IsNotNull(NewNode2);
+			DirectoryNode? newNode2 = await newBundle.FindDirectoryAsync(newNode1!, "node2");
+			Assert.IsNotNull(newNode2);
 
-			DirectoryNode? NewNode3 = await NewBundle.FindDirectoryAsync(NewNode2!, "node3");
-			Assert.IsNotNull(NewNode3);
+			DirectoryNode? newNode3 = await newBundle.FindDirectoryAsync(newNode2!, "node3");
+			Assert.IsNotNull(newNode3);
 
-			DirectoryNode? NewNode4 = await NewBundle.FindDirectoryAsync(NewNode3!, "node4");
-			Assert.IsNotNull(NewNode4);
+			DirectoryNode? newNode4 = await newBundle.FindDirectoryAsync(newNode3!, "node4");
+			Assert.IsNotNull(newNode4);
 		}
 
 		[TestMethod]
 		public async Task CompactTest()
 		{
-			BundleOptions Options = new BundleOptions();
-			Options.MaxBlobSize = 1024 * 1024;
-			Options.MaxInlineBlobSize = 1;
+			BundleOptions options = new BundleOptions();
+			options.MaxBlobSize = 1024 * 1024;
+			options.MaxInlineBlobSize = 1;
 
-			using Bundle<DirectoryNode> NewBundle = Bundle.Create<DirectoryNode>(StorageClient, NamespaceId, new DirectoryNode(), Options, Cache);
+			using Bundle<DirectoryNode> newBundle = Bundle.Create<DirectoryNode>(_storageClient, _namespaceId, new DirectoryNode(), options, _cache);
 
-			DirectoryNode Node1 = NewBundle.Root.AddDirectory("node1");
-			DirectoryNode Node2 = Node1.AddDirectory("node2");
-			DirectoryNode Node3 = Node2.AddDirectory("node3");
-			DirectoryNode Node4 = NewBundle.Root.AddDirectory("node4"); // same contents as node 3
+			DirectoryNode node1 = newBundle.Root.AddDirectory("node1");
+			DirectoryNode node2 = node1.AddDirectory("node2");
+			DirectoryNode node3 = node2.AddDirectory("node3");
+			DirectoryNode node4 = newBundle.Root.AddDirectory("node4"); // same contents as node 3
 
-			RefId RefId1 = new RefId("ref1");
-			await NewBundle.WriteAsync(BucketId, RefId1, CbObject.Empty, false);
+			RefId refId1 = new RefId("ref1");
+			await newBundle.WriteAsync(_bucketId, refId1, CbObject.Empty, false);
 
-			Assert.AreEqual(1, StorageClient.Refs.Count);
-			Assert.AreEqual(1, StorageClient.Blobs.Count);
+			Assert.AreEqual(1, _storageClient.Refs.Count);
+			Assert.AreEqual(1, _storageClient.Blobs.Count);
 
-			IRef Ref1 = StorageClient.Refs[(NamespaceId, BucketId, RefId1)];
+			IRef ref1 = _storageClient.Refs[(_namespaceId, _bucketId, refId1)];
 
-			BundleRoot Root1 = CbSerializer.Deserialize<BundleRoot>(Ref1.Value);
-			BundleObject RootObject1 = Root1.Object;
-			Assert.AreEqual(1, RootObject1.Exports.Count);
-			Assert.AreEqual(1, RootObject1.ImportObjects.Count);
+			BundleRoot root1 = CbSerializer.Deserialize<BundleRoot>(ref1.Value);
+			BundleObject rootObject1 = root1.Object;
+			Assert.AreEqual(1, rootObject1.Exports.Count);
+			Assert.AreEqual(1, rootObject1.ImportObjects.Count);
 
-			IoHash LeafHash1 = RootObject1.ImportObjects[0].Object.Hash;
-			BundleObject LeafObject1 = CbSerializer.Deserialize<BundleObject>(StorageClient.Blobs[(NamespaceId, LeafHash1)]);
-			Assert.AreEqual(3, LeafObject1.Exports.Count); // node1 + node2 + node3 (== node4)
-			Assert.AreEqual(0, LeafObject1.ImportObjects.Count);
+			IoHash leafHash1 = rootObject1.ImportObjects[0].Object.Hash;
+			BundleObject leafObject1 = CbSerializer.Deserialize<BundleObject>(_storageClient.Blobs[(_namespaceId, leafHash1)]);
+			Assert.AreEqual(3, leafObject1.Exports.Count); // node1 + node2 + node3 (== node4)
+			Assert.AreEqual(0, leafObject1.ImportObjects.Count);
 
 			// Remove one of the nodes from the root without compacting. the existing blob should be reused.
-			NewBundle.Root.DeleteDirectory("node1");
+			newBundle.Root.DeleteDirectory("node1");
 
-			RefId RefId2 = new RefId("ref2");
-			await NewBundle.WriteAsync(BucketId, RefId2, CbObject.Empty, false);
+			RefId refId2 = new RefId("ref2");
+			await newBundle.WriteAsync(_bucketId, refId2, CbObject.Empty, false);
 
-			IRef Ref2 = StorageClient.Refs[(NamespaceId, BucketId, RefId2)];
+			IRef ref2 = _storageClient.Refs[(_namespaceId, _bucketId, refId2)];
 
-			BundleRoot Root2 = CbSerializer.Deserialize<BundleRoot>(Ref2.Value);
-			BundleObject RootObject2 = Root2.Object;
-			Assert.AreEqual(1, RootObject2.Exports.Count);
-			Assert.AreEqual(1, RootObject2.ImportObjects.Count);
+			BundleRoot root2 = CbSerializer.Deserialize<BundleRoot>(ref2.Value);
+			BundleObject rootObject2 = root2.Object;
+			Assert.AreEqual(1, rootObject2.Exports.Count);
+			Assert.AreEqual(1, rootObject2.ImportObjects.Count);
 
-			IoHash LeafHash2 = RootObject2.ImportObjects[0].Object.Hash;
-			Assert.AreEqual(LeafHash1, LeafHash2);
-			Assert.AreEqual(3, LeafObject1.Exports.Count); // unused: node1 + node2 + node3, used: node4 (== node3)
-			Assert.AreEqual(0, LeafObject1.ImportObjects.Count);
+			IoHash leafHash2 = rootObject2.ImportObjects[0].Object.Hash;
+			Assert.AreEqual(leafHash1, leafHash2);
+			Assert.AreEqual(3, leafObject1.Exports.Count); // unused: node1 + node2 + node3, used: node4 (== node3)
+			Assert.AreEqual(0, leafObject1.ImportObjects.Count);
 
 			// Repack it and check that we make a new object
-			RefId RefId3 = new RefId("ref3");
-			await NewBundle.WriteAsync(BucketId, RefId3, CbObject.Empty, true);
+			RefId refId3 = new RefId("ref3");
+			await newBundle.WriteAsync(_bucketId, refId3, CbObject.Empty, true);
 
-			IRef Ref3 = StorageClient.Refs[(NamespaceId, BucketId, RefId3)];
+			IRef ref3 = _storageClient.Refs[(_namespaceId, _bucketId, refId3)];
 
-			BundleRoot Root3 = CbSerializer.Deserialize<BundleRoot>(Ref3.Value);
-			BundleObject RootObject3 = Root3.Object;
-			Assert.AreEqual(1, RootObject3.Exports.Count);
-			Assert.AreEqual(1, RootObject3.ImportObjects.Count);
+			BundleRoot root3 = CbSerializer.Deserialize<BundleRoot>(ref3.Value);
+			BundleObject rootObject3 = root3.Object;
+			Assert.AreEqual(1, rootObject3.Exports.Count);
+			Assert.AreEqual(1, rootObject3.ImportObjects.Count);
 
-			IoHash LeafHash3 = RootObject3.ImportObjects[0].Object.Hash;
-			Assert.AreNotEqual(LeafHash1, LeafHash3);
+			IoHash leafHash3 = rootObject3.ImportObjects[0].Object.Hash;
+			Assert.AreNotEqual(leafHash1, leafHash3);
 
-			BundleObject LeafObject3 = CbSerializer.Deserialize<BundleObject>(StorageClient.Blobs[(NamespaceId, LeafHash3)]);
-			Assert.AreEqual(1, LeafObject3.Exports.Count);
-			Assert.AreEqual(0, LeafObject3.ImportObjects.Count);
+			BundleObject leafObject3 = CbSerializer.Deserialize<BundleObject>(_storageClient.Blobs[(_namespaceId, leafHash3)]);
+			Assert.AreEqual(1, leafObject3.Exports.Count);
+			Assert.AreEqual(0, leafObject3.ImportObjects.Count);
 		}
 
 		[TestMethod]
 		public async Task CoreAppendTest()
 		{
-			byte[] Data = new byte[4096];
-			new Random(0).NextBytes(Data);
+			byte[] data = new byte[4096];
+			new Random(0).NextBytes(data);
 
-			ChunkingOptions Options = new ChunkingOptions();
-			Options.LeafOptions = new ChunkingOptionsForNodeType(16, 64, 256);
+			ChunkingOptions options = new ChunkingOptions();
+			options.LeafOptions = new ChunkingOptionsForNodeType(16, 64, 256);
 
-			FileNode Node = new FileNode();
-			for (int Idx = 0; Idx < Data.Length; Idx++)
+			FileNode node = new FileNode();
+			for (int idx = 0; idx < data.Length; idx++)
 			{
-				Node.Append(Data.AsMemory(Idx, 1), Options);
+				node.Append(data.AsMemory(idx, 1), options);
 
-				byte[] OutputData = await Node.ToByteArrayAsync(null!);
-				Assert.IsTrue(Data.AsMemory(0, Idx + 1).Span.SequenceEqual(OutputData.AsSpan(0, Idx + 1)));
+				byte[] outputData = await node.ToByteArrayAsync(null!);
+				Assert.IsTrue(data.AsMemory(0, idx + 1).Span.SequenceEqual(outputData.AsSpan(0, idx + 1)));
 			}
 		}
 
 		[TestMethod]
 		public async Task FixedSizeChunkingTests()
 		{
-			ChunkingOptions Options = new ChunkingOptions();
-			Options.LeafOptions = new ChunkingOptionsForNodeType(64, 64, 64);
-			Options.InteriorOptions = new ChunkingOptionsForNodeType(IoHash.NumBytes * 4, IoHash.NumBytes * 4, IoHash.NumBytes * 4);
+			ChunkingOptions options = new ChunkingOptions();
+			options.LeafOptions = new ChunkingOptionsForNodeType(64, 64, 64);
+			options.InteriorOptions = new ChunkingOptionsForNodeType(IoHash.NumBytes * 4, IoHash.NumBytes * 4, IoHash.NumBytes * 4);
 
-			await ChunkingTests(Options);
+			await ChunkingTests(options);
 		}
 
 		[TestMethod]
 		public async Task VariableSizeChunkingTests()
 		{
-			ChunkingOptions Options = new ChunkingOptions();
-			Options.LeafOptions = new ChunkingOptionsForNodeType(32, 64, 96);
-			Options.InteriorOptions = new ChunkingOptionsForNodeType(IoHash.NumBytes * 1, IoHash.NumBytes * 4, IoHash.NumBytes * 12);
+			ChunkingOptions options = new ChunkingOptions();
+			options.LeafOptions = new ChunkingOptionsForNodeType(32, 64, 96);
+			options.InteriorOptions = new ChunkingOptionsForNodeType(IoHash.NumBytes * 1, IoHash.NumBytes * 4, IoHash.NumBytes * 12);
 
-			await ChunkingTests(Options);
+			await ChunkingTests(options);
 		}
 
-		async Task ChunkingTests(ChunkingOptions Options)
+		async Task ChunkingTests(ChunkingOptions options)
 		{
-			byte[] Data = new byte[4096];
-			new Random(0).NextBytes(Data);
+			byte[] data = new byte[4096];
+			new Random(0).NextBytes(data);
 
-			for (int Idx = 0; Idx < Data.Length; Idx++)
+			for (int idx = 0; idx < data.Length; idx++)
 			{
-				Data[Idx] = (byte)Idx;
+				data[idx] = (byte)idx;
 			}
 
-			Bundle<FileNode> NewBundle = Bundle.Create<FileNode>(StorageClient, NamespaceId, new FileNode(), new BundleOptions(), Cache);
-			FileNode Root = NewBundle.Root;
+			Bundle<FileNode> newBundle = Bundle.Create<FileNode>(_storageClient, _namespaceId, new FileNode(), new BundleOptions(), _cache);
+			FileNode root = newBundle.Root;
 
 			const int NumIterations = 100;
-			for (int Idx = 0; Idx < NumIterations; Idx++)
+			for (int idx = 0; idx < NumIterations; idx++)
 			{
-				Root.Append(Data, Options);
+				root.Append(data, options);
 			}
 
-			byte[] Result = await Root.ToByteArrayAsync(null!);
-			Assert.AreEqual(NumIterations * Data.Length, Result.Length);
+			byte[] result = await root.ToByteArrayAsync(null!);
+			Assert.AreEqual(NumIterations * data.Length, result.Length);
 
-			for (int Idx = 0; Idx < NumIterations; Idx++)
+			for (int idx = 0; idx < NumIterations; idx++)
 			{
-				ReadOnlyMemory<byte> SpanData = Result.AsMemory(Idx * Data.Length, Data.Length);
-				Assert.IsTrue(SpanData.Span.SequenceEqual(Data));
+				ReadOnlyMemory<byte> spanData = result.AsMemory(idx * data.Length, data.Length);
+				Assert.IsTrue(spanData.Span.SequenceEqual(data));
 			}
 
-			await CheckSizes(NewBundle, Root, Options, true);
+			await CheckSizes(newBundle, root, options, true);
 		}
 
-		async Task CheckSizes(Bundle Bundle, FileNode Node, ChunkingOptions Options, bool Rightmost)
+		async Task CheckSizes(Bundle bundle, FileNode node, ChunkingOptions options, bool rightmost)
 		{
-			if (Node.Depth == 0)
+			if (node.Depth == 0)
 			{
-				Assert.IsTrue(Rightmost || Node.Payload.Length >= Options.LeafOptions.MinSize);
-				Assert.IsTrue(Node.Payload.Length <= Options.LeafOptions.MaxSize);
+				Assert.IsTrue(rightmost || node.Payload.Length >= options.LeafOptions.MinSize);
+				Assert.IsTrue(node.Payload.Length <= options.LeafOptions.MaxSize);
 			}
 			else
 			{
-				Assert.IsTrue(Rightmost || Node.Payload.Length >= Options.InteriorOptions.MinSize);
-				Assert.IsTrue(Node.Payload.Length <= Options.InteriorOptions.MaxSize);
+				Assert.IsTrue(rightmost || node.Payload.Length >= options.InteriorOptions.MinSize);
+				Assert.IsTrue(node.Payload.Length <= options.InteriorOptions.MaxSize);
 
-				int ChildCount = Node.Children.Count;
-				for (int Idx = 0; Idx < ChildCount; Idx++)
+				int childCount = node.Children.Count;
+				for (int idx = 0; idx < childCount; idx++)
 				{
-					FileNode ChildNode = await Bundle.GetAsync(Node.Children[Idx]);
-					await CheckSizes(Bundle, ChildNode, Options, Idx == ChildCount - 1);
+					FileNode childNode = await bundle.GetAsync(node.Children[idx]);
+					await CheckSizes(bundle, childNode, options, idx == childCount - 1);
 				}
 			}
 		}
@@ -413,62 +411,62 @@ namespace EpicGames.Horde.Tests
 		[TestMethod]
 		public async Task SpillTestAsync()
 		{
-			BundleOptions Options = new BundleOptions();
-			Options.MaxBlobSize = 1;
+			BundleOptions options = new BundleOptions();
+			options.MaxBlobSize = 1;
 
-			using Bundle<DirectoryNode> NewBundle = Bundle.Create<DirectoryNode>(StorageClient, NamespaceId, new DirectoryNode(), Options, Cache);
-			DirectoryNode Root = NewBundle.Root;
+			using Bundle<DirectoryNode> newBundle = Bundle.Create<DirectoryNode>(_storageClient, _namespaceId, new DirectoryNode(), options, _cache);
+			DirectoryNode root = newBundle.Root;
 
-			long TotalLength = 0;
-			for (int IdxA = 0; IdxA < 10; IdxA++)
+			long totalLength = 0;
+			for (int idxA = 0; idxA < 10; idxA++)
 			{
-				DirectoryNode NodeA = Root.AddDirectory($"{IdxA}");
-				for (int IdxB = 0; IdxB < 10; IdxB++)
+				DirectoryNode nodeA = root.AddDirectory($"{idxA}");
+				for (int idxB = 0; idxB < 10; idxB++)
 				{
-					DirectoryNode NodeB = NodeA.AddDirectory($"{IdxB}");
-					for (int IdxC = 0; IdxC < 10; IdxC++)
+					DirectoryNode nodeB = nodeA.AddDirectory($"{idxB}");
+					for (int idxC = 0; idxC < 10; idxC++)
 					{
-						DirectoryNode NodeC = NodeB.AddDirectory($"{IdxC}");
-						for (int IdxD = 0; IdxD < 10; IdxD++)
+						DirectoryNode nodeC = nodeB.AddDirectory($"{idxC}");
+						for (int idxD = 0; idxD < 10; idxD++)
 						{
-							FileNode File = NodeC.AddFile($"{IdxD}", 0);
-							byte[] Data = Encoding.UTF8.GetBytes($"This is file {IdxA}/{IdxB}/{IdxC}/{IdxD}");
-							TotalLength += Data.Length;
-							File.Append(Data, new ChunkingOptions());
+							FileNode file = nodeC.AddFile($"{idxD}", 0);
+							byte[] data = Encoding.UTF8.GetBytes($"This is file {idxA}/{idxB}/{idxC}/{idxD}");
+							totalLength += data.Length;
+							file.Append(data, new ChunkingOptions());
 						}
 					}
 
-					int OldWorkingSetSize = GetWorkingSetSize(NewBundle.Root);
-					await NewBundle.TrimAsync(20);
-					int NewWorkingSetSize = GetWorkingSetSize(NewBundle.Root);
-					Assert.IsTrue(NewWorkingSetSize <= OldWorkingSetSize);
-					Assert.IsTrue(NewWorkingSetSize <= 20);
+					int oldWorkingSetSize = GetWorkingSetSize(newBundle.Root);
+					await newBundle.TrimAsync(20);
+					int newWorkingSetSize = GetWorkingSetSize(newBundle.Root);
+					Assert.IsTrue(newWorkingSetSize <= oldWorkingSetSize);
+					Assert.IsTrue(newWorkingSetSize <= 20);
 				}
 			}
 
-			Assert.IsTrue(StorageClient.Blobs.Count > 0);
-			Assert.IsTrue(StorageClient.Refs.Count == 0);
+			Assert.IsTrue(_storageClient.Blobs.Count > 0);
+			Assert.IsTrue(_storageClient.Refs.Count == 0);
 
-			RefId RefId = new RefId("ref");
-			await NewBundle.WriteAsync(BucketId, RefId, CbObject.Empty, true);
+			RefId refId = new RefId("ref");
+			await newBundle.WriteAsync(_bucketId, refId, CbObject.Empty, true);
 
-			Assert.AreEqual(TotalLength, Root.Length);
+			Assert.AreEqual(totalLength, root.Length);
 
-			Assert.IsTrue(StorageClient.Blobs.Count > 0);
-			Assert.IsTrue(StorageClient.Refs.Count == 1);
+			Assert.IsTrue(_storageClient.Blobs.Count > 0);
+			Assert.IsTrue(_storageClient.Refs.Count == 1);
 		}
 
-		int GetWorkingSetSize(BundleNode Node)
+		int GetWorkingSetSize(BundleNode node)
 		{
-			int Size = 0;
-			foreach (BundleNodeRef NodeRef in Node.GetReferences())
+			int size = 0;
+			foreach (BundleNodeRef nodeRef in node.GetReferences())
 			{
-				if (NodeRef.Node != null)
+				if (nodeRef.Node != null)
 				{
-					Size += 1 + GetWorkingSetSize(NodeRef.Node);
+					size += 1 + GetWorkingSetSize(nodeRef.Node);
 				}
 			}
-			return Size;
+			return size;
 		}
 	}
 }
diff --git a/Engine/Source/Programs/Shared/EpicGames.Horde/Auth/AuthExtensions.cs b/Engine/Source/Programs/Shared/EpicGames.Horde/Auth/AuthExtensions.cs
--- a/Engine/Source/Programs/Shared/EpicGames.Horde/Auth/AuthExtensions.cs
+++ b/Engine/Source/Programs/Shared/EpicGames.Horde/Auth/AuthExtensions.cs
@@ -2,10 +2,8 @@
 
 using Microsoft.Extensions.DependencyInjection;
 using System;
-using System.Collections.Generic;
 using System.ComponentModel.DataAnnotations;
 using System.Net.Http;
-using System.Text;
 
 namespace EpicGames.Horde.Auth
 {
@@ -49,35 +47,35 @@ namespace EpicGames.Horde.Auth
 
 	internal static class AuthExtensions
 	{
-		public static void AddHttpClientWithAuth<TClient, TImplementation>(this IServiceCollection Services, Func<IServiceProvider, HttpServiceClientOptions> GetOptions) 
+		public static void AddHttpClientWithAuth<TClient, TImplementation>(this IServiceCollection services, Func<IServiceProvider, HttpServiceClientOptions> getOptions) 
 			where TClient : class 
 			where TImplementation : class, TClient
 		{
-			Services.AddScoped<OAuthHandlerFactory>();
-			Services.AddHttpClient<OAuthHandlerFactory>();
-			Services.AddScoped<OAuthHandler<TImplementation>>(ServiceProvider => ServiceProvider.GetRequiredService<OAuthHandlerFactory>().Create<TImplementation>(GetOptions(ServiceProvider)));
+			services.AddScoped<OAuthHandlerFactory>();
+			services.AddHttpClient<OAuthHandlerFactory>();
+			services.AddScoped<OAuthHandler<TImplementation>>(serviceProvider => serviceProvider.GetRequiredService<OAuthHandlerFactory>().Create<TImplementation>(getOptions(serviceProvider)));
 
-			Services.AddHttpClient<TClient, TImplementation>((ServiceProvider, Client) =>
+			services.AddHttpClient<TClient, TImplementation>((serviceProvider, client) =>
 				{
-					HttpServiceClientOptions Options = GetOptions(ServiceProvider);
-					Client.BaseAddress = new Uri(Options.Url);
+					HttpServiceClientOptions options = getOptions(serviceProvider);
+					client.BaseAddress = new Uri(options.Url);
 				})
-				.ConfigurePrimaryHttpMessageHandler(ServiceProvider =>
+				.ConfigurePrimaryHttpMessageHandler(serviceProvider =>
 				{
-					HttpServiceClientOptions Options = GetOptions(ServiceProvider);
-					return CreateMessageHandler<TImplementation>(ServiceProvider, Options);
+					HttpServiceClientOptions options = getOptions(serviceProvider);
+					return CreateMessageHandler<TImplementation>(serviceProvider, options);
 				});
 		}
 
-		static HttpMessageHandler CreateMessageHandler<TImplementation>(IServiceProvider ServiceProvider, HttpServiceClientOptions Options)
+		static HttpMessageHandler CreateMessageHandler<TImplementation>(IServiceProvider serviceProvider, HttpServiceClientOptions options)
 		{
-			if (!String.IsNullOrEmpty(Options.AuthUrl))
+			if (!String.IsNullOrEmpty(options.AuthUrl))
 			{
-				return ServiceProvider.GetRequiredService<OAuthHandler<TImplementation>>();
+				return serviceProvider.GetRequiredService<OAuthHandler<TImplementation>>();
 			}
-			else if (!String.IsNullOrEmpty(Options.Token))
+			else if (!String.IsNullOrEmpty(options.Token))
 			{
-				return new TokenHandler<TImplementation>(Options);
+				return new TokenHandler<TImplementation>(options);
 			}
 			else
 			{
diff --git a/Engine/Source/Programs/Shared/EpicGames.Horde/Auth/OAuthHandler.cs b/Engine/Source/Programs/Shared/EpicGames.Horde/Auth/OAuthHandler.cs
--- a/Engine/Source/Programs/Shared/EpicGames.Horde/Auth/OAuthHandler.cs
+++ b/Engine/Source/Programs/Shared/EpicGames.Horde/Auth/OAuthHandler.cs
@@ -1,11 +1,8 @@
 // Copyright Epic Games, Inc. All Rights Reserved.
 
-using EpicGames.Core;
-using EpicGames.Serialization;
-using Microsoft.Extensions.Options;
 using System;
 using System.Collections.Generic;
-using System.ComponentModel.DataAnnotations;
+using System.Diagnostics.CodeAnalysis;
 using System.Net;
 using System.Net.Http;
 using System.Text;
@@ -23,8 +20,8 @@ namespace EpicGames.Horde.Auth
 		/// <summary>
 		/// Constructor
 		/// </summary>
-		public AuthenticationException(string Message, Exception? InnerException)
-			: base(Message, InnerException)
+		public AuthenticationException(string message, Exception? innerException)
+			: base(message, innerException)
 		{
 		}
 	}
@@ -65,6 +62,7 @@ namespace EpicGames.Horde.Auth
 	/// </summary>
 	public class OAuthHandler<T> : HttpClientHandler
 	{
+		[SuppressMessage("Style", "IDE1006:Naming Styles")]
 		class ClientCredentialsResponse
 		{
 			public string? access_token { get; set; }
@@ -73,74 +71,74 @@ namespace EpicGames.Horde.Auth
 			public string? scope { get; set; }
 		}
 
-		HttpClient Client;
-		IOAuthOptions Options;
-		string CachedAccessToken = String.Empty;
-		DateTime ExpiresAt = DateTime.MinValue;
+		readonly HttpClient _client;
+		readonly IOAuthOptions _options;
+		string _cachedAccessToken = String.Empty;
+		DateTime _expiresAt = DateTime.MinValue;
 
 		/// <summary>
 		/// Constructor
 		/// </summary>
-		/// <param name="Client"></param>
-		/// <param name="Options"></param>
-		public OAuthHandler(HttpClient Client, IOAuthOptions Options)
+		/// <param name="client"></param>
+		/// <param name="options"></param>
+		public OAuthHandler(HttpClient client, IOAuthOptions options)
 		{
-			this.Client = Client;
-			this.Options = Options;
+			_client = client;
+			_options = options;
 		}
 
 		/// <inheritdoc/>
-		protected override async Task<HttpResponseMessage> SendAsync(HttpRequestMessage Request, CancellationToken CancellationToken)
+		protected override async Task<HttpResponseMessage> SendAsync(HttpRequestMessage request, CancellationToken cancellationToken)
 		{
-			if (DateTime.UtcNow > ExpiresAt)
+			if (DateTime.UtcNow > _expiresAt)
 			{
-				await UpdateAccessTokenAsync(CancellationToken);
+				await UpdateAccessTokenAsync(cancellationToken);
 			}
 
-			Request.Headers.Add("Authorization", $"Bearer {CachedAccessToken}");
-			return await base.SendAsync(Request, CancellationToken);
+			request.Headers.Add("Authorization", $"Bearer {_cachedAccessToken}");
+			return await base.SendAsync(request, cancellationToken);
 		}
 
 		/// <summary>
 		/// Updates the current access token
 		/// </summary>
 		/// <returns></returns>
-		async Task UpdateAccessTokenAsync(CancellationToken CancellationToken)
+		async Task UpdateAccessTokenAsync(CancellationToken cancellationToken)
 		{
-			KeyValuePair<string, string>[] Content = new KeyValuePair<string, string>[]
+			KeyValuePair<string, string>[] content = new KeyValuePair<string, string>[]
 			{
-				new KeyValuePair<string, string>("grant_type", Options.GrantType),
-				new KeyValuePair<string, string>("client_id", Options.ClientId),
-				new KeyValuePair<string, string>("client_secret", Options.ClientSecret),
-				new KeyValuePair<string, string>("scope", Options.Scope)
+				new KeyValuePair<string, string>("grant_type", _options.GrantType),
+				new KeyValuePair<string, string>("client_id", _options.ClientId),
+				new KeyValuePair<string, string>("client_secret", _options.ClientSecret),
+				new KeyValuePair<string, string>("scope", _options.Scope)
 			};
 
 			try
 			{
-				using HttpRequestMessage Message = new HttpRequestMessage(HttpMethod.Post, Options.AuthUrl);
-				Message.Content = new FormUrlEncodedContent(Content);
+				using HttpRequestMessage message = new HttpRequestMessage(HttpMethod.Post, _options.AuthUrl);
+				message.Content = new FormUrlEncodedContent(content);
 
-				HttpResponseMessage Response = await Client.SendAsync(Message, CancellationToken);
-				if (!Response.IsSuccessStatusCode)
+				HttpResponseMessage response = await _client.SendAsync(message, cancellationToken);
+				if (!response.IsSuccessStatusCode)
 				{
-					throw new AuthenticationException($"Authentication failed. Response: {Response.Content}", null);
+					throw new AuthenticationException($"Authentication failed. Response: {response.Content}", null);
 				}
 
-				byte[] ResponseData = await Response.Content.ReadAsByteArrayAsync();
-				ClientCredentialsResponse Result = JsonSerializer.Deserialize<ClientCredentialsResponse>(ResponseData)!;
+				byte[] responseData = await response.Content.ReadAsByteArrayAsync();
+				ClientCredentialsResponse result = JsonSerializer.Deserialize<ClientCredentialsResponse>(responseData)!;
 
-				string? accessToken = Result?.access_token;
-				if (string.IsNullOrEmpty(accessToken))
+				string? accessToken = result?.access_token;
+				if (String.IsNullOrEmpty(accessToken))
 				{
-					throw new AuthenticationException("The authentication token received by the server is null or empty. Body received was: " + Encoding.UTF8.GetString(ResponseData), null);
+					throw new AuthenticationException("The authentication token received by the server is null or empty. Body received was: " + Encoding.UTF8.GetString(responseData), null);
 				}
-				CachedAccessToken = accessToken;
+				_cachedAccessToken = accessToken;
 				// renew after half the renewal time
-				ExpiresAt = DateTime.UtcNow + TimeSpan.FromSeconds((Result?.expires_in ?? 3200) / 2.0);
+				_expiresAt = DateTime.UtcNow + TimeSpan.FromSeconds((result?.expires_in ?? 3200) / 2.0);
 			}
-			catch (WebException Ex)
+			catch (WebException ex)
 			{
-				throw new AuthenticationException("Unable to authenticate.", Ex);
+				throw new AuthenticationException("Unable to authenticate.", ex);
 			}
 		}
 	}
@@ -150,22 +148,22 @@ namespace EpicGames.Horde.Auth
 	/// </summary>
 	public class OAuthHandlerFactory
 	{
-		HttpClient HttpClient;
+		readonly HttpClient _httpClient;
 
 		/// <summary>
 		/// Constructor
 		/// </summary>
-		/// <param name="HttpClient"></param>
-		public OAuthHandlerFactory(HttpClient HttpClient)
+		/// <param name="httpClient"></param>
+		public OAuthHandlerFactory(HttpClient httpClient)
 		{
-			this.HttpClient = HttpClient;
+			_httpClient = httpClient;
 		}
 
 		/// <summary>
 		/// Create an instance of the auth provider
 		/// </summary>
-		/// <param name="Options"></param>
+		/// <param name="options"></param>
 		/// <returns></returns>
-		public OAuthHandler<T> Create<T>(IOAuthOptions Options) => new OAuthHandler<T>(HttpClient, Options);
+		public OAuthHandler<T> Create<T>(IOAuthOptions options) => new OAuthHandler<T>(_httpClient, options);
 	}
 }
diff --git a/Engine/Source/Programs/Shared/EpicGames.Horde/Auth/TokenHandler.cs b/Engine/Source/Programs/Shared/EpicGames.Horde/Auth/TokenHandler.cs
--- a/Engine/Source/Programs/Shared/EpicGames.Horde/Auth/TokenHandler.cs
+++ b/Engine/Source/Programs/Shared/EpicGames.Horde/Auth/TokenHandler.cs
@@ -1,15 +1,6 @@
 // Copyright Epic Games, Inc. All Rights Reserved.
 
-using EpicGames.Core;
-using EpicGames.Serialization;
-using Microsoft.Extensions.Options;
-using System;
-using System.Collections.Generic;
-using System.ComponentModel.DataAnnotations;
-using System.Net;
 using System.Net.Http;
-using System.Text;
-using System.Text.Json;
 using System.Threading;
 using System.Threading.Tasks;
 
@@ -31,22 +22,22 @@ namespace EpicGames.Horde.Auth
 	/// </summary>
 	public class TokenHandler<T> : HttpClientHandler
 	{
-		ITokenAuthOptions Options;
+		readonly ITokenAuthOptions _options;
 
 		/// <summary>
 		/// Constructor
 		/// </summary>
-		/// <param name="Options"></param>
-		public TokenHandler(ITokenAuthOptions Options)
+		/// <param name="options"></param>
+		public TokenHandler(ITokenAuthOptions options)
 		{
-			this.Options = Options;
+			_options = options;
 		}
 
 		/// <inheritdoc/>
-		protected override async Task<HttpResponseMessage> SendAsync(HttpRequestMessage Request, CancellationToken CancellationToken)
+		protected override async Task<HttpResponseMessage> SendAsync(HttpRequestMessage request, CancellationToken cancellationToken)
 		{
-			Request.Headers.Add("Authorization", $"Bearer {Options.Token}");
-			return await base.SendAsync(Request, CancellationToken);
+			request.Headers.Add("Authorization", $"Bearer {_options.Token}");
+			return await base.SendAsync(request, cancellationToken);
 		}
 	}
 }
