Horde: Use standard coding style for EpicGames.Horde.

diff --git a/Engine/Source/Programs/Shared/EpicGames.Horde/Bundles/Bundle.cs b/Engine/Source/Programs/Shared/EpicGames.Horde/Bundles/Bundle.cs
--- a/Engine/Source/Programs/Shared/EpicGames.Horde/Bundles/Bundle.cs
+++ b/Engine/Source/Programs/Shared/EpicGames.Horde/Bundles/Bundle.cs
@@ -6,22 +6,12 @@ using EpicGames.Horde.Storage;
 using EpicGames.Serialization;
 using K4os.Compression.LZ4;
 using Microsoft.Extensions.Caching.Memory;
-using Microsoft.Extensions.Logging;
 using System;
 using System.Buffers;
-using System.Buffers.Binary;
-using System.Collections;
 using System.Collections.Concurrent;
 using System.Collections.Generic;
 using System.Diagnostics;
-using System.Diagnostics.CodeAnalysis;
-using System.IO;
-using System.IO.Compression;
 using System.Linq;
-using System.Reflection;
-using System.Security.Cryptography;
-using System.Text;
-using System.Text.Json;
 using System.Threading;
 using System.Threading.Tasks;
 
@@ -116,46 +106,46 @@ namespace EpicGames.Horde.Bundles
 
 			public NodeState State
 			{
-				get => StatePrivate;
-				set => StatePrivate = value;
+				get => _state;
+				set => _state = value;
 			}
 
-			private NodeState StatePrivate;
+			private NodeState _state;
 
-			public NodeInfo(IoHash Hash, int Rank, int Length, NodeState State)
+			public NodeInfo(IoHash hash, int rank, int length, NodeState state)
 			{
-				if (Length <= 0)
+				if (length <= 0)
 				{
-					throw new ArgumentException("Node length must be greater than zero", nameof(Length));
+					throw new ArgumentException("Node length must be greater than zero", nameof(length));
 				}
 
-				this.Hash = Hash;
-				this.Rank = Rank;
-				this.Length = Length;
-				this.StatePrivate = State;
+				Hash = hash;
+				Rank = rank;
+				Length = length;
+				_state = state;
 			}
 
-			public void TrySetExport(BlobInfo Blob, BundleExport Export, IReadOnlyList<NodeInfo> References)
+			public void TrySetExport(BlobInfo blob, BundleExport export, IReadOnlyList<NodeInfo> references)
 			{
-				NodeState PrevState = State;
-				if (!PrevState.IsStandalone())
+				NodeState prevState = State;
+				if (!prevState.IsStandalone())
 				{
-					NodeState NextState = NodeState.Exported(Blob, Export, References);
-					while (!PrevState.IsStandalone() && !TryUpdate(PrevState, NextState))
+					NodeState nextState = NodeState.Exported(blob, export, references);
+					while (!prevState.IsStandalone() && !TryUpdate(prevState, nextState))
 					{
-						PrevState = State;
+						prevState = State;
 					}
 				}
 			}
 
-			public void SetStandalone(ReadOnlySequence<byte> Data, IReadOnlyList<NodeInfo> References)
+			public void SetStandalone(ReadOnlySequence<byte> data, IReadOnlyList<NodeInfo> references)
 			{
-				State = NodeState.Standalone(Data, References);
+				State = NodeState.Standalone(data, references);
 			}
 
-			private bool TryUpdate(NodeState PrevState, NodeState NextState)
+			private bool TryUpdate(NodeState prevState, NodeState nextState)
 			{
-				return Interlocked.CompareExchange(ref StatePrivate, NextState, PrevState) == PrevState;
+				return Interlocked.CompareExchange(ref _state, nextState, prevState) == prevState;
 			}
 		}
 
@@ -181,31 +171,31 @@ namespace EpicGames.Horde.Bundles
 			public ReadOnlySequence<byte> Data { get; }
 			public IReadOnlyList<NodeInfo>? References { get; }
 
-			private NodeState(BlobInfo? Blob, BundleExport? Export, ReadOnlySequence<byte> Data, IReadOnlyList<NodeInfo>? References)
+			private NodeState(BlobInfo? blob, BundleExport? export, ReadOnlySequence<byte> data, IReadOnlyList<NodeInfo>? references)
 			{
-				this.Blob = Blob;
-				this.Export = Export;
-				this.Data = Data;
-				this.References = References;
+				Blob = blob;
+				Export = export;
+				Data = data;
+				References = references;
 			}
 
 			public bool IsImport() => References == null;
 			public bool IsExport() => Export != null;
 			public bool IsStandalone() => Blob == null;
 
-			public static NodeState Imported(BlobInfo Blob)
+			public static NodeState Imported(BlobInfo blob)
 			{
-				return new NodeState(Blob, null, ReadOnlySequence<byte>.Empty, null);
+				return new NodeState(blob, null, ReadOnlySequence<byte>.Empty, null);
 			}
 
-			public static NodeState Exported(BlobInfo Blob, BundleExport Export, IReadOnlyList<NodeInfo> References)
+			public static NodeState Exported(BlobInfo blob, BundleExport export, IReadOnlyList<NodeInfo> references)
 			{
-				return new NodeState(Blob, Export, ReadOnlySequence<byte>.Empty, References);
+				return new NodeState(blob, export, ReadOnlySequence<byte>.Empty, references);
 			}
 
-			public static NodeState Standalone(ReadOnlySequence<byte> Data, IReadOnlyList<NodeInfo> References)
+			public static NodeState Standalone(ReadOnlySequence<byte> data, IReadOnlyList<NodeInfo> references)
 			{
-				return new NodeState(null, null, Data, References);
+				return new NodeState(null, null, data, references);
 			}
 		}
 
@@ -221,35 +211,35 @@ namespace EpicGames.Horde.Bundles
 			public bool Mounted { get; set; }
 			public List<NodeInfo>? LiveNodes { get; set; }
 
-			public BlobInfo(IoHash Hash, int TotalCost)
+			public BlobInfo(IoHash hash, int totalCost)
 			{
-				this.Hash = Hash;
-				this.TotalCost = TotalCost;
+				Hash = hash;
+				TotalCost = totalCost;
 			}
 		}
 
-		readonly IStorageClient StorageClient;
-		readonly NamespaceId NamespaceId;
+		readonly IStorageClient _storageClient;
+		readonly NamespaceId _namespaceId;
 
 		/// <summary>
 		/// Reference to the root node in the bundle
 		/// </summary>
 		public BundleNode Root { get; private set; }
 
-		ConcurrentDictionary<IoHash, NodeInfo> HashToNode = new ConcurrentDictionary<IoHash, NodeInfo>();
-		ConcurrentDictionary<IoHash, BlobInfo> HashToBlob = new ConcurrentDictionary<IoHash, BlobInfo>();
-		IMemoryCache Cache;
+		readonly ConcurrentDictionary<IoHash, NodeInfo> _hashToNode = new ConcurrentDictionary<IoHash, NodeInfo>();
+		readonly ConcurrentDictionary<IoHash, BlobInfo> _hashToBlob = new ConcurrentDictionary<IoHash, BlobInfo>();
+		readonly IMemoryCache _cache;
 
-		CbWriter SharedWriter = new CbWriter();
-		byte[] SharedBlobBuffer = Array.Empty<byte>();
+		readonly CbWriter _sharedWriter = new CbWriter();
+		byte[] _sharedBlobBuffer = Array.Empty<byte>();
 
-		object LockObject = new object();
+		readonly object _lockObject = new object();
 
-		SemaphoreSlim ReadSema;
+		readonly SemaphoreSlim _readSema;
 
 		// Active read tasks at any moment. If a BundleObject is not available in the cache, we start a read and add an entry to this dictionary
 		// so that other threads can also await it.
-		Dictionary<IoHash, Task<BundleObject>> ReadTasks = new Dictionary<IoHash, Task<BundleObject>>();
+		readonly Dictionary<IoHash, Task<BundleObject>> _readTasks = new Dictionary<IoHash, Task<BundleObject>>();
 
 		/// <summary>
 		/// Options for the bundle
@@ -261,132 +251,132 @@ namespace EpicGames.Horde.Bundles
 		/// </summary>
 		public BundleStats Stats { get; private set; } = new BundleStats();
 
-		BundleStats NextStats = new BundleStats();
-		byte[] BlockBuffer = Array.Empty<byte>();
-		byte[] EncodedBuffer = Array.Empty<byte>();
+		BundleStats _nextStats = new BundleStats();
+		byte[] _blockBuffer = Array.Empty<byte>();
+		byte[] _encodedBuffer = Array.Empty<byte>();
 
 		/// <summary>
 		/// Constructor
 		/// </summary>
-		protected Bundle(IStorageClient StorageClient, NamespaceId NamespaceId, BundleNode Root, BundleOptions Options, IMemoryCache Cache)
+		protected Bundle(IStorageClient storageClient, NamespaceId namespaceId, BundleNode root, BundleOptions options, IMemoryCache cache)
 		{
-			this.StorageClient = StorageClient;
-			this.NamespaceId = NamespaceId;
-			this.Root = Root;
-			this.Options = Options;
-			this.Cache = Cache;
-			this.ReadSema = new SemaphoreSlim(Options.MaxConcurrentReads);
+			_storageClient = storageClient;
+			_namespaceId = namespaceId;
+			Root = root;
+			Options = options;
+			_cache = cache;
+			_readSema = new SemaphoreSlim(options.MaxConcurrentReads);
 		}
 
 		/// <summary>
 		/// Constructor
 		/// </summary>
-		/// <param name="StorageClient">Client interface for the storage backend</param>
-		/// <param name="NamespaceId">Namespace for storing blobs in the storage system</param>
-		/// <param name="Object">Object containing the root node for this bundle</param>
-		/// <param name="Root">Reference to the root node for the bundle</param>
-		/// <param name="Options">Options for controlling how things are serialized</param>
-		/// <param name="Cache">Cache for storing decompressed objects, serialized blobs.</param>
-		protected Bundle(IStorageClient StorageClient, NamespaceId NamespaceId, BundleObject Object, BundleNode Root, BundleOptions Options, IMemoryCache Cache)
-			: this(StorageClient, NamespaceId, Root, Options, Cache)
+		/// <param name="storageClient">Client interface for the storage backend</param>
+		/// <param name="namespaceId">Namespace for storing blobs in the storage system</param>
+		/// <param name="rootObject">Object containing the root node for this bundle</param>
+		/// <param name="rootNode">Reference to the root node for the bundle</param>
+		/// <param name="options">Options for controlling how things are serialized</param>
+		/// <param name="cache">Cache for storing decompressed objects, serialized blobs.</param>
+		protected Bundle(IStorageClient storageClient, NamespaceId namespaceId, BundleObject rootObject, BundleNode rootNode, BundleOptions options, IMemoryCache cache)
+			: this(storageClient, namespaceId, rootNode, options, cache)
 		{
-			RegisterRootObject(Object);
+			RegisterRootObject(rootObject);
 		}
 
 		/// <inheritdoc/>
 		public void Dispose()
 		{
-			ReadSema.Dispose();
+			_readSema.Dispose();
 		}
 
 		/// <summary>
 		/// Creates a new typed bundle with a default-constructed root node
 		/// </summary>
 		/// <typeparam name="T">The node type</typeparam>
-		/// <param name="StorageClient">Client interface for the storage backend</param>
-		/// <param name="NamespaceId">Namespace for storing blobs in the storage system</param>
-		/// <param name="Options">Options for controlling how things are serialized</param>
-		/// <param name="Cache">Cache for storing decompressed objects, serialized blobs.</param>
+		/// <param name="storageClient">Client interface for the storage backend</param>
+		/// <param name="namespaceId">Namespace for storing blobs in the storage system</param>
+		/// <param name="options">Options for controlling how things are serialized</param>
+		/// <param name="cache">Cache for storing decompressed objects, serialized blobs.</param>
 		/// <returns>New bundle instance</returns>
-		public static Bundle<T> Create<T>(IStorageClient StorageClient, NamespaceId NamespaceId, BundleOptions Options, IMemoryCache Cache) where T : BundleNode, new()
+		public static Bundle<T> Create<T>(IStorageClient storageClient, NamespaceId namespaceId, BundleOptions options, IMemoryCache cache) where T : BundleNode, new()
 		{
-			return Create(StorageClient, NamespaceId, new T(), Options, Cache);
+			return Create(storageClient, namespaceId, new T(), options, cache);
 		}
 
 		/// <summary>
 		/// Creates a new typed bundle with a specific root node
 		/// </summary>
 		/// <typeparam name="T">The node type</typeparam>
-		/// <param name="StorageClient">Client interface for the storage backend</param>
-		/// <param name="NamespaceId">Namespace for storing blobs in the storage system</param>
-		/// <param name="Root">Reference to the root node for the bundle</param>
-		/// <param name="Options">Options for controlling how things are serialized</param>
-		/// <param name="Cache">Cache for storing decompressed objects, serialized blobs.</param>
+		/// <param name="storageClient">Client interface for the storage backend</param>
+		/// <param name="namespaceId">Namespace for storing blobs in the storage system</param>
+		/// <param name="root">Reference to the root node for the bundle</param>
+		/// <param name="options">Options for controlling how things are serialized</param>
+		/// <param name="cache">Cache for storing decompressed objects, serialized blobs.</param>
 		/// <returns>New bundle instance</returns>
-		public static Bundle<T> Create<T>(IStorageClient StorageClient, NamespaceId NamespaceId, T Root, BundleOptions Options, IMemoryCache Cache) where T : BundleNode
+		public static Bundle<T> Create<T>(IStorageClient storageClient, NamespaceId namespaceId, T root, BundleOptions options, IMemoryCache cache) where T : BundleNode
 		{
-			return new Bundle<T>(StorageClient, NamespaceId, Root, Options, Cache);
+			return new Bundle<T>(storageClient, namespaceId, root, options, cache);
 		}
 
 		/// <summary>
 		/// Reads a bundle from storage
 		/// </summary>
-		public static async ValueTask<Bundle<T>> ReadAsync<T>(IStorageClient StorageClient, NamespaceId NamespaceId, BucketId BucketId, RefId RefId, BundleOptions Options, IMemoryCache Cache) where T : BundleNode
+		public static async ValueTask<Bundle<T>> ReadAsync<T>(IStorageClient storageClient, NamespaceId namespaceId, BucketId bucketId, RefId refId, BundleOptions options, IMemoryCache cache) where T : BundleNode
 		{
-			Bundle<T> NewBundle = new Bundle<T>(StorageClient, NamespaceId, null!, Options, Cache);
+			Bundle<T> newBundle = new Bundle<T>(storageClient, namespaceId, null!, options, cache);
 
-			BundleRoot RootObject = await StorageClient.GetRefAsync<BundleRoot>(NamespaceId, BucketId, RefId);
-			NewBundle.RegisterRootObject(RootObject.Object);
+			BundleRoot rootObject = await storageClient.GetRefAsync<BundleRoot>(namespaceId, bucketId, refId);
+			newBundle.RegisterRootObject(rootObject.Object);
 
-			IoHash RootNodeHash = RootObject.Object.Exports[^1].Hash;
-			ReadOnlySequence<byte> RootNodeData = await NewBundle.GetDataAsync(RootNodeHash);
-			((Bundle)NewBundle).Root = BundleNode.Deserialize<T>(RootNodeData);
+			IoHash rootNodeHash = rootObject.Object.Exports[^1].Hash;
+			ReadOnlySequence<byte> rootNodeData = await newBundle.GetDataAsync(rootNodeHash);
+			((Bundle)newBundle).Root = BundleNode.Deserialize<T>(rootNodeData);
 
-			return NewBundle;
+			return newBundle;
 		}
 
 		/// <summary>
 		/// Flush a subset of nodes in the working set.
 		/// </summary>
-		private void SerializeWorkingSet(int IgnoreCount)
+		private void SerializeWorkingSet(int ignoreCount)
 		{
 			// Find all the nodes that need to be flushed, in order
-			List<BundleNodeRef> NodeRefs = new List<BundleNodeRef>();
-			LinkWorkingSet(Root, NodeRefs);
+			List<BundleNodeRef> nodeRefs = new List<BundleNodeRef>();
+			LinkWorkingSet(Root, nodeRefs);
 
 			// Find the working set of mutable nodes to retain in memory
-			HashSet<BundleNodeRef> WorkingSet = new HashSet<BundleNodeRef>(NodeRefs.Where(x => !x.Node!.IsReadOnly()).OrderByDescending(x => x.LastModifiedTime).Take(IgnoreCount));
-			NodeRefs.RemoveAll(x => WorkingSet.Contains(x));
-			NodeRefs.SortBy(x => x.LastModifiedTime);
+			HashSet<BundleNodeRef> workingSet = new HashSet<BundleNodeRef>(nodeRefs.Where(x => !x.Node!.IsReadOnly()).OrderByDescending(x => x._lastModifiedTime).Take(ignoreCount));
+			nodeRefs.RemoveAll(x => workingSet.Contains(x));
+			nodeRefs.SortBy(x => x._lastModifiedTime);
 
 			// Write them all to storage
-			foreach (BundleNodeRef NodeRef in NodeRefs)
+			foreach (BundleNodeRef nodeRef in nodeRefs)
 			{
-				BundleNode? Node = NodeRef.Node;
-				Debug.Assert(Node != null);
+				BundleNode? node = nodeRef.Node;
+				Debug.Assert(node != null);
 
-				NodeInfo NodeInfo = WriteNode(Node);
-				NodeRef.MarkAsClean(NodeInfo.Hash);
+				NodeInfo nodeInfo = WriteNode(node);
+				nodeRef.MarkAsClean(nodeInfo.Hash);
 			}
 		}
 
 		/// <summary>
 		/// Traverse a tree of <see cref="BundleNodeRef"/> objects and update the incoming reference to each node
 		/// </summary>
-		/// <param name="Node">The reference to update</param>
-		/// <param name="NodeRefs">List to accumulate the list of all modified references</param>
-		private void LinkWorkingSet(BundleNode Node, List<BundleNodeRef> NodeRefs)
+		/// <param name="node">The reference to update</param>
+		/// <param name="nodeRefs">List to accumulate the list of all modified references</param>
+		private void LinkWorkingSet(BundleNode node, List<BundleNodeRef> nodeRefs)
 		{
-			foreach (BundleNodeRef ChildRef in Node.GetReferences())
+			foreach (BundleNodeRef childRef in node.GetReferences())
 			{
-				if (ChildRef.Node != null && ChildRef.LastModifiedTime != 0)
+				if (childRef.Node != null && childRef._lastModifiedTime != 0)
 				{
-					LinkWorkingSet(ChildRef.Node, NodeRefs);
-					if (Node.IncomingRef != null)
+					LinkWorkingSet(childRef.Node, nodeRefs);
+					if (node.IncomingRef != null)
 					{
-						Node.IncomingRef.LastModifiedTime = Math.Max(Node.IncomingRef.LastModifiedTime, ChildRef.LastModifiedTime + 1);
+						node.IncomingRef._lastModifiedTime = Math.Max(node.IncomingRef._lastModifiedTime, childRef._lastModifiedTime + 1);
 					}
-					NodeRefs.Add(ChildRef);
+					nodeRefs.Add(childRef);
 				}
 			}
 		}
@@ -399,55 +389,55 @@ namespace EpicGames.Horde.Bundles
 		/// <summary>
 		/// Perform an incremental flush of the tree
 		/// </summary>
-		public async Task TrimAsync(int IgnoreCount)
+		public async Task TrimAsync(int ignoreCount)
 		{
 			// Serialize nodes currently in the working set
-			SerializeWorkingSet(IgnoreCount);
+			SerializeWorkingSet(ignoreCount);
 
 			// Find all the nodes that can be written, in order
-			List<NodeInfo> WriteNodes = new List<NodeInfo>();
-			FindNodesToWrite(Root, WriteNodes, new HashSet<NodeInfo>());
+			List<NodeInfo> writeNodes = new List<NodeInfo>();
+			FindNodesToWrite(Root, writeNodes, new HashSet<NodeInfo>());
 
 			// Write them all to storage
-			int MinNodeIdx = 0;
-			int NextBlobCost = 0;
-			for (int Idx = 0; Idx < WriteNodes.Count; Idx++)
+			int minNodeIdx = 0;
+			int nextBlobCost = 0;
+			for (int idx = 0; idx < writeNodes.Count; idx++)
 			{
-				NodeInfo Node = WriteNodes[Idx];
-				if (Idx > MinNodeIdx && NextBlobCost + Node.Length > Options.MaxBlobSize)
+				NodeInfo node = writeNodes[idx];
+				if (idx > minNodeIdx && nextBlobCost + node.Length > Options.MaxBlobSize)
 				{
-					await WriteObjectAsync(WriteNodes.Slice(MinNodeIdx, Idx - MinNodeIdx));
-					MinNodeIdx = Idx;
-					NextBlobCost = 0;
+					await WriteObjectAsync(writeNodes.Slice(minNodeIdx, idx - minNodeIdx));
+					minNodeIdx = idx;
+					nextBlobCost = 0;
 				}
-				NextBlobCost += Node.Length;
+				nextBlobCost += node.Length;
 			}
 		}
 
-		void FindNodesToWrite(BundleNode Node, List<NodeInfo> WriteNodes, HashSet<NodeInfo> WriteNodesSet)
+		void FindNodesToWrite(BundleNode node, List<NodeInfo> writeNodes, HashSet<NodeInfo> writeNodesSet)
 		{
-			foreach (BundleNodeRef ChildRef in Node.GetReferences())
+			foreach (BundleNodeRef childRef in node.GetReferences())
 			{
-				if (ChildRef.Node != null)
+				if (childRef.Node != null)
 				{
-					FindNodesToWrite(ChildRef.Node, WriteNodes, WriteNodesSet);
+					FindNodesToWrite(childRef.Node, writeNodes, writeNodesSet);
 				}
 				else
 				{
-					FindNodeInfosToWrite(HashToNode[ChildRef.Hash], WriteNodes, WriteNodesSet);
+					FindNodeInfosToWrite(_hashToNode[childRef.Hash], writeNodes, writeNodesSet);
 				}
 			}
 		}
 
-		void FindNodeInfosToWrite(NodeInfo Node, List<NodeInfo> WriteNodes, HashSet<NodeInfo> WriteNodesSet)
+		void FindNodeInfosToWrite(NodeInfo node, List<NodeInfo> writeNodes, HashSet<NodeInfo> writeNodesSet)
 		{
-			if (Node.State.Blob == null && WriteNodesSet.Add(Node))
+			if (node.State.Blob == null && writeNodesSet.Add(node))
 			{
-				foreach (NodeInfo ChildNode in Node.State.References!)
+				foreach (NodeInfo childNode in node.State.References!)
 				{
-					FindNodeInfosToWrite(ChildNode, WriteNodes, WriteNodesSet);
+					FindNodeInfosToWrite(childNode, writeNodes, writeNodesSet);
 				}
-				WriteNodes.Add(Node);
+				writeNodes.Add(node);
 			}
 		}
 
@@ -455,119 +445,119 @@ namespace EpicGames.Horde.Bundles
 		/// Gets the node referenced by a <see cref="BundleNodeRef{T}"/>, reading it from storage and deserializing it if necessary.
 		/// </summary>
 		/// <typeparam name="T">Type of node to return</typeparam>
-		/// <param name="Ref">The reference to get</param>
+		/// <param name="nodeRef">The reference to get</param>
 		/// <returns>The deserialized node</returns>
-		public async ValueTask<T> GetAsync<T>(BundleNodeRef<T> Ref) where T : BundleNode
+		public async ValueTask<T> GetAsync<T>(BundleNodeRef<T> nodeRef) where T : BundleNode
 		{
-			if (!Ref.MakeStrongRef())
+			if (!nodeRef.MakeStrongRef())
 			{
-				ReadOnlySequence<byte> Data = await GetDataAsync(Ref.Hash);
-				Ref.Node = BundleNode.Deserialize<T>(Data);
-				Ref.Node.IncomingRef = Ref;
+				ReadOnlySequence<byte> data = await GetDataAsync(nodeRef.Hash);
+				nodeRef.Node = BundleNode.Deserialize<T>(data);
+				nodeRef.Node.IncomingRef = nodeRef;
 			}
-			return Ref.Node!;
+			return nodeRef.Node!;
 		}
 
 		/// <summary>
 		/// Gets data for the node with the given hash
 		/// </summary>
-		/// <param name="Hash">Hash of the node to return data for</param>
+		/// <param name="hash">Hash of the node to return data for</param>
 		/// <returns>The node data</returns>
-		internal ValueTask<ReadOnlySequence<byte>> GetDataAsync(IoHash Hash)
+		internal ValueTask<ReadOnlySequence<byte>> GetDataAsync(IoHash hash)
 		{
-			return GetDataAsync(HashToNode[Hash]);
+			return GetDataAsync(_hashToNode[hash]);
 		}
 
 		/// <summary>
 		/// Gets the data for a given node
 		/// </summary>
-		/// <param name="Node">The node to get the data for</param>
+		/// <param name="node">The node to get the data for</param>
 		/// <returns>The node data</returns>
-		async ValueTask<ReadOnlySequence<byte>> GetDataAsync(NodeInfo Node)
+		async ValueTask<ReadOnlySequence<byte>> GetDataAsync(NodeInfo node)
 		{
-			NodeState NodeState = Node.State;
-			if (NodeState.Blob == null)
+			NodeState nodeState = node.State;
+			if (nodeState.Blob == null)
 			{
-				return NodeState.Data;
+				return nodeState.Data;
 			}
 			else
 			{
-				if (!NodeState.IsExport())
+				if (!nodeState.IsExport())
 				{
-					await MountBlobAsync(NodeState.Blob);
-					NodeState = Node.State;
+					await MountBlobAsync(nodeState.Blob);
+					nodeState = node.State;
 				}
 
-				Debug.Assert(NodeState.Blob != null);
-				Debug.Assert(NodeState.Export != null);
+				Debug.Assert(nodeState.Blob != null);
+				Debug.Assert(nodeState.Export != null);
 
-				BundleObject Object = await GetObjectAsync(NodeState.Blob.Hash);
-				ReadOnlyMemory<byte> PacketData = DecodePacket(NodeState.Blob.Hash, Object.Data, NodeState.Export.Packet);
-				ReadOnlyMemory<byte> NodeData = PacketData.Slice(NodeState.Export.Offset, NodeState.Export.Length);
+				BundleObject bundleObject = await GetObjectAsync(nodeState.Blob.Hash);
+				ReadOnlyMemory<byte> packetData = DecodePacket(nodeState.Blob.Hash, bundleObject.Data, nodeState.Export.Packet);
+				ReadOnlyMemory<byte> nodeData = packetData.Slice(nodeState.Export.Offset, nodeState.Export.Length);
 
-				return new ReadOnlySequence<byte>(NodeData);
+				return new ReadOnlySequence<byte>(nodeData);
 			}
 		}
 
-		NodeInfo WriteNode(BundleNode Node)
+		NodeInfo WriteNode(BundleNode node)
 		{
-			ReadOnlySequence<byte> Data = Node.Serialize();
+			ReadOnlySequence<byte> data = node.Serialize();
 
-			List<IoHash> References = new List<IoHash>();
-			foreach (BundleNodeRef Reference in Node.GetReferences())
+			List<IoHash> references = new List<IoHash>();
+			foreach (BundleNodeRef reference in node.GetReferences())
 			{
-				Debug.Assert(Reference.Hash != IoHash.Zero);
-				References.Add(Reference.Hash);
+				Debug.Assert(reference.Hash != IoHash.Zero);
+				references.Add(reference.Hash);
 			}
 
-			return WriteNode(Data, References);
+			return WriteNode(data, references);
 		}
 
 		/// <summary>
 		/// Write a node's data
 		/// </summary>
-		/// <param name="Data">The node data</param>
-		/// <param name="References">Hashes for referenced nodes</param>
+		/// <param name="data">The node data</param>
+		/// <param name="references">Hashes for referenced nodes</param>
 		/// <returns>Hash of the data</returns>
-		NodeInfo WriteNode(ReadOnlySequence<byte> Data, IEnumerable<IoHash> References)
+		NodeInfo WriteNode(ReadOnlySequence<byte> data, IEnumerable<IoHash> references)
 		{
-			IoHash Hash = IoHash.Compute(Data);
-			NodeInfo[] NodeReferences = References.Select(x => HashToNode[x]).Distinct().OrderBy(x => x.Hash).ToArray();
+			IoHash hash = IoHash.Compute(data);
+			NodeInfo[] nodeReferences = references.Select(x => _hashToNode[x]).Distinct().OrderBy(x => x.Hash).ToArray();
 
-			NodeState State = NodeState.Standalone(Data, NodeReferences);
+			NodeState state = NodeState.Standalone(data, nodeReferences);
 
-			NodeInfo? Node;
-			if (HashToNode.TryGetValue(Hash, out Node))
+			NodeInfo? node;
+			if (_hashToNode.TryGetValue(hash, out node))
 			{
-				Node.State = State;
+				node.State = state;
 			}
 			else
 			{
-				int Rank = 0;
-				foreach (NodeInfo Reference in NodeReferences)
+				int rank = 0;
+				foreach (NodeInfo reference in nodeReferences)
 				{
-					Rank = Math.Max(Rank, Reference.Rank + 1);
+					rank = Math.Max(rank, reference.Rank + 1);
 				}
 
-				Node = new NodeInfo(Hash, Rank, (int)Data.Length, State);
-				HashToNode[Hash] = Node;
+				node = new NodeInfo(hash, rank, (int)data.Length, state);
+				_hashToNode[hash] = node;
 
-				NextStats.NewNodeCount++;
-				NextStats.NewNodeBytes += Data.Length;
+				_nextStats.NewNodeCount++;
+				_nextStats.NewNodeBytes += data.Length;
 			}
-			return Node;
+			return node;
 		}
 
-		async ValueTask MountBlobAsync(BlobInfo Blob)
+		async ValueTask MountBlobAsync(BlobInfo blob)
 		{
-			if(!Blob.Mounted)
+			if(!blob.Mounted)
 			{
-				BundleObject Object = await GetObjectAsync(Blob.Hash);
-				lock (LockObject)
+				BundleObject bundleObject = await GetObjectAsync(blob.Hash);
+				lock (_lockObject)
 				{
-					if (!Blob.Mounted)
+					if (!blob.Mounted)
 					{
-						RegisterBlobObject(Object, Blob);
+						RegisterBlobObject(bundleObject, blob);
 					}
 				}
 			}
@@ -576,199 +566,199 @@ namespace EpicGames.Horde.Bundles
 		/// <summary>
 		/// Reads an object data from the store
 		/// </summary>
-		/// <param name="Hash">Hash of the object</param>
+		/// <param name="hash">Hash of the object</param>
 		/// <returns>The parsed object</returns>
-		async ValueTask<BundleObject> GetObjectAsync(IoHash Hash)
+		async ValueTask<BundleObject> GetObjectAsync(IoHash hash)
 		{
-			string CacheKey = $"object:{Hash}";
+			string cacheKey = $"object:{hash}";
 
-			BundleObject? Object;
-			if (!Cache.TryGetValue<BundleObject>(CacheKey, out Object))
+			BundleObject? bundleObject;
+			if (!_cache.TryGetValue<BundleObject>(cacheKey, out bundleObject))
 			{
-				Task<BundleObject>? ReadTask;
-				lock (LockObject)
+				Task<BundleObject>? readTask;
+				lock (_lockObject)
 				{
-					if (!ReadTasks.TryGetValue(Hash, out ReadTask))
+					if (!_readTasks.TryGetValue(hash, out readTask))
 					{
-						ReadTask = Task.Run(() => ReadObjectAsync(CacheKey, Hash));
-						ReadTasks.Add(Hash, ReadTask);
+						readTask = Task.Run(() => ReadObjectAsync(cacheKey, hash));
+						_readTasks.Add(hash, readTask);
 					}
 				}
-				Object = await ReadTask;
+				bundleObject = await readTask;
 			}
-			return Object;
+			return bundleObject;
 		}
 
-		async Task<BundleObject> ReadObjectAsync(string CacheKey, IoHash Hash)
+		async Task<BundleObject> ReadObjectAsync(string cacheKey, IoHash hash)
 		{
 			// Perform another (sequenced) check whether an object has been added to the cache, to counteract the race between a read task being added and a task completing.
-			lock (LockObject)
+			lock (_lockObject)
 			{
-				if (Cache.TryGetValue<BundleObject>(CacheKey, out BundleObject? CachedObject))
+				if (_cache.TryGetValue<BundleObject>(cacheKey, out BundleObject? cachedObject))
 				{
-					return CachedObject;
+					return cachedObject;
 				}
 			}
 
 			// Wait for the read semaphore to avoid triggering too many operations at once.
-			await ReadSema.WaitAsync();
+			await _readSema.WaitAsync();
 
 			// Read the data from storage
-			ReadOnlyMemory<byte> Data;
+			ReadOnlyMemory<byte> data;
 			try
 			{
-				Data = await StorageClient.ReadBlobToMemoryAsync(NamespaceId, Hash);
+				data = await _storageClient.ReadBlobToMemoryAsync(_namespaceId, hash);
 			}
 			finally
 			{
-				ReadSema.Release();
+				_readSema.Release();
 			}
 
 			// Add the object to the cache
-			BundleObject Object = CbSerializer.Deserialize<BundleObject>(Data);
-			using (ICacheEntry Entry = Cache.CreateEntry(CacheKey))
+			BundleObject bundleObject = CbSerializer.Deserialize<BundleObject>(data);
+			using (ICacheEntry entry = _cache.CreateEntry(cacheKey))
 			{
-				Entry.SetSize(Data.Length);
-				Entry.SetValue(Object);
+				entry.SetSize(data.Length);
+				entry.SetValue(bundleObject);
 			}
 
 			// Remove this object from the list of read tasks
-			lock (LockObject)
+			lock (_lockObject)
 			{
-				ReadTasks.Remove(Hash);
+				_readTasks.Remove(hash);
 			}
-			return Object;
+			return bundleObject;
 		}
 
 		/// <summary>
 		/// Gets a decoded block from the store
 		/// </summary>
-		/// <param name="BlobHash">Key for the blob</param>
-		/// <param name="Data">Raw blob data</param>
-		/// <param name="Packet">The decoded block location and size</param>
+		/// <param name="blobHash">Key for the blob</param>
+		/// <param name="data">Raw blob data</param>
+		/// <param name="packet">The decoded block location and size</param>
 		/// <returns>The decoded data</returns>
-		ReadOnlyMemory<byte> DecodePacket(IoHash BlobHash, ReadOnlyMemory<byte> Data, BundleCompressionPacket Packet)
+		ReadOnlyMemory<byte> DecodePacket(IoHash blobHash, ReadOnlyMemory<byte> data, BundleCompressionPacket packet)
 		{
-			string CacheKey = $"decode:{BlobHash}/{Packet.Offset}";
-			return Cache.GetOrCreate<ReadOnlyMemory<byte>>(CacheKey, Entry =>
+			string cacheKey = $"decode:{blobHash}/{packet.Offset}";
+			return _cache.GetOrCreate<ReadOnlyMemory<byte>>(cacheKey, entry =>
 			{
-				ReadOnlyMemory<byte> EncodedPacket = Data.Slice(Packet.Offset, Packet.EncodedLength);
+				ReadOnlyMemory<byte> encodedPacket = data.Slice(packet.Offset, packet.EncodedLength);
 
-				byte[] DecodedPacket = new byte[Packet.DecodedLength];
-				LZ4Codec.Decode(EncodedPacket.Span, DecodedPacket);
+				byte[] decodedPacket = new byte[packet.DecodedLength];
+				LZ4Codec.Decode(encodedPacket.Span, decodedPacket);
 
-				Entry.SetSize(Packet.DecodedLength);
-				return DecodedPacket;
+				entry.SetSize(packet.DecodedLength);
+				return decodedPacket;
 			});
 		}
 
 		/// <summary>
 		/// Find a node or create a new one
 		/// </summary>
-		NodeInfo FindOrAddNodeFromImport(BlobInfo Blob, BundleImport Import)
+		NodeInfo FindOrAddNodeFromImport(BlobInfo blob, BundleImport import)
 		{
-			NodeInfo? Node;
-			if (!HashToNode.TryGetValue(Import.Hash, out Node))
+			NodeInfo? node;
+			if (!_hashToNode.TryGetValue(import.Hash, out node))
 			{
-				NodeInfo NewNode = new NodeInfo(Import.Hash, Import.Rank, Import.Length, NodeState.Imported(Blob));
-				Node = HashToNode.GetOrAdd(NewNode.Hash, NewNode);
+				NodeInfo newNode = new NodeInfo(import.Hash, import.Rank, import.Length, NodeState.Imported(blob));
+				node = _hashToNode.GetOrAdd(newNode.Hash, newNode);
 			}
-			return Node;
+			return node;
 		}
 
 		/// <summary>
 		/// Creates nodes for a given root object. Data for each node will be decompressed and copied from the object.
 		/// </summary>
-		void RegisterRootObject(BundleObject Object)
+		void RegisterRootObject(BundleObject rootObject)
 		{
 			// Create all the export nodes without fixing up their dependencies yet
-			List<NodeInfo> Nodes = new List<NodeInfo>();
-			for (int Idx = 0; Idx < Object.Exports.Count;)
+			List<NodeInfo> nodes = new List<NodeInfo>();
+			for (int idx = 0; idx < rootObject.Exports.Count;)
 			{
-				BundleCompressionPacket Packet = Object.Exports[Idx].Packet;
+				BundleCompressionPacket packet = rootObject.Exports[idx].Packet;
 
-				byte[] DecodedData = new byte[Packet.DecodedLength];
-				LZ4Codec.Decode(Object.Data.Slice(Packet.Offset, Packet.EncodedLength).Span, DecodedData);
+				byte[] decodedData = new byte[packet.DecodedLength];
+				LZ4Codec.Decode(rootObject.Data.Slice(packet.Offset, packet.EncodedLength).Span, decodedData);
 
-				for (; Idx < Object.Exports.Count && Object.Exports[Idx].Packet == Packet; Idx++)
+				for (; idx < rootObject.Exports.Count && rootObject.Exports[idx].Packet == packet; idx++)
 				{
-					BundleExport Export = Object.Exports[Idx];
+					BundleExport export = rootObject.Exports[idx];
 
-					ReadOnlyMemory<byte> NodeData = DecodedData.AsMemory(Export.Offset, Export.Length);
-					if (Export.Length != DecodedData.Length)
+					ReadOnlyMemory<byte> nodeData = decodedData.AsMemory(export.Offset, export.Length);
+					if (export.Length != decodedData.Length)
 					{
-						NodeData = NodeData.ToArray();
+						nodeData = nodeData.ToArray();
 					}
 
-					NodeState State = NodeState.Standalone(new ReadOnlySequence<byte>(NodeData), Array.Empty<NodeInfo>());
+					NodeState state = NodeState.Standalone(new ReadOnlySequence<byte>(nodeData), Array.Empty<NodeInfo>());
 
-					NodeInfo? Node;
+					NodeInfo? node;
 					for (; ; )
 					{
-						if (HashToNode.TryGetValue(Export.Hash, out Node))
+						if (_hashToNode.TryGetValue(export.Hash, out node))
 						{
-							Node.State = State;
+							node.State = state;
 							break;
 						}
 
-						NodeInfo NewNode = new NodeInfo(Export.Hash, Export.Rank, Export.Length, State);
-						if (HashToNode.TryAdd(Export.Hash, NewNode))
+						NodeInfo newNode = new NodeInfo(export.Hash, export.Rank, export.Length, state);
+						if (_hashToNode.TryAdd(export.Hash, newNode))
 						{
-							Node = NewNode;
+							node = newNode;
 							break;
 						}
 					}
-					Nodes.Add(Node);
+					nodes.Add(node);
 				}
 			}
 
 			// Create all the imports
-			RegisterImports(Object.ImportObjects, Nodes);
+			RegisterImports(rootObject.ImportObjects, nodes);
 
 			// Fixup all the references
-			for (int Idx = 0; Idx < Object.Exports.Count; Idx++)
+			for (int idx = 0; idx < rootObject.Exports.Count; idx++)
 			{
-				NodeInfo Node = Nodes[Idx];
-				Node.State = NodeState.Standalone(Node.State.Data, Object.Exports[Idx].References.Select(x => Nodes[x]).ToArray());
+				NodeInfo node = nodes[idx];
+				node.State = NodeState.Standalone(node.State.Data, rootObject.Exports[idx].References.Select(x => nodes[x]).ToArray());
 			}
 		}
 
 		/// <summary>
 		/// Creates nodes for a given blob object.
 		/// </summary>
-		void RegisterBlobObject(BundleObject Object, BlobInfo Blob)
+		void RegisterBlobObject(BundleObject blobObject, BlobInfo blobInfo)
 		{
 			// Create the exports as imports first, since we can't resolve the reference list.
-			List<NodeInfo> Nodes = new List<NodeInfo>();
-			foreach (BundleExport Export in Object.Exports)
+			List<NodeInfo> nodes = new List<NodeInfo>();
+			foreach (BundleExport export in blobObject.Exports)
 			{
-				NodeInfo Node = FindOrAddNodeFromImport(Blob, Export);
-				Nodes.Add(Node);
+				NodeInfo node = FindOrAddNodeFromImport(blobInfo, export);
+				nodes.Add(node);
 			}
 
 			// Register the regular imports
-			RegisterImports(Object.ImportObjects, Nodes);
+			RegisterImports(blobObject.ImportObjects, nodes);
 
 			// Go back and resolve all the exports 
-			for (int Idx = 0; Idx < Object.Exports.Count; Idx++)
+			for (int idx = 0; idx < blobObject.Exports.Count; idx++)
 			{
-				BundleExport Export = Object.Exports[Idx];
-				Nodes[Idx].TrySetExport(Blob, Export, Export.References.Select(x => Nodes[x]).ToArray());
+				BundleExport export = blobObject.Exports[idx];
+				nodes[idx].TrySetExport(blobInfo, export, export.References.Select(x => nodes[x]).ToArray());
 			}
 
 			// Flag the blob as being mounted
-			Blob.Mounted = true;
+			blobInfo.Mounted = true;
 		}
 
-		void RegisterImports(List<BundleImportObject> ImportObjects, List<NodeInfo> Nodes)
+		void RegisterImports(List<BundleImportObject> importObjects, List<NodeInfo> nodes)
 		{
-			foreach (BundleImportObject ImportObject in ImportObjects)
+			foreach (BundleImportObject importObject in importObjects)
 			{
-				BlobInfo ImportBlob = HashToBlob.GetOrAdd(ImportObject.Object.Hash, new BlobInfo(ImportObject.Object.Hash, ImportObject.TotalCost));
-				foreach (BundleImport Import in ImportObject.Imports)
+				BlobInfo importBlob = _hashToBlob.GetOrAdd(importObject.Object.Hash, new BlobInfo(importObject.Object.Hash, importObject.TotalCost));
+				foreach (BundleImport import in importObject.Imports)
 				{
-					NodeInfo Node = FindOrAddNodeFromImport(ImportBlob, Import);
-					Nodes.Add(Node);
+					NodeInfo node = FindOrAddNodeFromImport(importBlob, import);
+					nodes.Add(node);
 				}
 			}
 		}
@@ -776,310 +766,310 @@ namespace EpicGames.Horde.Bundles
 		/// <summary>
 		/// Persist the bundle to storage
 		/// </summary>
-		/// <param name="BucketId"></param>
-		/// <param name="RefId"></param>
-		/// <param name="Metadata">Metadata for the root object</param>
-		/// <param name="Compact"></param>
-		public virtual async Task WriteAsync(BucketId BucketId, RefId RefId, CbObject Metadata, bool Compact)
+		/// <param name="bucketId"></param>
+		/// <param name="refId"></param>
+		/// <param name="metadata">Metadata for the root object</param>
+		/// <param name="compact"></param>
+		public virtual async Task WriteAsync(BucketId bucketId, RefId refId, CbObject metadata, bool compact)
 		{
 			// Completely flush the entire working set into NodeInfo objects
 			SerializeWorkingSet(0);
 
 			// Serialize the root node
-			NodeInfo RootNode = WriteNode(Root);
+			NodeInfo rootNode = WriteNode(Root);
 
 			// If we're compacting the output, see if there are any other blocks we should merge
-			if (Compact)
+			if (compact)
 			{
 				// Find the live set of objects
-				List<NodeInfo> LiveNodes = new List<NodeInfo>();
-				await FindLiveNodesAsync(RootNode, LiveNodes, new HashSet<NodeInfo>());
+				List<NodeInfo> liveNodes = new List<NodeInfo>();
+				await FindLiveNodesAsync(rootNode, liveNodes, new HashSet<NodeInfo>());
 
 				// Find the live set of blobs, and the cost of nodes within them
-				List<BlobInfo> LiveBlobs = new List<BlobInfo>();
-				foreach (NodeInfo LiveNode in LiveNodes)
+				List<BlobInfo> liveBlobs = new List<BlobInfo>();
+				foreach (NodeInfo liveNode in liveNodes)
 				{
-					BlobInfo? LiveBlob = LiveNode.State.Blob;
-					if (LiveBlob != null)
+					BlobInfo? liveBlob = liveNode.State.Blob;
+					if (liveBlob != null)
 					{
-						if (LiveBlob.LiveNodes == null)
+						if (liveBlob.LiveNodes == null)
 						{
-							LiveBlob.LiveNodes = new List<NodeInfo>();
-							LiveBlobs.Add(LiveBlob);
+							liveBlob.LiveNodes = new List<NodeInfo>();
+							liveBlobs.Add(liveBlob);
 						}
-						LiveBlob.LiveNodes.Add(LiveNode);
+						liveBlob.LiveNodes.Add(liveNode);
 					}
 				}
 
 				// Figure out which blobs need to be repacked
-				foreach (BlobInfo LiveBlob in LiveBlobs)
+				foreach (BlobInfo liveBlob in liveBlobs)
 				{
-					int LiveCost = LiveBlob.LiveNodes.Sum(x => x.Length);
-					int MinLiveCost = (int)(LiveBlob.TotalCost * Options.RepackRatio);
+					int liveCost = liveBlob.LiveNodes.Sum(x => x.Length);
+					int minLiveCost = (int)(liveBlob.TotalCost * Options.RepackRatio);
 
-					if (LiveCost < MinLiveCost)
+					if (liveCost < minLiveCost)
 					{
-						foreach (NodeInfo Node in LiveBlob.LiveNodes!)
+						foreach (NodeInfo node in liveBlob.LiveNodes!)
 						{
-							ReadOnlySequence<byte> Data = await GetDataAsync(Node);
-							Node.SetStandalone(Data, Node.State.References!);
+							ReadOnlySequence<byte> data = await GetDataAsync(node);
+							node.SetStandalone(data, node.State.References!);
 						}
 					}
 				}
 
 				// Walk backwards through the list of live nodes and make anything that references a repacked node also standalone
-				for (int Idx = LiveNodes.Count - 1; Idx >= 0; Idx--)
+				for (int idx = liveNodes.Count - 1; idx >= 0; idx--)
 				{
-					NodeInfo LiveNode = LiveNodes[Idx];
-					if (LiveNode.State.References != null && LiveNode.State.References.Any(x => x.State.Blob == null))
+					NodeInfo liveNode = liveNodes[idx];
+					if (liveNode.State.References != null && liveNode.State.References.Any(x => x.State.Blob == null))
 					{
-						ReadOnlySequence<byte> Data = await GetDataAsync(LiveNode);
-						LiveNode.SetStandalone(Data, LiveNode.State.References!);
+						ReadOnlySequence<byte> data = await GetDataAsync(liveNode);
+						liveNode.SetStandalone(data, liveNode.State.References!);
 					}
 				}
 
 				// Clear the reference list on each blob
-				foreach (BlobInfo LiveBlob in LiveBlobs)
+				foreach (BlobInfo liveBlob in liveBlobs)
 				{
-					LiveBlob.LiveNodes = null;
+					liveBlob.LiveNodes = null;
 				}
 			}
 
 			// Find all the modified nodes
-			HashSet<NodeInfo> NewNodes = new HashSet<NodeInfo>();
-			FindModifiedLiveSet(RootNode, NewNodes);
+			HashSet<NodeInfo> newNodes = new HashSet<NodeInfo>();
+			FindModifiedLiveSet(rootNode, newNodes);
 
 			// Write all the new blobs. Sort by rank and go from end to start, updating the list as we go so that we can figure out imports correctly.
-			NodeInfo[] WriteNodes = NewNodes.Where(x => x.State.Blob == null).OrderBy(x => x.Rank).ThenBy(x => x.Hash).ToArray();
+			NodeInfo[] writeNodes = newNodes.Where(x => x.State.Blob == null).OrderBy(x => x.Rank).ThenBy(x => x.Hash).ToArray();
 
-			int MinIdx = 0;
-			int NextBlobCost = WriteNodes[0].Length;
+			int minIdx = 0;
+			int nextBlobCost = writeNodes[0].Length;
 
-			for (int MaxIdx = 1; MaxIdx < WriteNodes.Length; MaxIdx++)
+			for (int maxIdx = 1; maxIdx < writeNodes.Length; maxIdx++)
 			{
-				NextBlobCost += WriteNodes[MaxIdx].Length;
-				if (NextBlobCost > Options.MaxBlobSize)
+				nextBlobCost += writeNodes[maxIdx].Length;
+				if (nextBlobCost > Options.MaxBlobSize)
 				{
-					await WriteObjectAsync(WriteNodes.Slice(MinIdx, MaxIdx - MinIdx));
-					MinIdx = MaxIdx;
-					NextBlobCost = WriteNodes[MaxIdx].Length;
+					await WriteObjectAsync(writeNodes.Slice(minIdx, maxIdx - minIdx));
+					minIdx = maxIdx;
+					nextBlobCost = writeNodes[maxIdx].Length;
 				}
 			}
 
-			if (NextBlobCost > Options.MaxInlineBlobSize && MinIdx + 1 < WriteNodes.Length)
+			if (nextBlobCost > Options.MaxInlineBlobSize && minIdx + 1 < writeNodes.Length)
 			{
-				await WriteObjectAsync(WriteNodes.Slice(MinIdx, WriteNodes.Length - 1 - MinIdx));
-				MinIdx = WriteNodes.Length - 1;
+				await WriteObjectAsync(writeNodes.Slice(minIdx, writeNodes.Length - 1 - minIdx));
+				minIdx = writeNodes.Length - 1;
 			}
 
 			// Write the final ref
-			await WriteRefAsync(WriteNodes.Slice(MinIdx), BucketId, RefId, Metadata);
+			await WriteRefAsync(writeNodes.Slice(minIdx), bucketId, refId, metadata);
 
 			// Copy the stats over
-			Stats = NextStats;
-			NextStats = new BundleStats();
+			Stats = _nextStats;
+			_nextStats = new BundleStats();
 		}
 
-		private void FindModifiedLiveSet(NodeInfo Node, HashSet<NodeInfo> Nodes)
+		private void FindModifiedLiveSet(NodeInfo node, HashSet<NodeInfo> nodes)
 		{
-			if (Nodes.Add(Node) && Node.Rank > 0)
+			if (nodes.Add(node) && node.Rank > 0)
 			{
-				foreach (NodeInfo Reference in Node.State.References!)
+				foreach (NodeInfo reference in node.State.References!)
 				{
-					if (Reference.State.Blob == null)
+					if (reference.State.Blob == null)
 					{
-						FindModifiedLiveSet(Reference, Nodes);
+						FindModifiedLiveSet(reference, nodes);
 					}
 				}
 			}
 		}
 
-		async Task FindLiveNodesAsync(NodeInfo Node, List<NodeInfo> LiveNodes, HashSet<NodeInfo> LiveNodeSet)
+		async Task FindLiveNodesAsync(NodeInfo node, List<NodeInfo> liveNodes, HashSet<NodeInfo> liveNodeSet)
 		{
-			if (LiveNodeSet.Add(Node))
+			if (liveNodeSet.Add(node))
 			{
-				LiveNodes.Add(Node);
-				if (Node.Rank > 0)
+				liveNodes.Add(node);
+				if (node.Rank > 0)
 				{
-					if (Node.State.Blob != null)
+					if (node.State.Blob != null)
 					{
-						await MountBlobAsync(Node.State.Blob);
+						await MountBlobAsync(node.State.Blob);
 					}
-					foreach (NodeInfo Reference in Node.State.References!)
+					foreach (NodeInfo reference in node.State.References!)
 					{
-						await FindLiveNodesAsync(Reference, LiveNodes, LiveNodeSet);
+						await FindLiveNodesAsync(reference, liveNodes, liveNodeSet);
 					}
 				}
 			}
 		}
 
-		async Task WriteRefAsync(IReadOnlyList<NodeInfo> Nodes, BucketId BucketId, RefId RefId, CbObject Metadata)
+		async Task WriteRefAsync(IReadOnlyList<NodeInfo> nodes, BucketId bucketId, RefId refId, CbObject metadata)
 		{
-			foreach (NodeInfo Node in Nodes)
+			foreach (NodeInfo node in nodes)
 			{
-				ReadOnlySequence<byte> NodeData = await GetDataAsync(Node);
-				Node.SetStandalone(NodeData, Node.State.References!);
+				ReadOnlySequence<byte> nodeData = await GetDataAsync(node);
+				node.SetStandalone(nodeData, node.State.References!);
 			}
 
-			BundleRoot Ref = new BundleRoot();
-			Ref.Metadata = Metadata;
-			Ref.Object = await CreateObjectAsync(Nodes);
+			BundleRoot rootRef = new BundleRoot();
+			rootRef.Metadata = metadata;
+			rootRef.Object = await CreateObjectAsync(nodes);
 
-			ReadOnlyMemory<byte> RefData = EncodeObject(Ref);
-			NextStats.NewRefCount++;
-			NextStats.NewRefBytes += RefData.Length;
+			ReadOnlyMemory<byte> rootRefData = EncodeObject(rootRef);
+			_nextStats.NewRefCount++;
+			_nextStats.NewRefBytes += rootRefData.Length;
 
-			await StorageClient.SetRefAsync(NamespaceId, BucketId, RefId, new CbField(RefData));
+			await _storageClient.SetRefAsync(_namespaceId, bucketId, refId, new CbField(rootRefData));
 		}
 
-		ReadOnlyMemory<byte> EncodeObject<T>(T Object)
+		ReadOnlyMemory<byte> EncodeObject<T>(T data)
 		{
-			SharedWriter.Clear();
-			CbSerializer.Serialize<T>(SharedWriter, Object);
+			_sharedWriter.Clear();
+			CbSerializer.Serialize<T>(_sharedWriter, data);
 
-			int Size = SharedWriter.GetSize();
-			CreateFreeSpace(ref SharedBlobBuffer, 0, Math.Max(Size, Options.MaxBlobSize));
+			int size = _sharedWriter.GetSize();
+			CreateFreeSpace(ref _sharedBlobBuffer, 0, Math.Max(size, Options.MaxBlobSize));
 
-			SharedWriter.CopyTo(SharedBlobBuffer.AsSpan(0, Size));
-			return SharedBlobBuffer.AsMemory(0, Size);
+			_sharedWriter.CopyTo(_sharedBlobBuffer.AsSpan(0, size));
+			return _sharedBlobBuffer.AsMemory(0, size);
 		}
 
-		async Task WriteObjectAsync(IReadOnlyList<NodeInfo> Nodes)
+		async Task WriteObjectAsync(IReadOnlyList<NodeInfo> nodes)
 		{
-			BundleObject Object = await CreateObjectAsync(Nodes);
+			BundleObject writeObject = await CreateObjectAsync(nodes);
 
-			ReadOnlyMemory<byte> Data = EncodeObject(Object);
-			NextStats.NewBlobCount++;
-			NextStats.NewBlobBytes += Data.Length;
+			ReadOnlyMemory<byte> writeData = EncodeObject(writeObject);
+			_nextStats.NewBlobCount++;
+			_nextStats.NewBlobBytes += writeData.Length;
 
-			IoHash Hash = await StorageClient.WriteBlobFromMemoryAsync(NamespaceId, Data);
+			IoHash hash = await _storageClient.WriteBlobFromMemoryAsync(_namespaceId, writeData);
 
-			BlobInfo Blob = new BlobInfo(Hash, Object.Exports.Sum(x => x.Length));
-			for (int Idx = 0; Idx < Object.Exports.Count; Idx++)
+			BlobInfo blob = new BlobInfo(hash, writeObject.Exports.Sum(x => x.Length));
+			for (int idx = 0; idx < writeObject.Exports.Count; idx++)
 			{
-				NodeInfo Node = Nodes[Idx];
-				Node.State = NodeState.Exported(Blob, Object.Exports[Idx], Node.State.References!);
+				NodeInfo node = nodes[idx];
+				node.State = NodeState.Exported(blob, writeObject.Exports[idx], node.State.References!);
 			}
 
-			Blob.Mounted = true;
+			blob.Mounted = true;
 		}
 
 		/// <summary>
 		/// Creates a BundleObject instance. 
 		///
-		/// WARNING: The <see cref="BundleObject.Data"/> member will be set to the value of <see cref="EncodedBuffer"/> on return, which will be
+		/// WARNING: The <see cref="BundleObject.Data"/> member will be set to the value of <see cref="_encodedBuffer"/> on return, which will be
 		/// reused on a subsequent call. If the object needs to be persisted across the lifetime of other objects, this field must be duplicated.
 		/// </summary>
-		async Task<BundleObject> CreateObjectAsync(IReadOnlyList<NodeInfo> Nodes)
+		async Task<BundleObject> CreateObjectAsync(IReadOnlyList<NodeInfo> nodes)
 		{
 			// Preallocate data in the encoded buffer to reduce fragmentation if we have to resize
-			CreateFreeSpace(ref EncodedBuffer, 0, Options.MaxBlobSize);
+			CreateFreeSpace(ref _encodedBuffer, 0, Options.MaxBlobSize);
 
 			// Find node indices for all the export
-			Dictionary<NodeInfo, int> NodeToIndex = new Dictionary<NodeInfo, int>();
-			foreach (NodeInfo Node in Nodes)
+			Dictionary<NodeInfo, int> nodeToIndex = new Dictionary<NodeInfo, int>();
+			foreach (NodeInfo node in nodes)
 			{
-				NodeToIndex[Node] = NodeToIndex.Count;
+				nodeToIndex[node] = nodeToIndex.Count;
 			}
 
 			// Create the new object
-			BundleObject Object = new BundleObject();
+			BundleObject newObject = new BundleObject();
 
 			// Find all the imported nodes
-			HashSet<NodeInfo> ImportedNodes = new HashSet<NodeInfo>();
-			foreach (NodeInfo Node in Nodes)
+			HashSet<NodeInfo> importedNodes = new HashSet<NodeInfo>();
+			foreach (NodeInfo node in nodes)
 			{
-				ImportedNodes.UnionWith(Node.State.References!);
+				importedNodes.UnionWith(node.State.References!);
 			}
-			ImportedNodes.ExceptWith(Nodes);
+			importedNodes.ExceptWith(nodes);
 
 			// Find all the imports
-			foreach (IGrouping<BlobInfo, NodeInfo> ImportGroup in ImportedNodes.GroupBy(x => x.State.Blob!))
+			foreach (IGrouping<BlobInfo, NodeInfo> importGroup in importedNodes.GroupBy(x => x.State.Blob!))
 			{
-				BlobInfo Blob = ImportGroup.Key;
+				BlobInfo blob = importGroup.Key;
 
-				NodeInfo[] ImportNodes = ImportGroup.ToArray();
-				foreach (NodeInfo ImportNode in ImportNodes)
+				NodeInfo[] importNodes = importGroup.ToArray();
+				foreach (NodeInfo importNode in importNodes)
 				{
-					NodeToIndex[ImportNode] = NodeToIndex.Count;
+					nodeToIndex[importNode] = nodeToIndex.Count;
 				}
 
-				List<BundleImport> Imports = ImportGroup.Select(x => new BundleImport(x.Hash, x.Rank, x.Length)).ToList();
-				BundleImportObject ImportObject = new BundleImportObject(new CbObjectAttachment(Blob.Hash), Blob.TotalCost, Imports);
-				Object.ImportObjects.Add(ImportObject);
+				List<BundleImport> imports = importGroup.Select(x => new BundleImport(x.Hash, x.Rank, x.Length)).ToList();
+				BundleImportObject importObject = new BundleImportObject(new CbObjectAttachment(blob.Hash), blob.TotalCost, imports);
+				newObject.ImportObjects.Add(importObject);
 			}
 
 			// Size of data currently stored in the block buffer
-			int BlockSize = 0;
+			int blockSize = 0;
 
 			// Compress all the nodes into the encoded buffer
-			BundleCompressionPacket Packet = new BundleCompressionPacket(0);
-			foreach (NodeInfo Node in Nodes)
+			BundleCompressionPacket packet = new BundleCompressionPacket(0);
+			foreach (NodeInfo node in nodes)
 			{
-				ReadOnlySequence<byte> NodeData = await GetDataAsync(Node);
+				ReadOnlySequence<byte> nodeData = await GetDataAsync(node);
 
 				// If we can't fit this data into the current block, flush the contents of it first
-				if (BlockSize > 0 && BlockSize + NodeData.Length > Options.MinCompressionPacketSize)
+				if (blockSize > 0 && blockSize + nodeData.Length > Options.MinCompressionPacketSize)
 				{
-					FlushPacket(BlockBuffer.AsMemory(0, BlockSize), Packet);
-					Packet = new BundleCompressionPacket(Packet.Offset + Packet.EncodedLength);
-					BlockSize = 0;
+					FlushPacket(_blockBuffer.AsMemory(0, blockSize), packet);
+					packet = new BundleCompressionPacket(packet.Offset + packet.EncodedLength);
+					blockSize = 0;
 				}
 
 				// Create the export for this node
-				int[] References = Node.State.References.Select(x => NodeToIndex[x]).ToArray();
-				BundleExport Export = new BundleExport(Node.Hash, Node.Rank, Packet, BlockSize, Node.Length, References);
-				Object.Exports.Add(Export);
+				int[] references = node.State.References.Select(x => nodeToIndex[x]).ToArray();
+				BundleExport export = new BundleExport(node.Hash, node.Rank, packet, blockSize, node.Length, references);
+				newObject.Exports.Add(export);
 
 				// Write out the new block
-				int Offset = Packet.EncodedLength;
-				if (NodeData.Length < Options.MinCompressionPacketSize || !NodeData.IsSingleSegment)
+				int offset = packet.EncodedLength;
+				if (nodeData.Length < Options.MinCompressionPacketSize || !nodeData.IsSingleSegment)
 				{
-					int RequiredSize = Math.Max(BlockSize + (int)NodeData.Length, (int)(Options.MaxBlobSize * 1.2));
-					CreateFreeSpace(ref BlockBuffer, BlockSize, RequiredSize);
-					foreach (ReadOnlyMemory<byte> NodeSegment in NodeData)
+					int requiredSize = Math.Max(blockSize + (int)nodeData.Length, (int)(Options.MaxBlobSize * 1.2));
+					CreateFreeSpace(ref _blockBuffer, blockSize, requiredSize);
+					foreach (ReadOnlyMemory<byte> nodeSegment in nodeData)
 					{
-						NodeSegment.CopyTo(BlockBuffer.AsMemory(BlockSize));
-						BlockSize += NodeSegment.Length;
+						nodeSegment.CopyTo(_blockBuffer.AsMemory(blockSize));
+						blockSize += nodeSegment.Length;
 					}
 				}
 				else
 				{
-					FlushPacket(NodeData.First, Packet);
-					Packet = new BundleCompressionPacket(Packet.Offset + Packet.EncodedLength);
+					FlushPacket(nodeData.First, packet);
+					packet = new BundleCompressionPacket(packet.Offset + packet.EncodedLength);
 				}
 			}
-			FlushPacket(BlockBuffer.AsMemory(0, BlockSize), Packet);
+			FlushPacket(_blockBuffer.AsMemory(0, blockSize), packet);
 
 			// Flush the data
-			Object.Data = EncodedBuffer.AsMemory(0, Packet.Offset + Packet.EncodedLength);
-			return Object;
+			newObject.Data = _encodedBuffer.AsMemory(0, packet.Offset + packet.EncodedLength);
+			return newObject;
 		}
 
-		void FlushPacket(ReadOnlyMemory<byte> InputData, BundleCompressionPacket Packet)
+		void FlushPacket(ReadOnlyMemory<byte> inputData, BundleCompressionPacket packet)
 		{
-			if (InputData.Length > 0)
+			if (inputData.Length > 0)
 			{
-				int MinFreeSpace = LZ4Codec.MaximumOutputSize(InputData.Length);
-				CreateFreeSpace(ref EncodedBuffer, Packet.Offset, Packet.Offset + MinFreeSpace);
+				int minFreeSpace = LZ4Codec.MaximumOutputSize(inputData.Length);
+				CreateFreeSpace(ref _encodedBuffer, packet.Offset, packet.Offset + minFreeSpace);
 
-				ReadOnlySpan<byte> InputSpan = InputData.Span;
-				Span<byte> OutputSpan = EncodedBuffer.AsSpan(Packet.Offset);
+				ReadOnlySpan<byte> inputSpan = inputData.Span;
+				Span<byte> outputSpan = _encodedBuffer.AsSpan(packet.Offset);
 
-				Packet.DecodedLength = InputData.Length;
-				Packet.EncodedLength = LZ4Codec.Encode(InputSpan, OutputSpan);
+				packet.DecodedLength = inputData.Length;
+				packet.EncodedLength = LZ4Codec.Encode(inputSpan, outputSpan);
 
-				Debug.Assert(Packet.EncodedLength >= 0);
+				Debug.Assert(packet.EncodedLength >= 0);
 			}
 		}
 
-		void CreateFreeSpace(ref byte[] Buffer, int UsedSize, int RequiredSize)
+		void CreateFreeSpace(ref byte[] buffer, int usedSize, int requiredSize)
 		{
-			if (RequiredSize > Buffer.Length)
+			if (requiredSize > buffer.Length)
 			{
-				byte[] NewBuffer = new byte[RequiredSize];
-				Buffer.AsSpan(0, UsedSize).CopyTo(NewBuffer);
-				Buffer = NewBuffer;
+				byte[] newBuffer = new byte[requiredSize];
+				buffer.AsSpan(0, usedSize).CopyTo(newBuffer);
+				buffer = newBuffer;
 			}
 		}
 	}
@@ -1098,8 +1088,8 @@ namespace EpicGames.Horde.Bundles
 		/// <summary>
 		/// Constructor
 		/// </summary>
-		internal Bundle(IStorageClient StorageClient, NamespaceId NamespaceId, T Root, BundleOptions Options, IMemoryCache Cache)
-			: base(StorageClient, NamespaceId, Root, Options, Cache)
+		internal Bundle(IStorageClient storageClient, NamespaceId namespaceId, T root, BundleOptions options, IMemoryCache cache)
+			: base(storageClient, namespaceId, root, options, cache)
 		{
 		}
 	}
diff --git a/Engine/Source/Programs/Shared/EpicGames.Horde/Bundles/BundleNode.cs b/Engine/Source/Programs/Shared/EpicGames.Horde/Bundles/BundleNode.cs
--- a/Engine/Source/Programs/Shared/EpicGames.Horde/Bundles/BundleNode.cs
+++ b/Engine/Source/Programs/Shared/EpicGames.Horde/Bundles/BundleNode.cs
@@ -6,8 +6,6 @@ using System.Buffers;
 using System.Collections.Generic;
 using System.Diagnostics;
 using System.Reflection;
-using System.Text;
-using System.Threading.Tasks;
 
 namespace EpicGames.Horde.Bundles.Nodes
 {
@@ -22,12 +20,12 @@ namespace EpicGames.Horde.Bundles.Nodes
 
 			static BundleNodeDeserializer<T> CreateInstance()
 			{
-				BundleNodeDeserializerAttribute? Attribute = typeof(T).GetCustomAttribute<BundleNodeDeserializerAttribute>();
-				if (Attribute == null)
+				BundleNodeDeserializerAttribute? attribute = typeof(T).GetCustomAttribute<BundleNodeDeserializerAttribute>();
+				if (attribute == null)
 				{
 					throw new InvalidOperationException($"No serializer is defined for {typeof(T).Name}");
 				}
-				return (BundleNodeDeserializer<T>)Activator.CreateInstance(Attribute.Type)!;
+				return (BundleNodeDeserializer<T>)Activator.CreateInstance(attribute.Type)!;
 			}
 		}
 
@@ -57,17 +55,17 @@ namespace EpicGames.Horde.Bundles.Nodes
 		/// Deserialize a node of a particular type, using the default serializer defined by an <see cref="BundleNodeDeserializerAttribute"/>.
 		/// </summary>
 		/// <typeparam name="T">The type to deserialize</typeparam>
-		/// <param name="Memory">Data to deserialize from</param>
+		/// <param name="memory">Data to deserialize from</param>
 		/// <returns>New instance of the node</returns>
-		public static T Deserialize<T>(ReadOnlyMemory<byte> Memory) where T : BundleNode => DefaultSerializer<T>.Instance.Deserialize(Memory);
+		public static T Deserialize<T>(ReadOnlyMemory<byte> memory) where T : BundleNode => DefaultSerializer<T>.Instance.Deserialize(memory);
 
 		/// <summary>
 		/// Deserialize a node of a particular type, using the default serializer defined by an <see cref="BundleNodeDeserializerAttribute"/>.
 		/// </summary>
 		/// <typeparam name="T">The type to deserialize</typeparam>
-		/// <param name="Sequence">Data to deserialize from</param>
+		/// <param name="sequence">Data to deserialize from</param>
 		/// <returns>New instance of the node</returns>
-		public static T Deserialize<T>(ReadOnlySequence<byte> Sequence) where T : BundleNode => DefaultSerializer<T>.Instance.Deserialize(Sequence);
+		public static T Deserialize<T>(ReadOnlySequence<byte> sequence) where T : BundleNode => DefaultSerializer<T>.Instance.Deserialize(sequence);
 
 		/// <summary>
 		/// Enumerates all the child references from this node
@@ -84,22 +82,22 @@ namespace EpicGames.Horde.Bundles.Nodes
 		/// <summary>
 		/// Cached reference to the parent node
 		/// </summary>
-		internal BundleNode Owner;
+		internal BundleNode _owner;
 
 		/// <summary>
 		/// Strong reference to the current node
 		/// </summary>
-		internal BundleNode? StrongRef;
+		internal BundleNode? _strongRef;
 
 		/// <summary>
 		/// Weak reference to the child node. Maintained after the object has been flushed.
 		/// </summary>
-		private WeakReference<BundleNode>? WeakRef { get; set; }
+		private WeakReference<BundleNode>? _weakRef;
 
 		/// <summary>
 		/// Last time that the node was modified
 		/// </summary>
-		internal long LastModifiedTime;
+		internal long _lastModifiedTime;
 
 		/// <summary>
 		/// Hash of the node. May be set to <see cref="IoHash.Zero"/> if the node in memory has been modified.
@@ -111,42 +109,46 @@ namespace EpicGames.Horde.Bundles.Nodes
 		/// </summary>
 		public BundleNode? Node
 		{
-			get => StrongRef;
-			set { StrongRef = value; WeakRef = null; }
+			get => _strongRef;
+			set
+			{ 
+				_strongRef = value; 
+				_weakRef = null; 
+			}
 		}
 
 		/// <summary>
 		/// Creates a reference to a node with the given hash
 		/// </summary>
-		/// <param name="Owner">The node which owns the reference</param>
-		/// <param name="Hash">Hash of the referenced node</param>
-		public BundleNodeRef(BundleNode Owner, IoHash Hash)
+		/// <param name="owner">The node which owns the reference</param>
+		/// <param name="hash">Hash of the referenced node</param>
+		public BundleNodeRef(BundleNode owner, IoHash hash)
 		{
-			this.Owner = Owner;
-			this.Hash = Hash;
+			_owner = owner;
+			Hash = hash;
 		}
 
 		/// <summary>
 		/// Creates a reference to the given node
 		/// </summary>
-		/// <param name="Owner">The node which owns the reference</param>
-		/// <param name="Node">The referenced node</param>
-		public BundleNodeRef(BundleNode Owner, BundleNode Node)
+		/// <param name="owner">The node which owns the reference</param>
+		/// <param name="node">The referenced node</param>
+		public BundleNodeRef(BundleNode owner, BundleNode node)
 		{
-			this.Owner = Owner;
-			this.StrongRef = Node;
-			this.LastModifiedTime = Stopwatch.GetTimestamp();
+			_owner = owner;
+			_strongRef = node;
+			_lastModifiedTime = Stopwatch.GetTimestamp();
 
-			Node.IncomingRef = this;
+			node.IncomingRef = this;
 		}
 
 		/// <summary>
 		/// Reparent this reference to a new owner.
 		/// </summary>
-		public void Reparent(BundleNode NewOwner)
+		public void Reparent(BundleNode newOwner)
 		{
 			MarkAsDirty();
-			Owner = NewOwner;
+			_owner = newOwner;
 			MarkAsDirty();
 		}
 
@@ -154,7 +156,7 @@ namespace EpicGames.Horde.Bundles.Nodes
 		/// Figure out whether this ref is dirty
 		/// </summary>
 		/// <returns></returns>
-		public bool IsDirty() => LastModifiedTime != 0;
+		public bool IsDirty() => _lastModifiedTime != 0;
 
 		/// <summary>
 		/// Converts the node in this reference from a weak to strong reference.
@@ -162,13 +164,13 @@ namespace EpicGames.Horde.Bundles.Nodes
 		/// <returns>True if the ref contains a valid node on return</returns>
 		internal bool MakeStrongRef()
 		{
-			if (StrongRef != null)
+			if (_strongRef != null)
 			{
 				return true;
 			}
-			if (WeakRef != null && WeakRef.TryGetTarget(out StrongRef))
+			if (_weakRef != null && _weakRef.TryGetTarget(out _strongRef))
 			{
-				WeakRef = null;
+				_weakRef = null;
 				return true;
 			}
 
@@ -180,37 +182,37 @@ namespace EpicGames.Horde.Bundles.Nodes
 		/// </summary>
 		internal void MarkAsDirty()
 		{
-			if (LastModifiedTime == 0)
+			if (_lastModifiedTime == 0)
 			{
-				this.LastModifiedTime = Stopwatch.GetTimestamp();
-				if (StrongRef == null)
+				_lastModifiedTime = Stopwatch.GetTimestamp();
+				if (_strongRef == null)
 				{
-					if (WeakRef == null || !WeakRef.TryGetTarget(out StrongRef))
+					if (_weakRef == null || !_weakRef.TryGetTarget(out _strongRef))
 					{
 						throw new InvalidOperationException("Unable to resolve weak reference to node");
 					}
-					this.WeakRef = null;
+					_weakRef = null;
 				}
-				Owner.IncomingRef?.MarkAsDirty();
+				_owner.IncomingRef?.MarkAsDirty();
 			}
 		}
 
 		/// <summary>
 		/// Marks this reference as clean
 		/// </summary>
-		internal void MarkAsClean(IoHash Hash)
+		internal void MarkAsClean(IoHash hash)
 		{
-			this.Hash = Hash;
+			Hash = hash;
 
-			if (StrongRef != null)
+			if (_strongRef != null)
 			{
 				OnCollapse();
 
-				this.WeakRef = new WeakReference<BundleNode>(StrongRef);
-				this.StrongRef = null;
+				_weakRef = new WeakReference<BundleNode>(_strongRef);
+				_strongRef = null;
 			}
 
-			this.LastModifiedTime = 0;
+			_lastModifiedTime = 0;
 		}
 
 		/// <summary>
@@ -237,18 +239,18 @@ namespace EpicGames.Horde.Bundles.Nodes
 		/// <summary>
 		/// Constructor
 		/// </summary>
-		/// <param name="Owner">The node which owns the reference</param>
-		/// <param name="Hash">Hash of the referenced node</param>
-		public BundleNodeRef(BundleNode Owner, IoHash Hash) : base(Owner, Hash)
+		/// <param name="owner">The node which owns the reference</param>
+		/// <param name="hash">Hash of the referenced node</param>
+		public BundleNodeRef(BundleNode owner, IoHash hash) : base(owner, hash)
 		{
 		}
 
 		/// <summary>
 		/// Constructor
 		/// </summary>
-		/// <param name="Owner">The node which owns the reference</param>
-		/// <param name="Node">The referenced node</param>
-		public BundleNodeRef(BundleNode Owner, T Node) : base(Owner, Node)
+		/// <param name="owner">The node which owns the reference</param>
+		/// <param name="node">The referenced node</param>
+		public BundleNodeRef(BundleNode owner, T node) : base(owner, node)
 		{
 		}
 	}
@@ -266,7 +268,7 @@ namespace EpicGames.Horde.Bundles.Nodes
 		/// <summary>
 		/// Constructor
 		/// </summary>
-		public BundleNodeDeserializerAttribute(Type Type) => this.Type = Type;
+		public BundleNodeDeserializerAttribute(Type type) => Type = type;
 	}
 
 	/// <summary>
@@ -278,27 +280,27 @@ namespace EpicGames.Horde.Bundles.Nodes
 		/// <summary>
 		/// Deserializes data from the given bundle
 		/// </summary>
-		/// <param name="Data">Data to deserialize</param>
+		/// <param name="data">Data to deserialize</param>
 		/// <returns>New node parsed from the data</returns>
-		public abstract T Deserialize(ReadOnlyMemory<byte> Data);
+		public abstract T Deserialize(ReadOnlyMemory<byte> data);
 
 		/// <summary>
 		/// Deserializes data from the given bundle
 		/// </summary>
-		/// <param name="Data">Data to deserialize</param>
+		/// <param name="data">Data to deserialize</param>
 		/// <returns>New node parsed from the data</returns>
-		public T Deserialize(ReadOnlySequence<byte> Data)
+		public T Deserialize(ReadOnlySequence<byte> data)
 		{
-			ReadOnlyMemory<byte> Memory;
-			if (Data.IsSingleSegment)
+			ReadOnlyMemory<byte> memory;
+			if (data.IsSingleSegment)
 			{
-				Memory = Data.First;
+				memory = data.First;
 			}
 			else
 			{
-				Memory = Data.ToArray();
+				memory = data.ToArray();
 			}
-			return Deserialize(Memory);
+			return Deserialize(memory);
 		}
 	}
 }
diff --git a/Engine/Source/Programs/Shared/EpicGames.Horde/Bundles/BundleObject.cs b/Engine/Source/Programs/Shared/EpicGames.Horde/Bundles/BundleObject.cs
--- a/Engine/Source/Programs/Shared/EpicGames.Horde/Bundles/BundleObject.cs
+++ b/Engine/Source/Programs/Shared/EpicGames.Horde/Bundles/BundleObject.cs
@@ -1,23 +1,11 @@
 // Copyright Epic Games, Inc. All Rights Reserved.
 
 using EpicGames.Core;
-using EpicGames.Horde.Bundles.Nodes;
 using EpicGames.Serialization;
-using K4os.Compression.LZ4;
-using Microsoft.Extensions.Logging;
 using System;
-using System.Buffers.Binary;
-using System.Collections;
 using System.Collections.Generic;
 using System.Diagnostics;
-using System.Diagnostics.CodeAnalysis;
-using System.IO;
-using System.IO.Compression;
 using System.Linq;
-using System.Security.Cryptography;
-using System.Text;
-using System.Text.Json;
-using System.Threading.Tasks;
 
 namespace EpicGames.Horde.Bundles
 {
@@ -110,11 +98,11 @@ namespace EpicGames.Horde.Bundles
 		/// <summary>
 		/// Constructor
 		/// </summary>
-		public BundleImportObject(CbObjectAttachment Object, int TotalCost, List<BundleImport> Imports)
+		public BundleImportObject(CbObjectAttachment sourceObject, int totalCost, List<BundleImport> imports)
 		{
-			this.Object = Object;
-			this.TotalCost = TotalCost;
-			this.Imports = Imports;
+			Object = sourceObject;
+			TotalCost = totalCost;
+			Imports = imports;
 		}
 	}
 
@@ -141,62 +129,62 @@ namespace EpicGames.Horde.Bundles
 		/// <summary>
 		/// Constructor
 		/// </summary>
-		public BundleImport(IoHash Hash, int Rank, int Length)
+		public BundleImport(IoHash hash, int rank, int length)
 		{
-			if (Length <= 0)
+			if (length <= 0)
 			{
-				throw new ArgumentException("Length must be greater than zero", nameof(Length));
+				throw new ArgumentException("Length must be greater than zero", nameof(length));
 			}
 
-			this.Hash = Hash;
-			this.Rank = Rank;
-			this.Length = Length;
+			Hash = hash;
+			Rank = rank;
+			Length = length;
 		}
 	}
 
 	class BundleImportListConverter : CbConverterBase<List<BundleImport>>
 	{
-		public override List<BundleImport> Read(CbField Field) => Deserialize(Field.AsBinary().Span);
-		public override void Write(CbWriter Writer, List<BundleImport> Value) => Writer.WriteBinaryArrayValue(Serialize(Value));
-		public override void WriteNamed(CbWriter Writer, Utf8String Name, List<BundleImport> Value) => Writer.WriteBinaryArray(Name, Serialize(Value));
+		public override List<BundleImport> Read(CbField field) => Deserialize(field.AsBinary().Span);
+		public override void Write(CbWriter writer, List<BundleImport> value) => writer.WriteBinaryArrayValue(Serialize(value));
+		public override void WriteNamed(CbWriter writer, Utf8String name, List<BundleImport> value) => writer.WriteBinaryArray(name, Serialize(value));
 
-		static List<BundleImport> Deserialize(ReadOnlySpan<byte> Span)
+		static List<BundleImport> Deserialize(ReadOnlySpan<byte> span)
 		{
-			List<BundleImport> Imports = new List<BundleImport>();
-			while (Span.Length > 0)
+			List<BundleImport> imports = new List<BundleImport>();
+			while (span.Length > 0)
 			{
-				IoHash Hash = new IoHash(Span);
-				Span = Span.Slice(IoHash.NumBytes);
+				IoHash hash = new IoHash(span);
+				span = span.Slice(IoHash.NumBytes);
 
-				int Rank = (int)VarInt.Read(Span, out int RankBytes);
-				Span = Span.Slice(RankBytes);
+				int rank = (int)VarInt.Read(span, out int rankBytes);
+				span = span.Slice(rankBytes);
 
-				int Cost = (int)VarInt.Read(Span, out int CostBytes);
-				Span = Span.Slice(CostBytes);
+				int cost = (int)VarInt.Read(span, out int costBytes);
+				span = span.Slice(costBytes);
 
-				Imports.Add(new BundleImport(Hash, Rank, Cost));
+				imports.Add(new BundleImport(hash, rank, cost));
 			}
-			return Imports;
+			return imports;
 		}
 
-		static byte[] Serialize(List<BundleImport> Imports)
+		static byte[] Serialize(List<BundleImport> imports)
 		{
-			byte[] Data = new byte[Imports.Sum(x => IoHash.NumBytes + VarInt.Measure(x.Rank) + VarInt.Measure(x.Length))];
+			byte[] data = new byte[imports.Sum(x => IoHash.NumBytes + VarInt.Measure(x.Rank) + VarInt.Measure(x.Length))];
 
-			Span<byte> Span = Data;
-			foreach (BundleImport Import in Imports)
+			Span<byte> span = data;
+			foreach (BundleImport import in imports)
 			{
-				Import.Hash.CopyTo(Span);
-				Span = Span.Slice(IoHash.NumBytes);
+				import.Hash.CopyTo(span);
+				span = span.Slice(IoHash.NumBytes);
 
-				int RankBytes = VarInt.Write(Span, Import.Rank);
-				Span = Span.Slice(RankBytes);
+				int rankBytes = VarInt.Write(span, import.Rank);
+				span = span.Slice(rankBytes);
 
-				int CostBytes = VarInt.Write(Span, Import.Length);
-				Span = Span.Slice(CostBytes);
+				int costBytes = VarInt.Write(span, import.Length);
+				span = span.Slice(costBytes);
 			}
 
-			return Data;
+			return data;
 		}
 	}
 
@@ -223,10 +211,10 @@ namespace EpicGames.Horde.Bundles
 		/// <summary>
 		/// Constructor
 		/// </summary>
-		/// <param name="Offset"></param>
-		public BundleCompressionPacket(int Offset)
+		/// <param name="offset"></param>
+		public BundleCompressionPacket(int offset)
 		{
-			this.Offset = Offset;
+			Offset = offset;
 		}
 	}
 
@@ -253,142 +241,142 @@ namespace EpicGames.Horde.Bundles
 		/// <summary>
 		/// Constructor
 		/// </summary>
-		public BundleExport(IoHash Hash, int Rank, BundleCompressionPacket Packet, int Offset, int Length, int[] References)
-			: base(Hash, Rank, Length)
+		public BundleExport(IoHash hash, int rank, BundleCompressionPacket packet, int offset, int length, int[] references)
+			: base(hash, rank, length)
 		{
-			this.Packet = Packet;
-			this.Offset = Offset;
-			this.References = References;
+			Packet = packet;
+			Offset = offset;
+			References = references;
 		}
 	}
 
 	class BundleExportListConverter : CbConverterBase<List<BundleExport>>
 	{
-		public override List<BundleExport> Read(CbField Field) => Deserialize(Field.AsBinary().Span);
-		public override void Write(CbWriter Writer, List<BundleExport> Value) => Writer.WriteBinaryArrayValue(Serialize(Value));
-		public override void WriteNamed(CbWriter Writer, Utf8String Name, List<BundleExport> Value) => Writer.WriteBinaryArray(Name, Serialize(Value));
+		public override List<BundleExport> Read(CbField field) => Deserialize(field.AsBinary().Span);
+		public override void Write(CbWriter writer, List<BundleExport> value) => writer.WriteBinaryArrayValue(Serialize(value));
+		public override void WriteNamed(CbWriter writer, Utf8String name, List<BundleExport> value) => writer.WriteBinaryArray(name, Serialize(value));
 
-		static List<BundleExport> Deserialize(ReadOnlySpan<byte> Span)
+		static List<BundleExport> Deserialize(ReadOnlySpan<byte> span)
 		{
-			List<BundleExport> Exports = new List<BundleExport>();
+			List<BundleExport> exports = new List<BundleExport>();
 
-			int PacketOffset = 0;
-			while (Span.Length > 0)
+			int packetOffset = 0;
+			while (span.Length > 0)
 			{
-				BundleCompressionPacket Packet = new BundleCompressionPacket(PacketOffset);
+				BundleCompressionPacket packet = new BundleCompressionPacket(packetOffset);
 
-				Packet.EncodedLength = (int)VarInt.Read(Span, out int EncodedLengthBytes);
-				Span = Span.Slice(EncodedLengthBytes);
+				packet.EncodedLength = (int)VarInt.Read(span, out int encodedLengthBytes);
+				span = span.Slice(encodedLengthBytes);
 
-				Packet.DecodedLength = (int)VarInt.Read(Span, out int DecodedLengthBytes);
-				Span = Span.Slice(DecodedLengthBytes);
+				packet.DecodedLength = (int)VarInt.Read(span, out int decodedLengthBytes);
+				span = span.Slice(decodedLengthBytes);
 
-				PacketOffset += Packet.EncodedLength;
+				packetOffset += packet.EncodedLength;
 
-				int Offset = 0;
-				while (Offset < Packet.DecodedLength)
+				int offset = 0;
+				while (offset < packet.DecodedLength)
 				{
-					IoHash Hash = new IoHash(Span);
-					Span = Span.Slice(IoHash.NumBytes);
+					IoHash hash = new IoHash(span);
+					span = span.Slice(IoHash.NumBytes);
 
-					int Rank = (int)VarInt.Read(Span, out int RankBytes);
-					Span = Span.Slice(RankBytes);
+					int rank = (int)VarInt.Read(span, out int rankBytes);
+					span = span.Slice(rankBytes);
 
-					int Length = (int)VarInt.Read(Span, out int LengthBytes);
-					Span = Span.Slice(LengthBytes);
+					int length = (int)VarInt.Read(span, out int lengthBytes);
+					span = span.Slice(lengthBytes);
 
-					int NumReferences = (int)VarInt.Read(Span, out int NumReferencesBytes);
-					Span = Span.Slice(NumReferencesBytes);
+					int numReferences = (int)VarInt.Read(span, out int numReferencesBytes);
+					span = span.Slice(numReferencesBytes);
 
-					int[] References = new int[NumReferences];
-					for (int ReferenceIdx = 0; ReferenceIdx < NumReferences; ReferenceIdx++)
+					int[] references = new int[numReferences];
+					for (int referenceIdx = 0; referenceIdx < numReferences; referenceIdx++)
 					{
-						References[ReferenceIdx] = (int)VarInt.Read(Span, out int ReferenceBytes);
-						Span = Span.Slice(ReferenceBytes);
+						references[referenceIdx] = (int)VarInt.Read(span, out int referenceBytes);
+						span = span.Slice(referenceBytes);
 					}
 
-					Exports.Add(new BundleExport(Hash, Rank, Packet, Offset, Length, References));
-					Offset += Length;
+					exports.Add(new BundleExport(hash, rank, packet, offset, length, references));
+					offset += length;
 				}
 			}
 
-			return Exports;
+			return exports;
 		}
 
-		static byte[] Serialize(List<BundleExport> Exports)
+		static byte[] Serialize(List<BundleExport> exports)
 		{
-			int Length = Measure(Exports);
+			int length = Measure(exports);
 
-			byte[] Data = new byte[Length];
-			Span<byte> Span = Data;
+			byte[] data = new byte[length];
+			Span<byte> span = data;
 
-			BundleCompressionPacket? PrevPacket = null;
-			foreach (BundleExport Export in Exports)
+			BundleCompressionPacket? prevPacket = null;
+			foreach (BundleExport export in exports)
 			{
-				Debug.Assert(Export.Length > 0);
+				Debug.Assert(export.Length > 0);
 
-				BundleCompressionPacket Packet = Export.Packet;
-				if (PrevPacket == null || Packet.Offset != PrevPacket.Offset)
+				BundleCompressionPacket packet = export.Packet;
+				if (prevPacket == null || packet.Offset != prevPacket.Offset)
 				{
-					CheckPacketSequence(PrevPacket, Packet);
+					CheckPacketSequence(prevPacket, packet);
 
-					int EncodedLengthBytes = VarInt.Write(Span, Packet.EncodedLength);
-					Span = Span.Slice(EncodedLengthBytes);
+					int encodedLengthBytes = VarInt.Write(span, packet.EncodedLength);
+					span = span.Slice(encodedLengthBytes);
 
-					int DecodedLengthBytes = VarInt.Write(Span, Packet.DecodedLength);
-					Span = Span.Slice(DecodedLengthBytes);
+					int decodedLengthBytes = VarInt.Write(span, packet.DecodedLength);
+					span = span.Slice(decodedLengthBytes);
 
-					PrevPacket = Packet;
+					prevPacket = packet;
 				}
 
-				Export.Hash.CopyTo(Span);
-				Span = Span.Slice(IoHash.NumBytes);
+				export.Hash.CopyTo(span);
+				span = span.Slice(IoHash.NumBytes);
 
-				int RankBytes = VarInt.Write(Span, Export.Rank);
-				Span = Span.Slice(RankBytes);
+				int rankBytes = VarInt.Write(span, export.Rank);
+				span = span.Slice(rankBytes);
 
-				int LengthBytes = VarInt.Write(Span, Export.Length);
-				Span = Span.Slice(LengthBytes);
+				int lengthBytes = VarInt.Write(span, export.Length);
+				span = span.Slice(lengthBytes);
 
-				int NumReferencesBytes = VarInt.Write(Span, Export.References.Length);
-				Span = Span.Slice(NumReferencesBytes);
+				int numReferencesBytes = VarInt.Write(span, export.References.Length);
+				span = span.Slice(numReferencesBytes);
 
-				foreach (int Reference in Export.References)
+				foreach (int reference in export.References)
 				{
-					int ReferenceBytes = VarInt.Write(Span, Reference);
-					Span = Span.Slice(ReferenceBytes);
+					int referenceBytes = VarInt.Write(span, reference);
+					span = span.Slice(referenceBytes);
 				}
 			}
 
-			Debug.Assert(Span.Length == 0);
-			return Data;
+			Debug.Assert(span.Length == 0);
+			return data;
 		}
 
-		static int Measure(List<BundleExport> Exports)
+		static int Measure(List<BundleExport> exports)
 		{
-			int Length = 0;
+			int length = 0;
 
-			BundleCompressionPacket? PrevPacket = null;
-			foreach (BundleExport Export in Exports)
+			BundleCompressionPacket? prevPacket = null;
+			foreach (BundleExport export in exports)
 			{
-				BundleCompressionPacket Packet = Export.Packet;
-				if (PrevPacket == null || Packet.Offset != PrevPacket.Offset)
+				BundleCompressionPacket packet = export.Packet;
+				if (prevPacket == null || packet.Offset != prevPacket.Offset)
 				{
-					CheckPacketSequence(PrevPacket, Packet);
-					Length += VarInt.Measure(Packet.EncodedLength) + VarInt.Measure(Packet.DecodedLength);
-					PrevPacket = Packet;
+					CheckPacketSequence(prevPacket, packet);
+					length += VarInt.Measure(packet.EncodedLength) + VarInt.Measure(packet.DecodedLength);
+					prevPacket = packet;
 				}
 
-				Length += IoHash.NumBytes + VarInt.Measure(Export.Rank) + VarInt.Measure(Export.Length) + VarInt.Measure(Export.References.Length) + Export.References.Sum(x => VarInt.Measure(x));
+				length += IoHash.NumBytes + VarInt.Measure(export.Rank) + VarInt.Measure(export.Length) + VarInt.Measure(export.References.Length) + export.References.Sum(x => VarInt.Measure(x));
 			}
 
-			return Length;
+			return length;
 		}
 
-		static void CheckPacketSequence(BundleCompressionPacket? PrevPacket, BundleCompressionPacket Packet)
+		static void CheckPacketSequence(BundleCompressionPacket? prevPacket, BundleCompressionPacket packet)
 		{
-			int ExpectedOffset = (PrevPacket == null) ? 0 : (PrevPacket.Offset + PrevPacket.EncodedLength);
-			if (Packet.Offset != ExpectedOffset)
+			int expectedOffset = (prevPacket == null) ? 0 : (prevPacket.Offset + prevPacket.EncodedLength);
+			if (packet.Offset != expectedOffset)
 			{
 				throw new InvalidOperationException("Bundle compression packets are not sequential");
 			}
