Horde: Use standard coding style for EpicGames.Horde.

diff --git a/Engine/Source/Programs/Shared/EpicGames.Horde/Bundles/BundleTools.cs b/Engine/Source/Programs/Shared/EpicGames.Horde/Bundles/BundleTools.cs
--- a/Engine/Source/Programs/Shared/EpicGames.Horde/Bundles/BundleTools.cs
+++ b/Engine/Source/Programs/Shared/EpicGames.Horde/Bundles/BundleTools.cs
@@ -1,23 +1,14 @@
 // Copyright Epic Games, Inc. All Rights Reserved.
 
 using EpicGames.Core;
-using EpicGames.Horde.Bundles.Nodes;
 using EpicGames.Horde.Storage;
 using EpicGames.Serialization;
-using K4os.Compression.LZ4;
 using Microsoft.Extensions.Logging;
 using System;
-using System.Buffers.Binary;
-using System.Collections;
 using System.Collections.Generic;
 using System.Diagnostics;
-using System.Diagnostics.CodeAnalysis;
 using System.IO;
-using System.IO.Compression;
 using System.Linq;
-using System.Security.Cryptography;
-using System.Text;
-using System.Text.Json;
 using System.Threading.Tasks;
 
 namespace EpicGames.Horde.Bundles
@@ -36,12 +27,12 @@ namespace EpicGames.Horde.Bundles
 
 			public BundleExport? Export { get; set; }
 
-			public NodeInfo(IoHash Hash, int Rank, int Length, BlobInfo Blob)
+			public NodeInfo(IoHash hash, int rank, int length, BlobInfo blob)
 			{
-				this.Blob = Blob;
-				this.Hash = Hash;
-				this.Rank = Rank;
-				this.Length = Length;
+				Blob = blob;
+				Hash = hash;
+				Rank = rank;
+				Length = length;
 			}
 		}
 
@@ -53,201 +44,201 @@ namespace EpicGames.Horde.Bundles
 			public List<NodeInfo>? References { get; set; }
 			public Dictionary<IoHash, NodeInfo> Nodes { get; } = new Dictionary<IoHash, NodeInfo>();
 
-			public BlobInfo(IoHash Hash, int DataSize)
+			public BlobInfo(IoHash hash, int dataSize)
 			{
-				this.Hash = Hash;
-				this.DataSize = DataSize;
+				Hash = hash;
+				DataSize = dataSize;
 			}
 		}
 
 		/// <summary>
 		/// Finds the delta between two bundles
 		/// </summary>
-		/// <param name="StorageClient"></param>
-		/// <param name="NamespaceId"></param>
-		/// <param name="BucketId"></param>
-		/// <param name="RefId1">First ref to compare</param>
-		/// <param name="RefId2">Second ref to compare</param>
-		/// <param name="Logger">Logger to write output data</param>
-		public static async Task FindDeltaAsync(IStorageClient StorageClient, NamespaceId NamespaceId, BucketId BucketId, RefId RefId1, RefId RefId2, ILogger Logger)
+		/// <param name="storageClient"></param>
+		/// <param name="namespaceId"></param>
+		/// <param name="bucketId"></param>
+		/// <param name="refId1">First ref to compare</param>
+		/// <param name="refId2">Second ref to compare</param>
+		/// <param name="logger">Logger to write output data</param>
+		public static async Task FindDeltaAsync(IStorageClient storageClient, NamespaceId namespaceId, BucketId bucketId, RefId refId1, RefId refId2, ILogger logger)
 		{
-			Dictionary<IoHash, BlobInfo> BlobCache = new Dictionary<IoHash, BlobInfo>();
+			Dictionary<IoHash, BlobInfo> blobCache = new Dictionary<IoHash, BlobInfo>();
 
-			Logger.LogInformation("Reading internal blobs...");
-			HashSet<NodeInfo> Nodes1 = new HashSet<NodeInfo>();
-			BlobInfo RootBlob1 = await ScanTreeAsync(StorageClient, NamespaceId, BucketId, RefId1, BlobCache, Nodes1);
+			logger.LogInformation("Reading internal blobs...");
+			HashSet<NodeInfo> nodes1 = new HashSet<NodeInfo>();
+			BlobInfo rootBlob1 = await ScanTreeAsync(storageClient, namespaceId, bucketId, refId1, blobCache, nodes1);
 
-			HashSet<NodeInfo> Nodes2 = new HashSet<NodeInfo>();
-			BlobInfo RootBlob2 = await ScanTreeAsync(StorageClient, NamespaceId, BucketId, RefId2, BlobCache, Nodes2);
+			HashSet<NodeInfo> nodes2 = new HashSet<NodeInfo>();
+			BlobInfo rootBlob2 = await ScanTreeAsync(storageClient, namespaceId, bucketId, refId2, blobCache, nodes2);
 
-			Logger.LogInformation("Reading leaf blobs...");
-			foreach (BlobInfo Blob in BlobCache.Values)
+			logger.LogInformation("Reading leaf blobs...");
+			foreach (BlobInfo blob in blobCache.Values)
 			{
-				if (Blob.References == null)
+				if (blob.References == null)
 				{
-					ReadOnlyMemory<byte> Data = await StorageClient.ReadBlobToMemoryAsync(NamespaceId, Blob.Hash);
-					MountObject(StorageClient, NamespaceId, Blob, Data, null!);
+					ReadOnlyMemory<byte> data = await storageClient.ReadBlobToMemoryAsync(namespaceId, blob.Hash);
+					MountObject(blob, data, null!);
 				}
 			}
 
-			Logger.LogInformation("");
+			logger.LogInformation("");
 
-			Logger.LogInformation("{Row}", "                                               |            |       Total       |        Used (Ref 1)        |        Used (Ref 2)        ");
-			Logger.LogInformation("{Row}", " Identifier                                    |    RawSize |    Num       Size |    Num         Size    Pct |    Num         Size    Pct ");
-			Logger.LogInformation("{Row}", " ----------------------------------------------|------------|-------------------|----------------------------|----------------------------");
+			logger.LogInformation("{Row}", "                                               |            |       Total       |        Used (Ref 1)        |        Used (Ref 2)        ");
+			logger.LogInformation("{Row}", " Identifier                                    |    RawSize |    Num       Size |    Num         Size    Pct |    Num         Size    Pct ");
+			logger.LogInformation("{Row}", " ----------------------------------------------|------------|-------------------|----------------------------|----------------------------");
 
-			WriteBlobInfo($"Ref1 {RefId1.Hash}", RootBlob1, Nodes1, Nodes2, Logger);
-			WriteBlobInfo($"Ref2 {RefId2.Hash}", RootBlob2, Nodes1, Nodes2, Logger);
+			WriteBlobInfo($"Ref1 {refId1.Hash}", rootBlob1, nodes1, nodes2, logger);
+			WriteBlobInfo($"Ref2 {refId2.Hash}", rootBlob2, nodes1, nodes2, logger);
 
-			foreach (BlobInfo Blob in BlobCache.Values)
+			foreach (BlobInfo blob in blobCache.Values)
 			{
-				WriteBlobInfo($"Blob {Blob.Hash}", Blob, Nodes1, Nodes2, Logger);
+				WriteBlobInfo($"Blob {blob.Hash}", blob, nodes1, nodes2, logger);
 			}
 
-			Logger.LogInformation("");
+			logger.LogInformation("");
 
-			List<BlobInfo> NewBlobs = Nodes2.Select(x => x.Blob).Distinct().Except(Nodes2.Select(x => x.Blob)).ToList();
-			Logger.LogInformation("Total blob download size: {NumBytes:n0} bytes", NewBlobs.Sum(x => x.Size));
+			List<BlobInfo> newBlobs = nodes2.Select(x => x.Blob).Distinct().Except(nodes2.Select(x => x.Blob)).ToList();
+			logger.LogInformation("Total blob download size: {NumBytes:n0} bytes", newBlobs.Sum(x => x.Size));
 		}
 
-		static void WriteBlobInfo(string Name, BlobInfo Blob, HashSet<NodeInfo> Nodes1, HashSet<NodeInfo> Nodes2, ILogger Logger)
+		static void WriteBlobInfo(string name, BlobInfo blob, HashSet<NodeInfo> nodes1, HashSet<NodeInfo> nodes2, ILogger logger)
 		{
-			int UsedCount1 = 0;
-			int UsedSize1 = 0;
+			int usedCount1 = 0;
+			int usedSize1 = 0;
 
-			int UsedCount2 = 0;
-			int UsedSize2 = 0;
+			int usedCount2 = 0;
+			int usedSize2 = 0;
 
-			foreach (NodeInfo Node in Blob.Nodes.Values)
+			foreach (NodeInfo node in blob.Nodes.Values)
 			{
-				if (Nodes1.Contains(Node))
+				if (nodes1.Contains(node))
 				{
-					UsedCount1++;
-					UsedSize1 += Node.Length;
+					usedCount1++;
+					usedSize1 += node.Length;
 				}
-				if (Nodes2.Contains(Node))
+				if (nodes2.Contains(node))
 				{
-					UsedCount2++;
-					UsedSize2 += Node.Length;
+					usedCount2++;
+					usedSize2 += node.Length;
 				}
 			}
 
-			string UsedPct1 = $"{(UsedSize1 * 100) / Blob.DataSize}%";
-			string UsedPct2 = $"{(UsedSize2 * 100) / Blob.DataSize}%";
+			string usedPct1 = $"{(usedSize1 * 100) / blob.DataSize}%";
+			string usedPct2 = $"{(usedSize2 * 100) / blob.DataSize}%";
 
-			string Row = $"{Name,42} | {Blob.Size,10:n0} | {Blob.Nodes.Count,6:n0} {Blob.DataSize,10:n0} | {UsedCount1,6:n0} {UsedSize1,12:n0} {UsedPct1,6} | {UsedCount2,6:n0} {UsedSize2,12:n0} {UsedPct2,6}";
-			Logger.LogInformation(" {Row}", Row);
+			string row = $"{name,42} | {blob.Size,10:n0} | {blob.Nodes.Count,6:n0} {blob.DataSize,10:n0} | {usedCount1,6:n0} {usedSize1,12:n0} {usedPct1,6} | {usedCount2,6:n0} {usedSize2,12:n0} {usedPct2,6}";
+			logger.LogInformation(" {Row}", row);
 		}
 
-		static async Task<BlobInfo> ScanTreeAsync(IStorageClient StorageClient, NamespaceId NamespaceId, BucketId BucketId, RefId RefId, Dictionary<IoHash, BlobInfo> BlobCache, HashSet<NodeInfo> Nodes)
+		static async Task<BlobInfo> ScanTreeAsync(IStorageClient storageClient, NamespaceId namespaceId, BucketId bucketId, RefId refId, Dictionary<IoHash, BlobInfo> blobCache, HashSet<NodeInfo> nodes)
 		{
-			IRef Ref = await StorageClient.GetRefAsync(NamespaceId, BucketId, RefId);
+			IRef rootRef = await storageClient.GetRefAsync(namespaceId, bucketId, refId);
 
-			BlobInfo RootBlob = new BlobInfo(IoHash.Zero, 0);
-			BundleObject Object = MountObject(StorageClient, NamespaceId, RootBlob, Ref.Value.GetView(), BlobCache);
+			BlobInfo rootBlob = new BlobInfo(IoHash.Zero, 0);
+			BundleObject rootObject = MountObject(rootBlob, rootRef.Value.GetView(), blobCache);
 
-			NodeInfo RootNode = RootBlob.Nodes[Object.Exports[^1].Hash];
-			await ScanTreeAsync(StorageClient, NamespaceId, RootNode, BlobCache, Nodes);
+			NodeInfo rootNode = rootBlob.Nodes[rootObject.Exports[^1].Hash];
+			await ScanTreeAsync(storageClient, namespaceId, rootNode, blobCache, nodes);
 
-			return RootBlob;
+			return rootBlob;
 		}
 
-		static async Task ScanTreeAsync(IStorageClient StorageClient, NamespaceId NamespaceId, NodeInfo Node, Dictionary<IoHash, BlobInfo> BlobCache, HashSet<NodeInfo> Nodes)
+		static async Task ScanTreeAsync(IStorageClient storageClient, NamespaceId namespaceId, NodeInfo node, Dictionary<IoHash, BlobInfo> blobCache, HashSet<NodeInfo> nodes)
 		{
-			if (Nodes.Add(Node) && Node.Rank > 0)
+			if (nodes.Add(node) && node.Rank > 0)
 			{
-				if (Node.Export == null)
+				if (node.Export == null)
 				{
-					ReadOnlyMemory<byte> Data = await StorageClient.ReadBlobToMemoryAsync(NamespaceId, Node.Blob.Hash);
-					MountObject(StorageClient, NamespaceId, Node.Blob, Data, BlobCache);
+					ReadOnlyMemory<byte> data = await storageClient.ReadBlobToMemoryAsync(namespaceId, node.Blob.Hash);
+					MountObject(node.Blob, data, blobCache);
 				}
 
-				Debug.Assert(Node.Export != null);
+				Debug.Assert(node.Export != null);
 
-				foreach (int RefIdx in Node.Export.References)
+				foreach (int refIdx in node.Export.References)
 				{
-					NodeInfo RefNode = Node.Blob.References![RefIdx];
-					await ScanTreeAsync(StorageClient, NamespaceId, RefNode, BlobCache, Nodes);
+					NodeInfo refNode = node.Blob.References![refIdx];
+					await ScanTreeAsync(storageClient, namespaceId, refNode, blobCache, nodes);
 				}
 			}
 		}
 
-		static BundleObject MountObject(IStorageClient StorageClient, NamespaceId NamespaceId, BlobInfo Blob, ReadOnlyMemory<byte> BlobData, Dictionary<IoHash, BlobInfo> BlobCache)
+		static BundleObject MountObject(BlobInfo blob, ReadOnlyMemory<byte> blobData, Dictionary<IoHash, BlobInfo> blobCache)
 		{
-			BundleObject Object = CbSerializer.Deserialize<BundleObject>(BlobData);
+			BundleObject blobObject = CbSerializer.Deserialize<BundleObject>(blobData);
 
-			Blob.Size = BlobData.Length;
-			Blob.DataSize = Object.Data.Length;
-			Blob.References = new List<NodeInfo>();
+			blob.Size = blobData.Length;
+			blob.DataSize = blobObject.Data.Length;
+			blob.References = new List<NodeInfo>();
 
-			foreach (BundleExport Export in Object.Exports)
+			foreach (BundleExport export in blobObject.Exports)
 			{
-				NodeInfo? Node;
-				if (!Blob.Nodes.TryGetValue(Export.Hash, out Node))
+				NodeInfo? node;
+				if (!blob.Nodes.TryGetValue(export.Hash, out node))
 				{
-					Node = new NodeInfo(Export.Hash, Export.Rank, Export.Length, Blob);
-					Blob.Nodes.Add(Export.Hash, Node);
+					node = new NodeInfo(export.Hash, export.Rank, export.Length, blob);
+					blob.Nodes.Add(export.Hash, node);
 				}
-				Blob.References.Add(Node);
+				blob.References.Add(node);
 
-				Node.Export = Export;
+				node.Export = export;
 			}
 
-			foreach (BundleImportObject ImportObject in Object.ImportObjects)
+			foreach (BundleImportObject importObject in blobObject.ImportObjects)
 			{
-				BlobInfo? ImportBlob;
-				if (!BlobCache.TryGetValue(ImportObject.Object.Hash, out ImportBlob))
+				BlobInfo? importBlob;
+				if (!blobCache.TryGetValue(importObject.Object.Hash, out importBlob))
 				{
-					ImportBlob = new BlobInfo(ImportObject.Object.Hash, ImportObject.TotalCost);
-					BlobCache[ImportObject.Object.Hash] = ImportBlob;
+					importBlob = new BlobInfo(importObject.Object.Hash, importObject.TotalCost);
+					blobCache[importObject.Object.Hash] = importBlob;
 				}
-				foreach (BundleImport Import in ImportObject.Imports)
+				foreach (BundleImport import in importObject.Imports)
 				{
-					NodeInfo? Node;
-					if (!ImportBlob.Nodes.TryGetValue(Import.Hash, out Node))
+					NodeInfo? node;
+					if (!importBlob.Nodes.TryGetValue(import.Hash, out node))
 					{
-						Node = new NodeInfo(Import.Hash, Import.Rank, Import.Length, ImportBlob);
-						ImportBlob.Nodes.Add(Import.Hash, Node);
+						node = new NodeInfo(import.Hash, import.Rank, import.Length, importBlob);
+						importBlob.Nodes.Add(import.Hash, node);
 					}
-					Blob.References.Add(Node);
+					blob.References.Add(node);
 				}
 			}
 
-			return Object;
+			return blobObject;
 		}
 
 		/// <summary>
 		/// Creates a YAML summary of the given bundle tree
 		/// </summary>
-		/// <param name="StorageClient"></param>
-		/// <param name="NamespaceId"></param>
-		/// <param name="BucketId"></param>
-		/// <param name="RefId"></param>
-		/// <param name="OutputDir"></param>
-		public static async Task WriteSummaryAsync(IStorageClient StorageClient, NamespaceId NamespaceId, BucketId BucketId, RefId RefId, DirectoryReference OutputDir)
+		/// <param name="storageClient"></param>
+		/// <param name="namespaceId"></param>
+		/// <param name="bucketId"></param>
+		/// <param name="refId"></param>
+		/// <param name="outputDir"></param>
+		public static async Task WriteSummaryAsync(IStorageClient storageClient, NamespaceId namespaceId, BucketId bucketId, RefId refId, DirectoryReference outputDir)
 		{
-			IRef Ref = await StorageClient.GetRefAsync(NamespaceId, BucketId, RefId);
-			BundleObject RootObject = CbSerializer.Deserialize<BundleObject>(Ref.Value);
+			IRef rootRef = await storageClient.GetRefAsync(namespaceId, bucketId, refId);
+			BundleObject rootObject = CbSerializer.Deserialize<BundleObject>(rootRef.Value);
 
-			FileReference RootFile = FileReference.Combine(OutputDir, "root.yml");
-			WriteSummaryForObject(RootObject, RootFile);
+			FileReference rootFile = FileReference.Combine(outputDir, "root.yml");
+			WriteSummaryForObject(rootObject, rootFile);
 
-			await WriteSummaryForImportsAsync(StorageClient, NamespaceId, RootObject.ImportObjects, new HashSet<IoHash>(), OutputDir);
+			await WriteSummaryForImportsAsync(storageClient, namespaceId, rootObject.ImportObjects, new HashSet<IoHash>(), outputDir);
 		}
 
-		static async Task WriteSummaryForImportsAsync(IStorageClient StorageClient, NamespaceId NamespaceId, IEnumerable<BundleImportObject> ImportObjects, HashSet<IoHash> Hashes, DirectoryReference OutputDir)
+		static async Task WriteSummaryForImportsAsync(IStorageClient storageClient, NamespaceId namespaceId, IEnumerable<BundleImportObject> importObjects, HashSet<IoHash> hashes, DirectoryReference outputDir)
 		{
-			foreach (BundleImportObject ImportObject in ImportObjects)
+			foreach (BundleImportObject importObject in importObjects)
 			{
-				IoHash Hash = ImportObject.Object.Hash;
-				if (Hashes.Add(Hash))
+				IoHash hash = importObject.Object.Hash;
+				if (hashes.Add(hash))
 				{
-					BundleObject Object = await StorageClient.ReadBlobAsync<BundleObject>(NamespaceId, Hash);
+					BundleObject blobObject = await storageClient.ReadBlobAsync<BundleObject>(namespaceId, hash);
 
-					FileReference OutputFile = FileReference.Combine(OutputDir, $"{Hash}.yml");
-					WriteSummaryForObject(Object, OutputFile);
+					FileReference outputFile = FileReference.Combine(outputDir, $"{hash}.yml");
+					WriteSummaryForObject(blobObject, outputFile);
 
-					await WriteSummaryForImportsAsync(StorageClient, NamespaceId, Object.ImportObjects, Hashes, OutputDir);
+					await WriteSummaryForImportsAsync(storageClient, namespaceId, blobObject.ImportObjects, hashes, outputDir);
 				}
 			}
 		}
@@ -255,66 +246,66 @@ namespace EpicGames.Horde.Bundles
 		/// <summary>
 		/// Writes a summary of the given bundle object to disk in YAML format
 		/// </summary>
-		/// <param name="Object"></param>
-		/// <param name="File"></param>
-		public static void WriteSummaryForObject(BundleObject Object, FileReference File)
+		/// <param name="blobObject"></param>
+		/// <param name="file"></param>
+		public static void WriteSummaryForObject(BundleObject blobObject, FileReference file)
 		{
-			DirectoryReference.CreateDirectory(File.Directory);
+			DirectoryReference.CreateDirectory(file.Directory);
 
-			using StreamWriter Writer = new StreamWriter(FileReference.Open(File, FileMode.Create));
+			using StreamWriter writer = new StreamWriter(FileReference.Open(file, FileMode.Create));
 
-			Writer.WriteLine("schema: {0}", Object.Schema);
+			writer.WriteLine("schema: {0}", blobObject.Schema);
 
-			List<IoHash> References = new List<IoHash>(Object.Exports.Select(x => x.Hash));
-			if (Object.ImportObjects.Count > 0)
+			List<IoHash> references = new List<IoHash>(blobObject.Exports.Select(x => x.Hash));
+			if (blobObject.ImportObjects.Count > 0)
 			{
-				Writer.WriteLine("import-objects:");
-				foreach (BundleImportObject ImportObject in Object.ImportObjects)
+				writer.WriteLine("import-objects:");
+				foreach (BundleImportObject importObject in blobObject.ImportObjects)
 				{
-					Writer.WriteLine($"- blob: \"{ImportObject.Object.Hash}\"");
-					Writer.WriteLine($"  imports:");
-					foreach (BundleImport Import in ImportObject.Imports)
+					writer.WriteLine($"- blob: \"{importObject.Object.Hash}\"");
+					writer.WriteLine($"  imports:");
+					foreach (BundleImport import in importObject.Imports)
 					{
-						Writer.WriteLine($"  - hash: \"{Import.Hash}\"");
-						Writer.WriteLine($"    rank: \"{Import.Rank}\"");
-						Writer.WriteLine($"    cost: \"{Import.Length}\"");
-						Writer.WriteLine("");
-						References.Add(Import.Hash);
+						writer.WriteLine($"  - hash: \"{import.Hash}\"");
+						writer.WriteLine($"    rank: \"{import.Rank}\"");
+						writer.WriteLine($"    cost: \"{import.Length}\"");
+						writer.WriteLine("");
+						references.Add(import.Hash);
 					}
-					Writer.WriteLine();
+					writer.WriteLine();
 				}
 			}
 
-			if (Object.Exports.Count > 0)
+			if (blobObject.Exports.Count > 0)
 			{
-				Writer.WriteLine("packets:");
-				foreach (IGrouping<int, BundleCompressionPacket> Group in Object.Exports.Select(x => x.Packet).GroupBy(x => x.Offset).OrderBy(x => x.Key))
+				writer.WriteLine("packets:");
+				foreach (IGrouping<int, BundleCompressionPacket> group in blobObject.Exports.Select(x => x.Packet).GroupBy(x => x.Offset).OrderBy(x => x.Key))
 				{
-					BundleCompressionPacket Packet = Group.First();
-					Writer.WriteLine($"- offset: {Packet.Offset}");
-					Writer.WriteLine($"  encodedLength: {Packet.EncodedLength}");
-					Writer.WriteLine($"  decodedLength: {Packet.DecodedLength}");
-					Writer.WriteLine($"  compression-ratio: \"{(Packet.EncodedLength * 100) / Packet.DecodedLength}%\"");
-					Writer.WriteLine($"  exports: {Group.Count()}");
-					Writer.WriteLine();
+					BundleCompressionPacket packet = group.First();
+					writer.WriteLine($"- offset: {packet.Offset}");
+					writer.WriteLine($"  encodedLength: {packet.EncodedLength}");
+					writer.WriteLine($"  decodedLength: {packet.DecodedLength}");
+					writer.WriteLine($"  compression-ratio: \"{(packet.EncodedLength * 100) / packet.DecodedLength}%\"");
+					writer.WriteLine($"  exports: {group.Count()}");
+					writer.WriteLine();
 				}
 
-				Writer.WriteLine("exports:");
-				foreach (BundleExport Export in Object.Exports)
+				writer.WriteLine("exports:");
+				foreach (BundleExport export in blobObject.Exports)
 				{
-					Writer.WriteLine($"- hash: {Export.Hash}");
-					Writer.WriteLine($"  packet-offset: {Export.Packet.Offset}");
-					Writer.WriteLine($"  size: {Export.Length}");
-					if (Export.References.Length > 0)
+					writer.WriteLine($"- hash: {export.Hash}");
+					writer.WriteLine($"  packet-offset: {export.Packet.Offset}");
+					writer.WriteLine($"  size: {export.Length}");
+					if (export.References.Length > 0)
 					{
-						Writer.WriteLine($"  references:");
-						foreach (int Reference in Export.References)
+						writer.WriteLine($"  references:");
+						foreach (int reference in export.References)
 						{
-							IoHash ReferenceHash = References[Reference];
-							Writer.WriteLine($"  - \"{ReferenceHash}\"");
+							IoHash referenceHash = references[reference];
+							writer.WriteLine($"  - \"{referenceHash}\"");
 						}
 					}
-					Writer.WriteLine();
+					writer.WriteLine();
 				}
 			}
 		}
diff --git a/Engine/Source/Programs/Shared/EpicGames.Horde/Bundles/Nodes/DirectoryNode.cs b/Engine/Source/Programs/Shared/EpicGames.Horde/Bundles/Nodes/DirectoryNode.cs
--- a/Engine/Source/Programs/Shared/EpicGames.Horde/Bundles/Nodes/DirectoryNode.cs
+++ b/Engine/Source/Programs/Shared/EpicGames.Horde/Bundles/Nodes/DirectoryNode.cs
@@ -5,13 +5,11 @@ using EpicGames.Serialization;
 using Microsoft.Extensions.Logging;
 using System;
 using System.Buffers;
-using System.Buffers.Binary;
 using System.Collections.Generic;
 using System.Diagnostics;
 using System.Diagnostics.CodeAnalysis;
 using System.IO;
 using System.Linq;
-using System.Text;
 using System.Threading.Tasks;
 
 namespace EpicGames.Horde.Bundles.Nodes
@@ -65,38 +63,39 @@ namespace EpicGames.Horde.Bundles.Nodes
 		/// <summary>
 		/// Length of this entry
 		/// </summary>
-		public long Length => (Node == null) ? CachedLength : Node.Length;
+		public long Length => (Node == null) ? _cachedLength : Node.Length;
 
 		/// <summary>
 		/// Cached length of this node
 		/// </summary>
-		long CachedLength;
+		long _cachedLength;
 
 		/// <summary>
 		/// Constructor
 		/// </summary>
-		public FileEntry(DirectoryNode Owner, Utf8String Name, FileEntryFlags Flags, FileNode Node)
-			: base(Owner, Node)
+		public FileEntry(DirectoryNode owner, Utf8String name, FileEntryFlags flags, FileNode node)
+			: base(owner, node)
 		{
-			this.Name = Name;
-			this.Flags = Flags;
+			Name = name;
+			Flags = flags;
 		}
 
 		/// <summary>
 		/// Constructor
 		/// </summary>
-		public FileEntry(DirectoryNode Owner, Utf8String Name, FileEntryFlags Flags, long Length, IoHash Hash)
-			: base(Owner, Hash)
+		public FileEntry(DirectoryNode owner, Utf8String name, FileEntryFlags flags, long length, IoHash hash)
+			: base(owner, hash)
 		{
-			this.Name = Name;
-			this.Flags = Flags;
-			this.CachedLength = Length;
+			Name = name;
+			Flags = flags;
+			
+			_cachedLength = length;
 		}
 
 		/// <inheritdoc/>
 		protected override void OnCollapse()
 		{
-			CachedLength = Node!.Length;
+			_cachedLength = Node!.Length;
 		}
 
 		/// <inheritdoc/>
@@ -116,36 +115,37 @@ namespace EpicGames.Horde.Bundles.Nodes
 		/// <summary>
 		/// Length of this directory tree
 		/// </summary>
-		public long Length => (Node == null) ? CachedLength : Node.Length;
+		public long Length => (Node == null) ? _cachedLength : Node.Length;
 
 		/// <summary>
 		/// Cached value for the length of this tree
 		/// </summary>
-		long CachedLength;
+		long _cachedLength;
 
 		/// <summary>
 		/// Constructor
 		/// </summary>
-		public DirectoryEntry(DirectoryNode Owner, Utf8String Name, DirectoryNode Node)
-			: base(Owner, Node)
+		public DirectoryEntry(DirectoryNode owner, Utf8String name, DirectoryNode node)
+			: base(owner, node)
 		{
-			this.Name = Name;
+			Name = name;
 		}
 
 		/// <summary>
 		/// Constructor
 		/// </summary>
-		public DirectoryEntry(DirectoryNode Owner, Utf8String Name, long Length, IoHash Hash)
-			: base(Owner, Hash)
+		public DirectoryEntry(DirectoryNode owner, Utf8String name, long length, IoHash hash)
+			: base(owner, hash)
 		{
-			this.Name = Name;
-			this.CachedLength = Length;
+			Name = name;
+
+			_cachedLength = length;
 		}
 
 		/// <inheritdoc/>
 		protected override void OnCollapse()
 		{
-			CachedLength = Node!.Length;
+			_cachedLength = Node!.Length;
 		}
 
 		/// <inheritdoc/>
@@ -160,8 +160,8 @@ namespace EpicGames.Horde.Bundles.Nodes
 	{
 		internal const byte TypeId = (byte)'d';
 
-		Dictionary<Utf8String, FileEntry> NameToFileEntry = new Dictionary<Utf8String, FileEntry>();
-		Dictionary<Utf8String, DirectoryEntry> NameToDirectoryEntry = new Dictionary<Utf8String, DirectoryEntry>();
+		readonly Dictionary<Utf8String, FileEntry> _nameToFileEntry = new Dictionary<Utf8String, FileEntry>();
+		readonly Dictionary<Utf8String, DirectoryEntry> _nameToDirectoryEntry = new Dictionary<Utf8String, DirectoryEntry>();
 
 		/// <summary>
 		/// Total size of this directory
@@ -171,29 +171,29 @@ namespace EpicGames.Horde.Bundles.Nodes
 		/// <summary>
 		/// All the files within this directory
 		/// </summary>
-		public IReadOnlyCollection<FileEntry> Files => NameToFileEntry.Values;
+		public IReadOnlyCollection<FileEntry> Files => _nameToFileEntry.Values;
 
 		/// <summary>
 		/// All the subdirectories within this directory
 		/// </summary>
-		public IReadOnlyCollection<DirectoryEntry> Directories => NameToDirectoryEntry.Values;
+		public IReadOnlyCollection<DirectoryEntry> Directories => _nameToDirectoryEntry.Values;
 
 		/// <summary>
 		/// Clear the contents of this directory
 		/// </summary>
 		public void Clear()
 		{
-			NameToFileEntry.Clear();
-			NameToDirectoryEntry.Clear();
+			_nameToFileEntry.Clear();
+			_nameToDirectoryEntry.Clear();
 			MarkAsDirty();
 		}
 
 		/// <summary>
 		/// Check whether an entry with the given name exists in this directory
 		/// </summary>
-		/// <param name="Name">Name of the entry to search for</param>
+		/// <param name="name">Name of the entry to search for</param>
 		/// <returns>True if the entry exists</returns>
-		public bool Contains(Utf8String Name) => TryGetFileEntry(Name, out _) || TryGetDirectoryEntry(Name, out _);
+		public bool Contains(Utf8String name) => TryGetFileEntry(name, out _) || TryGetDirectoryEntry(name, out _);
 
 		/// <inheritdoc/>
 		public override IEnumerable<BundleNodeRef> GetReferences() => Enumerable.Concat<BundleNodeRef>(Files, Directories);
@@ -203,90 +203,90 @@ namespace EpicGames.Horde.Bundles.Nodes
 		/// <summary>
 		/// Adds a new directory with the given name
 		/// </summary>
-		/// <param name="Name">Name of the new directory</param>
-		/// <param name="Flags">Flags for the new file</param>
+		/// <param name="name">Name of the new directory</param>
+		/// <param name="flags">Flags for the new file</param>
 		/// <returns>The new directory object</returns>
-		public FileNode AddFile(Utf8String Name, FileEntryFlags Flags)
+		public FileNode AddFile(Utf8String name, FileEntryFlags flags)
 		{
-			if(TryGetDirectoryEntry(Name, out _))
+			if(TryGetDirectoryEntry(name, out _))
 			{
-				throw new ArgumentException($"A directory with the name {Name} already exists");
+				throw new ArgumentException($"A directory with the name {name} already exists");
 			}
 
-			FileNode NewNode = new FileNode();
+			FileNode newNode = new FileNode();
 
-			FileEntry Entry = new FileEntry(this, Name, Flags, NewNode);
-			NameToFileEntry[Name] = Entry;
+			FileEntry entry = new FileEntry(this, name, flags, newNode);
+			_nameToFileEntry[name] = entry;
 			MarkAsDirty();
 
-			return NewNode;
+			return newNode;
 		}
 
 		/// <summary>
 		/// Finds or adds a file with the given path
 		/// </summary>
-		/// <param name="Bundle">Bundle that contains the node</param>
-		/// <param name="Path">Path to the file</param>
-		/// <param name="Flags">Flags for the new file</param>
+		/// <param name="bundle">Bundle that contains the node</param>
+		/// <param name="path">Path to the file</param>
+		/// <param name="flags">Flags for the new file</param>
 		/// <returns>The new directory object</returns>
-		public async ValueTask<FileNode> AddFileByPathAsync(Bundle Bundle, Utf8String Path, FileEntryFlags Flags)
+		public async ValueTask<FileNode> AddFileByPathAsync(Bundle bundle, Utf8String path, FileEntryFlags flags)
 		{
-			DirectoryNode Directory = this;
+			DirectoryNode directory = this;
 
-			Utf8String RemainingPath = Path;
-			if (RemainingPath[0] == '/' || RemainingPath[0] == '\\')
+			Utf8String remainingPath = path;
+			if (remainingPath[0] == '/' || remainingPath[0] == '\\')
 			{
-				RemainingPath = RemainingPath.Substring(1);
+				remainingPath = remainingPath.Substring(1);
 			}
 
 			for (; ; )
 			{
-				int Length = 0;
-				for (; ; Length++)
+				int length = 0;
+				for (; ; length++)
 				{
-					if (Length == RemainingPath.Length)
+					if (length == remainingPath.Length)
 					{
-						return Directory.AddFile(RemainingPath, Flags);
+						return directory.AddFile(remainingPath, flags);
 					}
 
-					byte Character = RemainingPath[Length];
-					if (Character == '\\' || Character == '/')
+					byte character = remainingPath[length];
+					if (character == '\\' || character == '/')
 					{
 						break;
 					}
 				}
 
-				if (Length > 0)
+				if (length > 0)
 				{
-					Directory = await Bundle.FindOrAddDirectoryAsync(Directory, RemainingPath.Slice(0, Length));
+					directory = await bundle.FindOrAddDirectoryAsync(directory, remainingPath.Slice(0, length));
 				}
-				RemainingPath = RemainingPath.Slice(Length + 1);
+				remainingPath = remainingPath.Slice(length + 1);
 			}
 		}
 
 		/// <summary>
 		/// Attempts to get a file entry with the given name
 		/// </summary>
-		/// <param name="Name">Name of the file</param>
+		/// <param name="name">Name of the file</param>
 		/// <returns>Entry for the given name</returns>
-		public FileEntry GetFileEntry(Utf8String Name) => NameToFileEntry[Name];
+		public FileEntry GetFileEntry(Utf8String name) => _nameToFileEntry[name];
 
 		/// <summary>
 		/// Attempts to get a file entry with the given name
 		/// </summary>
-		/// <param name="Name">Name of the file</param>
-		/// <param name="Entry">Entry for the file</param>
+		/// <param name="name">Name of the file</param>
+		/// <param name="entry">Entry for the file</param>
 		/// <returns>True if the file was found</returns>
-		public bool TryGetFileEntry(Utf8String Name, [NotNullWhen(true)] out FileEntry? Entry) => NameToFileEntry.TryGetValue(Name, out Entry);
+		public bool TryGetFileEntry(Utf8String name, [NotNullWhen(true)] out FileEntry? entry) => _nameToFileEntry.TryGetValue(name, out entry);
 
 		/// <summary>
 		/// Deletes the file entry with the given name
 		/// </summary>
-		/// <param name="Name">Name of the entry to delete</param>
+		/// <param name="name">Name of the entry to delete</param>
 		/// <returns>True if the entry was found, false otherwise</returns>
-		public bool DeleteFile(Utf8String Name)
+		public bool DeleteFile(Utf8String name)
 		{
-			if (NameToFileEntry.Remove(Name))
+			if (_nameToFileEntry.Remove(name))
 			{
 				MarkAsDirty();
 				return true;
@@ -297,26 +297,26 @@ namespace EpicGames.Horde.Bundles.Nodes
 		/// <summary>
 		/// 
 		/// </summary>
-		/// <param name="Bundle">Bundle that contains the node</param>
-		/// <param name="Path"></param>
+		/// <param name="bundle">Bundle that contains the node</param>
+		/// <param name="path"></param>
 		/// <returns></returns>
-		public async ValueTask<bool> DeleteFileByPathAsync(Bundle Bundle, Utf8String Path)
+		public async ValueTask<bool> DeleteFileByPathAsync(Bundle bundle, Utf8String path)
 		{
-			DirectoryNode Directory = this;
+			DirectoryNode directory = this;
 
-			Utf8String RemainingPath = Path;
+			Utf8String remainingPath = path;
 			for (; ; )
 			{
-				int Length = RemainingPath.IndexOf('/');
-				if (Length == -1)
+				int length = remainingPath.IndexOf('/');
+				if (length == -1)
 				{
-					return Directory.DeleteFile(RemainingPath);
+					return directory.DeleteFile(remainingPath);
 				}
-				if (Length > 0)
+				if (length > 0)
 				{
-					Directory = await Bundle.FindOrAddDirectoryAsync(Directory, RemainingPath.Slice(0, Length));
+					directory = await bundle.FindOrAddDirectoryAsync(directory, remainingPath.Slice(0, length));
 				}
-				RemainingPath = RemainingPath.Slice(Length + 1);
+				remainingPath = remainingPath.Slice(length + 1);
 			}
 		}
 
@@ -327,47 +327,47 @@ namespace EpicGames.Horde.Bundles.Nodes
 		/// <summary>
 		/// Adds a new directory with the given name
 		/// </summary>
-		/// <param name="Name">Name of the new directory</param>
+		/// <param name="name">Name of the new directory</param>
 		/// <returns>The new directory object</returns>
-		public DirectoryNode AddDirectory(Utf8String Name)
+		public DirectoryNode AddDirectory(Utf8String name)
 		{
-			if (TryGetFileEntry(Name, out _))
+			if (TryGetFileEntry(name, out _))
 			{
-				throw new ArgumentException($"A file with the name '{Name}' already exists in this directory", nameof(Name));
+				throw new ArgumentException($"A file with the name '{name}' already exists in this directory", nameof(name));
 			}
 
-			DirectoryNode Node = new DirectoryNode();
+			DirectoryNode node = new DirectoryNode();
 
-			DirectoryEntry Entry = new DirectoryEntry(this, Name, Node);
-			NameToDirectoryEntry.Add(Name, Entry);
+			DirectoryEntry entry = new DirectoryEntry(this, name, node);
+			_nameToDirectoryEntry.Add(name, entry);
 			MarkAsDirty();
 
-			return Node;
+			return node;
 		}
 
 		/// <summary>
 		/// Get a directory entry with the given name
 		/// </summary>
-		/// <param name="Name">Name of the directory</param>
+		/// <param name="name">Name of the directory</param>
 		/// <returns>The entry with the given name</returns>
-		public DirectoryEntry GetDirectoryEntry(Utf8String Name) => NameToDirectoryEntry[Name];
+		public DirectoryEntry GetDirectoryEntry(Utf8String name) => _nameToDirectoryEntry[name];
 
 		/// <summary>
 		/// Attempts to get a directory entry with the given name
 		/// </summary>
-		/// <param name="Name">Name of the directory</param>
-		/// <param name="Entry">Entry for the directory</param>
+		/// <param name="name">Name of the directory</param>
+		/// <param name="entry">Entry for the directory</param>
 		/// <returns>True if the directory was found</returns>
-		public bool TryGetDirectoryEntry(Utf8String Name, [NotNullWhen(true)] out DirectoryEntry? Entry) => NameToDirectoryEntry.TryGetValue(Name, out Entry);
+		public bool TryGetDirectoryEntry(Utf8String name, [NotNullWhen(true)] out DirectoryEntry? entry) => _nameToDirectoryEntry.TryGetValue(name, out entry);
 
 		/// <summary>
 		/// Deletes the file entry with the given name
 		/// </summary>
-		/// <param name="Name">Name of the entry to delete</param>
+		/// <param name="name">Name of the entry to delete</param>
 		/// <returns>True if the entry was found, false otherwise</returns>
-		public bool DeleteDirectory(Utf8String Name)
+		public bool DeleteDirectory(Utf8String name)
 		{
-			if(NameToDirectoryEntry.Remove(Name))
+			if(_nameToDirectoryEntry.Remove(name))
 			{
 				MarkAsDirty();
 				return true;
@@ -381,173 +381,173 @@ namespace EpicGames.Horde.Bundles.Nodes
 		public override ReadOnlySequence<byte> Serialize()
 		{
 			// Measure the required size of the write buffer
-			int Size = 1;
-			Size += VarInt.Measure(NameToFileEntry.Count);
-			foreach (FileEntry Entry in NameToFileEntry.Values)
+			int size = 1;
+			size += VarInt.Measure(_nameToFileEntry.Count);
+			foreach (FileEntry entry in _nameToFileEntry.Values)
 			{
-				Size += (Entry.Name.Length + 1) + 1 + VarInt.Measure((ulong)Entry.Length) + IoHash.NumBytes;
+				size += (entry.Name.Length + 1) + 1 + VarInt.Measure((ulong)entry.Length) + IoHash.NumBytes;
 			}
-			Size += VarInt.Measure(NameToDirectoryEntry.Count);
-			foreach (DirectoryEntry Entry in NameToDirectoryEntry.Values)
+			size += VarInt.Measure(_nameToDirectoryEntry.Count);
+			foreach (DirectoryEntry entry in _nameToDirectoryEntry.Values)
 			{
-				Size += (Entry.Name.Length + 1) + VarInt.Measure((ulong)Entry.Length) + IoHash.NumBytes;
+				size += (entry.Name.Length + 1) + VarInt.Measure((ulong)entry.Length) + IoHash.NumBytes;
 			}
 
 			// Allocate the buffer and copy the node to it
-			byte[] Data = new byte[Size];
-			Span<byte> Span = Data.AsSpan();
+			byte[] data = new byte[size];
+			Span<byte> span = data.AsSpan();
 
-			Span[0] = TypeId;
-			Span = Span.Slice(1);
+			span[0] = TypeId;
+			span = span.Slice(1);
 
-			int FileCountBytes = VarInt.Write(Span, NameToFileEntry.Count);
-			Span = Span.Slice(FileCountBytes);
+			int fileCountBytes = VarInt.Write(span, _nameToFileEntry.Count);
+			span = span.Slice(fileCountBytes);
 
-			foreach (FileEntry FileEntry in NameToFileEntry.Values.OrderBy(x => x.Name))
+			foreach (FileEntry fileEntry in _nameToFileEntry.Values.OrderBy(x => x.Name))
 			{
-				FileEntry.Name.Span.CopyTo(Span);
-				Span = Span.Slice(FileEntry.Name.Length + 1);
+				fileEntry.Name.Span.CopyTo(span);
+				span = span.Slice(fileEntry.Name.Length + 1);
 
-				Span[0] = (byte)FileEntry.Flags;
-				Span = Span.Slice(1);
+				span[0] = (byte)fileEntry.Flags;
+				span = span.Slice(1);
 
-				int LengthBytes = VarInt.Write(Span, FileEntry.Length);
-				Span = Span.Slice(LengthBytes);
+				int lengthBytes = VarInt.Write(span, fileEntry.Length);
+				span = span.Slice(lengthBytes);
 
-				FileEntry.Hash.CopyTo(Span);
-				Span = Span.Slice(IoHash.NumBytes);
+				fileEntry.Hash.CopyTo(span);
+				span = span.Slice(IoHash.NumBytes);
 			}
 
-			int DirectoryCountBytes = VarInt.Write(Span, NameToDirectoryEntry.Count);
-			Span = Span.Slice(DirectoryCountBytes);
+			int directoryCountBytes = VarInt.Write(span, _nameToDirectoryEntry.Count);
+			span = span.Slice(directoryCountBytes);
 
-			foreach (DirectoryEntry DirectoryEntry in NameToDirectoryEntry.Values.OrderBy(x => x.Name))
+			foreach (DirectoryEntry directoryEntry in _nameToDirectoryEntry.Values.OrderBy(x => x.Name))
 			{
-				DirectoryEntry.Name.Span.CopyTo(Span);
-				Span = Span.Slice(DirectoryEntry.Name.Length + 1);
+				directoryEntry.Name.Span.CopyTo(span);
+				span = span.Slice(directoryEntry.Name.Length + 1);
 
-				int LengthBytes = VarInt.Write(Span, DirectoryEntry.Length);
-				Span = Span.Slice(LengthBytes);
+				int lengthBytes = VarInt.Write(span, directoryEntry.Length);
+				span = span.Slice(lengthBytes);
 
-				DirectoryEntry.Hash.CopyTo(Span);
-				Span = Span.Slice(IoHash.NumBytes);
+				directoryEntry.Hash.CopyTo(span);
+				span = span.Slice(IoHash.NumBytes);
 			}
 
-			Debug.Assert(Span.Length == 0);
-			return new ReadOnlySequence<byte>(Data);
+			Debug.Assert(span.Length == 0);
+			return new ReadOnlySequence<byte>(data);
 		}
 
 		/// <summary>
 		/// Deserialize a directory node from data
 		/// </summary>
-		/// <param name="Data">The data to read from</param>
+		/// <param name="data">The data to read from</param>
 		/// <returns>The deserialized directory node</returns>
-		public static DirectoryNode Deserialize(ReadOnlyMemory<byte> Data)
+		public static DirectoryNode Deserialize(ReadOnlyMemory<byte> data)
 		{
-			ReadOnlySpan<byte> Span = Data.Span;
-			if(Span[0] != TypeId)
+			ReadOnlySpan<byte> span = data.Span;
+			if(span[0] != TypeId)
 			{
 				throw new InvalidOperationException("Invalid signature byte for directory");
 			}
 
-			Span = Span.Slice(1);
+			span = span.Slice(1);
 
-			DirectoryNode Node = new DirectoryNode();
+			DirectoryNode node = new DirectoryNode();
 
-			int FileCount = (int)VarInt.Read(Span, out int FileCountBytes);
-			Span = Span.Slice(FileCountBytes);
+			int fileCount = (int)VarInt.Read(span, out int fileCountBytes);
+			span = span.Slice(fileCountBytes);
 
-			Node.NameToFileEntry.EnsureCapacity(FileCount);
-			for (int Idx = 0; Idx < FileCount; Idx++)
+			node._nameToFileEntry.EnsureCapacity(fileCount);
+			for (int idx = 0; idx < fileCount; idx++)
 			{
-				int NameLen = Span.IndexOf((byte)0);
-				Utf8String Name = new Utf8String(Span.Slice(0, NameLen).ToArray());
-				Span = Span.Slice(NameLen + 1);
+				int nameLen = span.IndexOf((byte)0);
+				Utf8String name = new Utf8String(span.Slice(0, nameLen).ToArray());
+				span = span.Slice(nameLen + 1);
 
-				FileEntryFlags Flags = (FileEntryFlags)Span[0];
-				Span = Span.Slice(1);
+				FileEntryFlags flags = (FileEntryFlags)span[0];
+				span = span.Slice(1);
 
-				long Length = (long)VarInt.Read(Span, out int LengthBytes);
-				Span = Span.Slice(LengthBytes);
+				long length = (long)VarInt.Read(span, out int lengthBytes);
+				span = span.Slice(lengthBytes);
 
-				IoHash Hash = new IoHash(Span);
-				Span = Span.Slice(IoHash.NumBytes);
+				IoHash hash = new IoHash(span);
+				span = span.Slice(IoHash.NumBytes);
 
-				Node.NameToFileEntry[Name] = new FileEntry(Node, Name, Flags, Length, Hash);
+				node._nameToFileEntry[name] = new FileEntry(node, name, flags, length, hash);
 			}
 
-			int DirectoryCount = (int)VarInt.Read(Span, out int DirectoryCountBytes);
-			Span = Span.Slice(DirectoryCountBytes);
+			int directoryCount = (int)VarInt.Read(span, out int directoryCountBytes);
+			span = span.Slice(directoryCountBytes);
 
-			Node.NameToDirectoryEntry.EnsureCapacity(DirectoryCount);
-			for (int Idx = 0; Idx < DirectoryCount; Idx++)
+			node._nameToDirectoryEntry.EnsureCapacity(directoryCount);
+			for (int idx = 0; idx < directoryCount; idx++)
 			{
-				int NameLen = Span.IndexOf((byte)0);
-				Utf8String Name = new Utf8String(Span.Slice(0, NameLen).ToArray());
-				Span = Span.Slice(NameLen + 1);
+				int nameLen = span.IndexOf((byte)0);
+				Utf8String name = new Utf8String(span.Slice(0, nameLen).ToArray());
+				span = span.Slice(nameLen + 1);
 
-				long Length = (long)VarInt.Read(Span, out int LengthBytes);
-				Span = Span.Slice(LengthBytes);
+				long length = (long)VarInt.Read(span, out int lengthBytes);
+				span = span.Slice(lengthBytes);
 
-				IoHash Hash = new IoHash(Span);
-				Span = Span.Slice(IoHash.NumBytes);
+				IoHash hash = new IoHash(span);
+				span = span.Slice(IoHash.NumBytes);
 
-				Node.NameToDirectoryEntry[Name] = new DirectoryEntry(Node, Name, Length, Hash);
+				node._nameToDirectoryEntry[name] = new DirectoryEntry(node, name, length, hash);
 			}
 
-			Debug.Assert(Span.Length == 0);
-			return Node;
+			Debug.Assert(span.Length == 0);
+			return node;
 		}
 
 		/// <summary>
 		/// Adds files from a directory on disk
 		/// </summary>
-		/// <param name="DirectoryInfo"></param>
-		/// <param name="Options">Options for chunking file content</param>
-		/// <param name="Logger"></param>
+		/// <param name="directoryInfo"></param>
+		/// <param name="options">Options for chunking file content</param>
+		/// <param name="logger"></param>
 		/// <returns></returns>
-		public async Task CopyFromDirectoryAsync(DirectoryInfo DirectoryInfo, ChunkingOptions Options, ILogger Logger)
+		public async Task CopyFromDirectoryAsync(DirectoryInfo directoryInfo, ChunkingOptions options, ILogger logger)
 		{
-			foreach (DirectoryInfo SubDirectoryInfo in DirectoryInfo.EnumerateDirectories())
+			foreach (DirectoryInfo subDirectoryInfo in directoryInfo.EnumerateDirectories())
 			{
-				Logger.LogInformation("Adding {Directory}", SubDirectoryInfo.FullName);
-				DirectoryNode SubDirectoryNode = AddDirectory(SubDirectoryInfo.Name);
-				await SubDirectoryNode.CopyFromDirectoryAsync(SubDirectoryInfo, Options, Logger);
+				logger.LogInformation("Adding {Directory}", subDirectoryInfo.FullName);
+				DirectoryNode subDirectoryNode = AddDirectory(subDirectoryInfo.Name);
+				await subDirectoryNode.CopyFromDirectoryAsync(subDirectoryInfo, options, logger);
 			}
-			foreach (FileInfo FileInfo in DirectoryInfo.EnumerateFiles())
+			foreach (FileInfo fileInfo in directoryInfo.EnumerateFiles())
 			{
-				Logger.LogInformation("Adding {File}", FileInfo.FullName);
-				FileNode FileNode = AddFile(FileInfo.Name, 0);
- 				await FileNode.CopyFromFileAsync(FileInfo, Options);
+				logger.LogInformation("Adding {File}", fileInfo.FullName);
+				FileNode fileNode = AddFile(fileInfo.Name, 0);
+ 				await fileNode.CopyFromFileAsync(fileInfo, options);
 			}
 		}
 
 		/// <summary>
 		/// Utility function to allow extracting a packed directory to disk
 		/// </summary>
-		/// <param name="Bundle">Bundle that contains the node</param>
-		/// <param name="DirectoryInfo"></param>
-		/// <param name="Logger"></param>
+		/// <param name="bundle">Bundle that contains the node</param>
+		/// <param name="directoryInfo"></param>
+		/// <param name="logger"></param>
 		/// <returns></returns>
-		public async Task CopyToDirectoryAsync(Bundle Bundle, DirectoryInfo DirectoryInfo, ILogger Logger)
+		public async Task CopyToDirectoryAsync(Bundle bundle, DirectoryInfo directoryInfo, ILogger logger)
 		{
-			DirectoryInfo.Create();
+			directoryInfo.Create();
 
-			List<Task> Tasks = new List<Task>();
-			foreach (FileEntry FileEntry in Files)
+			List<Task> tasks = new List<Task>();
+			foreach (FileEntry fileEntry in Files)
 			{
-				FileInfo FileInfo = new FileInfo(Path.Combine(DirectoryInfo.FullName, FileEntry.Name.ToString()));
-				FileNode FileNode = await Bundle.GetAsync(FileEntry);
-				Tasks.Add(Task.Run(() => FileNode.CopyToFileAsync(Bundle, FileInfo)));
+				FileInfo fileInfo = new FileInfo(Path.Combine(directoryInfo.FullName, fileEntry.Name.ToString()));
+				FileNode fileNode = await bundle.GetAsync(fileEntry);
+				tasks.Add(Task.Run(() => fileNode.CopyToFileAsync(bundle, fileInfo)));
 			}
-			foreach (DirectoryEntry DirectoryEntry in Directories)
+			foreach (DirectoryEntry directoryEntry in Directories)
 			{
-				DirectoryInfo SubDirectoryInfo = DirectoryInfo.CreateSubdirectory(DirectoryEntry.Name.ToString());
-				DirectoryNode SubDirectoryNode = await Bundle.GetAsync(DirectoryEntry);
-				Tasks.Add(Task.Run(() => SubDirectoryNode.CopyToDirectoryAsync(Bundle, SubDirectoryInfo, Logger)));
+				DirectoryInfo subDirectoryInfo = directoryInfo.CreateSubdirectory(directoryEntry.Name.ToString());
+				DirectoryNode subDirectoryNode = await bundle.GetAsync(directoryEntry);
+				tasks.Add(Task.Run(() => subDirectoryNode.CopyToDirectoryAsync(bundle, subDirectoryInfo, logger)));
 			}
 
-			await Task.WhenAll(Tasks);
+			await Task.WhenAll(tasks);
 		}
 	}
 
@@ -557,7 +557,7 @@ namespace EpicGames.Horde.Bundles.Nodes
 	public class DirectoryNodeDeserializer : BundleNodeDeserializer<DirectoryNode>
 	{
 		/// <inheritdoc/>
-		public override DirectoryNode Deserialize(ReadOnlyMemory<byte> Data) => DirectoryNode.Deserialize(Data);
+		public override DirectoryNode Deserialize(ReadOnlyMemory<byte> data) => DirectoryNode.Deserialize(data);
 	}
 
 	/// <summary>
@@ -568,15 +568,15 @@ namespace EpicGames.Horde.Bundles.Nodes
 		/// <summary>
 		/// Tries to get a directory with the given name
 		/// </summary>
-		/// <param name="Bundle">Bundle that contains the node</param>
-		/// <param name="ParentNode">Parent directory to search</param>
-		/// <param name="Name">Name of the new directory</param>
+		/// <param name="bundle">Bundle that contains the node</param>
+		/// <param name="parentNode">Parent directory to search</param>
+		/// <param name="name">Name of the new directory</param>
 		/// <returns>The new directory object</returns>
-		public static async ValueTask<DirectoryNode?> FindDirectoryAsync(this Bundle Bundle, DirectoryNode ParentNode, Utf8String Name)
+		public static async ValueTask<DirectoryNode?> FindDirectoryAsync(this Bundle bundle, DirectoryNode parentNode, Utf8String name)
 		{
-			if (ParentNode.TryGetDirectoryEntry(Name, out DirectoryEntry? Entry))
+			if (parentNode.TryGetDirectoryEntry(name, out DirectoryEntry? entry))
 			{
-				return await Bundle.GetAsync(Entry);
+				return await bundle.GetAsync(entry);
 			}
 			else
 			{
@@ -587,14 +587,14 @@ namespace EpicGames.Horde.Bundles.Nodes
 		/// <summary>
 		/// Tries to get a directory with the given name, or adds it if not present
 		/// </summary>
-		/// <param name="Bundle">Bundle that contains the node</param>
-		/// <param name="ParentNode">Parent directory to search</param>
-		/// <param name="Name">Name of the new directory</param>
+		/// <param name="bundle">Bundle that contains the node</param>
+		/// <param name="parentNode">Parent directory to search</param>
+		/// <param name="name">Name of the new directory</param>
 		/// <returns>The new directory object</returns>
-		public static async ValueTask<DirectoryNode> FindOrAddDirectoryAsync(this Bundle Bundle, DirectoryNode ParentNode, Utf8String Name)
+		public static async ValueTask<DirectoryNode> FindOrAddDirectoryAsync(this Bundle bundle, DirectoryNode parentNode, Utf8String name)
 		{
-			DirectoryNode? Node = await FindDirectoryAsync(Bundle, ParentNode, Name);
-			return Node ?? ParentNode.AddDirectory(Name);
+			DirectoryNode? node = await FindDirectoryAsync(bundle, parentNode, name);
+			return node ?? parentNode.AddDirectory(name);
 		}
 	}
 }
diff --git a/Engine/Source/Programs/Shared/EpicGames.Horde/Bundles/Nodes/FileNode.cs b/Engine/Source/Programs/Shared/EpicGames.Horde/Bundles/Nodes/FileNode.cs
--- a/Engine/Source/Programs/Shared/EpicGames.Horde/Bundles/Nodes/FileNode.cs
+++ b/Engine/Source/Programs/Shared/EpicGames.Horde/Bundles/Nodes/FileNode.cs
@@ -7,8 +7,6 @@ using System.Buffers;
 using System.Collections.Generic;
 using System.Diagnostics;
 using System.IO;
-using System.Linq;
-using System.Text;
 using System.Threading.Tasks;
 
 namespace EpicGames.Horde.Bundles.Nodes
@@ -61,11 +59,11 @@ namespace EpicGames.Horde.Bundles.Nodes
 		/// <summary>
 		/// Default constructor
 		/// </summary>
-		public ChunkingOptionsForNodeType(int MinSize, int MaxSize, int TargetSize)
+		public ChunkingOptionsForNodeType(int minSize, int maxSize, int targetSize)
 		{
-			this.MinSize = MinSize;
-			this.MaxSize = MaxSize;
-			this.TargetSize = TargetSize;
+			MinSize = minSize;
+			MaxSize = maxSize;
+			TargetSize = targetSize;
 		}
 	}
 
@@ -79,33 +77,33 @@ namespace EpicGames.Horde.Bundles.Nodes
 	{
 		const byte TypeId = (byte)'c';
 
-		static readonly ReadOnlyMemory<byte> DefaultSegment = CreateFirstSegmentData(0, Array.Empty<byte>());
+		static readonly ReadOnlyMemory<byte> s_defaultSegment = CreateFirstSegmentData(0, Array.Empty<byte>());
 
 		class DataSegment : ReadOnlySequenceSegment<byte>
 		{
-			public DataSegment(long RunningIndex, ReadOnlyMemory<byte> Data)
+			public DataSegment(long runningIndex, ReadOnlyMemory<byte> data)
 			{
-				this.RunningIndex = RunningIndex;
-				this.Memory = Data;
+				RunningIndex = runningIndex;
+				Memory = data;
 			}
 
-			public void SetNext(DataSegment Next)
+			public void SetNext(DataSegment next)
 			{
-				this.Next = Next;
+				Next = next;
 			}
 		}
 
 		internal int Depth { get; private set; }
 		internal ReadOnlySequence<byte> Payload { get; private set; }
-		ReadOnlySequence<byte> Data = new ReadOnlySequence<byte>(DefaultSegment); // Full serialized data, including type id and header fields.
-		IoHash? Hash; // Hash of the serialized data buffer.
+		ReadOnlySequence<byte> _data = new ReadOnlySequence<byte>(s_defaultSegment); // Full serialized data, including type id and header fields.
+		IoHash? _hash; // Hash of the serialized data buffer.
 
 		// In-memory state
-		bool IsRoot = true;
-		uint RollingHash;
-		DataSegment? FirstSegment;
-		DataSegment? LastSegment;
-		List<BundleNodeRef<FileNode>>? ChildNodeRefs;
+		bool _isRoot = true;
+		uint _rollingHash;
+		DataSegment? _firstSegment;
+		DataSegment? _lastSegment;
+		List<BundleNodeRef<FileNode>>? _childNodeRefs;
 
 		/// <summary>
 		/// Length of this tree
@@ -115,18 +113,18 @@ namespace EpicGames.Horde.Bundles.Nodes
 		/// <summary>
 		/// Whether this node is read-only
 		/// </summary>
-		public override bool IsReadOnly() => Hash != null;
+		public override bool IsReadOnly() => _hash != null;
 
 		/// <summary>
 		/// Accessor for the children of this node
 		/// </summary>
-		public IReadOnlyList<BundleNodeRef<FileNode>> Children => (IReadOnlyList<BundleNodeRef<FileNode>>?)ChildNodeRefs ?? Array.Empty<BundleNodeRef<FileNode>>();
+		public IReadOnlyList<BundleNodeRef<FileNode>> Children => (IReadOnlyList<BundleNodeRef<FileNode>>?)_childNodeRefs ?? Array.Empty<BundleNodeRef<FileNode>>();
 
 		/// <inheritdoc/>
 		public override ReadOnlySequence<byte> Serialize()
 		{
 			MarkComplete();
-			return Data;
+			return _data;
 		}
 
 		/// <inheritdoc/>
@@ -135,112 +133,112 @@ namespace EpicGames.Horde.Bundles.Nodes
 		/// <summary>
 		/// Create a file node from deserialized data
 		/// </summary>
-		public static FileNode Deserialize(ReadOnlyMemory<byte> Data)
+		public static FileNode Deserialize(ReadOnlyMemory<byte> data)
 		{
-			ReadOnlySpan<byte> Span = Data.Span;
-			if (Span[0] != TypeId)
+			ReadOnlySpan<byte> span = data.Span;
+			if (span[0] != TypeId)
 			{
 				throw new InvalidDataException("Invalid type id");
 			}
 
-			FileNode Node = new FileNode();
-			Node.Data = new ReadOnlySequence<byte>(Data);
-			Node.Hash = IoHash.Compute(Data.Span);
-			Node.Depth = (int)VarInt.Read(Span.Slice(1), out int DepthBytes);
+			FileNode node = new FileNode();
+			node._data = new ReadOnlySequence<byte>(data);
+			node._hash = IoHash.Compute(data.Span);
+			node.Depth = (int)VarInt.Read(span.Slice(1), out int depthBytes);
 
-			int HeaderLength = 1 + DepthBytes;
-			Span = Span.Slice(HeaderLength);
+			int headerLength = 1 + depthBytes;
+			span = span.Slice(headerLength);
 
-			ReadOnlyMemory<byte> Payload = Data.Slice(HeaderLength);
-			Node.Payload = new ReadOnlySequence<byte>(Payload);
+			ReadOnlyMemory<byte> payload = data.Slice(headerLength);
+			node.Payload = new ReadOnlySequence<byte>(payload);
 
-			if (Node.Depth == 0)
+			if (node.Depth == 0)
 			{
-				Node.Length = Node.Payload.Length;
+				node.Length = node.Payload.Length;
 			}
 			else
 			{
-				Node.Length = (long)VarInt.Read(Span, out int LengthBytes);
-				Span = Span.Slice(LengthBytes);
+				node.Length = (long)VarInt.Read(span, out int lengthBytes);
+				span = span.Slice(lengthBytes);
 
-				Node.ChildNodeRefs = new List<BundleNodeRef<FileNode>>(Payload.Length / IoHash.NumBytes);
-				while (Span.Length > 0)
+				node._childNodeRefs = new List<BundleNodeRef<FileNode>>(payload.Length / IoHash.NumBytes);
+				while (span.Length > 0)
 				{
-					IoHash ChildHash = new IoHash(Span);
-					Node.ChildNodeRefs.Add(new BundleNodeRef<FileNode>(Node, ChildHash));
-					Span = Span.Slice(IoHash.NumBytes);
+					IoHash childHash = new IoHash(span);
+					node._childNodeRefs.Add(new BundleNodeRef<FileNode>(node, childHash));
+					span = span.Slice(IoHash.NumBytes);
 				}
 			}
-			return Node;
+			return node;
 		}
 
 		/// <summary>
 		/// Serialize this node and its children into a byte array
 		/// </summary>
 		/// <returns>Array of data stored by the tree</returns>
-		public async Task<byte[]> ToByteArrayAsync(Bundle Bundle)
+		public async Task<byte[]> ToByteArrayAsync(Bundle bundle)
 		{
-			using (MemoryStream Stream = new MemoryStream())
+			using (MemoryStream stream = new MemoryStream())
 			{
-				await CopyToStreamAsync(Bundle, Stream);
-				return Stream.ToArray();
+				await CopyToStreamAsync(bundle, stream);
+				return stream.ToArray();
 			}
 		}
 
 		/// <summary>
 		/// Copy data from the given stream into this file node
 		/// </summary>
-		/// <param name="InputStream"></param>
-		/// <param name="Options"></param>
+		/// <param name="inputStream"></param>
+		/// <param name="options"></param>
 		/// <returns></returns>
-		public async Task CopyFromStreamAsync(Stream InputStream, ChunkingOptions Options)
+		public async Task CopyFromStreamAsync(Stream inputStream, ChunkingOptions options)
 		{
-			byte[] Buffer = new byte[64 * 1024];
+			byte[] buffer = new byte[64 * 1024];
 			for (; ; )
 			{
-				int ReadSize = await InputStream.ReadAsync(Buffer);
-				if(ReadSize == 0)
+				int readSize = await inputStream.ReadAsync(buffer);
+				if(readSize == 0)
 				{
 					break;
 				}
-				Append(Buffer.AsMemory(0, ReadSize), Options);
+				Append(buffer.AsMemory(0, readSize), options);
 			}
 		}
 
 		/// <summary>
 		/// Copies the contents of this file from disk
 		/// </summary>
-		/// <param name="File"></param>
-		/// <param name="Options"></param>
+		/// <param name="file"></param>
+		/// <param name="options"></param>
 		/// <returns></returns>
-		public async Task CopyFromFileAsync(FileInfo File, ChunkingOptions Options)
+		public async Task CopyFromFileAsync(FileInfo file, ChunkingOptions options)
 		{
-			using (FileStream Stream = File.OpenRead())
+			using (FileStream stream = file.OpenRead())
 			{
-				await CopyFromStreamAsync(Stream, Options);
+				await CopyFromStreamAsync(stream, options);
 			}
 		}
 
 		/// <summary>
 		/// Copies the contents of this node and its children to the given output stream
 		/// </summary>
-		/// <param name="Bundle">Bundle that can provide node data</param>
-		/// <param name="OutputStream">The output stream to receive the data</param>
-		public async Task CopyToStreamAsync(Bundle Bundle, Stream OutputStream)
+		/// <param name="bundle">Bundle that can provide node data</param>
+		/// <param name="outputStream">The output stream to receive the data</param>
+		public async Task CopyToStreamAsync(Bundle bundle, Stream outputStream)
 		{
-			if (ChildNodeRefs != null)
+			if (_childNodeRefs != null)
 			{
-				foreach (BundleNodeRef<FileNode> ChildNodeRef in ChildNodeRefs)
+				foreach (BundleNodeRef<FileNode> childNodeRef in _childNodeRefs)
 				{
-					FileNode ChildNode = ChildNodeRef.Node ?? await Bundle.GetAsync(ChildNodeRef);
-					await ChildNode.CopyToStreamAsync(Bundle, OutputStream);
+					FileNode childNode = childNodeRef.Node ?? await bundle.GetAsync(childNodeRef);
+					await childNode.CopyToStreamAsync(bundle, outputStream);
 				}
 			}
 			else
 			{
-				foreach (ReadOnlyMemory<byte> PayloadSegment in Payload)
+				foreach (ReadOnlyMemory<byte> payloadSegment in Payload)
 				{
-					await OutputStream.WriteAsync(PayloadSegment);
+					await outputStream.WriteAsync(payloadSegment);
 				}
 			}
 		}
@@ -248,25 +246,25 @@ namespace EpicGames.Horde.Bundles.Nodes
 		/// <summary>
 		/// Extracts the contents of this node to a file
 		/// </summary>
-		/// <param name="Bundle">Bundle that can provide node data</param>
-		/// <param name="File">File to write with the contents of this node</param>
+		/// <param name="bundle">Bundle that can provide node data</param>
+		/// <param name="file">File to write with the contents of this node</param>
 		/// <returns></returns>
-		public async Task CopyToFileAsync(Bundle Bundle, FileInfo File)
+		public async Task CopyToFileAsync(Bundle bundle, FileInfo file)
 		{
-			using (FileStream Stream = File.OpenWrite())
+			using (FileStream stream = file.OpenWrite())
 			{
-				await CopyToStreamAsync(Bundle, Stream);
+				await CopyToStreamAsync(bundle, stream);
 			}
 		}
 
 		/// <summary>
 		/// Append data to this chunk. Must only be called on the root node in a chunk tree.
 		/// </summary>
-		/// <param name="Input">The data to write</param>
-		/// <param name="Options">Settings for chunking the data</param>
-		public void Append(ReadOnlyMemory<byte> Input, ChunkingOptions Options)
+		/// <param name="input">The data to write</param>
+		/// <param name="options">Settings for chunking the data</param>
+		public void Append(ReadOnlyMemory<byte> input, ChunkingOptions options)
 		{
-			if (!IsRoot)
+			if (!_isRoot)
 			{
 				throw new InvalidOperationException("Data may only be appended to the root of a file node tree");
 			}
@@ -274,275 +272,273 @@ namespace EpicGames.Horde.Bundles.Nodes
 			for (; ; )
 			{
 				// Append as much data as possible to the existing tree
-				Input = AppendToNode(Input, Options);
-				if (Input.IsEmpty)
+				input = AppendToNode(input, options);
+				if (input.IsEmpty)
 				{
 					break;
 				}
 
 				// Increase the height of the tree by pushing the contents of this node into a new child node
-				FileNode NewNode = new FileNode();
-				NewNode.Depth = Depth;
-				NewNode.Payload = Payload;
-				NewNode.Data = Data;
-				NewNode.Hash = Hash;
+				FileNode newNode = new FileNode();
+				newNode.Depth = Depth;
+				newNode.Payload = Payload;
+				newNode._data = _data;
+				newNode._hash = _hash;
 
-				NewNode.IsRoot = false;
-				NewNode.RollingHash = RollingHash;
-				NewNode.FirstSegment = FirstSegment;
-				NewNode.LastSegment = LastSegment;
-				NewNode.ChildNodeRefs = ChildNodeRefs;
+				newNode._isRoot = false;
+				newNode._rollingHash = _rollingHash;
+				newNode._firstSegment = _firstSegment;
+				newNode._lastSegment = _lastSegment;
+				newNode._childNodeRefs = _childNodeRefs;
 
-				NewNode.Length = Length;
+				newNode.Length = Length;
 
 				// Detach all the child refs
-				if (ChildNodeRefs != null)
+				if (_childNodeRefs != null)
 				{
-					foreach (BundleNodeRef<FileNode> ChildNodeRef in ChildNodeRefs)
+					foreach (BundleNodeRef<FileNode> childNodeRef in _childNodeRefs)
 					{
-						ChildNodeRef.Reparent(NewNode);
+						childNodeRef.Reparent(newNode);
 					}
-					ChildNodeRefs = null;
+					_childNodeRefs = null;
 				}
 
 				// Increase the depth and reset the current node
 				Depth++;
 				Payload = ReadOnlySequence<byte>.Empty;
-				Data = ReadOnlySequence<byte>.Empty;
-				Hash = null;
+				_data = ReadOnlySequence<byte>.Empty;
+				_hash = null;
 
-				IsRoot = true;
-				RollingHash = BuzHash.Add(0, NewNode.Hash!.Value.ToByteArray());
-				FirstSegment = null;
-				LastSegment = null;
-				ChildNodeRefs = new List<BundleNodeRef<FileNode>>();
+				_isRoot = true;
+				_rollingHash = BuzHash.Add(0, newNode._hash!.Value.ToByteArray());
+				_firstSegment = null;
+				_lastSegment = null;
+				_childNodeRefs = new List<BundleNodeRef<FileNode>>();
 
 				// Append the new node as a child
-				ChildNodeRefs.Add(new BundleNodeRef<FileNode>(this, NewNode));
+				_childNodeRefs.Add(new BundleNodeRef<FileNode>(this, newNode));
 			}
 		}
 
-		private ReadOnlyMemory<byte> AppendToNode(ReadOnlyMemory<byte> Data, ChunkingOptions Options)
+		private ReadOnlyMemory<byte> AppendToNode(ReadOnlyMemory<byte> data, ChunkingOptions options)
 		{
-			if (Data.Length == 0 || IsReadOnly())
+			if (data.Length == 0 || IsReadOnly())
 			{
-				return Data;
+				return data;
 			}
 
 			if (Depth == 0)
 			{
-				return AppendToLeafNode(Data, Options);
+				return AppendToLeafNode(data, options);
 			}
 			else
 			{
-				return AppendToInteriorNode(Data, Options);
+				return AppendToInteriorNode(data, options);
 			}
 		}
 
-		private ReadOnlyMemory<byte> AppendToLeafNode(ReadOnlyMemory<byte> NewData, ChunkingOptions Options)
+		private ReadOnlyMemory<byte> AppendToLeafNode(ReadOnlyMemory<byte> newData, ChunkingOptions options)
 		{
 			// Fast path for appending data to the buffer up to the chunk window size
-			int WindowSize = Options.LeafOptions.MinSize;
-			if (Payload.Length < WindowSize)
+			int windowSize = options.LeafOptions.MinSize;
+			if (Payload.Length < windowSize)
 			{
-				int AppendLength = Math.Min(WindowSize - (int)Payload.Length, NewData.Length);
-				AppendLeafData(NewData.Span.Slice(0, AppendLength), Options);
-				NewData = NewData.Slice(AppendLength);
+				int appendLength = Math.Min(windowSize - (int)Payload.Length, newData.Length);
+				AppendLeafData(newData.Span.Slice(0, appendLength));
+				newData = newData.Slice(appendLength);
 			}
 
 			// Cap the maximum amount of data to append to this node
-			int MaxLength = Math.Min(NewData.Length, Options.LeafOptions.MaxSize - (int)Payload.Length);
-			if (MaxLength > 0)
+			int maxLength = Math.Min(newData.Length, options.LeafOptions.MaxSize - (int)Payload.Length);
+			if (maxLength > 0)
 			{
-				ReadOnlySpan<byte> InputSpan = NewData.Span.Slice(0, MaxLength);
-				int Length = AppendLeafDataToChunkBoundary(InputSpan, Options);
-				NewData = NewData.Slice(Length);
+				ReadOnlySpan<byte> inputSpan = newData.Span.Slice(0, maxLength);
+				int length = AppendLeafDataToChunkBoundary(inputSpan, options);
+				newData = newData.Slice(length);
 			}
 
 			// Compute the hash if this node is complete
-			if (Payload.Length == Options.LeafOptions.MaxSize)
+			if (Payload.Length == options.LeafOptions.MaxSize)
 			{
 				MarkComplete();
 			}
-			return NewData;
+			return newData;
 		}
 
-		private int AppendLeafDataToChunkBoundary(ReadOnlySpan<byte> InputSpan, ChunkingOptions Options)
+		private int AppendLeafDataToChunkBoundary(ReadOnlySpan<byte> inputSpan, ChunkingOptions options)
 		{
-			int WindowSize = Options.LeafOptions.MinSize;
-			Debug.Assert(Payload.Length >= WindowSize);
+			int windowSize = options.LeafOptions.MinSize;
+			Debug.Assert(Payload.Length >= windowSize);
 
 			// If this is the first time we're hashing the data, compute the hash of the existing window
-			if (Payload.Length == WindowSize)
+			if (Payload.Length == windowSize)
 			{
-				foreach (ReadOnlyMemory<byte> Segment in Payload)
+				foreach (ReadOnlyMemory<byte> segment in Payload)
 				{
-					RollingHash = BuzHash.Add(RollingHash, Segment.Span);
+					_rollingHash = BuzHash.Add(_rollingHash, segment.Span);
 				}
 			}
 
 			// Get the threshold for the rolling hash
-			uint RollingHashThreshold = (uint)((1L << 32) / Options.LeafOptions.TargetSize);
+			uint rollingHashThreshold = (uint)((1L << 32) / options.LeafOptions.TargetSize);
 
 			// Length of the data taken from the input span, updated as we step through it.
-			int Length = 0;
+			int length = 0;
 
 			// Step the window through the tail end of the existing payload window. In this state, update the hash to remove data from the current payload, and add data from the new payload.
-			int SplitLength = Math.Min(InputSpan.Length, WindowSize);
-			ReadOnlySequence<byte> SplitPayload = Payload.Slice(Payload.Length - WindowSize, SplitLength);
+			int splitLength = Math.Min(inputSpan.Length, windowSize);
+			ReadOnlySequence<byte> splitPayload = Payload.Slice(Payload.Length - windowSize, splitLength);
 
-			foreach (ReadOnlyMemory<byte> PayloadSegment in SplitPayload)
+			foreach (ReadOnlyMemory<byte> payloadSegment in splitPayload)
 			{
-				int BaseLength = Length;
-				int SpanLength = Length + PayloadSegment.Length;
+				int baseLength = length;
+				int spanLength = length + payloadSegment.Length;
 
-				ReadOnlySpan<byte> PayloadSpan = PayloadSegment.Span;
-				for (; Length < SpanLength; Length++)
+				ReadOnlySpan<byte> payloadSpan = payloadSegment.Span;
+				for (; length < spanLength; length++)
 				{
-					RollingHash = BuzHash.Add(RollingHash, InputSpan[Length]);
-					if (RollingHash < RollingHashThreshold)
+					_rollingHash = BuzHash.Add(_rollingHash, inputSpan[length]);
+					if (_rollingHash < rollingHashThreshold)
 					{
-						AppendLeafData(InputSpan.Slice(0, Length), Options);
+						AppendLeafData(inputSpan.Slice(0, length));
 						MarkComplete();
-						return Length;
+						return length;
 					}
-					RollingHash = BuzHash.Sub(RollingHash, PayloadSpan[Length - BaseLength], WindowSize);
+					_rollingHash = BuzHash.Sub(_rollingHash, payloadSpan[length - baseLength], windowSize);
 				}
 			}
 
 			// Step through the new window.
-			for (; Length < InputSpan.Length; Length++)
+			for (; length < inputSpan.Length; length++)
 			{
-				RollingHash = BuzHash.Add(RollingHash, InputSpan[Length]);
-				if (RollingHash < RollingHashThreshold)
+				_rollingHash = BuzHash.Add(_rollingHash, inputSpan[length]);
+				if (_rollingHash < rollingHashThreshold)
 				{
-					AppendLeafData(InputSpan.Slice(0, Length), Options);
+					AppendLeafData(inputSpan.Slice(0, length));
 					MarkComplete();
-					return Length;
+					return length;
 				}
-				RollingHash = BuzHash.Sub(RollingHash, InputSpan[Length - WindowSize], WindowSize);
+				_rollingHash = BuzHash.Sub(_rollingHash, inputSpan[length - windowSize], windowSize);
 			}
 
-			AppendLeafData(InputSpan, Options);
-			return InputSpan.Length;
+			AppendLeafData(inputSpan);
+			return inputSpan.Length;
 		}
 
-		private void AppendLeafData(ReadOnlySpan<byte> LeafData, ChunkingOptions Options)
+		private void AppendLeafData(ReadOnlySpan<byte> leafData)
 		{
-			int DepthBytes = VarInt.Measure(Depth);
-
-			if (LastSegment == null)
+			if (_lastSegment == null)
 			{
-				byte[] Buffer = CreateFirstSegmentData(Depth, LeafData);
-				FirstSegment = new DataSegment(0, Buffer);
-				LastSegment = FirstSegment;
+				byte[] buffer = CreateFirstSegmentData(Depth, leafData);
+				_firstSegment = new DataSegment(0, buffer);
+				_lastSegment = _firstSegment;
 			}
 			else
 			{
-				DataSegment NewSegment = new DataSegment(LastSegment.RunningIndex + LastSegment.Memory.Length, LeafData.ToArray());
-				LastSegment.SetNext(NewSegment);
-				LastSegment = NewSegment;
+				DataSegment newSegment = new DataSegment(_lastSegment.RunningIndex + _lastSegment.Memory.Length, leafData.ToArray());
+				_lastSegment.SetNext(newSegment);
+				_lastSegment = newSegment;
 			}
 
-			int HeaderSize = (int)(Data.Length - Payload.Length);
-			Data = new ReadOnlySequence<byte>(FirstSegment!, 0, LastSegment, LastSegment.Memory.Length);
-			Payload = Data.Slice(HeaderSize);
+			int headerSize = (int)(_data.Length - Payload.Length);
+			_data = new ReadOnlySequence<byte>(_firstSegment!, 0, _lastSegment, _lastSegment.Memory.Length);
+			Payload = _data.Slice(headerSize);
 
-			Length += LeafData.Length;
+			Length += leafData.Length;
 		}
 
-		private static byte[] CreateFirstSegmentData(int Depth, ReadOnlySpan<byte> LeafData)
+		private static byte[] CreateFirstSegmentData(int depth, ReadOnlySpan<byte> leafData)
 		{
-			int DepthBytes = VarInt.Measure(Depth);
+			int depthBytes = VarInt.Measure(depth);
 
-			byte[] Buffer = new byte[1 + DepthBytes + LeafData.Length];
-			Buffer[0] = TypeId;
+			byte[] buffer = new byte[1 + depthBytes + leafData.Length];
+			buffer[0] = TypeId;
 
-			VarInt.Write(Buffer.AsSpan(1, DepthBytes), Depth);
-			LeafData.CopyTo(Buffer.AsSpan(1 + DepthBytes));
+			VarInt.Write(buffer.AsSpan(1, depthBytes), depth);
+			leafData.CopyTo(buffer.AsSpan(1 + depthBytes));
 
-			return Buffer;
+			return buffer;
 		}
 
 		private void MarkComplete()
 		{
-			if (Hash == null)
+			if (_hash == null)
 			{
-				if (ChildNodeRefs != null)
+				if (_childNodeRefs != null)
 				{
-					int DepthBytes = VarInt.Measure(Depth);
-					int LengthBytes = VarInt.Measure((ulong)Length);
+					int depthBytes = VarInt.Measure(Depth);
+					int lengthBytes = VarInt.Measure((ulong)Length);
 
-					byte[] NewData = new byte[1 + DepthBytes + LengthBytes + (ChildNodeRefs.Count * IoHash.NumBytes)];
-					Data = new ReadOnlySequence<byte>(NewData);
-					Payload = new ReadOnlySequence<byte>(NewData.AsMemory(1 + DepthBytes));
+					byte[] newData = new byte[1 + depthBytes + lengthBytes + (_childNodeRefs.Count * IoHash.NumBytes)];
+					_data = new ReadOnlySequence<byte>(newData);
+					Payload = new ReadOnlySequence<byte>(newData.AsMemory(1 + depthBytes));
 
-					NewData[0] = TypeId;
-					Span<byte> Span = NewData.AsSpan(1);
+					newData[0] = TypeId;
+					Span<byte> span = newData.AsSpan(1);
 
-					VarInt.Write(Span, Depth);
-					Span = Span.Slice(DepthBytes);
+					VarInt.Write(span, Depth);
+					span = span.Slice(depthBytes);
 
-					VarInt.Write(Span, Length);
-					Span = Span.Slice(LengthBytes);
+					VarInt.Write(span, Length);
+					span = span.Slice(lengthBytes);
 
-					foreach (BundleNodeRef<FileNode> ChildNodeRef in ChildNodeRefs)
+					foreach (BundleNodeRef<FileNode> childNodeRef in _childNodeRefs)
 					{
-						ChildNodeRef.Hash.CopyTo(Span);
-						Span = Span.Slice(IoHash.NumBytes);
+						childNodeRef.Hash.CopyTo(span);
+						span = span.Slice(IoHash.NumBytes);
 					}
 
-					Debug.Assert(Span.Length == 0);
+					Debug.Assert(span.Length == 0);
 				}
-				Hash = IoHash.Compute(Data);
+				_hash = IoHash.Compute(_data);
 			}
 		}
 
-		private ReadOnlyMemory<byte> AppendToInteriorNode(ReadOnlyMemory<byte> NewData, ChunkingOptions Options)
+		private ReadOnlyMemory<byte> AppendToInteriorNode(ReadOnlyMemory<byte> newData, ChunkingOptions options)
 		{
 			for (; ; )
 			{
-				Debug.Assert(ChildNodeRefs != null);
+				Debug.Assert(_childNodeRefs != null);
 
 				// Try to write to the last node
-				if (ChildNodeRefs.Count > 0)
+				if (_childNodeRefs.Count > 0)
 				{
-					FileNode? LastNode = ChildNodeRefs[^1].Node;
-					if (LastNode != null)
+					FileNode? lastNode = _childNodeRefs[^1].Node;
+					if (lastNode != null)
 					{
 						// Update the length to match the new node
-						Length -= LastNode.Length;
-						NewData = LastNode.AppendToNode(NewData, Options);
-						Length += LastNode.Length;
+						Length -= lastNode.Length;
+						newData = lastNode.AppendToNode(newData, options);
+						Length += lastNode.Length;
 
 						// If the last node is complete, write it to the buffer
-						if (LastNode.IsReadOnly())
+						if (lastNode.IsReadOnly())
 						{
 							// Update the hash
-							byte[] HashData = LastNode.Hash!.Value.ToByteArray();
-							RollingHash = BuzHash.Add(RollingHash, HashData);
+							byte[] hashData = lastNode._hash!.Value.ToByteArray();
+							_rollingHash = BuzHash.Add(_rollingHash, hashData);
 
 							// Check if it's time to finish this chunk
-							uint HashThreshold = (uint)(((1L << 32) * IoHash.NumBytes) / Options.LeafOptions.TargetSize);
-							if ((Payload.Length >= Options.InteriorOptions.MinSize && RollingHash < HashThreshold) || (Payload.Length >= Options.InteriorOptions.MaxSize))
+							uint hashThreshold = (uint)(((1L << 32) * IoHash.NumBytes) / options.LeafOptions.TargetSize);
+							if ((Payload.Length >= options.InteriorOptions.MinSize && _rollingHash < hashThreshold) || (Payload.Length >= options.InteriorOptions.MaxSize))
 							{
 								MarkComplete();
-								return NewData;
+								return newData;
 							}
 						}
 
 						// Bail out if there's nothing left to write
-						if (NewData.Length == 0)
+						if (newData.Length == 0)
 						{
-							return NewData;
+							return newData;
 						}
 					}
 				}
 
 				// Add a new child node
-				FileNode ChildNode = new FileNode();
-				ChildNode.IsRoot = false;
-				ChildNodeRefs.Add(new BundleNodeRef<FileNode>(this, ChildNode));
+				FileNode childNode = new FileNode();
+				childNode._isRoot = false;
+				_childNodeRefs.Add(new BundleNodeRef<FileNode>(this, childNode));
 			}
 		}
 	}
@@ -553,6 +549,6 @@ namespace EpicGames.Horde.Bundles.Nodes
 	public class FileNodeDeserializer : BundleNodeDeserializer<FileNode>
 	{
 		/// <inheritdoc/>
-		public override FileNode Deserialize(ReadOnlyMemory<byte> Data) => FileNode.Deserialize(Data);
+		public override FileNode Deserialize(ReadOnlyMemory<byte> data) => FileNode.Deserialize(data);
 	}
 }
