Horde: Use standard coding style for EpicGames.Horde.

diff --git a/Engine/Source/Programs/Shared/EpicGames.Horde/Common/Condition.cs b/Engine/Source/Programs/Shared/EpicGames.Horde/Common/Condition.cs
--- a/Engine/Source/Programs/Shared/EpicGames.Horde/Common/Condition.cs
+++ b/Engine/Source/Programs/Shared/EpicGames.Horde/Common/Condition.cs
@@ -2,7 +2,6 @@
 
 using EpicGames.Core;
 using EpicGames.Serialization;
-using EpicGames.Serialization.Converters;
 using System;
 using System.Collections.Generic;
 using System.ComponentModel;
@@ -10,7 +9,6 @@ using System.Diagnostics;
 using System.Diagnostics.CodeAnalysis;
 using System.Globalization;
 using System.Linq;
-using System.Text;
 using System.Text.Json;
 using System.Text.Json.Serialization;
 using System.Text.RegularExpressions;
@@ -49,17 +47,17 @@ namespace EpicGames.Horde.Common
 		public TokenType Type { get; private set; }
 		public string? Scalar { get; private set; }
 
-		public TokenReader(string Input)
+		public TokenReader(string input)
 		{
-			this.Input = Input;
+			Input = input;
 			MoveNext();
 		}
 
-		private void SetCurrent(int Length, TokenType Type, string? Scalar = null)
+		private void SetCurrent(int length, TokenType type, string? scalar = null)
 		{
-			this.Length = Length;
-			this.Type = Type;
-			this.Scalar = Scalar;
+			Length = length;
+			Type = type;
+			Scalar = scalar;
 		}
 
 		public void MoveNext()
@@ -138,31 +136,31 @@ namespace EpicGames.Horde.Common
 					case '\'':
 						ParseString();
 						return;
-					case char Char when IsNumber(Char):
+					case char character when IsNumber(character):
 						ParseNumber();
 						return;
-					case char Char when IsIdentifier(Char):
-						int EndIdx = Offset + 1;
-						while (EndIdx < Input.Length && IsIdentifierTail(Input[EndIdx]))
+					case char character when IsIdentifier(character):
+						int endIdx = Offset + 1;
+						while (endIdx < Input.Length && IsIdentifierTail(Input[endIdx]))
 						{
-							EndIdx++;
+							endIdx++;
 						}
 
-						TokenType Type;
-						if (EndIdx == Offset + 4 && String.Compare(Input, Offset, "true", 0, 4, StringComparison.OrdinalIgnoreCase) == 0)
+						TokenType type;
+						if (endIdx == Offset + 4 && String.Compare(Input, Offset, "true", 0, 4, StringComparison.OrdinalIgnoreCase) == 0)
 						{
-							Type = TokenType.True;
+							type = TokenType.True;
 						}
-						else if (EndIdx == Offset + 5 && String.Compare(Input, Offset, "false", 0, 5, StringComparison.OrdinalIgnoreCase) == 0)
+						else if (endIdx == Offset + 5 && String.Compare(Input, Offset, "false", 0, 5, StringComparison.OrdinalIgnoreCase) == 0)
 						{
-							Type = TokenType.False;
+							type = TokenType.False;
 						}
 						else
 						{
-							Type = TokenType.Identifier;
+							type = TokenType.Identifier;
 						}
 
-						SetCurrent(EndIdx - Offset, Type);
+						SetCurrent(endIdx - Offset, type);
 						return;
 					default:
 						SetCurrent(1, TokenType.Error, $"Invalid character at offset {Offset}: '{Input[Offset]}'");
@@ -175,52 +173,52 @@ namespace EpicGames.Horde.Common
 
 		bool ParseString()
 		{
-			char QuoteChar = Input[Offset];
-			if (QuoteChar != '\'' && QuoteChar != '\"')
+			char quoteChar = Input[Offset];
+			if (quoteChar != '\'' && quoteChar != '\"')
 			{
-				SetCurrent(1, TokenType.Error, $"Invalid quote character '{(char)QuoteChar}' at offset {Offset}");
+				SetCurrent(1, TokenType.Error, $"Invalid quote character '{(char)quoteChar}' at offset {Offset}");
 				return false;
 			}
 
-			int NumEscapeChars = 0;
+			int numEscapeChars = 0;
 
-			int EndIdx = Offset + 1;
-			for (; ; EndIdx++)
+			int endIdx = Offset + 1;
+			for (; ; endIdx++)
 			{
-				if (EndIdx >= Input.Length)
+				if (endIdx >= Input.Length)
 				{
-					SetCurrent(EndIdx - Offset, TokenType.Error, "Unterminated string in expression");
+					SetCurrent(endIdx - Offset, TokenType.Error, "Unterminated string in expression");
 					return false;
 				}
-				else if (Input[EndIdx] == '\\')
+				else if (Input[endIdx] == '\\')
 				{
-					NumEscapeChars++;
-					EndIdx++;
+					numEscapeChars++;
+					endIdx++;
 				}
-				else if (Input[EndIdx] == QuoteChar)
+				else if (Input[endIdx] == quoteChar)
 				{
 					break;
 				}
 			}
 
-			char[] Copy = new char[EndIdx - (Offset + 1) - NumEscapeChars];
+			char[] copy = new char[endIdx - (Offset + 1) - numEscapeChars];
 
-			int InputIdx = Offset + 1;
-			int OutputIdx = 0;
-			while (OutputIdx < Copy.Length)
+			int inputIdx = Offset + 1;
+			int outputIdx = 0;
+			while (outputIdx < copy.Length)
 			{
-				if (Input[InputIdx] == '\\')
+				if (Input[inputIdx] == '\\')
 				{
-					InputIdx++;
+					inputIdx++;
 				}
-				Copy[OutputIdx++] = Input[InputIdx++];
+				copy[outputIdx++] = Input[inputIdx++];
 			}
 
-			SetCurrent(EndIdx + 1 - Offset, TokenType.Scalar, new string(Copy));
+			SetCurrent(endIdx + 1 - Offset, TokenType.Scalar, new string(copy));
 			return true;
 		}
 
-		static Dictionary<StringView, ulong> SizeSuffixes = new Dictionary<StringView, ulong>(StringViewComparer.OrdinalIgnoreCase)
+		static readonly Dictionary<StringView, ulong> s_sizeSuffixes = new Dictionary<StringView, ulong>(StringViewComparer.OrdinalIgnoreCase)
 		{
 			["kb"] = 1024UL,
 			["Mb"] = 1024UL * 1024,
@@ -230,61 +228,61 @@ namespace EpicGames.Horde.Common
 
 		bool ParseNumber()
 		{
-			ulong Value = 0;
+			ulong value = 0;
 
-			int EndIdx = Offset;
-			while (EndIdx < Input.Length && IsNumber(Input[EndIdx]))
+			int endIdx = Offset;
+			while (endIdx < Input.Length && IsNumber(Input[endIdx]))
 			{
-				Value = (Value * 10) + (uint)(Input[EndIdx] - '0');
-				EndIdx++;
+				value = (value * 10) + (uint)(Input[endIdx] - '0');
+				endIdx++;
 			}
 
-			if (EndIdx < Input.Length && IsIdentifier(Input[EndIdx]))
+			if (endIdx < Input.Length && IsIdentifier(Input[endIdx]))
 			{
-				int Offset = EndIdx++;
-				while (EndIdx < Input.Length && IsIdentifierTail(Input[EndIdx]))
+				int offset = endIdx++;
+				while (endIdx < Input.Length && IsIdentifierTail(Input[endIdx]))
 				{
-					EndIdx++;
+					endIdx++;
 				}
 
-				StringView Suffix = new StringView(Input, Offset, EndIdx - Offset);
+				StringView suffix = new StringView(Input, offset, endIdx - offset);
 
-				ulong Size;
-				if (!SizeSuffixes.TryGetValue(Suffix, out Size))
+				ulong size;
+				if (!s_sizeSuffixes.TryGetValue(suffix, out size))
 				{
-					SetCurrent((EndIdx + 1) - Offset, TokenType.Error, $"'{Suffix}' is not a valid numerical suffix");
+					SetCurrent((endIdx + 1) - offset, TokenType.Error, $"'{suffix}' is not a valid numerical suffix");
 					return false;
 				}
-				Value *= Size;
+				value *= size;
 			}
 
-			SetCurrent(EndIdx - Offset, TokenType.Scalar, Value.ToString(CultureInfo.InvariantCulture));
+			SetCurrent(endIdx - Offset, TokenType.Scalar, value.ToString(CultureInfo.InvariantCulture));
 			return true;
 		}
 
-		bool RequireCharacter(string Text, int Idx, char Character)
+		bool RequireCharacter(string text, int idx, char character)
 		{
-			if (Idx == Text.Length || Text[Idx] != Character)
+			if (idx == text.Length || text[idx] != character)
 			{
-				SetCurrent(1, TokenType.Error, $"Invalid character at position {Idx}; expected '{Character}'.");
+				SetCurrent(1, TokenType.Error, $"Invalid character at position {idx}; expected '{character}'.");
 				return false;
 			}
 			return true;
 		}
 
-		static bool IsNumber(char Character)
+		static bool IsNumber(char character)
 		{
-			return Character >= '0' && Character <= '9';
+			return character >= '0' && character <= '9';
 		}
 
-		static bool IsIdentifier(char Character)
+		static bool IsIdentifier(char character)
 		{
-			return (Character >= 'a' && Character <= 'z') || (Character >= 'A' && Character <= 'Z') || Character == '_';
+			return (character >= 'a' && character <= 'z') || (character >= 'A' && character <= 'Z') || character == '_';
 		}
 
-		static bool IsIdentifierTail(char Character)
+		static bool IsIdentifierTail(char character)
 		{
-			return IsIdentifier(Character) || IsNumber(Character) || Character == '-' || Character == '.';
+			return IsIdentifier(character) || IsNumber(character) || character == '-' || character == '.';
 		}
 	}
 
@@ -296,9 +294,9 @@ namespace EpicGames.Horde.Common
 		/// <summary>
 		/// Constructor
 		/// </summary>
-		/// <param name="Error"></param>
-		internal ConditionException(string Error)
-			: base(Error)
+		/// <param name="error"></param>
+		internal ConditionException(string error)
+			: base(error)
 		{
 		}
 	}
@@ -318,10 +316,10 @@ namespace EpicGames.Horde.Common
 			public readonly TokenType Type;
 			public readonly byte Index;
 
-			public Token(TokenType Type, int Index)
+			public Token(TokenType type, int index)
 			{
-				this.Type = Type;
-				this.Index = (byte)Index;
+				Type = type;
+				Index = (byte)index;
 			}
 		}
 
@@ -334,25 +332,26 @@ namespace EpicGames.Horde.Common
 		/// Error produced when parsing the condition
 		/// </summary>
 		public string? Error { get; private set; }
-		List<Token> Tokens = new List<Token>();
-		List<string> Strings = new List<string>();
+		
+		readonly List<Token> _tokens = new List<Token>();
+		readonly List<string> _strings = new List<string>();
 
-		private Condition(string Text)
+		private Condition(string text)
 		{
-			this.Text = Text;
+			Text = text;
 
-			TokenReader Reader = new TokenReader(Text);
-			if (Reader.Type != TokenType.End)
+			TokenReader reader = new TokenReader(text);
+			if (reader.Type != TokenType.End)
 			{
-				Error = ParseOrExpr(Reader);
+				Error = ParseOrExpr(reader);
 
-				if (Reader.Type == TokenType.Error)
+				if (reader.Type == TokenType.Error)
 				{
-					Error = Reader.Scalar;
+					Error = reader.Scalar;
 				}
-				else if (Error == null && Reader.Type != TokenType.End)
+				else if (Error == null && reader.Type != TokenType.End)
 				{
-					Error = $"Unexpected token at offset {Reader.Offset}: {Reader.Token}";
+					Error = $"Unexpected token at offset {reader.Offset}: {reader.Token}";
 				}
 			}
 		}
@@ -361,7 +360,7 @@ namespace EpicGames.Horde.Common
 		/// Determines if the condition is empty
 		/// </summary>
 		/// <returns>True if the condition is empty</returns>
-		public bool IsEmpty() => Tokens.Count == 0 && IsValid();
+		public bool IsEmpty() => _tokens.Count == 0 && IsValid();
 
 		/// <summary>
 		/// Checks if the condition has been parsed correctly
@@ -371,109 +370,109 @@ namespace EpicGames.Horde.Common
 		/// <summary>
 		/// Parse the given text as a condition
 		/// </summary>
-		/// <param name="Text">Condition text to parse</param>
+		/// <param name="text">Condition text to parse</param>
 		/// <returns>The new condition object</returns>
-		public static Condition Parse(string Text)
+		public static Condition Parse(string text)
 		{
-			Condition Condition = new Condition(Text);
-			if(!Condition.IsValid())
+			Condition condition = new Condition(text);
+			if(!condition.IsValid())
 			{
-				throw new ConditionException(Condition.Error!);
+				throw new ConditionException(condition.Error!);
 			}
-			return Condition;
+			return condition;
 		}
 
 		/// <summary>
 		/// Attempts to parse the given text as a condition
 		/// </summary>
-		/// <param name="Text">Condition to parse</param>
+		/// <param name="text">Condition to parse</param>
 		/// <returns>The parsed condition. Does not validate whether the parse completed successfully; call <see cref="Condition.IsValid"/> to verify.</returns>
-		public static Condition TryParse(string Text)
+		public static Condition TryParse(string text)
 		{
-			return new Condition(Text);
+			return new Condition(text);
 		}
 
-		string? ParseOrExpr(TokenReader Reader)
+		string? ParseOrExpr(TokenReader reader)
 		{
-			int StartCount = Tokens.Count;
+			int startCount = _tokens.Count;
 
-			string? Error = ParseAndExpr(Reader);
-			while (Error == null && Reader.Type == TokenType.LogicalOr)
+			string? error = ParseAndExpr(reader);
+			while (error == null && reader.Type == TokenType.LogicalOr)
 			{
-				Tokens.Insert(StartCount++, new Token(TokenType.LogicalOr, 0));
-				Reader.MoveNext();
-				Error = ParseAndExpr(Reader);
+				_tokens.Insert(startCount++, new Token(TokenType.LogicalOr, 0));
+				reader.MoveNext();
+				error = ParseAndExpr(reader);
 			}
 
-			return Error;
+			return error;
 		}
 
-		string? ParseAndExpr(TokenReader Reader)
+		string? ParseAndExpr(TokenReader reader)
 		{
-			int StartCount = Tokens.Count;
+			int startCount = _tokens.Count;
 
-			string? Error = ParseBooleanExpr(Reader);
-			while (Error == null && Reader.Type == TokenType.LogicalAnd)
+			string? error = ParseBooleanExpr(reader);
+			while (error == null && reader.Type == TokenType.LogicalAnd)
 			{
-				Tokens.Insert(StartCount++, new Token(TokenType.LogicalAnd, 0));
-				Reader.MoveNext();
-				Error = ParseBooleanExpr(Reader);
+				_tokens.Insert(startCount++, new Token(TokenType.LogicalAnd, 0));
+				reader.MoveNext();
+				error = ParseBooleanExpr(reader);
 			}
 
-			return Error;
+			return error;
 		}
 
-		string? ParseBooleanExpr(TokenReader Reader)
+		string? ParseBooleanExpr(TokenReader reader)
 		{
-			switch (Reader.Type)
+			switch (reader.Type)
 			{
 				case TokenType.Not:
-					Tokens.Add(new Token(Reader.Type, 0));
-					Reader.MoveNext();
-					if (Reader.Type != TokenType.Lparen)
+					_tokens.Add(new Token(reader.Type, 0));
+					reader.MoveNext();
+					if (reader.Type != TokenType.Lparen)
 					{
-						return $"Expected '(' at offset {Reader.Offset}";
+						return $"Expected '(' at offset {reader.Offset}";
 					}
-					return ParseSubExpr(Reader);
+					return ParseSubExpr(reader);
 				case TokenType.Lparen:
-					return ParseSubExpr(Reader);
+					return ParseSubExpr(reader);
 				case TokenType.True:
 				case TokenType.False:
-					Tokens.Add(new Token(Reader.Type, 0));
-					Reader.MoveNext();
+					_tokens.Add(new Token(reader.Type, 0));
+					reader.MoveNext();
 					return null;
 				default:
-					return ParseComparisonExpr(Reader);
+					return ParseComparisonExpr(reader);
 			}
 		}
 
-		string? ParseSubExpr(TokenReader Reader)
+		string? ParseSubExpr(TokenReader reader)
 		{
-			Reader.MoveNext();
+			reader.MoveNext();
 
-			string? Error = ParseOrExpr(Reader);
-			if (Error == null)
+			string? error = ParseOrExpr(reader);
+			if (error == null)
 			{
-				if (Reader.Type == TokenType.Rparen)
+				if (reader.Type == TokenType.Rparen)
 				{
-					Reader.MoveNext();
+					reader.MoveNext();
 				}
 				else
 				{
-					Error = $"Missing ')' at offset {Reader.Offset}";
+					error = $"Missing ')' at offset {reader.Offset}";
 				}
 			}
-			return Error;
+			return error;
 		}
 
-		string? ParseComparisonExpr(TokenReader Reader)
+		string? ParseComparisonExpr(TokenReader reader)
 		{
-			int StartCount = Tokens.Count;
+			int startCount = _tokens.Count;
 
-			string? Error = ParseScalarExpr(Reader);
-			if (Error == null)
+			string? error = ParseScalarExpr(reader);
+			if (error == null)
 			{
-				switch(Reader.Type)
+				switch(reader.Type)
 				{
 					case TokenType.Lt:
 					case TokenType.Lte:
@@ -482,130 +481,126 @@ namespace EpicGames.Horde.Common
 					case TokenType.Eq:
 					case TokenType.Neq:
 					case TokenType.Regex:
-						Tokens.Insert(StartCount, new Token(Reader.Type, 0));
-						Reader.MoveNext();
-						Error = ParseScalarExpr(Reader);
+						_tokens.Insert(startCount, new Token(reader.Type, 0));
+						reader.MoveNext();
+						error = ParseScalarExpr(reader);
 						break;
 				}
 			}
 
-			return Error;
+			return error;
 		}
 
-		string? ParseScalarExpr(TokenReader Reader)
+		string? ParseScalarExpr(TokenReader reader)
 		{
-			switch (Reader.Type)
+			switch (reader.Type)
 			{
 				case TokenType.Identifier:
-					Strings.Add(Reader.Token.ToString());
-					Tokens.Add(new Token(TokenType.Identifier, Strings.Count - 1));
-					Reader.MoveNext();
+					_strings.Add(reader.Token.ToString());
+					_tokens.Add(new Token(TokenType.Identifier, _strings.Count - 1));
+					reader.MoveNext();
 					return null;
 				case TokenType.Scalar:
-					Strings.Add(Reader.Scalar!);
-					Tokens.Add(new Token(TokenType.Scalar, Strings.Count - 1));
-					Reader.MoveNext();
+					_strings.Add(reader.Scalar!);
+					_tokens.Add(new Token(TokenType.Scalar, _strings.Count - 1));
+					reader.MoveNext();
 					return null;
 				default:
-					return $"Unexpected token '{Reader.Token}' at offset {Reader.Offset}";
+					return $"Unexpected token '{reader.Token}' at offset {reader.Offset}";
 			}
 		}
 
 		/// <summary>
 		/// Evaluate the condition using the given callback to retreive property values
 		/// </summary>
-		/// <param name="GetPropertyValues"></param>
+		/// <param name="getPropertyValues"></param>
 		/// <returns></returns>
-		public bool Evaluate(Func<string, IEnumerable<string>> GetPropertyValues)
+		public bool Evaluate(Func<string, IEnumerable<string>> getPropertyValues)
 		{
 			if (IsEmpty())
 			{
 				return true;
 			}
 
-			int Idx = 0;
-			return IsValid() && EvaluateCondition(ref Idx, GetPropertyValues);
+			int idx = 0;
+			return IsValid() && EvaluateCondition(ref idx, getPropertyValues);
 		}
 
-		bool EvaluateCondition(ref int Idx, Func<string, IEnumerable<string>> GetPropertyValues)
+		bool EvaluateCondition(ref int idx, Func<string, IEnumerable<string>> getPropertyValues)
 		{
-			bool LhsBool;
-			bool RhsBool;
-			IEnumerable<string> LhsScalar;
-			IEnumerable<string> RhsScalar;
+			bool lhsBool;
+			bool rhsBool;
+			IEnumerable<string> lhsScalar;
+			IEnumerable<string> rhsScalar;
 
-			Token Token = Tokens[Idx++];
-			switch (Token.Type)
+			Token token = _tokens[idx++];
+			switch (token.Type)
 			{
 				case TokenType.True:
 					return true;
 				case TokenType.False:
 					return false;
 				case TokenType.Not:
-					return !EvaluateCondition(ref Idx, GetPropertyValues);
+					return !EvaluateCondition(ref idx, getPropertyValues);
 				case TokenType.Eq:
-					LhsScalar = EvaluateScalar(ref Idx, GetPropertyValues);
-					RhsScalar = EvaluateScalar(ref Idx, GetPropertyValues);
-					return LhsScalar.Any(x => RhsScalar.Contains(x, StringComparer.OrdinalIgnoreCase));
+					lhsScalar = EvaluateScalar(ref idx, getPropertyValues);
+					rhsScalar = EvaluateScalar(ref idx, getPropertyValues);
+					return lhsScalar.Any(x => rhsScalar.Contains(x, StringComparer.OrdinalIgnoreCase));
 				case TokenType.Neq:
-					LhsScalar = EvaluateScalar(ref Idx, GetPropertyValues);
-					RhsScalar = EvaluateScalar(ref Idx, GetPropertyValues);
-					return !LhsScalar.Any(x => RhsScalar.Contains(x, StringComparer.OrdinalIgnoreCase));
+					lhsScalar = EvaluateScalar(ref idx, getPropertyValues);
+					rhsScalar = EvaluateScalar(ref idx, getPropertyValues);
+					return !lhsScalar.Any(x => rhsScalar.Contains(x, StringComparer.OrdinalIgnoreCase));
 				case TokenType.LogicalOr:
-					LhsBool = EvaluateCondition(ref Idx, GetPropertyValues);
-					RhsBool = EvaluateCondition(ref Idx, GetPropertyValues);
-					return LhsBool || RhsBool;
+					lhsBool = EvaluateCondition(ref idx, getPropertyValues);
+					rhsBool = EvaluateCondition(ref idx, getPropertyValues);
+					return lhsBool || rhsBool;
 				case TokenType.LogicalAnd:
-					LhsBool = EvaluateCondition(ref Idx, GetPropertyValues);
-					RhsBool = EvaluateCondition(ref Idx, GetPropertyValues);
-					return LhsBool && RhsBool;
+					lhsBool = EvaluateCondition(ref idx, getPropertyValues);
+					rhsBool = EvaluateCondition(ref idx, getPropertyValues);
+					return lhsBool && rhsBool;
 				case TokenType.Lt:
-					LhsScalar = EvaluateScalar(ref Idx, GetPropertyValues);
-					RhsScalar = EvaluateScalar(ref Idx, GetPropertyValues);
-					return AsIntegers(LhsScalar).Any(x => AsIntegers(RhsScalar).Any(y => x < y));
+					lhsScalar = EvaluateScalar(ref idx, getPropertyValues);
+					rhsScalar = EvaluateScalar(ref idx, getPropertyValues);
+					return AsIntegers(lhsScalar).Any(x => AsIntegers(rhsScalar).Any(y => x < y));
 				case TokenType.Lte:
-					LhsScalar = EvaluateScalar(ref Idx, GetPropertyValues);
-					RhsScalar = EvaluateScalar(ref Idx, GetPropertyValues);
-					return AsIntegers(LhsScalar).Any(x => AsIntegers(RhsScalar).Any(y => x <= y));
+					lhsScalar = EvaluateScalar(ref idx, getPropertyValues);
+					rhsScalar = EvaluateScalar(ref idx, getPropertyValues);
+					return AsIntegers(lhsScalar).Any(x => AsIntegers(rhsScalar).Any(y => x <= y));
 				case TokenType.Gt:
-					LhsScalar = EvaluateScalar(ref Idx, GetPropertyValues);
-					RhsScalar = EvaluateScalar(ref Idx, GetPropertyValues);
-					return AsIntegers(LhsScalar).Any(x => AsIntegers(RhsScalar).Any(y => x > y));
+					lhsScalar = EvaluateScalar(ref idx, getPropertyValues);
+					rhsScalar = EvaluateScalar(ref idx, getPropertyValues);
+					return AsIntegers(lhsScalar).Any(x => AsIntegers(rhsScalar).Any(y => x > y));
 				case TokenType.Gte:
-					LhsScalar = EvaluateScalar(ref Idx, GetPropertyValues);
-					RhsScalar = EvaluateScalar(ref Idx, GetPropertyValues);
-					return AsIntegers(LhsScalar).Any(x => AsIntegers(RhsScalar).Any(y => x >= y));
+					lhsScalar = EvaluateScalar(ref idx, getPropertyValues);
+					rhsScalar = EvaluateScalar(ref idx, getPropertyValues);
+					return AsIntegers(lhsScalar).Any(x => AsIntegers(rhsScalar).Any(y => x >= y));
 				case TokenType.Regex:
-					LhsScalar = EvaluateScalar(ref Idx, GetPropertyValues);
-					RhsScalar = EvaluateScalar(ref Idx, GetPropertyValues);
-					return LhsScalar.Any(x => RhsScalar.Any(y => Regex.IsMatch(x, y, RegexOptions.IgnoreCase)));
+					lhsScalar = EvaluateScalar(ref idx, getPropertyValues);
+					rhsScalar = EvaluateScalar(ref idx, getPropertyValues);
+					return lhsScalar.Any(x => rhsScalar.Any(y => Regex.IsMatch(x, y, RegexOptions.IgnoreCase)));
 				default:
 					throw new InvalidOperationException("Invalid token type");
 			}
 		}
 
-		IEnumerable<string> EvaluateScalar(ref int Idx, Func<string, IEnumerable<string>> GetPropertyValues)
+		IEnumerable<string> EvaluateScalar(ref int idx, Func<string, IEnumerable<string>> getPropertyValues)
 		{
-			Token Token = Tokens[Idx++];
-			switch (Token.Type)
+			Token token = _tokens[idx++];
+			return token.Type switch
 			{
-				case TokenType.Identifier:
-					return GetPropertyValues(Strings[Token.Index]);
-				case TokenType.Scalar:
-					return new string[] { Strings[Token.Index] };
-				default:
-					throw new InvalidOperationException("Invalid token type");
-			}
+				TokenType.Identifier => getPropertyValues(_strings[token.Index]),
+				TokenType.Scalar => new string[] { _strings[token.Index] },
+				_ => throw new InvalidOperationException("Invalid token type")
+			};
 		}
 
-		static IEnumerable<long> AsIntegers(IEnumerable<string> Scalars)
+		static IEnumerable<long> AsIntegers(IEnumerable<string> scalars)
 		{
-			foreach (string Scalar in Scalars)
+			foreach (string scalar in scalars)
 			{
-				long Value;
-				if (long.TryParse(Scalar, out Value))
+				if (Int64.TryParse(scalar, out long value))
 				{
-					yield return Value;
+					yield return value;
 				}
 			}
 		}
@@ -613,17 +608,17 @@ namespace EpicGames.Horde.Common
 		/// <summary>
 		/// Implicit conversion from string to conditions
 		/// </summary>
-		/// <param name="Text"></param>
+		/// <param name="text"></param>
 		[return: NotNullIfNotNull("Text")]
-		public static implicit operator Condition?(string? Text)
+		public static implicit operator Condition?(string? text)
 		{
-			if (Text == null)
+			if (text == null)
 			{
 				return null;
 			}
 			else
 			{
-				return new Condition(Text);
+				return new Condition(text);
 			}
 		}
 
@@ -637,13 +632,13 @@ namespace EpicGames.Horde.Common
 	public class ConditionCbConverter : CbConverterBase<Condition>
 	{
 		/// <inheritdoc/>
-		public override Condition Read(CbField Field) => Condition.TryParse(Field.AsUtf8String().ToString());
+		public override Condition Read(CbField field) => Condition.TryParse(field.AsUtf8String().ToString());
 
 		/// <inheritdoc/>
-		public override void Write(CbWriter Writer, Condition Value) => Writer.WriteUtf8StringValue(Value.Text);
+		public override void Write(CbWriter writer, Condition value) => writer.WriteUtf8StringValue(value.Text);
 
 		/// <inheritdoc/>
-		public override void WriteNamed(CbWriter Writer, Utf8String Name, Condition Value) => Writer.WriteUtf8String(Name, Value.Text);
+		public override void WriteNamed(CbWriter writer, Utf8String name, Condition value) => writer.WriteUtf8String(name, value.Text);
 	}
 
 	/// <summary>
@@ -652,16 +647,16 @@ namespace EpicGames.Horde.Common
 	sealed class ConditionTypeConverter : TypeConverter
 	{
 		/// <inheritdoc/>
-		public override bool CanConvertFrom(ITypeDescriptorContext Context, Type SourceType) => SourceType == typeof(string);
+		public override bool CanConvertFrom(ITypeDescriptorContext context, Type sourceType) => sourceType == typeof(string);
 
 		/// <inheritdoc/>
-		public override object ConvertFrom(ITypeDescriptorContext Context, CultureInfo Culture, object Value) => Condition.TryParse((string)Value);
+		public override object ConvertFrom(ITypeDescriptorContext context, CultureInfo culture, object value) => Condition.TryParse((string)value);
 
 		/// <inheritdoc/>
-		public override bool CanConvertTo(ITypeDescriptorContext Context, Type DestinationType) => DestinationType == typeof(string);
+		public override bool CanConvertTo(ITypeDescriptorContext context, Type destinationType) => destinationType == typeof(string);
 
 		/// <inheritdoc/>
-		public override object? ConvertTo(ITypeDescriptorContext Context, CultureInfo Culture, object Value, Type DestinationType) => ((Condition)Value).Text;
+		public override object? ConvertTo(ITypeDescriptorContext context, CultureInfo culture, object value, Type destinationType) => ((Condition)value).Text;
 	}
 
 	/// <summary>
@@ -674,12 +669,12 @@ namespace EpicGames.Horde.Common
 
 		public override Condition Read(ref Utf8JsonReader reader, Type typeToConvert, JsonSerializerOptions options)
 		{
-			string? String = reader.GetString();
-			if(String == null)
+			string? text = reader.GetString();
+			if(text == null)
 			{
 				throw new InvalidOperationException();
 			}
-			return Condition.Parse(String);
+			return Condition.Parse(text);
 		}
 
 		public override void Write(Utf8JsonWriter writer, Condition value, JsonSerializerOptions options)
diff --git a/Engine/Source/Programs/Shared/EpicGames.Horde/Compute/ChannelId.cs b/Engine/Source/Programs/Shared/EpicGames.Horde/Compute/ChannelId.cs
--- a/Engine/Source/Programs/Shared/EpicGames.Horde/Compute/ChannelId.cs
+++ b/Engine/Source/Programs/Shared/EpicGames.Horde/Compute/ChannelId.cs
@@ -3,11 +3,8 @@
 using EpicGames.Core;
 using EpicGames.Serialization;
 using System;
-using System.Collections.Generic;
 using System.ComponentModel;
-using System.Diagnostics.CodeAnalysis;
 using System.Globalization;
-using System.Text;
 using System.Text.Json;
 using System.Text.Json.Serialization;
 
@@ -24,34 +21,34 @@ namespace EpicGames.Horde.Compute
 		/// <summary>
 		/// The text representing this id
 		/// </summary>
-		readonly StringId Inner;
+		readonly StringId _inner;
 
 		/// <summary>
 		/// Constructor
 		/// </summary>
-		/// <param name="Input">Unique id for the string</param>
-		public ChannelId(string Input)
+		/// <param name="input">Unique id for the string</param>
+		public ChannelId(string input)
 		{
-			Inner = new StringId(Input);
+			_inner = new StringId(input);
 		}
 
 		/// <inheritdoc/>
-		public override bool Equals(object? Obj) => Obj is ChannelId Id && Inner.Equals(Id.Inner);
+		public override bool Equals(object? obj) => obj is ChannelId id && _inner.Equals(id._inner);
 
 		/// <inheritdoc/>
-		public override int GetHashCode() => Inner.GetHashCode();
+		public override int GetHashCode() => _inner.GetHashCode();
 
 		/// <inheritdoc/>
-		public bool Equals(ChannelId Other) => Inner.Equals(Other.Inner);
+		public bool Equals(ChannelId other) => _inner.Equals(other._inner);
 
 		/// <inheritdoc/>
-		public override string ToString() => Inner.ToString();
+		public override string ToString() => _inner.ToString();
 
 		/// <inheritdoc cref="StringId.op_Equality"/>
-		public static bool operator ==(ChannelId Left, ChannelId Right) => Left.Inner == Right.Inner;
+		public static bool operator ==(ChannelId left, ChannelId right) => left._inner == right._inner;
 
 		/// <inheritdoc cref="StringId.op_Inequality"/>
-		public static bool operator !=(ChannelId Left, ChannelId Right) => Left.Inner != Right.Inner;
+		public static bool operator !=(ChannelId left, ChannelId right) => left._inner != right._inner;
 	}
 
 	/// <summary>
@@ -60,13 +57,13 @@ namespace EpicGames.Horde.Compute
 	sealed class ChannelIdCbConverter : CbConverterBase<ChannelId>
 	{
 		/// <inheritdoc/>
-		public override ChannelId Read(CbField Field) => new ChannelId(Field.AsString());
+		public override ChannelId Read(CbField field) => new ChannelId(field.AsString());
 
 		/// <inheritdoc/>
-		public override void Write(CbWriter Writer, ChannelId Value) => Writer.WriteStringValue(Value.ToString());
+		public override void Write(CbWriter writer, ChannelId value) => writer.WriteStringValue(value.ToString());
 
 		/// <inheritdoc/>
-		public override void WriteNamed(CbWriter Writer, Utf8String Name, ChannelId Value) => Writer.WriteString(Name, Value.ToString());
+		public override void WriteNamed(CbWriter writer, Utf8String name, ChannelId value) => writer.WriteString(name, value.ToString());
 	}
 
 	/// <summary>
@@ -75,10 +72,10 @@ namespace EpicGames.Horde.Compute
 	sealed class ChannelIdJsonConverter : JsonConverter<ChannelId>
 	{
 		/// <inheritdoc/>
-		public override ChannelId Read(ref Utf8JsonReader Reader, Type TypeToConvert, JsonSerializerOptions Options) => new ChannelId(Reader.GetString() ?? String.Empty);
+		public override ChannelId Read(ref Utf8JsonReader reader, Type typeToConvert, JsonSerializerOptions options) => new ChannelId(reader.GetString() ?? String.Empty);
 
 		/// <inheritdoc/>
-		public override void Write(Utf8JsonWriter Writer, ChannelId Value, JsonSerializerOptions Options) => Writer.WriteStringValue(Value.ToString());
+		public override void Write(Utf8JsonWriter writer, ChannelId value, JsonSerializerOptions options) => writer.WriteStringValue(value.ToString());
 	}
 
 	/// <summary>
@@ -87,9 +84,9 @@ namespace EpicGames.Horde.Compute
 	sealed class ChannelIdTypeConverter : TypeConverter
 	{
 		/// <inheritdoc/>
-		public override bool CanConvertFrom(ITypeDescriptorContext Context, Type SourceType) => SourceType == typeof(string);
+		public override bool CanConvertFrom(ITypeDescriptorContext context, Type sourceType) => sourceType == typeof(string);
 
 		/// <inheritdoc/>
-		public override object ConvertFrom(ITypeDescriptorContext Context, CultureInfo Culture, object Value) => new ChannelId((string)Value);
+		public override object ConvertFrom(ITypeDescriptorContext context, CultureInfo culture, object value) => new ChannelId((string)value);
 	}
 }
diff --git a/Engine/Source/Programs/Shared/EpicGames.Horde/Compute/ClusterId.cs b/Engine/Source/Programs/Shared/EpicGames.Horde/Compute/ClusterId.cs
--- a/Engine/Source/Programs/Shared/EpicGames.Horde/Compute/ClusterId.cs
+++ b/Engine/Source/Programs/Shared/EpicGames.Horde/Compute/ClusterId.cs
@@ -3,11 +3,8 @@
 using EpicGames.Core;
 using EpicGames.Serialization;
 using System;
-using System.Collections.Generic;
 using System.ComponentModel;
-using System.Diagnostics.CodeAnalysis;
 using System.Globalization;
-using System.Text;
 using System.Text.Json;
 using System.Text.Json.Serialization;
 
@@ -24,34 +21,34 @@ namespace EpicGames.Horde.Compute
 		/// <summary>
 		/// The text representing this id
 		/// </summary>
-		readonly StringId Inner;
+		readonly StringId _inner;
 
 		/// <summary>
 		/// Constructor
 		/// </summary>
-		/// <param name="Input">Unique id for the string</param>
-		public ClusterId(string Input)
+		/// <param name="input">Unique id for the string</param>
+		public ClusterId(string input)
 		{
-			Inner = new StringId(Input);
+			_inner = new StringId(input);
 		}
 
 		/// <inheritdoc/>
-		public override bool Equals(object? Obj) => Obj is ClusterId Id && Inner.Equals(Id.Inner);
+		public override bool Equals(object? obj) => obj is ClusterId id && _inner.Equals(id._inner);
 
 		/// <inheritdoc/>
-		public override int GetHashCode() => Inner.GetHashCode();
+		public override int GetHashCode() => _inner.GetHashCode();
 
 		/// <inheritdoc/>
-		public bool Equals(ClusterId Other) => Inner.Equals(Other.Inner);
+		public bool Equals(ClusterId other) => _inner.Equals(other._inner);
 
 		/// <inheritdoc/>
-		public override string ToString() => Inner.ToString();
+		public override string ToString() => _inner.ToString();
 
 		/// <inheritdoc cref="StringId.op_Equality"/>
-		public static bool operator ==(ClusterId Left, ClusterId Right) => Left.Inner == Right.Inner;
+		public static bool operator ==(ClusterId left, ClusterId right) => left._inner == right._inner;
 
 		/// <inheritdoc cref="StringId.op_Inequality"/>
-		public static bool operator !=(ClusterId Left, ClusterId Right) => Left.Inner != Right.Inner;
+		public static bool operator !=(ClusterId left, ClusterId right) => left._inner != right._inner;
 	}
 
 	/// <summary>
@@ -60,13 +57,13 @@ namespace EpicGames.Horde.Compute
 	sealed class ClusterIdCbConverter : CbConverterBase<ClusterId>
 	{
 		/// <inheritdoc/>
-		public override ClusterId Read(CbField Field) => new ClusterId(Field.AsString());
+		public override ClusterId Read(CbField field) => new ClusterId(field.AsString());
 
 		/// <inheritdoc/>
-		public override void Write(CbWriter Writer, ClusterId Value) => Writer.WriteStringValue(Value.ToString());
+		public override void Write(CbWriter writer, ClusterId value) => writer.WriteStringValue(value.ToString());
 
 		/// <inheritdoc/>
-		public override void WriteNamed(CbWriter Writer, Utf8String Name, ClusterId Value) => Writer.WriteString(Name, Value.ToString());
+		public override void WriteNamed(CbWriter writer, Utf8String name, ClusterId value) => writer.WriteString(name, value.ToString());
 	}
 
 	/// <summary>
@@ -75,10 +72,10 @@ namespace EpicGames.Horde.Compute
 	sealed class ClusterIdJsonConverter : JsonConverter<ClusterId>
 	{
 		/// <inheritdoc/>
-		public override ClusterId Read(ref Utf8JsonReader Reader, Type TypeToConvert, JsonSerializerOptions Options) => new ClusterId(Reader.GetString() ?? String.Empty);
+		public override ClusterId Read(ref Utf8JsonReader reader, Type typeToConvert, JsonSerializerOptions options) => new ClusterId(reader.GetString() ?? String.Empty);
 
 		/// <inheritdoc/>
-		public override void Write(Utf8JsonWriter Writer, ClusterId Value, JsonSerializerOptions Options) => Writer.WriteStringValue(Value.ToString());
+		public override void Write(Utf8JsonWriter writer, ClusterId value, JsonSerializerOptions options) => writer.WriteStringValue(value.ToString());
 	}
 
 	/// <summary>
@@ -87,10 +84,10 @@ namespace EpicGames.Horde.Compute
 	sealed class ClusterIdTypeConverter : TypeConverter
 	{
 		/// <inheritdoc/>
-		public override bool CanConvertFrom(ITypeDescriptorContext Context, Type SourceType) => SourceType == typeof(string);
+		public override bool CanConvertFrom(ITypeDescriptorContext context, Type sourceType) => sourceType == typeof(string);
 
 		/// <inheritdoc/>
-		public override object ConvertFrom(ITypeDescriptorContext Context, CultureInfo Culture, object Value) => new ClusterId((string)Value);
+		public override object ConvertFrom(ITypeDescriptorContext context, CultureInfo culture, object value) => new ClusterId((string)value);
 	}
 }
 
diff --git a/Engine/Source/Programs/Shared/EpicGames.Horde/Compute/ComputeExtensions.cs b/Engine/Source/Programs/Shared/EpicGames.Horde/Compute/ComputeExtensions.cs
--- a/Engine/Source/Programs/Shared/EpicGames.Horde/Compute/ComputeExtensions.cs
+++ b/Engine/Source/Programs/Shared/EpicGames.Horde/Compute/ComputeExtensions.cs
@@ -5,8 +5,6 @@ using EpicGames.Horde.Compute.Impl;
 using Microsoft.Extensions.DependencyInjection;
 using Microsoft.Extensions.Options;
 using System;
-using System.Collections.Generic;
-using System.Text;
 
 namespace EpicGames.Horde.Compute
 {
@@ -25,24 +23,24 @@ namespace EpicGames.Horde.Compute
 		/// <summary>
 		/// Registers services for Horde Compute
 		/// </summary>
-		/// <param name="Services">The current service collection</param>
-		public static void AddHordeCompute(this IServiceCollection Services)
+		/// <param name="services">The current service collection</param>
+		public static void AddHordeCompute(this IServiceCollection services)
 		{
-			Services.AddOptions<ComputeOptions>();
+			services.AddOptions<ComputeOptions>();
 
-			Services.AddScoped<IComputeClient, HttpComputeClient>();
-			Services.AddHttpClientWithAuth<IComputeClient, HttpComputeClient>(ServiceProvider => ServiceProvider.GetRequiredService<IOptions<ComputeOptions>>().Value);
+			services.AddScoped<IComputeClient, HttpComputeClient>();
+			services.AddHttpClientWithAuth<IComputeClient, HttpComputeClient>(serviceProvider => serviceProvider.GetRequiredService<IOptions<ComputeOptions>>().Value);
 		}
 
 		/// <summary>
 		/// Registers services for Horde Compute
 		/// </summary>
-		/// <param name="Services">The current service collection</param>
-		/// <param name="Configure">Callback for configuring the storage service</param>
-		public static void AddHordeCompute(this IServiceCollection Services, Action<ComputeOptions> Configure)
+		/// <param name="services">The current service collection</param>
+		/// <param name="configure">Callback for configuring the storage service</param>
+		public static void AddHordeCompute(this IServiceCollection services, Action<ComputeOptions> configure)
 		{
-			AddHordeCompute(Services);
-			Services.Configure<ComputeOptions>(Configure);
+			AddHordeCompute(services);
+			services.Configure<ComputeOptions>(configure);
 		}
 	}
 }
diff --git a/Engine/Source/Programs/Shared/EpicGames.Horde/Compute/ComputeTask.cs b/Engine/Source/Programs/Shared/EpicGames.Horde/Compute/ComputeTask.cs
--- a/Engine/Source/Programs/Shared/EpicGames.Horde/Compute/ComputeTask.cs
+++ b/Engine/Source/Programs/Shared/EpicGames.Horde/Compute/ComputeTask.cs
@@ -2,9 +2,7 @@
 
 using EpicGames.Core;
 using EpicGames.Serialization;
-using System;
 using System.Collections.Generic;
-using System.Text;
 
 namespace EpicGames.Horde.Compute
 {
@@ -65,16 +63,16 @@ namespace EpicGames.Horde.Compute
 		/// <summary>
 		/// Constructor
 		/// </summary>
-		/// <param name="Executable">The executable to run</param>
-		/// <param name="Arguments">Arguments for the executable to run</param>
-		/// <param name="WorkingDirectory">Working directory for execution</param>
-		/// <param name="SandboxHash">Hash of the sandbox</param>
-		public ComputeTask(Utf8String Executable, List<Utf8String> Arguments, Utf8String WorkingDirectory, CbObjectAttachment SandboxHash)
+		/// <param name="executable">The executable to run</param>
+		/// <param name="arguments">Arguments for the executable to run</param>
+		/// <param name="workingDirectory">Working directory for execution</param>
+		/// <param name="sandboxHash">Hash of the sandbox</param>
+		public ComputeTask(Utf8String executable, List<Utf8String> arguments, Utf8String workingDirectory, CbObjectAttachment sandboxHash)
 		{
-			this.Executable = Executable;
-			this.Arguments = Arguments;
-			this.WorkingDirectory = WorkingDirectory;
-			this.SandboxHash = SandboxHash;
+			Executable = executable;
+			Arguments = arguments;
+			WorkingDirectory = workingDirectory;
+			SandboxHash = sandboxHash;
 		}
 	}
 }
diff --git a/Engine/Source/Programs/Shared/EpicGames.Horde/Compute/ComputeTaskResult.cs b/Engine/Source/Programs/Shared/EpicGames.Horde/Compute/ComputeTaskResult.cs
--- a/Engine/Source/Programs/Shared/EpicGames.Horde/Compute/ComputeTaskResult.cs
+++ b/Engine/Source/Programs/Shared/EpicGames.Horde/Compute/ComputeTaskResult.cs
@@ -1,10 +1,6 @@
 // Copyright Epic Games, Inc. All Rights Reserved.
 
-using EpicGames.Core;
 using EpicGames.Serialization;
-using System;
-using System.Collections.Generic;
-using System.Text;
 
 namespace EpicGames.Horde.Compute
 {
@@ -47,10 +43,10 @@ namespace EpicGames.Horde.Compute
 		/// <summary>
 		/// Constructor
 		/// </summary>
-		/// <param name="ExitCode">Exit code of the process</param>
-		public ComputeTaskResult(int ExitCode)
+		/// <param name="exitCode">Exit code of the process</param>
+		public ComputeTaskResult(int exitCode)
 		{
-			this.ExitCode = ExitCode;
+			ExitCode = exitCode;
 		}
 	}
 }
diff --git a/Engine/Source/Programs/Shared/EpicGames.Horde/Compute/DirectoryTree.cs b/Engine/Source/Programs/Shared/EpicGames.Horde/Compute/DirectoryTree.cs
--- a/Engine/Source/Programs/Shared/EpicGames.Horde/Compute/DirectoryTree.cs
+++ b/Engine/Source/Programs/Shared/EpicGames.Horde/Compute/DirectoryTree.cs
@@ -2,9 +2,7 @@
 
 using EpicGames.Core;
 using EpicGames.Serialization;
-using System;
 using System.Collections.Generic;
-using System.Text;
 
 namespace EpicGames.Horde.Compute
 {
@@ -65,12 +63,12 @@ namespace EpicGames.Horde.Compute
 		/// <summary>
 		/// Constructor
 		/// </summary>
-		public FileNode(Utf8String Name, CbBinaryAttachment Hash, long Size, int Attributes)
+		public FileNode(Utf8String name, CbBinaryAttachment hash, long size, int attributes)
 		{
-			this.Name = Name;
-			this.Hash = Hash;
-			this.Size = Size;
-			this.Attributes = Attributes;
+			Name = name;
+			Hash = hash;
+			Size = size;
+			Attributes = attributes;
 		}
 	}
 
@@ -101,10 +99,10 @@ namespace EpicGames.Horde.Compute
 		/// <summary>
 		/// Constructor
 		/// </summary>
-		public DirectoryNode(Utf8String Name, CbObjectAttachment Hash)
+		public DirectoryNode(Utf8String name, CbObjectAttachment hash)
 		{
-			this.Name = Name;
-			this.Hash = Hash;
+			Name = name;
+			Hash = hash;
 		}
 	}
 }
diff --git a/Engine/Source/Programs/Shared/EpicGames.Horde/Compute/IComputeClient.cs b/Engine/Source/Programs/Shared/EpicGames.Horde/Compute/IComputeClient.cs
--- a/Engine/Source/Programs/Shared/EpicGames.Horde/Compute/IComputeClient.cs
+++ b/Engine/Source/Programs/Shared/EpicGames.Horde/Compute/IComputeClient.cs
@@ -1,13 +1,9 @@
 // Copyright Epic Games, Inc. All Rights Reserved.
 
 using EpicGames.Core;
-using EpicGames.Horde.Compute.Impl;
 using EpicGames.Horde.Storage;
-using EpicGames.Serialization;
 using System;
 using System.Collections.Generic;
-using System.Runtime.CompilerServices;
-using System.Text;
 using System.Threading;
 using System.Threading.Tasks;
 
@@ -155,31 +151,31 @@ namespace EpicGames.Horde.Compute
 		/// <summary>
 		/// Gets information about a cluster
 		/// </summary>
-		/// <param name="ClusterId">The cluster to retreive information on</param>
-		/// <param name="CancellationToken">Cancellation token for the request</param>
+		/// <param name="clusterId">The cluster to retreive information on</param>
+		/// <param name="cancellationToken">Cancellation token for the request</param>
 		/// <returns>Cluster information</returns>
-		Task<IComputeClusterInfo> GetClusterInfoAsync(ClusterId ClusterId, CancellationToken CancellationToken = default);
+		Task<IComputeClusterInfo> GetClusterInfoAsync(ClusterId clusterId, CancellationToken cancellationToken = default);
 
 		/// <summary>
 		/// Queues a set of tasks for execution
 		/// </summary>
-		/// <param name="ClusterId">Cluster for executing the tasks</param>
-		/// <param name="ChannelId">Indicates an identifier generated by the client to query responses on</param>
-		/// <param name="RequirementsHash">Requirements of the task to execute</param>
-		/// <param name="TaskRefIds">Refs describing the request</param>
-		/// <param name="SkipCacheLookup">Whether to skip the lookup of cached output for this task</param>
-		/// <param name="CancellationToken">Cancellation token for the request</param>
+		/// <param name="clusterId">Cluster for executing the tasks</param>
+		/// <param name="channelId">Indicates an identifier generated by the client to query responses on</param>
+		/// <param name="requirementsHash">Requirements of the task to execute</param>
+		/// <param name="taskRefIds">Refs describing the request</param>
+		/// <param name="skipCacheLookup">Whether to skip the lookup of cached output for this task</param>
+		/// <param name="cancellationToken">Cancellation token for the request</param>
 		/// <returns>The new channel id</returns>
-		Task AddTasksAsync(ClusterId ClusterId, ChannelId ChannelId, IEnumerable<RefId> TaskRefIds, IoHash RequirementsHash, bool SkipCacheLookup, CancellationToken CancellationToken = default);
+		Task AddTasksAsync(ClusterId clusterId, ChannelId channelId, IEnumerable<RefId> taskRefIds, IoHash requirementsHash, bool skipCacheLookup, CancellationToken cancellationToken = default);
 
 		/// <summary>
 		/// Read updates from the given remote execution channel
 		/// </summary>
-		/// <param name="ClusterId">Cluster for executing the tasks</param>
-		/// <param name="ChannelId">Channel to receive updates on</param>
-		/// <param name="CancellationToken">Cancellation token for the operation</param>
+		/// <param name="clusterId">Cluster for executing the tasks</param>
+		/// <param name="channelId">Channel to receive updates on</param>
+		/// <param name="cancellationToken">Cancellation token for the operation</param>
 		/// <returns></returns>
-		IAsyncEnumerable<IComputeTaskInfo> GetTaskUpdatesAsync(ClusterId ClusterId, ChannelId ChannelId, CancellationToken CancellationToken = default);
+		IAsyncEnumerable<IComputeTaskInfo> GetTaskUpdatesAsync(ClusterId clusterId, ChannelId channelId, CancellationToken cancellationToken = default);
 	}
 
 	/// <summary>
@@ -190,16 +186,16 @@ namespace EpicGames.Horde.Compute
 		/// <summary>
 		/// Queues a single task for execution
 		/// </summary>
-		/// <param name="ComputeClient">The compute client instance</param>
-		/// <param name="ClusterId">Name of the profile to use for executing the tasks</param>
-		/// <param name="ChannelId">Channel to receive updates on</param>
-		/// <param name="TaskRefId">Ref describing the request</param>
-		/// <param name="RequirementsHash">Requirements of the task to execute</param>
-		/// <param name="SkipCacheLookup">Whether to skip the lookup of cached output for this task</param>
-		/// <param name="CancellationToken">Cancellation token for the operation</param>
-		public static Task AddTaskAsync(this IComputeClient ComputeClient, ClusterId ClusterId, ChannelId ChannelId, RefId TaskRefId, IoHash RequirementsHash, bool SkipCacheLookup, CancellationToken CancellationToken = default)
+		/// <param name="computeClient">The compute client instance</param>
+		/// <param name="clusterId">Name of the profile to use for executing the tasks</param>
+		/// <param name="channelId">Channel to receive updates on</param>
+		/// <param name="taskRefId">Ref describing the request</param>
+		/// <param name="requirementsHash">Requirements of the task to execute</param>
+		/// <param name="skipCacheLookup">Whether to skip the lookup of cached output for this task</param>
+		/// <param name="cancellationToken">Cancellation token for the operation</param>
+		public static Task AddTaskAsync(this IComputeClient computeClient, ClusterId clusterId, ChannelId channelId, RefId taskRefId, IoHash requirementsHash, bool skipCacheLookup, CancellationToken cancellationToken = default)
 		{
-			return ComputeClient.AddTasksAsync(ClusterId, ChannelId, new[] { TaskRefId }, RequirementsHash, SkipCacheLookup, CancellationToken);
+			return computeClient.AddTasksAsync(clusterId, channelId, new[] { taskRefId }, requirementsHash, skipCacheLookup, cancellationToken);
 		}
 	}
 }
diff --git a/Engine/Source/Programs/Shared/EpicGames.Horde/Compute/Impl/HttpComputeClient.cs b/Engine/Source/Programs/Shared/EpicGames.Horde/Compute/Impl/HttpComputeClient.cs
--- a/Engine/Source/Programs/Shared/EpicGames.Horde/Compute/Impl/HttpComputeClient.cs
+++ b/Engine/Source/Programs/Shared/EpicGames.Horde/Compute/Impl/HttpComputeClient.cs
@@ -1,16 +1,13 @@
 // Copyright Epic Games, Inc. All Rights Reserved.
 
 using EpicGames.Core;
-using EpicGames.Horde.Common;
 using EpicGames.Horde.Storage;
 using EpicGames.Serialization;
 using System;
 using System.Collections.Generic;
 using System.ComponentModel.DataAnnotations;
-using System.Linq;
 using System.Net.Http;
 using System.Runtime.CompilerServices;
-using System.Text;
 using System.Threading;
 using System.Threading.Tasks;
 
@@ -123,60 +120,60 @@ namespace EpicGames.Horde.Compute.Impl
 	/// </summary>
 	public class HttpComputeClient : IComputeClient
 	{
-		HttpClient HttpClient;
+		readonly HttpClient _httpClient;
 
 		/// <summary>
 		/// Constructor
 		/// </summary>
-		/// <param name="HttpClient"></param>
-		public HttpComputeClient(HttpClient HttpClient)
+		/// <param name="httpClient"></param>
+		public HttpComputeClient(HttpClient httpClient)
 		{
-			this.HttpClient = HttpClient;
+			_httpClient = httpClient;
 		}
 
 		/// <inheritdoc/>
-		public async Task<IComputeClusterInfo> GetClusterInfoAsync(ClusterId ClusterId, CancellationToken CancellationToken)
+		public async Task<IComputeClusterInfo> GetClusterInfoAsync(ClusterId clusterId, CancellationToken cancellationToken)
 		{
-			return await HttpClient.GetAsync<GetComputeClusterInfo>($"api/v1/compute/{ClusterId}", CancellationToken);
+			return await _httpClient.GetAsync<GetComputeClusterInfo>($"api/v1/compute/{clusterId}", cancellationToken);
 		}
 
 		/// <inheritdoc/>
-		public async Task AddTasksAsync(ClusterId ClusterId, ChannelId ChannelId, IEnumerable<RefId> TaskRefIds, IoHash RequirementsHash, bool SkipCacheLookup, CancellationToken CancellationToken)
+		public async Task AddTasksAsync(ClusterId clusterId, ChannelId channelId, IEnumerable<RefId> taskRefIds, IoHash requirementsHash, bool skipCacheLookup, CancellationToken cancellationToken)
 		{
-			AddTasksRequest AddTasks = new AddTasksRequest();
-			AddTasks.ChannelId = ChannelId;
-			AddTasks.TaskRefIds.AddRange(TaskRefIds);
-			AddTasks.RequirementsHash = RequirementsHash;
-			AddTasks.DoNotCache = SkipCacheLookup;
+			AddTasksRequest addTasks = new AddTasksRequest();
+			addTasks.ChannelId = channelId;
+			addTasks.TaskRefIds.AddRange(taskRefIds);
+			addTasks.RequirementsHash = requirementsHash;
+			addTasks.DoNotCache = skipCacheLookup;
 
-			ReadOnlyMemoryContent Content = new ReadOnlyMemoryContent(CbSerializer.Serialize(AddTasks).GetView());
-			Content.Headers.Add("Content-Type", "application/x-ue-cb");
+			ReadOnlyMemoryContent content = new ReadOnlyMemoryContent(CbSerializer.Serialize(addTasks).GetView());
+			content.Headers.Add("Content-Type", "application/x-ue-cb");
 
-			HttpRequestMessage Request = new HttpRequestMessage(HttpMethod.Post, $"api/v1/compute/{ClusterId}");
-			Request.Content = Content;
+			HttpRequestMessage request = new HttpRequestMessage(HttpMethod.Post, $"api/v1/compute/{clusterId}");
+			request.Content = content;
 
-			HttpResponseMessage Response = await HttpClient.SendAsync(Request, CancellationToken);
-			Response.EnsureSuccessStatusCode();
+			HttpResponseMessage response = await _httpClient.SendAsync(request, cancellationToken);
+			response.EnsureSuccessStatusCode();
 		}
 
 		/// <inheritdoc/>
-		public async IAsyncEnumerable<IComputeTaskInfo> GetTaskUpdatesAsync(ClusterId ClusterId, ChannelId ChannelId, [EnumeratorCancellation] CancellationToken CancellationToken)
+		public async IAsyncEnumerable<IComputeTaskInfo> GetTaskUpdatesAsync(ClusterId clusterId, ChannelId channelId, [EnumeratorCancellation] CancellationToken cancellationToken)
 		{
 			for (; ; )
 			{
-				HttpRequestMessage Request = new HttpRequestMessage(HttpMethod.Post, $"api/v1/compute/{ClusterId}/updates/{ChannelId}?wait=10");
-				Request.Headers.Add("Accept", "application/x-ue-cb");
+				HttpRequestMessage request = new HttpRequestMessage(HttpMethod.Post, $"api/v1/compute/{clusterId}/updates/{channelId}?wait=10");
+				request.Headers.Add("Accept", "application/x-ue-cb");
 
-				using HttpResponseMessage Response = await HttpClient.SendAsync(Request, HttpCompletionOption.ResponseContentRead, CancellationToken);
-				Response.EnsureSuccessStatusCode();
+				using HttpResponseMessage response = await _httpClient.SendAsync(request, HttpCompletionOption.ResponseContentRead, cancellationToken);
+				response.EnsureSuccessStatusCode();
 
-				byte[] Data = await Response.Content.ReadAsByteArrayAsync();
-				GetTaskUpdatesResponse ParsedResponse = CbSerializer.Deserialize<GetTaskUpdatesResponse>(new CbField(Data));
+				byte[] data = await response.Content.ReadAsByteArrayAsync();
+				GetTaskUpdatesResponse parsedResponse = CbSerializer.Deserialize<GetTaskUpdatesResponse>(new CbField(data));
 
-				foreach (GetTaskUpdateResponse Update in ParsedResponse.Updates)
+				foreach (GetTaskUpdateResponse update in parsedResponse.Updates)
 				{
-					CancellationToken.ThrowIfCancellationRequested();
-					yield return Update;
+					cancellationToken.ThrowIfCancellationRequested();
+					yield return update;
 				}
 			}
 		}
diff --git a/Engine/Source/Programs/Shared/EpicGames.Horde/Compute/Requirements.cs b/Engine/Source/Programs/Shared/EpicGames.Horde/Compute/Requirements.cs
--- a/Engine/Source/Programs/Shared/EpicGames.Horde/Compute/Requirements.cs
+++ b/Engine/Source/Programs/Shared/EpicGames.Horde/Compute/Requirements.cs
@@ -1,11 +1,9 @@
 // Copyright Epic Games, Inc. All Rights Reserved.
 
-using EpicGames.Core;
 using EpicGames.Horde.Common;
 using EpicGames.Serialization;
 using System;
 using System.Collections.Generic;
-using System.Text;
 
 namespace EpicGames.Horde.Compute
 {
@@ -42,29 +40,29 @@ namespace EpicGames.Horde.Compute
 		/// <summary>
 		/// Construct a requirements object with a condition
 		/// </summary>
-		/// <param name="Condition"></param>
-		public Requirements(Condition? Condition)
+		/// <param name="condition"></param>
+		public Requirements(Condition? condition)
 		{
-			this.Condition = Condition;
+			Condition = condition;
 		}
 
 		/// <inheritdoc/>
 		public override string ToString()
 		{
-			List<string> List = new List<string>();
+			List<string> list = new List<string>();
 			if (Condition != null)
 			{
-				List.Add($"\"{Condition}\"");
+				list.Add($"\"{Condition}\"");
 			}
-			foreach ((string Name, int Count) in Resources)
+			foreach ((string name, int count) in Resources)
 			{
-				List.Add($"{Name}: {Count}");
+				list.Add($"{name}: {count}");
 			}
 			if (Exclusive)
 			{
-				List.Add("Exclusive");
+				list.Add("Exclusive");
 			}
-			return String.Join(", ", List);
+			return String.Join(", ", list);
 		}
 	}
 }
diff --git a/Engine/Source/Programs/Shared/EpicGames.Horde/Storage/BucketId.cs b/Engine/Source/Programs/Shared/EpicGames.Horde/Storage/BucketId.cs
--- a/Engine/Source/Programs/Shared/EpicGames.Horde/Storage/BucketId.cs
+++ b/Engine/Source/Programs/Shared/EpicGames.Horde/Storage/BucketId.cs
@@ -1,11 +1,8 @@
 // Copyright Epic Games, Inc. All Rights Reserved.
 
 using System;
-using System.Collections.Generic;
 using System.ComponentModel;
-using System.Diagnostics.CodeAnalysis;
 using System.Globalization;
-using System.Text;
 using System.Text.Json;
 using System.Text.Json.Serialization;
 using EpicGames.Core;
@@ -24,34 +21,34 @@ namespace EpicGames.Horde.Storage
 		/// <summary>
 		/// The text representing this id
 		/// </summary>
-		readonly StringId Inner;
+		readonly StringId _inner;
 
 		/// <summary>
 		/// Constructor
 		/// </summary>
-		/// <param name="Input">Unique id for the string</param>
-		public BucketId(string Input)
+		/// <param name="input">Unique id for the string</param>
+		public BucketId(string input)
 		{
-			Inner = new StringId(Input);
+			_inner = new StringId(input);
 		}
 
 		/// <inheritdoc/>
-		public override bool Equals(object? Obj) => Obj is BucketId Id && Inner.Equals(Id.Inner);
+		public override bool Equals(object? obj) => obj is BucketId id && _inner.Equals(id._inner);
 
 		/// <inheritdoc/>
-		public override int GetHashCode() => Inner.GetHashCode();
+		public override int GetHashCode() => _inner.GetHashCode();
 
 		/// <inheritdoc/>
-		public bool Equals(BucketId Other) => Inner.Equals(Other.Inner);
+		public bool Equals(BucketId other) => _inner.Equals(other._inner);
 
 		/// <inheritdoc/>
-		public override string ToString() => Inner.ToString();
+		public override string ToString() => _inner.ToString();
 
 		/// <inheritdoc cref="StringId.op_Equality"/>
-		public static bool operator ==(BucketId Left, BucketId Right) => Left.Inner == Right.Inner;
+		public static bool operator ==(BucketId left, BucketId right) => left._inner == right._inner;
 
 		/// <inheritdoc cref="StringId.op_Inequality"/>
-		public static bool operator !=(BucketId Left, BucketId Right) => Left.Inner != Right.Inner;
+		public static bool operator !=(BucketId left, BucketId right) => left._inner != right._inner;
 	}
 
 	/// <summary>
@@ -60,10 +57,10 @@ namespace EpicGames.Horde.Storage
 	sealed class BucketIdJsonConverter : JsonConverter<BucketId>
 	{
 		/// <inheritdoc/>
-		public override BucketId Read(ref Utf8JsonReader Reader, Type TypeToConvert, JsonSerializerOptions Options) => new BucketId(Reader.GetString() ?? String.Empty);
+		public override BucketId Read(ref Utf8JsonReader reader, Type typeToConvert, JsonSerializerOptions options) => new BucketId(reader.GetString() ?? String.Empty);
 
 		/// <inheritdoc/>
-		public override void Write(Utf8JsonWriter Writer, BucketId Value, JsonSerializerOptions Options) => Writer.WriteStringValue(Value.ToString());
+		public override void Write(Utf8JsonWriter writer, BucketId value, JsonSerializerOptions options) => writer.WriteStringValue(value.ToString());
 	}
 
 	/// <summary>
@@ -72,10 +69,10 @@ namespace EpicGames.Horde.Storage
 	sealed class BucketIdTypeConverter : TypeConverter
 	{
 		/// <inheritdoc/>
-		public override bool CanConvertFrom(ITypeDescriptorContext Context, Type SourceType) => SourceType == typeof(string);
+		public override bool CanConvertFrom(ITypeDescriptorContext context, Type sourceType) => sourceType == typeof(string);
 
 		/// <inheritdoc/>
-		public override object ConvertFrom(ITypeDescriptorContext Context, CultureInfo Culture, object Value) => new BucketId((string)Value);
+		public override object ConvertFrom(ITypeDescriptorContext context, CultureInfo culture, object value) => new BucketId((string)value);
 	}
 
 	sealed class BucketIdCbConverter : CbConverterBase<BucketId>
diff --git a/Engine/Source/Programs/Shared/EpicGames.Horde/Storage/IStorageClient.cs b/Engine/Source/Programs/Shared/EpicGames.Horde/Storage/IStorageClient.cs
--- a/Engine/Source/Programs/Shared/EpicGames.Horde/Storage/IStorageClient.cs
+++ b/Engine/Source/Programs/Shared/EpicGames.Horde/Storage/IStorageClient.cs
@@ -5,7 +5,6 @@ using EpicGames.Serialization;
 using System;
 using System.Collections.Generic;
 using System.IO;
-using System.Text;
 using System.Threading;
 using System.Threading.Tasks;
 
@@ -19,8 +18,8 @@ namespace EpicGames.Horde.Storage
 		/// <summary>
 		/// Constructor
 		/// </summary>
-		public StorageException(string Message, Exception? InnerException)
-			: base(Message, InnerException)
+		public StorageException(string message, Exception? innerException)
+			: base(message, innerException)
 		{
 		}
 	}
@@ -43,11 +42,11 @@ namespace EpicGames.Horde.Storage
 		/// <summary>
 		/// Constructor
 		/// </summary>
-		public BlobException(NamespaceId NamespaceId, IoHash Hash, string Message, Exception? InnerException = null)
-			: base(Message, InnerException)
+		public BlobException(NamespaceId namespaceId, IoHash hash, string message, Exception? innerException = null)
+			: base(message, innerException)
 		{
-			this.NamespaceId = NamespaceId;
-			this.Hash = Hash;
+			NamespaceId = namespaceId;
+			Hash = hash;
 		}
 	}
 
@@ -59,8 +58,8 @@ namespace EpicGames.Horde.Storage
 		/// <summary>
 		/// Constructor
 		/// </summary>
-		public BlobNotFoundException(NamespaceId NamespaceId, IoHash Hash, Exception? InnerException = null)
-			: base(NamespaceId, Hash, $"Unable to find blob {Hash} in {NamespaceId}", InnerException)
+		public BlobNotFoundException(NamespaceId namespaceId, IoHash hash, Exception? innerException = null)
+			: base(namespaceId, hash, $"Unable to find blob {hash} in {namespaceId}", innerException)
 		{
 		}
 	}
@@ -88,12 +87,12 @@ namespace EpicGames.Horde.Storage
 		/// <summary>
 		/// Constructor
 		/// </summary>
-		public RefException(NamespaceId NamespaceId, BucketId BucketId, RefId RefId, string Message, Exception? InnerException = null)
-			: base(Message, InnerException)
+		public RefException(NamespaceId namespaceId, BucketId bucketId, RefId refId, string message, Exception? innerException = null)
+			: base(message, innerException)
 		{
-			this.NamespaceId = NamespaceId;
-			this.BucketId = BucketId;
-			this.RefId = RefId;
+			NamespaceId = namespaceId;
+			BucketId = bucketId;
+			RefId = refId;
 		}
 	}
 
@@ -105,8 +104,8 @@ namespace EpicGames.Horde.Storage
 		/// <summary>
 		/// Constructor
 		/// </summary>
-		public RefNotFoundException(NamespaceId NamespaceId, BucketId BucketId, RefId RefId, Exception? InnerException = null)
-			: base(NamespaceId, BucketId, RefId, $"Ref {NamespaceId}/{BucketId}/{RefId} not found", InnerException)
+		public RefNotFoundException(NamespaceId namespaceId, BucketId bucketId, RefId refId, Exception? innerException = null)
+			: base(namespaceId, bucketId, refId, $"Ref {namespaceId}/{bucketId}/{refId} not found", innerException)
 		{
 		}
 	}
@@ -119,8 +118,8 @@ namespace EpicGames.Horde.Storage
 		/// <summary>
 		/// Constructor
 		/// </summary>
-		public RefMissingBlobException(NamespaceId NamespaceId, BucketId BucketId, RefId RefId, List<IoHash> MissingBlobs, Exception? InnerException = null)
-			: base(NamespaceId, BucketId, RefId, $"Ref {NamespaceId}/{BucketId}/{RefId} cannot be finalized; missing {MissingBlobs.Count} blobs ({MissingBlobs[0]}...)", InnerException)
+		public RefMissingBlobException(NamespaceId namespaceId, BucketId bucketId, RefId refId, List<IoHash> missingBlobs, Exception? innerException = null)
+			: base(namespaceId, bucketId, refId, $"Ref {namespaceId}/{bucketId}/{refId} cannot be finalized; missing {missingBlobs.Count} blobs ({missingBlobs[0]}...)", innerException)
 		{
 		}
 	}
@@ -161,38 +160,38 @@ namespace EpicGames.Horde.Storage
 		/// <summary>
 		/// Opens a blob read stream
 		/// </summary>
-		/// <param name="NamespaceId">Namespace to operate on</param>
-		/// <param name="Hash">Hash of the blob</param>
-		/// <param name="CancellationToken">Cancellation token for the operation</param>
+		/// <param name="namespaceId">Namespace to operate on</param>
+		/// <param name="hash">Hash of the blob</param>
+		/// <param name="cancellationToken">Cancellation token for the operation</param>
 		/// <returns>Stream for the blob, or null if it does not exist</returns>
-		Task<Stream> ReadBlobAsync(NamespaceId NamespaceId, IoHash Hash, CancellationToken CancellationToken = default);
+		Task<Stream> ReadBlobAsync(NamespaceId namespaceId, IoHash hash, CancellationToken cancellationToken = default);
 
 		/// <summary>
 		/// Writes a blob to storage
 		/// </summary>
-		/// <param name="NamespaceId">Namespace to operate on</param>
-		/// <param name="Hash">Hash of the blob</param>
-		/// <param name="Stream">The stream to write</param>
-		/// <param name="CancellationToken">Cancellation token for the operation</param>
-		Task WriteBlobAsync(NamespaceId NamespaceId, IoHash Hash, Stream Stream, CancellationToken CancellationToken = default);
+		/// <param name="namespaceId">Namespace to operate on</param>
+		/// <param name="hash">Hash of the blob</param>
+		/// <param name="stream">The stream to write</param>
+		/// <param name="cancellationToken">Cancellation token for the operation</param>
+		Task WriteBlobAsync(NamespaceId namespaceId, IoHash hash, Stream stream, CancellationToken cancellationToken = default);
 
 		/// <summary>
 		/// Checks if the given blob exists
 		/// </summary>
-		/// <param name="NamespaceId">Namespace to operate on</param>
-		/// <param name="Hash">Hash of the blob</param>
-		/// <param name="CancellationToken">Cancellation token for the operation</param>
+		/// <param name="namespaceId">Namespace to operate on</param>
+		/// <param name="hash">Hash of the blob</param>
+		/// <param name="cancellationToken">Cancellation token for the operation</param>
 		/// <returns>True if the blob exists, false if it did not exist</returns>
-		Task<bool> HasBlobAsync(NamespaceId NamespaceId, IoHash Hash, CancellationToken CancellationToken = default);
+		Task<bool> HasBlobAsync(NamespaceId namespaceId, IoHash hash, CancellationToken cancellationToken = default);
 
 		/// <summary>
 		/// Checks if a list of blobs exist, returning the set of missing blobs
 		/// </summary>
-		/// <param name="NamespaceId">Namespace to operate on</param>
-		/// <param name="Hashes">Set of hashes to check</param>
-		/// <param name="CancellationToken">Cancellation token for the operation</param>
+		/// <param name="namespaceId">Namespace to operate on</param>
+		/// <param name="hashes">Set of hashes to check</param>
+		/// <param name="cancellationToken">Cancellation token for the operation</param>
 		/// <returns>A set of the missing hashes</returns>
-		Task<HashSet<IoHash>> FindMissingBlobsAsync(NamespaceId NamespaceId, HashSet<IoHash> Hashes, CancellationToken CancellationToken = default);
+		Task<HashSet<IoHash>> FindMissingBlobsAsync(NamespaceId namespaceId, HashSet<IoHash> hashes, CancellationToken cancellationToken = default);
 
 		#endregion
 		#region Refs
@@ -200,64 +199,64 @@ namespace EpicGames.Horde.Storage
 		/// <summary>
 		/// Gets the given reference
 		/// </summary>
-		/// <param name="NamespaceId">Namespace identifier</param>
-		/// <param name="BucketId">Bucket identifier</param>
-		/// <param name="RefId">Name of the reference</param>
-		/// <param name="CancellationToken">Cancellation token for the operation</param>
+		/// <param name="namespaceId">Namespace identifier</param>
+		/// <param name="bucketId">Bucket identifier</param>
+		/// <param name="refId">Name of the reference</param>
+		/// <param name="cancellationToken">Cancellation token for the operation</param>
 		/// <returns>The reference data if the ref exists</returns>
-		Task<IRef> GetRefAsync(NamespaceId NamespaceId, BucketId BucketId, RefId RefId, CancellationToken CancellationToken = default);
+		Task<IRef> GetRefAsync(NamespaceId namespaceId, BucketId bucketId, RefId refId, CancellationToken cancellationToken = default);
 
 		/// <summary>
 		/// Checks if the given reference exists
 		/// </summary>
-		/// <param name="NamespaceId">Namespace identifier</param>
-		/// <param name="BucketId">Bucket identifier</param>
-		/// <param name="RefId">Ref identifier</param>
-		/// <param name="CancellationToken">Cancellation token for the operation</param>
+		/// <param name="namespaceId">Namespace identifier</param>
+		/// <param name="bucketId">Bucket identifier</param>
+		/// <param name="refId">Ref identifier</param>
+		/// <param name="cancellationToken">Cancellation token for the operation</param>
 		/// <returns>True if the ref exists, false if it did not exist</returns>
-		Task<bool> HasRefAsync(NamespaceId NamespaceId, BucketId BucketId, RefId RefId, CancellationToken CancellationToken = default);
+		Task<bool> HasRefAsync(NamespaceId namespaceId, BucketId bucketId, RefId refId, CancellationToken cancellationToken = default);
 
 		/// <summary>
 		/// Determines which refs are missing
 		/// </summary>
-		/// <param name="NamespaceId">Namespace identifier</param>
-		/// <param name="BucketId">Bucket identifier</param>
-		/// <param name="RefIds">Names of the references</param>
-		/// <param name="CancellationToken">Cancellation token for the operation</param>
+		/// <param name="namespaceId">Namespace identifier</param>
+		/// <param name="bucketId">Bucket identifier</param>
+		/// <param name="refIds">Names of the references</param>
+		/// <param name="cancellationToken">Cancellation token for the operation</param>
 		/// <returns>List of missing references</returns>
-		Task<List<RefId>> FindMissingRefsAsync(NamespaceId NamespaceId, BucketId BucketId, List<RefId> RefIds, CancellationToken CancellationToken = default);
+		Task<List<RefId>> FindMissingRefsAsync(NamespaceId namespaceId, BucketId bucketId, List<RefId> refIds, CancellationToken cancellationToken = default);
 
 		/// <summary>
 		/// Attempts to sets the given reference, returning a list of missing objects on failure.
 		/// </summary>
-		/// <param name="NamespaceId">Namespace identifier</param>
-		/// <param name="BucketId">Bucket identifier</param>
-		/// <param name="RefId">Ref identifier</param>
-		/// <param name="Value">New value for the reference</param>
-		/// <param name="CancellationToken">Cancellation token for the operation</param>
+		/// <param name="namespaceId">Namespace identifier</param>
+		/// <param name="bucketId">Bucket identifier</param>
+		/// <param name="refId">Ref identifier</param>
+		/// <param name="value">New value for the reference</param>
+		/// <param name="cancellationToken">Cancellation token for the operation</param>
 		/// <returns>List of missing references</returns>
-		Task<List<IoHash>> TrySetRefAsync(NamespaceId NamespaceId, BucketId BucketId, RefId RefId, CbObject Value, CancellationToken CancellationToken = default);
+		Task<List<IoHash>> TrySetRefAsync(NamespaceId namespaceId, BucketId bucketId, RefId refId, CbObject value, CancellationToken cancellationToken = default);
 
 		/// <summary>
 		/// Attempts to finalize a reference, turning its references into hard references
 		/// </summary>
-		/// <param name="NamespaceId">Namespace identifier</param>
-		/// <param name="BucketId">Bucket identifier</param>
-		/// <param name="RefId">Ref identifier</param>
-		/// <param name="Hash">Hash of the referenced object</param>
-		/// <param name="CancellationToken">Cancellation token for the operation</param>
+		/// <param name="namespaceId">Namespace identifier</param>
+		/// <param name="bucketId">Bucket identifier</param>
+		/// <param name="refId">Ref identifier</param>
+		/// <param name="hash">Hash of the referenced object</param>
+		/// <param name="cancellationToken">Cancellation token for the operation</param>
 		/// <returns></returns>
-		Task<List<IoHash>> TryFinalizeRefAsync(NamespaceId NamespaceId, BucketId BucketId, RefId RefId, IoHash Hash, CancellationToken CancellationToken = default);
+		Task<List<IoHash>> TryFinalizeRefAsync(NamespaceId namespaceId, BucketId bucketId, RefId refId, IoHash hash, CancellationToken cancellationToken = default);
 
 		/// <summary>
 		/// Removes the given reference
 		/// </summary>
-		/// <param name="NamespaceId">Namespace identifier</param>
-		/// <param name="BucketId">Bucket identifier</param>
-		/// <param name="RefId">Ref identifier</param>
-		/// <param name="CancellationToken">Cancellation token for the operation</param>
+		/// <param name="namespaceId">Namespace identifier</param>
+		/// <param name="bucketId">Bucket identifier</param>
+		/// <param name="refId">Ref identifier</param>
+		/// <param name="cancellationToken">Cancellation token for the operation</param>
 		/// <returns>True if the ref was deleted, false if it did not exist</returns>
-		Task<bool> DeleteRefAsync(NamespaceId NamespaceId, BucketId BucketId, RefId RefId, CancellationToken CancellationToken = default);
+		Task<bool> DeleteRefAsync(NamespaceId namespaceId, BucketId bucketId, RefId refId, CancellationToken cancellationToken = default);
 
 		#endregion
 	}
@@ -272,146 +271,146 @@ namespace EpicGames.Horde.Storage
 		/// <summary>
 		/// Gets a blob as a byte array
 		/// </summary>
-		/// <param name="StorageClient">The storage interface</param>
-		/// <param name="NamespaceId">Namespace containing the blob</param>
-		/// <param name="Hash">Hash of the blob to read</param>
-		/// <param name="MaxInMemoryBlobLength">Maximum allowed memory allocation to store the blob</param>
-		/// <param name="CancellationToken">Cancellation token for the operation</param>
+		/// <param name="storageClient">The storage interface</param>
+		/// <param name="namespaceId">Namespace containing the blob</param>
+		/// <param name="hash">Hash of the blob to read</param>
+		/// <param name="maxInMemoryBlobLength">Maximum allowed memory allocation to store the blob</param>
+		/// <param name="cancellationToken">Cancellation token for the operation</param>
 		/// <returns>Data for the blob that was read. Throws an exception if the blob was not found.</returns>
-		public static async Task<byte[]> ReadBlobToMemoryAsync(this IStorageClient StorageClient, NamespaceId NamespaceId, IoHash Hash, int MaxInMemoryBlobLength = DefaultMaxInMemoryBlobLength, CancellationToken CancellationToken = default)
+		public static async Task<byte[]> ReadBlobToMemoryAsync(this IStorageClient storageClient, NamespaceId namespaceId, IoHash hash, int maxInMemoryBlobLength = DefaultMaxInMemoryBlobLength, CancellationToken cancellationToken = default)
 		{
-			using Stream Stream = await StorageClient.ReadBlobAsync(NamespaceId, Hash, CancellationToken);
+			using Stream stream = await storageClient.ReadBlobAsync(namespaceId, hash, cancellationToken);
 
-			long Length = Stream.Length;
-			if (Length > MaxInMemoryBlobLength)
+			long length = stream.Length;
+			if (length > maxInMemoryBlobLength)
 			{
-				throw new BlobException(NamespaceId, Hash, $"Blob {Hash} is too large ({Length} > {MaxInMemoryBlobLength})");
+				throw new BlobException(namespaceId, hash, $"Blob {hash} is too large ({length} > {maxInMemoryBlobLength})");
 			}
 
-			byte[] Buffer = new byte[Length];
-			for (int Offset = 0; Offset < Length;)
+			byte[] buffer = new byte[length];
+			for (int offset = 0; offset < length;)
 			{
-				int Count = await Stream.ReadAsync(Buffer, Offset, (int)Length - Offset, CancellationToken);
-				if (Count == 0)
+				int count = await stream.ReadAsync(buffer, offset, (int)length - offset, cancellationToken);
+				if (count == 0)
 				{
-					throw new BlobException(NamespaceId, Hash, $"Unexpected end of stream reading blob {Hash}");
+					throw new BlobException(namespaceId, hash, $"Unexpected end of stream reading blob {hash}");
 				}
-				Offset += Count;
+				offset += count;
 			}
 
-			return Buffer;
+			return buffer;
 		}
 
 		/// <summary>
 		/// Writes a blob from memory to storage
 		/// </summary>
-		/// <param name="StorageClient">The storage interface</param>
-		/// <param name="NamespaceId">Namespace containing the blob</param>
-		/// <param name="Data">Data to write</param>
-		/// <param name="CancellationToken">Cancellation token for the operation</param>
+		/// <param name="storageClient">The storage interface</param>
+		/// <param name="namespaceId">Namespace containing the blob</param>
+		/// <param name="data">Data to write</param>
+		/// <param name="cancellationToken">Cancellation token for the operation</param>
 		/// <returns></returns>
-		public static async Task<IoHash> WriteBlobFromMemoryAsync(this IStorageClient StorageClient, NamespaceId NamespaceId, ReadOnlyMemory<byte> Data, CancellationToken CancellationToken = default)
+		public static async Task<IoHash> WriteBlobFromMemoryAsync(this IStorageClient storageClient, NamespaceId namespaceId, ReadOnlyMemory<byte> data, CancellationToken cancellationToken = default)
 		{
-			IoHash Hash = IoHash.Compute(Data.Span);
-			await WriteBlobFromMemoryAsync(StorageClient, NamespaceId, Hash, Data, CancellationToken);
-			return Hash;
+			IoHash hash = IoHash.Compute(data.Span);
+			await WriteBlobFromMemoryAsync(storageClient, namespaceId, hash, data, cancellationToken);
+			return hash;
 		}
 
 		/// <summary>
 		/// Writes a blob from memory to storage
 		/// </summary>
-		/// <param name="StorageClient">The storage interface</param>
-		/// <param name="NamespaceId">Namespace containing the blob</param>
-		/// <param name="Hash">Hash of the data</param>
-		/// <param name="Data">The data to be written</param>
-		/// <param name="CancellationToken">Cancellation token for the operation</param>
+		/// <param name="storageClient">The storage interface</param>
+		/// <param name="namespaceId">Namespace containing the blob</param>
+		/// <param name="hash">Hash of the data</param>
+		/// <param name="data">The data to be written</param>
+		/// <param name="cancellationToken">Cancellation token for the operation</param>
 		/// <returns></returns>
-		public static async Task WriteBlobFromMemoryAsync(this IStorageClient StorageClient, NamespaceId NamespaceId, IoHash Hash, ReadOnlyMemory<byte> Data, CancellationToken CancellationToken = default)
+		public static async Task WriteBlobFromMemoryAsync(this IStorageClient storageClient, NamespaceId namespaceId, IoHash hash, ReadOnlyMemory<byte> data, CancellationToken cancellationToken = default)
 		{
-			using ReadOnlyMemoryStream Stream = new ReadOnlyMemoryStream(Data);
-			await StorageClient.WriteBlobAsync(NamespaceId, Hash, Stream, CancellationToken);
+			using ReadOnlyMemoryStream stream = new ReadOnlyMemoryStream(data);
+			await storageClient.WriteBlobAsync(namespaceId, hash, stream, cancellationToken);
 		}
 
 		/// <summary>
 		/// Reads a blob and deserializes it as the given compact-binary encoded type
 		/// </summary>
 		/// <typeparam name="T">Type of object to deserialize</typeparam>
-		/// <param name="StorageClient">The storage interface</param>
-		/// <param name="NamespaceId">Namespace containing the blob</param>
-		/// <param name="Hash">Hash of the blob to read</param>
-		/// <param name="MaxInMemoryBlobLength">Maximum allowed memory allocation to store the blob</param>
-		/// <param name="CancellationToken">Cancellation token for the operation</param>
+		/// <param name="storageClient">The storage interface</param>
+		/// <param name="namespaceId">Namespace containing the blob</param>
+		/// <param name="hash">Hash of the blob to read</param>
+		/// <param name="maxInMemoryBlobLength">Maximum allowed memory allocation to store the blob</param>
+		/// <param name="cancellationToken">Cancellation token for the operation</param>
 		/// <returns>The decoded object</returns>
-		public static async Task<T> ReadBlobAsync<T>(this IStorageClient StorageClient, NamespaceId NamespaceId, IoHash Hash, int MaxInMemoryBlobLength = DefaultMaxInMemoryBlobLength, CancellationToken CancellationToken = default)
+		public static async Task<T> ReadBlobAsync<T>(this IStorageClient storageClient, NamespaceId namespaceId, IoHash hash, int maxInMemoryBlobLength = DefaultMaxInMemoryBlobLength, CancellationToken cancellationToken = default)
 		{
-			ReadOnlyMemory<byte> Data = await ReadBlobToMemoryAsync(StorageClient, NamespaceId, Hash, MaxInMemoryBlobLength, CancellationToken);
-			return CbSerializer.Deserialize<T>(Data);
+			ReadOnlyMemory<byte> data = await ReadBlobToMemoryAsync(storageClient, namespaceId, hash, maxInMemoryBlobLength, cancellationToken);
+			return CbSerializer.Deserialize<T>(data);
 		}
 
 		/// <summary>
 		/// Serializes an object to compact binary format and writes it to storage
 		/// </summary>
-		/// <param name="StorageClient">The storage interface</param>
-		/// <param name="NamespaceId">Namespace containing the blob</param>
-		/// <param name="Object">The object to be written</param>
-		/// <param name="CancellationToken">Cancellation token for the operation</param>
+		/// <param name="storageClient">The storage interface</param>
+		/// <param name="namespaceId">Namespace containing the blob</param>
+		/// <param name="blob">The object to be written</param>
+		/// <param name="cancellationToken">Cancellation token for the operation</param>
 		/// <returns></returns>
-		public static async Task<IoHash> WriteBlobAsync<T>(this IStorageClient StorageClient, NamespaceId NamespaceId, T Object, CancellationToken CancellationToken = default)
+		public static async Task<IoHash> WriteBlobAsync<T>(this IStorageClient storageClient, NamespaceId namespaceId, T blob, CancellationToken cancellationToken = default)
 		{
-			CbWriter Writer = new CbWriter();
-			CbSerializer.Serialize<T>(Writer, Object);
+			CbWriter writer = new CbWriter();
+			CbSerializer.Serialize<T>(writer, blob);
 
-			IoHash Hash = Writer.ComputeHash();
-			await StorageClient.WriteBlobAsync(NamespaceId, Hash, Writer.AsStream(), CancellationToken);
-			return Hash;
+			IoHash hash = writer.ComputeHash();
+			await storageClient.WriteBlobAsync(namespaceId, hash, writer.AsStream(), cancellationToken);
+			return hash;
 		}
 
 		/// <summary>
 		/// Reads a reference as a specific type
 		/// </summary>
 		/// <typeparam name="T"></typeparam>
-		/// <param name="StorageClient">The storage interface</param>
-		/// <param name="NamespaceId">Namespace containing the ref</param>
-		/// <param name="BucketId">Bucket containing the ref</param>
-		/// <param name="RefId">The ref id</param>
-		/// <param name="CancellationToken">Cancellation token for the operation</param>
+		/// <param name="storageClient">The storage interface</param>
+		/// <param name="namespaceId">Namespace containing the ref</param>
+		/// <param name="bucketId">Bucket containing the ref</param>
+		/// <param name="refId">The ref id</param>
+		/// <param name="cancellationToken">Cancellation token for the operation</param>
 		/// <returns>Deserialized object for the given ref</returns>
-		public static async Task<T> GetRefAsync<T>(this IStorageClient StorageClient, NamespaceId NamespaceId, BucketId BucketId, RefId RefId, CancellationToken CancellationToken = default)
+		public static async Task<T> GetRefAsync<T>(this IStorageClient storageClient, NamespaceId namespaceId, BucketId bucketId, RefId refId, CancellationToken cancellationToken = default)
 		{
-			IRef Ref = await StorageClient.GetRefAsync(NamespaceId, BucketId, RefId, CancellationToken);
-			return CbSerializer.Deserialize<T>(Ref.Value);
+			IRef item = await storageClient.GetRefAsync(namespaceId, bucketId, refId, cancellationToken);
+			return CbSerializer.Deserialize<T>(item.Value);
 		}
 
 		/// <summary>
 		/// Sets a ref to a particular value
 		/// </summary>
-		/// <param name="StorageClient">The storage interface</param>
-		/// <param name="NamespaceId">Namespace containing the ref</param>
-		/// <param name="BucketId">Bucket containing the ref</param>
-		/// <param name="RefId">The ref id</param>
-		/// <param name="Value">The new object for the ref</param>
-		/// <param name="CancellationToken">Cancellation token for the operation</param>
-		public static Task SetRefAsync<T>(this IStorageClient StorageClient, NamespaceId NamespaceId, BucketId BucketId, RefId RefId, T Value, CancellationToken CancellationToken = default)
+		/// <param name="storageClient">The storage interface</param>
+		/// <param name="namespaceId">Namespace containing the ref</param>
+		/// <param name="bucketId">Bucket containing the ref</param>
+		/// <param name="refId">The ref id</param>
+		/// <param name="value">The new object for the ref</param>
+		/// <param name="cancellationToken">Cancellation token for the operation</param>
+		public static Task SetRefAsync<T>(this IStorageClient storageClient, NamespaceId namespaceId, BucketId bucketId, RefId refId, T value, CancellationToken cancellationToken = default)
 		{
-			CbObject Object = CbSerializer.Serialize<T>(Value);
-			return SetRefAsync(StorageClient, NamespaceId, BucketId, RefId, Object, CancellationToken);
+			CbObject objectValue = CbSerializer.Serialize<T>(value);
+			return SetRefAsync(storageClient, namespaceId, bucketId, refId, objectValue, cancellationToken);
 		}
 
 		/// <summary>
 		/// Sets a ref to a particular value
 		/// </summary>
-		/// <param name="StorageClient">The storage interface</param>
-		/// <param name="NamespaceId">Namespace containing the ref</param>
-		/// <param name="BucketId">Bucket containing the ref</param>
-		/// <param name="RefId">The ref id</param>
-		/// <param name="Value">The new object for the ref</param>
-		/// <param name="CancellationToken">Cancellation token for the operation</param>
-		public static async Task SetRefAsync(this IStorageClient StorageClient, NamespaceId NamespaceId, BucketId BucketId, RefId RefId, CbObject Value, CancellationToken CancellationToken = default)
+		/// <param name="storageClient">The storage interface</param>
+		/// <param name="namespaceId">Namespace containing the ref</param>
+		/// <param name="bucketId">Bucket containing the ref</param>
+		/// <param name="refId">The ref id</param>
+		/// <param name="value">The new object for the ref</param>
+		/// <param name="cancellationToken">Cancellation token for the operation</param>
+		public static async Task SetRefAsync(this IStorageClient storageClient, NamespaceId namespaceId, BucketId bucketId, RefId refId, CbObject value, CancellationToken cancellationToken = default)
 		{
-			List<IoHash> MissingHashes = await StorageClient.TrySetRefAsync(NamespaceId, BucketId, RefId, Value, CancellationToken);
-			if (MissingHashes.Count > 0)
+			List<IoHash> missingHashes = await storageClient.TrySetRefAsync(namespaceId, bucketId, refId, value, cancellationToken);
+			if (missingHashes.Count > 0)
 			{
-				throw new RefMissingBlobException(NamespaceId, BucketId, RefId, MissingHashes);
+				throw new RefMissingBlobException(namespaceId, bucketId, refId, missingHashes);
 			}
 		}
 
@@ -419,34 +418,34 @@ namespace EpicGames.Horde.Storage
 		/// Attempts to set a ref to a particular value
 		/// </summary>
 		/// <typeparam name="T"></typeparam>
-		/// <param name="StorageClient">The storage interface</param>
-		/// <param name="NamespaceId">Namespace containing the ref</param>
-		/// <param name="BucketId">Bucket containing the ref</param>
-		/// <param name="RefId">The ref id</param>
-		/// <param name="Value">The new object for the ref</param>
-		/// <param name="CancellationToken">Cancellation token for the operation</param>
+		/// <param name="storageClient">The storage interface</param>
+		/// <param name="namespaceId">Namespace containing the ref</param>
+		/// <param name="bucketId">Bucket containing the ref</param>
+		/// <param name="refId">The ref id</param>
+		/// <param name="value">The new object for the ref</param>
+		/// <param name="cancellationToken">Cancellation token for the operation</param>
 		/// <returns>List of missing blob hashes</returns>
-		public static Task<List<IoHash>> TrySetRefAsync<T>(this IStorageClient StorageClient, NamespaceId NamespaceId, BucketId BucketId, RefId RefId, T Value, CancellationToken CancellationToken = default)
+		public static Task<List<IoHash>> TrySetRefAsync<T>(this IStorageClient storageClient, NamespaceId namespaceId, BucketId bucketId, RefId refId, T value, CancellationToken cancellationToken = default)
 		{
-			CbObject Object = CbSerializer.Serialize<T>(Value);
-			return StorageClient.TrySetRefAsync(NamespaceId, BucketId, RefId, Object, CancellationToken);
+			CbObject objectValue = CbSerializer.Serialize<T>(value);
+			return storageClient.TrySetRefAsync(namespaceId, bucketId, refId, objectValue, cancellationToken);
 		}
 
 		/// <summary>
 		/// Finalize a ref, throwing an exception if finalization fails
 		/// </summary>
-		/// <param name="StorageClient">The storage interface</param>
-		/// <param name="NamespaceId">Namespace containing the ref</param>
-		/// <param name="BucketId">Bucket containing the ref</param>
-		/// <param name="RefId">The ref id</param>
-		/// <param name="ValueHash">Hash of the ref value</param>
-		/// <param name="CancellationToken">Cancellation token for the operation</param>
-		public static async Task FinalizeRefAsync(this IStorageClient StorageClient, NamespaceId NamespaceId, BucketId BucketId, RefId RefId, IoHash ValueHash, CancellationToken CancellationToken = default)
+		/// <param name="storageClient">The storage interface</param>
+		/// <param name="namespaceId">Namespace containing the ref</param>
+		/// <param name="bucketId">Bucket containing the ref</param>
+		/// <param name="refId">The ref id</param>
+		/// <param name="valueHash">Hash of the ref value</param>
+		/// <param name="cancellationToken">Cancellation token for the operation</param>
+		public static async Task FinalizeRefAsync(this IStorageClient storageClient, NamespaceId namespaceId, BucketId bucketId, RefId refId, IoHash valueHash, CancellationToken cancellationToken = default)
 		{
-			List<IoHash> MissingHashes = await StorageClient.TryFinalizeRefAsync(NamespaceId, BucketId, RefId, ValueHash, CancellationToken);
-			if (MissingHashes.Count > 0)
+			List<IoHash> missingHashes = await storageClient.TryFinalizeRefAsync(namespaceId, bucketId, refId, valueHash, cancellationToken);
+			if (missingHashes.Count > 0)
 			{
-				throw new RefMissingBlobException(NamespaceId, BucketId, RefId, MissingHashes);
+				throw new RefMissingBlobException(namespaceId, bucketId, refId, missingHashes);
 			}
 		}
 	}
diff --git a/Engine/Source/Programs/Shared/EpicGames.Horde/Storage/Impl/FileStorageClient.cs b/Engine/Source/Programs/Shared/EpicGames.Horde/Storage/Impl/FileStorageClient.cs
--- a/Engine/Source/Programs/Shared/EpicGames.Horde/Storage/Impl/FileStorageClient.cs
+++ b/Engine/Source/Programs/Shared/EpicGames.Horde/Storage/Impl/FileStorageClient.cs
@@ -7,7 +7,6 @@ using System;
 using System.Collections.Generic;
 using System.IO;
 using System.Linq;
-using System.Text;
 using System.Threading;
 using System.Threading.Tasks;
 
@@ -25,12 +24,12 @@ namespace EpicGames.Horde.Storage.Impl
 			public RefId RefId { get; set; }
 			public CbObject Value { get; set; }
 
-			public Ref(NamespaceId NamespaceId, BucketId BucketId, RefId RefId, CbObject Value)
+			public Ref(NamespaceId namespaceId, BucketId bucketId, RefId refId, CbObject value)
 			{
-				this.NamespaceId = NamespaceId;
-				this.BucketId = BucketId;
-				this.RefId = RefId;
-				this.Value = Value;
+				NamespaceId = namespaceId;
+				BucketId = bucketId;
+				RefId = refId;
+				Value = value;
 			}
 		}
 
@@ -38,111 +37,112 @@ namespace EpicGames.Horde.Storage.Impl
 		/// Base directory for storing files
 		/// </summary>
 		public DirectoryReference BaseDir { get; }
-		ILogger Logger;
+
+		private readonly ILogger _logger;
 
 		/// <summary>
 		/// Constructor
 		/// </summary>
-		/// <param name="BaseDir"></param>
-		/// <param name="Logger"></param>
-		public FileStorageClient(DirectoryReference BaseDir, ILogger Logger)
+		/// <param name="baseDir"></param>
+		/// <param name="logger"></param>
+		public FileStorageClient(DirectoryReference baseDir, ILogger logger)
 		{
-			this.BaseDir = BaseDir;
-			this.Logger = Logger;
+			BaseDir = baseDir;
+			_logger = logger;
 
-			DirectoryReference.CreateDirectory(BaseDir);
+			DirectoryReference.CreateDirectory(baseDir);
 		}
 
-		FileReference GetBlobFile(NamespaceId NamespaceId, IoHash Hash)
+		FileReference GetBlobFile(NamespaceId namespaceId, IoHash hash)
 		{
-			return FileReference.Combine(BaseDir, NamespaceId.ToString(), $"{Hash}.blob");
+			return FileReference.Combine(BaseDir, namespaceId.ToString(), $"{hash}.blob");
 		}
 
 		/// <inheritdoc/>
-		public Task<Stream> ReadBlobAsync(NamespaceId NamespaceId, IoHash Hash, CancellationToken CancellationToken = default)
+		public Task<Stream> ReadBlobAsync(NamespaceId namespaceId, IoHash hash, CancellationToken cancellationToken = default)
 		{
-			FileReference File = GetBlobFile(NamespaceId, Hash);
-			Logger.LogInformation("Reading {File} ({Size:n0} bytes)", File, new FileInfo(File.FullName).Length);
-			return Task.FromResult<Stream>(FileReference.Open(File, FileMode.Open, FileAccess.Read, FileShare.Read));
+			FileReference file = GetBlobFile(namespaceId, hash);
+			_logger.LogInformation("Reading {File} ({Size:n0} bytes)", file, new FileInfo(file.FullName).Length);
+			return Task.FromResult<Stream>(FileReference.Open(file, FileMode.Open, FileAccess.Read, FileShare.Read));
 		}
 
 		/// <inheritdoc/>
-		public async Task WriteBlobAsync(NamespaceId NamespaceId, IoHash Hash, Stream Stream, CancellationToken CancellationToken = default)
+		public async Task WriteBlobAsync(NamespaceId namespaceId, IoHash hash, Stream stream, CancellationToken cancellationToken = default)
 		{
-			FileReference File = GetBlobFile(NamespaceId, Hash);
-			DirectoryReference.CreateDirectory(File.Directory);
+			FileReference file = GetBlobFile(namespaceId, hash);
+			DirectoryReference.CreateDirectory(file.Directory);
 
-			Logger.LogInformation("Writing {File} ({Size:n0} bytes)", File, Stream.Length);
-			using (Stream OutputStream = FileReference.Open(File, FileMode.Create, FileAccess.Write, FileShare.Read))
+			_logger.LogInformation("Writing {File} ({Size:n0} bytes)", file, stream.Length);
+			using (Stream outputStream = FileReference.Open(file, FileMode.Create, FileAccess.Write, FileShare.Read))
 			{
-				await Stream.CopyToAsync(OutputStream);
+				await stream.CopyToAsync(outputStream);
 			}
 		}
 
 		/// <inheritdoc/>
-		public Task<bool> HasBlobAsync(NamespaceId NamespaceId, IoHash Hash, CancellationToken CancellationToken = default)
+		public Task<bool> HasBlobAsync(NamespaceId namespaceId, IoHash hash, CancellationToken cancellationToken = default)
 		{
-			return Task.FromResult(FileReference.Exists(GetBlobFile(NamespaceId, Hash)));
+			return Task.FromResult(FileReference.Exists(GetBlobFile(namespaceId, hash)));
 		}
 
 		/// <inheritdoc/>
-		public Task<HashSet<IoHash>> FindMissingBlobsAsync(NamespaceId NamespaceId, HashSet<IoHash> Hashes, CancellationToken CancellationToken)
+		public Task<HashSet<IoHash>> FindMissingBlobsAsync(NamespaceId namespaceId, HashSet<IoHash> hashes, CancellationToken cancellationToken)
 		{
-			return Task.FromResult(Hashes.Where(x => !FileReference.Exists(GetBlobFile(NamespaceId, x))).ToHashSet());
+			return Task.FromResult(hashes.Where(x => !FileReference.Exists(GetBlobFile(namespaceId, x))).ToHashSet());
 		}
 
 		/// <inheritdoc/>
-		public Task<bool> DeleteRefAsync(NamespaceId NamespaceId, BucketId BucketId, RefId RefId, CancellationToken CancellationToken = default)
+		public Task<bool> DeleteRefAsync(NamespaceId namespaceId, BucketId bucketId, RefId refId, CancellationToken cancellationToken = default)
 		{
-			FileReference File = GetRefFile(NamespaceId, BucketId, RefId);
-			if (FileReference.Exists(File))
+			FileReference file = GetRefFile(namespaceId, bucketId, refId);
+			if (FileReference.Exists(file))
 			{
-				FileReference.Delete(File);
+				FileReference.Delete(file);
 				return Task.FromResult(true);
 			}
 			return Task.FromResult(false);
 		}
 
 		/// <inheritdoc/>
-		public Task<List<RefId>> FindMissingRefsAsync(NamespaceId NamespaceId, BucketId BucketId, List<RefId> RefIds, CancellationToken CancellationToken = default)
+		public Task<List<RefId>> FindMissingRefsAsync(NamespaceId namespaceId, BucketId bucketId, List<RefId> refIds, CancellationToken cancellationToken = default)
 		{
-			return Task.FromResult(RefIds.Where(x => !FileReference.Exists(GetRefFile(NamespaceId, BucketId, x))).ToList());
+			return Task.FromResult(refIds.Where(x => !FileReference.Exists(GetRefFile(namespaceId, bucketId, x))).ToList());
 		}
 
-		FileReference GetRefFile(NamespaceId NamespaceId, BucketId BucketId, RefId RefId)
+		FileReference GetRefFile(NamespaceId namespaceId, BucketId bucketId, RefId refId)
 		{
-			return FileReference.Combine(BaseDir, NamespaceId.ToString(), BucketId.ToString(), $"{RefId}.ref");
+			return FileReference.Combine(BaseDir, namespaceId.ToString(), bucketId.ToString(), $"{refId}.ref");
 		}
 
 		/// <inheritdoc/>
-		public async Task<IRef> GetRefAsync(NamespaceId NamespaceId, BucketId BucketId, RefId RefId, CancellationToken CancellationToken = default)
+		public async Task<IRef> GetRefAsync(NamespaceId namespaceId, BucketId bucketId, RefId refId, CancellationToken cancellationToken = default)
 		{
-			FileReference File = GetRefFile(NamespaceId, BucketId, RefId);
-			Logger.LogInformation("Reading {File} ({Size:n0} bytes)", File, new FileInfo(File.FullName).Length);
-			byte[] Data = await FileReference.ReadAllBytesAsync(File);
-			return new Ref(NamespaceId, BucketId, RefId, new CbObject(Data));
+			FileReference file = GetRefFile(namespaceId, bucketId, refId);
+			_logger.LogInformation("Reading {File} ({Size:n0} bytes)", file, new FileInfo(file.FullName).Length);
+			byte[] data = await FileReference.ReadAllBytesAsync(file);
+			return new Ref(namespaceId, bucketId, refId, new CbObject(data));
 		}
 
 		/// <inheritdoc/>
-		public Task<bool> HasRefAsync(NamespaceId NamespaceId, BucketId BucketId, RefId RefId, CancellationToken CancellationToken = default)
+		public Task<bool> HasRefAsync(NamespaceId namespaceId, BucketId bucketId, RefId refId, CancellationToken cancellationToken = default)
 		{
-			return Task.FromResult(FileReference.Exists(GetRefFile(NamespaceId, BucketId, RefId)));
+			return Task.FromResult(FileReference.Exists(GetRefFile(namespaceId, bucketId, refId)));
 		}
 
 		/// <inheritdoc/>
-		public Task<List<IoHash>> TryFinalizeRefAsync(NamespaceId NamespaceId, BucketId BucketId, RefId RefId, IoHash Hash, CancellationToken CancellationToken = default)
+		public Task<List<IoHash>> TryFinalizeRefAsync(NamespaceId namespaceId, BucketId bucketId, RefId refId, IoHash hash, CancellationToken cancellationToken = default)
 		{
 			throw new NotImplementedException();
 		}
 
 		/// <inheritdoc/>
-		public async Task<List<IoHash>> TrySetRefAsync(NamespaceId NamespaceId, BucketId BucketId, RefId RefId, CbObject Value, CancellationToken CancellationToken = default)
+		public async Task<List<IoHash>> TrySetRefAsync(NamespaceId namespaceId, BucketId bucketId, RefId refId, CbObject value, CancellationToken cancellationToken = default)
 		{
-			FileReference File = GetRefFile(NamespaceId, BucketId, RefId);
-			DirectoryReference.CreateDirectory(File.Directory);
+			FileReference file = GetRefFile(namespaceId, bucketId, refId);
+			DirectoryReference.CreateDirectory(file.Directory);
 
-			Logger.LogInformation("Writing {File} ({Size:n0} bytes)", File, Value.GetView().Length);
-			await FileReference.WriteAllBytesAsync(File, Value.GetView().ToArray(), CancellationToken);
+			_logger.LogInformation("Writing {File} ({Size:n0} bytes)", file, value.GetView().Length);
+			await FileReference.WriteAllBytesAsync(file, value.GetView().ToArray(), cancellationToken);
 			return new List<IoHash>();
 		}
 	}
