Horde: Use standard coding style for EpicGames.Horde.

diff --git a/Engine/Source/Programs/Shared/EpicGames.Horde/Storage/Impl/HttpStorageClient.cs b/Engine/Source/Programs/Shared/EpicGames.Horde/Storage/Impl/HttpStorageClient.cs
--- a/Engine/Source/Programs/Shared/EpicGames.Horde/Storage/Impl/HttpStorageClient.cs
+++ b/Engine/Source/Programs/Shared/EpicGames.Horde/Storage/Impl/HttpStorageClient.cs
@@ -1,9 +1,7 @@
 // Copyright Epic Games, Inc. All Rights Reserved.
 
 using EpicGames.Core;
-using EpicGames.Horde.Auth;
 using EpicGames.Serialization;
-using Microsoft.Extensions.Options;
 using System;
 using System.Collections.Generic;
 using System.IO;
@@ -65,12 +63,12 @@ namespace EpicGames.Horde.Storage.Impl
 
 			public CbObject Value { get; }
 
-			public RefImpl(NamespaceId NamespaceId, BucketId BucketId, RefId RefId, CbObject Value)
+			public RefImpl(NamespaceId namespaceId, BucketId bucketId, RefId refId, CbObject value)
 			{
-				this.NamespaceId = NamespaceId;
-				this.BucketId = BucketId;
-				this.RefId = RefId;
-				this.Value = Value;
+				NamespaceId = namespaceId;
+				BucketId = bucketId;
+				RefId = refId;
+				Value = value;
 			}
 		}
 
@@ -79,97 +77,97 @@ namespace EpicGames.Horde.Storage.Impl
 
 		const string CompactBinaryMimeType = "application/x-ue-cb";
 
-		HttpClient HttpClient;
+		readonly HttpClient _httpClient;
 
 		/// <summary>
 		/// Constructor
 		/// </summary>
-		/// <param name="HttpClient">Http client for the blob store. Should be pre-configured with a base address and appropriate authentication headers.</param>
-		public HttpStorageClient(HttpClient HttpClient)
+		/// <param name="httpClient">Http client for the blob store. Should be pre-configured with a base address and appropriate authentication headers.</param>
+		public HttpStorageClient(HttpClient httpClient)
 		{
-			this.HttpClient = HttpClient;
+			_httpClient = httpClient;
 		}
 
 		#region Blobs
 
 		/// <inheritdoc/>
-		public async Task<Stream> ReadBlobAsync(NamespaceId NamespaceId, IoHash Hash, CancellationToken CancellationToken)
+		public async Task<Stream> ReadBlobAsync(NamespaceId namespaceId, IoHash hash, CancellationToken cancellationToken)
 		{
-			using HttpRequestMessage Request = new HttpRequestMessage(HttpMethod.Get, $"api/v1/blobs/{NamespaceId}/{Hash}");
+			using HttpRequestMessage request = new HttpRequestMessage(HttpMethod.Get, $"api/v1/blobs/{namespaceId}/{hash}");
 
-			HttpResponseMessage Response = await HttpClient.SendAsync(Request, CancellationToken);
+			HttpResponseMessage response = await _httpClient.SendAsync(request, cancellationToken);
 			try
 			{
-				if (Response.StatusCode == System.Net.HttpStatusCode.NotFound)
+				if (response.StatusCode == System.Net.HttpStatusCode.NotFound)
 				{
-					throw new BlobNotFoundException(NamespaceId, Hash);
+					throw new BlobNotFoundException(namespaceId, hash);
 				}
 
-				Response.EnsureSuccessStatusCode();
-				return await Response.Content.ReadAsStreamAsync();
+				response.EnsureSuccessStatusCode();
+				return await response.Content.ReadAsStreamAsync();
 			}
 			catch
 			{
-				Response.Dispose();
+				response.Dispose();
 				throw;
 			}
 		}
 
 		/// <inheritdoc/>
-		public async Task WriteBlobAsync(NamespaceId NamespaceId, IoHash Hash, Stream Stream, CancellationToken CancellationToken)
+		public async Task WriteBlobAsync(NamespaceId namespaceId, IoHash hash, Stream stream, CancellationToken cancellationToken)
 		{
-			using HttpRequestMessage Request = new HttpRequestMessage(HttpMethod.Put, $"api/v1/blobs/{NamespaceId}/{Hash}");
+			using HttpRequestMessage request = new HttpRequestMessage(HttpMethod.Put, $"api/v1/blobs/{namespaceId}/{hash}");
 
-			StreamContent Content = new StreamContent(Stream);
-			Content.Headers.ContentType = new MediaTypeHeaderValue("application/octet-stream");
-			Request.Content = Content;
+			StreamContent content = new StreamContent(stream);
+			content.Headers.ContentType = new MediaTypeHeaderValue("application/octet-stream");
+			request.Content = content;
 
-			HttpResponseMessage Response = await HttpClient.SendAsync(Request, CancellationToken);
-			Response.EnsureSuccessStatusCode();
+			HttpResponseMessage response = await _httpClient.SendAsync(request, cancellationToken);
+			response.EnsureSuccessStatusCode();
 		}
 
 		/// <inheritdoc/>
-		public async Task<bool> HasBlobAsync(NamespaceId NamespaceId, IoHash Hash, CancellationToken CancellationToken = default)
+		public async Task<bool> HasBlobAsync(NamespaceId namespaceId, IoHash hash, CancellationToken cancellationToken = default)
 		{
-			using HttpRequestMessage Request = new HttpRequestMessage(HttpMethod.Head, $"api/v1/blobs/{NamespaceId}/{Hash}");
+			using HttpRequestMessage request = new HttpRequestMessage(HttpMethod.Head, $"api/v1/blobs/{namespaceId}/{hash}");
 
-			HttpResponseMessage Response = await HttpClient.SendAsync(Request, CancellationToken);
+			HttpResponseMessage response = await _httpClient.SendAsync(request, cancellationToken);
 			try
 			{
-				if (Response.StatusCode != System.Net.HttpStatusCode.NotFound)
+				if (response.StatusCode != System.Net.HttpStatusCode.NotFound)
 				{
-					Response.EnsureSuccessStatusCode();
+					response.EnsureSuccessStatusCode();
 				}
-				return Response.StatusCode == System.Net.HttpStatusCode.OK;
+				return response.StatusCode == System.Net.HttpStatusCode.OK;
 			}
 			catch
 			{
-				Response.Dispose();
+				response.Dispose();
 				throw;
 			}
 		}
 
 		/// <inheritdoc/>
-		public async Task<HashSet<IoHash>> FindMissingBlobsAsync(NamespaceId NamespaceId, HashSet<IoHash> Hashes, CancellationToken CancellationToken)
+		public async Task<HashSet<IoHash>> FindMissingBlobsAsync(NamespaceId namespaceId, HashSet<IoHash> hashes, CancellationToken cancellationToken)
 		{
-			using StringContent Content = new StringContent(JsonSerializer.Serialize(Hashes));
+			using StringContent content = new StringContent(JsonSerializer.Serialize(hashes));
 
-			using HttpRequestMessage Request = new HttpRequestMessage(HttpMethod.Post, $"api/v1/blobs/{NamespaceId}/exist");
+			using HttpRequestMessage request = new HttpRequestMessage(HttpMethod.Post, $"api/v1/blobs/{namespaceId}/exist");
 
-			Request.Headers.Accept.Clear();
-			Request.Headers.Accept.Add(new MediaTypeWithQualityHeaderValue("application/json"));
-			Content.Headers.ContentType = new MediaTypeHeaderValue("application/json");
-			Request.Content = Content;
+			request.Headers.Accept.Clear();
+			request.Headers.Accept.Add(new MediaTypeWithQualityHeaderValue("application/json"));
+			content.Headers.ContentType = new MediaTypeHeaderValue("application/json");
+			request.Content = content;
 
-			HttpResponseMessage Response = await HttpClient.SendAsync(Request, CancellationToken);
-			Response.EnsureSuccessStatusCode();
+			HttpResponseMessage response = await _httpClient.SendAsync(request, cancellationToken);
+			response.EnsureSuccessStatusCode();
 
-			BlobExistsResponse? ResponseBody = await ReadJsonResponse<BlobExistsResponse>(Response.Content);
-			if (ResponseBody == null)
+			BlobExistsResponse? responseBody = await ReadJsonResponse<BlobExistsResponse>(response.Content);
+			if (responseBody == null)
 			{
 				throw new StorageException("Unable to parse response body", null);
 			}
-			return ResponseBody.Needs;
+			return responseBody.Needs;
 		}
 
 		#endregion
@@ -177,132 +175,132 @@ namespace EpicGames.Horde.Storage.Impl
 		#region Refs
 
 		/// <inheritdoc/>
-		public async Task<IRef> GetRefAsync(NamespaceId NamespaceId, BucketId BucketId, RefId RefId, CancellationToken CancellationToken)
+		public async Task<IRef> GetRefAsync(NamespaceId namespaceId, BucketId bucketId, RefId refId, CancellationToken cancellationToken)
 		{
-			using HttpRequestMessage Request = new HttpRequestMessage(HttpMethod.Get, $"api/v1/refs/{NamespaceId}/{BucketId}/{RefId}");
+			using HttpRequestMessage request = new HttpRequestMessage(HttpMethod.Get, $"api/v1/refs/{namespaceId}/{bucketId}/{refId}");
 
-			Request.Headers.Accept.Clear();
-			Request.Headers.Accept.Add(new MediaTypeWithQualityHeaderValue(CompactBinaryMimeType));
+			request.Headers.Accept.Clear();
+			request.Headers.Accept.Add(new MediaTypeWithQualityHeaderValue(CompactBinaryMimeType));
 
-			using HttpResponseMessage Response = await HttpClient.SendAsync(Request);
-			if (!Response.IsSuccessStatusCode)
+			using HttpResponseMessage response = await _httpClient.SendAsync(request);
+			if (!response.IsSuccessStatusCode)
 			{
-				if (Response.StatusCode == System.Net.HttpStatusCode.NotFound)
+				if (response.StatusCode == System.Net.HttpStatusCode.NotFound)
 				{
-					throw new RefNotFoundException(NamespaceId, BucketId, RefId);
+					throw new RefNotFoundException(namespaceId, bucketId, refId);
 				}
 				else
 				{
-					throw new RefException(NamespaceId, BucketId, RefId, await GetMessageFromResponse(Response));
+					throw new RefException(namespaceId, bucketId, refId, await GetMessageFromResponse(response));
 				}
 			}
 
-			byte[] Data = await Response.Content.ReadAsByteArrayAsync();
-			return new RefImpl(NamespaceId, BucketId, RefId, new CbObject(Data));
+			byte[] data = await response.Content.ReadAsByteArrayAsync();
+			return new RefImpl(namespaceId, bucketId, refId, new CbObject(data));
 		}
 
 		/// <inheritdoc/>
-		public async Task<List<IoHash>> TrySetRefAsync(NamespaceId NamespaceId, BucketId BucketId, RefId RefId, CbObject Value, CancellationToken CancellationToken)
+		public async Task<List<IoHash>> TrySetRefAsync(NamespaceId namespaceId, BucketId bucketId, RefId refId, CbObject value, CancellationToken cancellationToken)
 		{
-			using ReadOnlyMemoryStream Stream = new ReadOnlyMemoryStream(Value.GetView());
+			using ReadOnlyMemoryStream stream = new ReadOnlyMemoryStream(value.GetView());
 
-			StreamContent Content = new StreamContent(Stream);
-			Content.Headers.ContentType = new MediaTypeHeaderValue(CompactBinaryMimeType);
-			Content.Headers.Add(HashHeaderName, IoHash.Compute(Value.GetView().Span).ToString());
+			StreamContent content = new StreamContent(stream);
+			content.Headers.ContentType = new MediaTypeHeaderValue(CompactBinaryMimeType);
+			content.Headers.Add(HashHeaderName, IoHash.Compute(value.GetView().Span).ToString());
 
-			using HttpRequestMessage Request = new HttpRequestMessage(HttpMethod.Put, $"api/v1/refs/{NamespaceId}/{BucketId}/{RefId}");
+			using HttpRequestMessage request = new HttpRequestMessage(HttpMethod.Put, $"api/v1/refs/{namespaceId}/{bucketId}/{refId}");
 
-			Request.Headers.Accept.Clear();
-			Request.Headers.Accept.Add(new MediaTypeWithQualityHeaderValue("application/json"));
-			Request.Content = Content;
+			request.Headers.Accept.Clear();
+			request.Headers.Accept.Add(new MediaTypeWithQualityHeaderValue("application/json"));
+			request.Content = content;
 
-			using HttpResponseMessage Response = await HttpClient.SendAsync(Request);
-			if (!Response.IsSuccessStatusCode)
+			using HttpResponseMessage response = await _httpClient.SendAsync(request);
+			if (!response.IsSuccessStatusCode)
 			{
-				throw new RefException(NamespaceId, BucketId, RefId, await GetMessageFromResponse(Response));
+				throw new RefException(namespaceId, bucketId, refId, await GetMessageFromResponse(response));
 			}
 
-			RefPutResponse? ResponseBody = await ReadJsonResponse<RefPutResponse>(Response.Content);
-			if (ResponseBody == null)
+			RefPutResponse? responseBody = await ReadJsonResponse<RefPutResponse>(response.Content);
+			if (responseBody == null)
 			{
-				throw new RefException(NamespaceId, BucketId, RefId, "Unable to parse response body");
+				throw new RefException(namespaceId, bucketId, refId, "Unable to parse response body");
 			}
 
-			return ResponseBody.Needs;
+			return responseBody.Needs;
 		}
 
 		/// <inheritdoc/>
-		public async Task<List<IoHash>> TryFinalizeRefAsync(NamespaceId NamespaceId, BucketId BucketId, RefId RefId, IoHash Hash, CancellationToken CancellationToken)
+		public async Task<List<IoHash>> TryFinalizeRefAsync(NamespaceId namespaceId, BucketId bucketId, RefId refId, IoHash hash, CancellationToken cancellationToken)
 		{
-			using HttpRequestMessage Request = new HttpRequestMessage(HttpMethod.Post, $"api/v1/refs/{NamespaceId}/{BucketId}/{RefId}/finalize/{Hash}");
+			using HttpRequestMessage request = new HttpRequestMessage(HttpMethod.Post, $"api/v1/refs/{namespaceId}/{bucketId}/{refId}/finalize/{hash}");
 
-			Request.Headers.Accept.Clear();
-			Request.Headers.Accept.Add(new MediaTypeWithQualityHeaderValue("application/json"));
+			request.Headers.Accept.Clear();
+			request.Headers.Accept.Add(new MediaTypeWithQualityHeaderValue("application/json"));
 
-			using HttpResponseMessage Response = await HttpClient.SendAsync(Request);
-			if (!Response.IsSuccessStatusCode)
+			using HttpResponseMessage response = await _httpClient.SendAsync(request);
+			if (!response.IsSuccessStatusCode)
 			{
-				throw new RefException(NamespaceId, BucketId, RefId, await GetMessageFromResponse(Response));
+				throw new RefException(namespaceId, bucketId, refId, await GetMessageFromResponse(response));
 			}
 
-			RefPutResponse? ResponseBody = await ReadJsonResponse<RefPutResponse>(Response.Content);
-			if (ResponseBody == null)
+			RefPutResponse? responseBody = await ReadJsonResponse<RefPutResponse>(response.Content);
+			if (responseBody == null)
 			{
-				throw new RefException(NamespaceId, BucketId, RefId, "Unable to parse response body");
+				throw new RefException(namespaceId, bucketId, refId, "Unable to parse response body");
 			}
 
-			return ResponseBody.Needs;
+			return responseBody.Needs;
 		}
 
 		/// <inheritdoc/>
-		public async Task<bool> DeleteRefAsync(NamespaceId NamespaceId, BucketId BucketId, RefId RefId, CancellationToken CancellationToken)
+		public async Task<bool> DeleteRefAsync(NamespaceId namespaceId, BucketId bucketId, RefId refId, CancellationToken cancellationToken)
 		{
-			using HttpRequestMessage Request = new HttpRequestMessage(HttpMethod.Delete, $"api/v1/refs/{NamespaceId}/{BucketId}/{RefId}");
+			using HttpRequestMessage request = new HttpRequestMessage(HttpMethod.Delete, $"api/v1/refs/{namespaceId}/{bucketId}/{refId}");
 
-			using HttpResponseMessage Response = await HttpClient.SendAsync(Request, CancellationToken);
-			return Response.IsSuccessStatusCode;
+			using HttpResponseMessage response = await _httpClient.SendAsync(request, cancellationToken);
+			return response.IsSuccessStatusCode;
 		}
 
 		/// <inheritdoc/>
-		public async Task<bool> HasRefAsync(NamespaceId NamespaceId, BucketId BucketId, RefId RefId, CancellationToken CancellationToken)
+		public async Task<bool> HasRefAsync(NamespaceId namespaceId, BucketId bucketId, RefId refId, CancellationToken cancellationToken)
 		{
-			using HttpRequestMessage Request = new HttpRequestMessage(HttpMethod.Head, $"api/v1/refs/{NamespaceId}/{BucketId}/{RefId}");
+			using HttpRequestMessage request = new HttpRequestMessage(HttpMethod.Head, $"api/v1/refs/{namespaceId}/{bucketId}/{refId}");
 
-			using HttpResponseMessage Response = await HttpClient.SendAsync(Request, CancellationToken);
-			if (Response.IsSuccessStatusCode)
+			using HttpResponseMessage response = await _httpClient.SendAsync(request, cancellationToken);
+			if (response.IsSuccessStatusCode)
 			{
 				return true;
 			}
-			else if (Response.StatusCode == System.Net.HttpStatusCode.NotFound)
+			else if (response.StatusCode == System.Net.HttpStatusCode.NotFound)
 			{
 				return false;
 			}
 			else
 			{
-				throw new RefException(NamespaceId, BucketId, RefId, $"Unexpected response for {nameof(HasRefAsync)} call: {Response.StatusCode}");
+				throw new RefException(namespaceId, bucketId, refId, $"Unexpected response for {nameof(HasRefAsync)} call: {response.StatusCode}");
 			}
 		}
 
 		/// <inheritdoc/>
-		public async Task<List<RefId>> FindMissingRefsAsync(NamespaceId NamespaceId, BucketId BucketId, List<RefId> RefIds, CancellationToken CancellationToken)
+		public async Task<List<RefId>> FindMissingRefsAsync(NamespaceId namespaceId, BucketId bucketId, List<RefId> refIds, CancellationToken cancellationToken)
 		{
-			string RefList = String.Join("&", RefIds.Select(x => $"id={x}"));
+			string refList = String.Join("&", refIds.Select(x => $"id={x}"));
 
-			using HttpRequestMessage Request = new HttpRequestMessage(HttpMethod.Post, $"/api/v1/refs/{NamespaceId}/{BucketId}/exists?{RefList}");
+			using HttpRequestMessage request = new HttpRequestMessage(HttpMethod.Post, $"/api/v1/refs/{namespaceId}/{bucketId}/exists?{refList}");
 
-			using HttpResponseMessage Response = await HttpClient.SendAsync(Request, CancellationToken);
-			if (!Response.IsSuccessStatusCode)
+			using HttpResponseMessage response = await _httpClient.SendAsync(request, cancellationToken);
+			if (!response.IsSuccessStatusCode)
 			{
-				throw new RefException(NamespaceId, BucketId, new RefId(IoHash.Zero), await GetMessageFromResponse(Response));
+				throw new RefException(namespaceId, bucketId, new RefId(IoHash.Zero), await GetMessageFromResponse(response));
 			}
 
-			RefExistsResponse? ResponseContent = await ReadJsonResponse<RefExistsResponse>(Response.Content);
-			if (ResponseContent == null)
+			RefExistsResponse? responseContent = await ReadJsonResponse<RefExistsResponse>(response.Content);
+			if (responseContent == null)
 			{
-				throw new RefException(NamespaceId, BucketId, new RefId(IoHash.Zero), "Unable to parse response body");
+				throw new RefException(namespaceId, bucketId, new RefId(IoHash.Zero), "Unable to parse response body");
 			}
 
-			return ResponseContent.Needs;
+			return responseContent.Needs;
 		}
 
 		#endregion
@@ -312,26 +310,26 @@ namespace EpicGames.Horde.Storage.Impl
 			public string? Title { get; set; }
 		}
 
-		static async Task<string> GetMessageFromResponse(HttpResponseMessage Response)
+		static async Task<string> GetMessageFromResponse(HttpResponseMessage response)
 		{
-			StringBuilder Message = new StringBuilder($"HTTP {Response.StatusCode}");
+			StringBuilder message = new StringBuilder($"HTTP {response.StatusCode}");
 			try
 			{
-				byte[] Data = await Response.Content.ReadAsByteArrayAsync();
-				ProblemDetails? Details = JsonSerializer.Deserialize<ProblemDetails>(Data);
-				Message.Append($": {Details?.Title ?? "No description available"}");
+				byte[] data = await response.Content.ReadAsByteArrayAsync();
+				ProblemDetails? details = JsonSerializer.Deserialize<ProblemDetails>(data);
+				message.Append($": {details?.Title ?? "No description available"}");
 			}
 			catch
 			{
-				Message.Append(" (Unable to parse response data)");
+				message.Append(" (Unable to parse response data)");
 			}
-			return Message.ToString();
+			return message.ToString();
 		}
 
-		static async Task<T> ReadJsonResponse<T>(HttpContent Content)
+		static async Task<T> ReadJsonResponse<T>(HttpContent content)
 		{
-			byte[] Data = await Content.ReadAsByteArrayAsync();
-			return JsonSerializer.Deserialize<T>(Data, new JsonSerializerOptions { PropertyNameCaseInsensitive = true });
+			byte[] data = await content.ReadAsByteArrayAsync();
+			return JsonSerializer.Deserialize<T>(data, new JsonSerializerOptions { PropertyNameCaseInsensitive = true });
 		}
 	}
 }
diff --git a/Engine/Source/Programs/Shared/EpicGames.Horde/Storage/Impl/MemoryStorageClient.cs b/Engine/Source/Programs/Shared/EpicGames.Horde/Storage/Impl/MemoryStorageClient.cs
--- a/Engine/Source/Programs/Shared/EpicGames.Horde/Storage/Impl/MemoryStorageClient.cs
+++ b/Engine/Source/Programs/Shared/EpicGames.Horde/Storage/Impl/MemoryStorageClient.cs
@@ -2,12 +2,10 @@
 
 using EpicGames.Core;
 using EpicGames.Serialization;
-using Microsoft.Extensions.Logging;
 using System;
 using System.Collections.Generic;
 using System.IO;
 using System.Linq;
-using System.Text;
 using System.Threading;
 using System.Threading.Tasks;
 
@@ -25,12 +23,12 @@ namespace EpicGames.Horde.Storage.Impl
 			public RefId RefId { get; set; }
 			public CbObject Value { get; set; }
 
-			public Ref(NamespaceId NamespaceId, BucketId BucketId, RefId RefId, CbObject Value)
+			public Ref(NamespaceId namespaceId, BucketId bucketId, RefId refId, CbObject value)
 			{
-				this.NamespaceId = NamespaceId;
-				this.BucketId = BucketId;
-				this.RefId = RefId;
-				this.Value = Value;
+				NamespaceId = namespaceId;
+				BucketId = bucketId;
+				RefId = refId;
+				Value = value;
 			}
 		}
 
@@ -41,68 +39,68 @@ namespace EpicGames.Horde.Storage.Impl
 		public Dictionary<(NamespaceId, BucketId, RefId), IRef> Refs { get; } = new Dictionary<(NamespaceId, BucketId, RefId), IRef>();
 
 		/// <inheritdoc/>
-		public Task<Stream> ReadBlobAsync(NamespaceId NamespaceId, IoHash Hash, CancellationToken CancellationToken = default)
+		public Task<Stream> ReadBlobAsync(NamespaceId namespaceId, IoHash hash, CancellationToken cancellationToken = default)
 		{
-			ReadOnlyMemory<byte> Data = Blobs[(NamespaceId, Hash)];
-			return Task.FromResult<Stream>(new ReadOnlyMemoryStream(Data));
+			ReadOnlyMemory<byte> data = Blobs[(namespaceId, hash)];
+			return Task.FromResult<Stream>(new ReadOnlyMemoryStream(data));
 		}
 
 		/// <inheritdoc/>
-		public async Task WriteBlobAsync(NamespaceId NamespaceId, IoHash Hash, Stream Stream, CancellationToken CancellationToken = default)
+		public async Task WriteBlobAsync(NamespaceId namespaceId, IoHash hash, Stream stream, CancellationToken cancellationToken = default)
 		{
-			using (MemoryStream MemoryStream = new MemoryStream())
+			using (MemoryStream memoryStream = new MemoryStream())
 			{
-				await Stream.CopyToAsync(MemoryStream);
-				Blobs[(NamespaceId, Hash)] = MemoryStream.ToArray();
+				await stream.CopyToAsync(memoryStream);
+				Blobs[(namespaceId, hash)] = memoryStream.ToArray();
 			}
 		}
 
 		/// <inheritdoc/>
-		public Task<bool> HasBlobAsync(NamespaceId NamespaceId, IoHash Hash, CancellationToken CancellationToken = default)
+		public Task<bool> HasBlobAsync(NamespaceId namespaceId, IoHash hash, CancellationToken cancellationToken = default)
 		{
-			return Task.FromResult(Blobs.ContainsKey((NamespaceId, Hash)));
+			return Task.FromResult(Blobs.ContainsKey((namespaceId, hash)));
 		}
 
 		/// <inheritdoc/>
-		public Task<HashSet<IoHash>> FindMissingBlobsAsync(NamespaceId NamespaceId, HashSet<IoHash> Hashes, CancellationToken CancellationToken)
+		public Task<HashSet<IoHash>> FindMissingBlobsAsync(NamespaceId namespaceId, HashSet<IoHash> hashes, CancellationToken cancellationToken)
 		{
-			return Task.FromResult(Hashes.Where(x => !Blobs.ContainsKey((NamespaceId, x))).ToHashSet());
+			return Task.FromResult(hashes.Where(x => !Blobs.ContainsKey((namespaceId, x))).ToHashSet());
 		}
 
 		/// <inheritdoc/>
-		public Task<bool> DeleteRefAsync(NamespaceId NamespaceId, BucketId BucketId, RefId RefId, CancellationToken CancellationToken = default)
+		public Task<bool> DeleteRefAsync(NamespaceId namespaceId, BucketId bucketId, RefId refId, CancellationToken cancellationToken = default)
 		{
-			return Task.FromResult(Refs.Remove((NamespaceId, BucketId, RefId)));
+			return Task.FromResult(Refs.Remove((namespaceId, bucketId, refId)));
 		}
 
 		/// <inheritdoc/>
-		public Task<List<RefId>> FindMissingRefsAsync(NamespaceId NamespaceId, BucketId BucketId, List<RefId> RefIds, CancellationToken CancellationToken = default)
+		public Task<List<RefId>> FindMissingRefsAsync(NamespaceId namespaceId, BucketId bucketId, List<RefId> refIds, CancellationToken cancellationToken = default)
 		{
-			return Task.FromResult(RefIds.Where(x => !Refs.ContainsKey((NamespaceId, BucketId, x))).ToList());
+			return Task.FromResult(refIds.Where(x => !Refs.ContainsKey((namespaceId, bucketId, x))).ToList());
 		}
 
 		/// <inheritdoc/>
-		public Task<IRef> GetRefAsync(NamespaceId NamespaceId, BucketId BucketId, RefId RefId, CancellationToken CancellationToken = default)
+		public Task<IRef> GetRefAsync(NamespaceId namespaceId, BucketId bucketId, RefId refId, CancellationToken cancellationToken = default)
 		{
-			return Task.FromResult(Refs[(NamespaceId, BucketId, RefId)]);
+			return Task.FromResult(Refs[(namespaceId, bucketId, refId)]);
 		}
 
 		/// <inheritdoc/>
-		public Task<bool> HasRefAsync(NamespaceId NamespaceId, BucketId BucketId, RefId RefId, CancellationToken CancellationToken = default)
+		public Task<bool> HasRefAsync(NamespaceId namespaceId, BucketId bucketId, RefId refId, CancellationToken cancellationToken = default)
 		{
-			return Task.FromResult(Refs.ContainsKey((NamespaceId, BucketId, RefId)));
+			return Task.FromResult(Refs.ContainsKey((namespaceId, bucketId, refId)));
 		}
 
 		/// <inheritdoc/>
-		public Task<List<IoHash>> TryFinalizeRefAsync(NamespaceId NamespaceId, BucketId BucketId, RefId RefId, IoHash Hash, CancellationToken CancellationToken = default)
+		public Task<List<IoHash>> TryFinalizeRefAsync(NamespaceId namespaceId, BucketId bucketId, RefId refId, IoHash hash, CancellationToken cancellationToken = default)
 		{
 			throw new NotImplementedException();
 		}
 
 		/// <inheritdoc/>
-		public Task<List<IoHash>> TrySetRefAsync(NamespaceId NamespaceId, BucketId BucketId, RefId RefId, CbObject Value, CancellationToken CancellationToken = default)
+		public Task<List<IoHash>> TrySetRefAsync(NamespaceId namespaceId, BucketId bucketId, RefId refId, CbObject value, CancellationToken cancellationToken = default)
 		{
-			Refs[(NamespaceId, BucketId, RefId)] = new Ref(NamespaceId, BucketId, RefId, Value);
+			Refs[(namespaceId, bucketId, refId)] = new Ref(namespaceId, bucketId, refId, value);
 			return Task.FromResult(new List<IoHash>());
 		}
 	}
diff --git a/Engine/Source/Programs/Shared/EpicGames.Horde/Storage/Impl/TracingStorageClient.cs b/Engine/Source/Programs/Shared/EpicGames.Horde/Storage/Impl/TracingStorageClient.cs
--- a/Engine/Source/Programs/Shared/EpicGames.Horde/Storage/Impl/TracingStorageClient.cs
+++ b/Engine/Source/Programs/Shared/EpicGames.Horde/Storage/Impl/TracingStorageClient.cs
@@ -3,11 +3,8 @@
 using EpicGames.Core;
 using EpicGames.Serialization;
 using Microsoft.Extensions.Logging;
-using System;
 using System.Collections.Generic;
 using System.IO;
-using System.Linq;
-using System.Text;
 using System.Threading;
 using System.Threading.Tasks;
 
@@ -18,83 +15,83 @@ namespace EpicGames.Horde.Storage.Impl
 	/// </summary>
 	public class TracingStorageClient : IStorageClient
 	{
-		IStorageClient Inner;
-		ILogger Logger;
+		readonly IStorageClient _inner;
+		readonly ILogger _logger;
 
 		/// <summary>
 		/// Constructor
 		/// </summary>
-		/// <param name="Inner"></param>
-		/// <param name="Logger"></param>
-		public TracingStorageClient(IStorageClient Inner, ILogger Logger)
+		/// <param name="inner"></param>
+		/// <param name="logger"></param>
+		public TracingStorageClient(IStorageClient inner, ILogger logger)
 		{
-			this.Inner = Inner;
-			this.Logger = Logger;
+			_inner = inner;
+			_logger = logger;
 		}
 
 		/// <inheritdoc/>
-		public Task<Stream> ReadBlobAsync(NamespaceId NamespaceId, IoHash Hash, CancellationToken CancellationToken = default)
+		public Task<Stream> ReadBlobAsync(NamespaceId namespaceId, IoHash hash, CancellationToken cancellationToken = default)
 		{
-			Logger.LogDebug("Reading blob {NamespaceId}/{Hash}", NamespaceId, Hash);
-			return Inner.ReadBlobAsync(NamespaceId, Hash, CancellationToken);
+			_logger.LogDebug("Reading blob {NamespaceId}/{Hash}", namespaceId, hash);
+			return _inner.ReadBlobAsync(namespaceId, hash, cancellationToken);
 		}
 
 		/// <inheritdoc/>
-		public Task WriteBlobAsync(NamespaceId NamespaceId, IoHash Hash, Stream Stream, CancellationToken CancellationToken = default)
+		public Task WriteBlobAsync(NamespaceId namespaceId, IoHash hash, Stream stream, CancellationToken cancellationToken = default)
 		{
-			Logger.LogDebug("Writing blob {NamespaceId}/{Hash}", NamespaceId, Hash);
-			return Inner.WriteBlobAsync(NamespaceId, Hash, Stream, CancellationToken);
+			_logger.LogDebug("Writing blob {NamespaceId}/{Hash}", namespaceId, hash);
+			return _inner.WriteBlobAsync(namespaceId, hash, stream, cancellationToken);
 		}
 
 		/// <inheritdoc/>
-		public Task<bool> HasBlobAsync(NamespaceId NamespaceId, IoHash Hash, CancellationToken CancellationToken = default)
+		public Task<bool> HasBlobAsync(NamespaceId namespaceId, IoHash hash, CancellationToken cancellationToken = default)
 		{
-			return Inner.HasBlobAsync(NamespaceId, Hash, CancellationToken);
+			return _inner.HasBlobAsync(namespaceId, hash, cancellationToken);
 		}
 
 		/// <inheritdoc/>
-		public Task<HashSet<IoHash>> FindMissingBlobsAsync(NamespaceId NamespaceId, HashSet<IoHash> Hashes, CancellationToken CancellationToken)
+		public Task<HashSet<IoHash>> FindMissingBlobsAsync(NamespaceId namespaceId, HashSet<IoHash> hashes, CancellationToken cancellationToken)
 		{
-			return Inner.FindMissingBlobsAsync(NamespaceId, Hashes, CancellationToken);
+			return _inner.FindMissingBlobsAsync(namespaceId, hashes, cancellationToken);
 		}
 
 		/// <inheritdoc/>
-		public Task<bool> DeleteRefAsync(NamespaceId NamespaceId, BucketId BucketId, RefId RefId, CancellationToken CancellationToken = default)
+		public Task<bool> DeleteRefAsync(NamespaceId namespaceId, BucketId bucketId, RefId refId, CancellationToken cancellationToken = default)
 		{
-			Logger.LogDebug("Deleting ref {NamespaceId}/{BucketId}/{RefId}", NamespaceId, BucketId, RefId);
-			return Inner.DeleteRefAsync(NamespaceId, BucketId, RefId, CancellationToken);
+			_logger.LogDebug("Deleting ref {NamespaceId}/{BucketId}/{RefId}", namespaceId, bucketId, refId);
+			return _inner.DeleteRefAsync(namespaceId, bucketId, refId, cancellationToken);
 		}
 
 		/// <inheritdoc/>
-		public Task<List<RefId>> FindMissingRefsAsync(NamespaceId NamespaceId, BucketId BucketId, List<RefId> RefIds, CancellationToken CancellationToken = default)
+		public Task<List<RefId>> FindMissingRefsAsync(NamespaceId namespaceId, BucketId bucketId, List<RefId> refIds, CancellationToken cancellationToken = default)
 		{
-			return Inner.FindMissingRefsAsync(NamespaceId, BucketId, RefIds, CancellationToken);
+			return _inner.FindMissingRefsAsync(namespaceId, bucketId, refIds, cancellationToken);
 		}
 
 		/// <inheritdoc/>
-		public Task<IRef> GetRefAsync(NamespaceId NamespaceId, BucketId BucketId, RefId RefId, CancellationToken CancellationToken = default)
+		public Task<IRef> GetRefAsync(NamespaceId namespaceId, BucketId bucketId, RefId refId, CancellationToken cancellationToken = default)
 		{
-			Logger.LogDebug("Getting ref {NamespaceId}/{BucketId}/{RefId}", NamespaceId, BucketId, RefId);
-			return Inner.GetRefAsync(NamespaceId, BucketId, RefId, CancellationToken);
+			_logger.LogDebug("Getting ref {NamespaceId}/{BucketId}/{RefId}", namespaceId, bucketId, refId);
+			return _inner.GetRefAsync(namespaceId, bucketId, refId, cancellationToken);
 		}
 
 		/// <inheritdoc/>
-		public Task<bool> HasRefAsync(NamespaceId NamespaceId, BucketId BucketId, RefId RefId, CancellationToken CancellationToken = default)
+		public Task<bool> HasRefAsync(NamespaceId namespaceId, BucketId bucketId, RefId refId, CancellationToken cancellationToken = default)
 		{
-			return Inner.HasRefAsync(NamespaceId, BucketId, RefId, CancellationToken);
+			return _inner.HasRefAsync(namespaceId, bucketId, refId, cancellationToken);
 		}
 
 		/// <inheritdoc/>
-		public Task<List<IoHash>> TryFinalizeRefAsync(NamespaceId NamespaceId, BucketId BucketId, RefId RefId, IoHash Hash, CancellationToken CancellationToken = default)
+		public Task<List<IoHash>> TryFinalizeRefAsync(NamespaceId namespaceId, BucketId bucketId, RefId refId, IoHash hash, CancellationToken cancellationToken = default)
 		{
-			return Inner.TryFinalizeRefAsync(NamespaceId, BucketId, RefId, Hash, CancellationToken);
+			return _inner.TryFinalizeRefAsync(namespaceId, bucketId, refId, hash, cancellationToken);
 		}
 
 		/// <inheritdoc/>
-		public Task<List<IoHash>> TrySetRefAsync(NamespaceId NamespaceId, BucketId BucketId, RefId RefId, CbObject Value, CancellationToken CancellationToken = default)
+		public Task<List<IoHash>> TrySetRefAsync(NamespaceId namespaceId, BucketId bucketId, RefId refId, CbObject value, CancellationToken cancellationToken = default)
 		{
-			Logger.LogDebug("Setting ref {NamespaceId}/{BucketId}/{RefId}", NamespaceId, BucketId, RefId);
-			return Inner.TrySetRefAsync(NamespaceId, BucketId, RefId, Value, CancellationToken);
+			_logger.LogDebug("Setting ref {NamespaceId}/{BucketId}/{RefId}", namespaceId, bucketId, refId);
+			return _inner.TrySetRefAsync(namespaceId, bucketId, refId, value, cancellationToken);
 		}
 	}
 }
diff --git a/Engine/Source/Programs/Shared/EpicGames.Horde/Storage/NamespaceId.cs b/Engine/Source/Programs/Shared/EpicGames.Horde/Storage/NamespaceId.cs
--- a/Engine/Source/Programs/Shared/EpicGames.Horde/Storage/NamespaceId.cs
+++ b/Engine/Source/Programs/Shared/EpicGames.Horde/Storage/NamespaceId.cs
@@ -1,11 +1,8 @@
 // Copyright Epic Games, Inc. All Rights Reserved.
 
 using System;
-using System.Collections.Generic;
 using System.ComponentModel;
-using System.Diagnostics.CodeAnalysis;
 using System.Globalization;
-using System.Text;
 using System.Text.Json;
 using System.Text.Json.Serialization;
 
@@ -21,34 +18,34 @@ namespace EpicGames.Horde.Storage
 		/// <summary>
 		/// The text representing this id
 		/// </summary>
-		readonly StringId Inner;
+		readonly StringId _inner;
 
 		/// <summary>
 		/// Constructor
 		/// </summary>
-		/// <param name="Input">Unique id for the string</param>
-		public NamespaceId(string Input)
+		/// <param name="input">Unique id for the string</param>
+		public NamespaceId(string input)
 		{
-			Inner = new StringId(Input);
+			_inner = new StringId(input);
 		}
 
 		/// <inheritdoc/>
-		public override bool Equals(object? Obj) => Obj is NamespaceId Id && Inner.Equals(Id.Inner);
+		public override bool Equals(object? obj) => obj is NamespaceId id && _inner.Equals(id._inner);
 
 		/// <inheritdoc/>
-		public override int GetHashCode() => Inner.GetHashCode();
+		public override int GetHashCode() => _inner.GetHashCode();
 
 		/// <inheritdoc/>
-		public bool Equals(NamespaceId Other) => Inner.Equals(Other.Inner);
+		public bool Equals(NamespaceId other) => _inner.Equals(other._inner);
 
 		/// <inheritdoc/>
-		public override string ToString() => Inner.ToString();
+		public override string ToString() => _inner.ToString();
 
 		/// <inheritdoc cref="StringId.op_Equality"/>
-		public static bool operator ==(NamespaceId Left, NamespaceId Right) => Left.Inner == Right.Inner;
+		public static bool operator ==(NamespaceId left, NamespaceId right) => left._inner == right._inner;
 
 		/// <inheritdoc cref="StringId.op_Inequality"/>
-		public static bool operator !=(NamespaceId Left, NamespaceId Right) => Left.Inner != Right.Inner;
+		public static bool operator !=(NamespaceId left, NamespaceId right) => left._inner != right._inner;
 	}
 
 	/// <summary>
@@ -57,10 +54,10 @@ namespace EpicGames.Horde.Storage
 	sealed class NamespaceIdJsonConverter : JsonConverter<NamespaceId>
 	{
 		/// <inheritdoc/>
-		public override NamespaceId Read(ref Utf8JsonReader Reader, Type TypeToConvert, JsonSerializerOptions Options) => new NamespaceId(Reader.GetString() ?? String.Empty);
+		public override NamespaceId Read(ref Utf8JsonReader reader, Type typeToConvert, JsonSerializerOptions options) => new NamespaceId(reader.GetString() ?? String.Empty);
 
 		/// <inheritdoc/>
-		public override void Write(Utf8JsonWriter Writer, NamespaceId Value, JsonSerializerOptions Options) => Writer.WriteStringValue(Value.ToString());
+		public override void Write(Utf8JsonWriter writer, NamespaceId value, JsonSerializerOptions options) => writer.WriteStringValue(value.ToString());
 	}
 
 	/// <summary>
@@ -69,9 +66,9 @@ namespace EpicGames.Horde.Storage
 	sealed class NamespaceIdTypeConverter : TypeConverter
 	{
 		/// <inheritdoc/>
-		public override bool CanConvertFrom(ITypeDescriptorContext Context, Type SourceType) => SourceType == typeof(string);
+		public override bool CanConvertFrom(ITypeDescriptorContext context, Type sourceType) => sourceType == typeof(string);
 
 		/// <inheritdoc/>
-		public override object ConvertFrom(ITypeDescriptorContext Context, CultureInfo Culture, object Value) => new NamespaceId((string)Value);
+		public override object ConvertFrom(ITypeDescriptorContext context, CultureInfo culture, object value) => new NamespaceId((string)value);
 	}
 }
diff --git a/Engine/Source/Programs/Shared/EpicGames.Horde/Storage/RefId.cs b/Engine/Source/Programs/Shared/EpicGames.Horde/Storage/RefId.cs
--- a/Engine/Source/Programs/Shared/EpicGames.Horde/Storage/RefId.cs
+++ b/Engine/Source/Programs/Shared/EpicGames.Horde/Storage/RefId.cs
@@ -3,7 +3,6 @@
 using EpicGames.Core;
 using EpicGames.Serialization;
 using System;
-using System.Collections.Generic;
 using System.ComponentModel;
 using System.Globalization;
 using System.Text;
@@ -28,38 +27,38 @@ namespace EpicGames.Horde.Storage
 		/// <summary>
 		/// Constructor
 		/// </summary>
-		/// <param name="Hash"></param>
-		public RefId(IoHash Hash)
+		/// <param name="hash"></param>
+		public RefId(IoHash hash)
 		{
-			this.Hash = Hash;
+			Hash = hash;
 		}
 
 		/// <summary>
 		/// Constructor
 		/// </summary>
-		/// <param name="Name">Name to identify this ref</param>
-		public RefId(string Name)
+		/// <param name="name">Name to identify this ref</param>
+		public RefId(string name)
 		{
-			this.Hash = IoHash.Compute(Encoding.UTF8.GetBytes(Name));
+			Hash = IoHash.Compute(Encoding.UTF8.GetBytes(name));
 		}
 
 		/// <inheritdoc/>
-		public override bool Equals(object? Obj) => Obj is RefId RefId && Equals(RefId);
+		public override bool Equals(object? obj) => obj is RefId refId && Equals(refId);
 
 		/// <inheritdoc/>
 		public override int GetHashCode() => Hash.GetHashCode();
 
 		/// <inheritdoc/>
-		public bool Equals(RefId RefId) => Hash == RefId.Hash;
+		public bool Equals(RefId refId) => Hash == refId.Hash;
 
 		/// <inheritdoc/>
 		public override string ToString() => Hash.ToString();
 
 		/// <inheritdoc/>
-		public static bool operator ==(RefId Lhs, RefId Rhs) => Lhs.Equals(Rhs);
+		public static bool operator ==(RefId lhs, RefId rhs) => lhs.Equals(rhs);
 
 		/// <inheritdoc/>
-		public static bool operator !=(RefId Lhs, RefId Rhs) => !Lhs.Equals(Rhs);
+		public static bool operator !=(RefId lhs, RefId rhs) => !lhs.Equals(rhs);
 	}
 
 	/// <summary>
@@ -68,10 +67,10 @@ namespace EpicGames.Horde.Storage
 	sealed class RefIdJsonConverter : JsonConverter<RefId>
 	{
 		/// <inheritdoc/>
-		public override RefId Read(ref Utf8JsonReader Reader, Type TypeToConvert, JsonSerializerOptions Options) => new RefId(IoHash.Parse(Reader.ValueSpan));
+		public override RefId Read(ref Utf8JsonReader reader, Type typeToConvert, JsonSerializerOptions options) => new RefId(IoHash.Parse(reader.ValueSpan));
 
 		/// <inheritdoc/>
-		public override void Write(Utf8JsonWriter Writer, RefId Value, JsonSerializerOptions Options) => Writer.WriteStringValue(Value.Hash.ToUtf8String().Span);
+		public override void Write(Utf8JsonWriter writer, RefId value, JsonSerializerOptions options) => writer.WriteStringValue(value.Hash.ToUtf8String().Span);
 	}
 
 	/// <summary>
@@ -80,22 +79,22 @@ namespace EpicGames.Horde.Storage
 	sealed class RefIdTypeConverter : TypeConverter
 	{
 		/// <inheritdoc/>
-		public override bool CanConvertFrom(ITypeDescriptorContext Context, Type SourceType)
+		public override bool CanConvertFrom(ITypeDescriptorContext context, Type sourceType)
 		{
-			return SourceType == typeof(string);
+			return sourceType == typeof(string);
 		}
 
 		/// <inheritdoc/>
-		public override object ConvertFrom(ITypeDescriptorContext Context, CultureInfo Culture, object Value)
+		public override object ConvertFrom(ITypeDescriptorContext context, CultureInfo culture, object value)
 		{
-			IoHash Hash;
-			if (IoHash.TryParse((string)Value, out Hash))
+			IoHash hash;
+			if (IoHash.TryParse((string)value, out hash))
 			{
-				return new RefId(Hash);
+				return new RefId(hash);
 			}
 			else
 			{
-				return new RefId((string)Value);
+				return new RefId((string)value);
 			}
 		}
 	}
@@ -106,12 +105,12 @@ namespace EpicGames.Horde.Storage
 	sealed class RefIdCbConverter : CbConverterBase<RefId>
 	{
 		/// <inheritdoc/>
-		public override RefId Read(CbField Field) => new RefId(Field.AsHash());
+		public override RefId Read(CbField field) => new RefId(field.AsHash());
 
 		/// <inheritdoc/>
-		public override void Write(CbWriter Writer, RefId Value) => Writer.WriteHashValue(Value.Hash);
+		public override void Write(CbWriter writer, RefId value) => writer.WriteHashValue(value.Hash);
 
 		/// <inheritdoc/>
-		public override void WriteNamed(CbWriter Writer, Utf8String Name, RefId Value) => Writer.WriteHash(Name, Value.Hash);
+		public override void WriteNamed(CbWriter writer, Utf8String name, RefId value) => writer.WriteHash(name, value.Hash);
 	}
 }
diff --git a/Engine/Source/Programs/Shared/EpicGames.Horde/Storage/StorageExtensions.cs b/Engine/Source/Programs/Shared/EpicGames.Horde/Storage/StorageExtensions.cs
--- a/Engine/Source/Programs/Shared/EpicGames.Horde/Storage/StorageExtensions.cs
+++ b/Engine/Source/Programs/Shared/EpicGames.Horde/Storage/StorageExtensions.cs
@@ -5,16 +5,12 @@ using EpicGames.Horde.Storage.Impl;
 using Microsoft.Extensions.DependencyInjection;
 using Microsoft.Extensions.Options;
 using System;
-using System.Collections.Generic;
-using System.Net.Http;
-using System.Text;
 
 namespace EpicGames.Horde.Storage
 {
 	/// <summary>
 	/// Settings to configure a connected Horde.Storage instance
 	/// </summary>
-	[System.Diagnostics.CodeAnalysis.SuppressMessage("Design", "CA1056:URI-like properties should not be strings")]
 	public class StorageOptions : HttpServiceClientOptions
 	{
 	}
@@ -27,24 +23,24 @@ namespace EpicGames.Horde.Storage
 		/// <summary>
 		/// Registers services for Horde Storage
 		/// </summary>
-		/// <param name="Services">The current service collection</param>
-		public static void AddHordeStorage(this IServiceCollection Services)
+		/// <param name="services">The current service collection</param>
+		public static void AddHordeStorage(this IServiceCollection services)
 		{
-			Services.AddOptions<StorageOptions>();
+			services.AddOptions<StorageOptions>();
 
-			Services.AddScoped<IStorageClient, HttpStorageClient>();
-			Services.AddHttpClientWithAuth<IStorageClient, HttpStorageClient>(ServiceProvider => ServiceProvider.GetRequiredService<IOptions<StorageOptions>>().Value);
+			services.AddScoped<IStorageClient, HttpStorageClient>();
+			services.AddHttpClientWithAuth<IStorageClient, HttpStorageClient>(serviceProvider => serviceProvider.GetRequiredService<IOptions<StorageOptions>>().Value);
 		}
 
 		/// <summary>
 		/// Registers services for Horde Storage
 		/// </summary>
-		/// <param name="Services">The current service collection</param>
-		/// <param name="Configure">Callback for configuring the storage service</param>
-		public static void AddHordeStorage(this IServiceCollection Services, Action<StorageOptions> Configure)
+		/// <param name="services">The current service collection</param>
+		/// <param name="configure">Callback for configuring the storage service</param>
+		public static void AddHordeStorage(this IServiceCollection services, Action<StorageOptions> configure)
 		{
-			AddHordeStorage(Services);
-			Services.Configure<StorageOptions>(Configure);
+			AddHordeStorage(services);
+			services.Configure<StorageOptions>(configure);
 		}
 	}
 }
diff --git a/Engine/Source/Programs/Shared/EpicGames.Horde/StringId.cs b/Engine/Source/Programs/Shared/EpicGames.Horde/StringId.cs
--- a/Engine/Source/Programs/Shared/EpicGames.Horde/StringId.cs
+++ b/Engine/Source/Programs/Shared/EpicGames.Horde/StringId.cs
@@ -1,9 +1,6 @@
 // Copyright Epic Games, Inc. All Rights Reserved.
 
 using System;
-using System.Collections.Generic;
-using System.Diagnostics.CodeAnalysis;
-using System.Text;
 
 namespace EpicGames.Horde
 {
@@ -20,11 +17,10 @@ namespace EpicGames.Horde
 		/// <summary>
 		/// Constructor
 		/// </summary>
-		/// <param name="Input">Unique id for the string</param>
-		[SuppressMessage("Globalization", "CA1308:Normalize strings to uppercase", Justification = "Known limited character set")]
-		public StringId(string Input)
+		/// <param name="input">Unique id for the string</param>
+		public StringId(string input)
 		{
-			this.Text = Input;
+			Text = input;
 
 			if (Text.Length == 0)
 			{
@@ -37,12 +33,12 @@ namespace EpicGames.Horde
 				throw new ArgumentException($"String id may not be longer than {MaxLength} characters");
 			}
 
-			for (int Idx = 0; Idx < Text.Length; Idx++)
+			for (int idx = 0; idx < Text.Length; idx++)
 			{
-				char Character = Text[Idx];
-				if (!IsValidCharacter(Character))
+				char character = Text[idx];
+				if (!IsValidCharacter(character))
 				{
-					if (Character >= 'A' && Character <= 'Z')
+					if (character >= 'A' && character <= 'Z')
 					{
 						Text = Text.ToLowerInvariant();
 					}
@@ -57,27 +53,24 @@ namespace EpicGames.Horde
 		/// <summary>
 		/// Checks whether this StringId is set
 		/// </summary>
-		public bool IsEmpty
-		{
-			get { return String.IsNullOrEmpty(Text); }
-		}
+		public bool IsEmpty => String.IsNullOrEmpty(Text);
 
 		/// <summary>
 		/// Checks whether the given character is valid within a string id
 		/// </summary>
-		/// <param name="Character">The character to check</param>
+		/// <param name="character">The character to check</param>
 		/// <returns>True if the character is valid</returns>
-		static bool IsValidCharacter(char Character)
+		static bool IsValidCharacter(char character)
 		{
-			if (Character >= 'a' && Character <= 'z')
+			if (character >= 'a' && character <= 'z')
 			{
 				return true;
 			}
-			if (Character >= '0' && Character <= '9')
+			if (character >= '0' && character <= '9')
 			{
 				return true;
 			}
-			if (Character == '-' || Character == '_' || Character == '.')
+			if (character == '-' || character == '_' || character == '.')
 			{
 				return true;
 			}
@@ -85,13 +78,13 @@ namespace EpicGames.Horde
 		}
 
 		/// <inheritdoc/>
-		public override bool Equals(object? Obj) => Obj is StringId Id && Equals(Id);
+		public override bool Equals(object? obj) => obj is StringId id && Equals(id);
 
 		/// <inheritdoc/>
 		public override int GetHashCode() => Text.GetHashCode(StringComparison.Ordinal);
 
 		/// <inheritdoc/>
-		public bool Equals(StringId Other) => Text.Equals(Other.Text, StringComparison.Ordinal);
+		public bool Equals(StringId other) => Text.Equals(other.Text, StringComparison.Ordinal);
 
 		/// <inheritdoc/>
 		public override string ToString() => Text;
@@ -99,17 +92,17 @@ namespace EpicGames.Horde
 		/// <summary>
 		/// Compares two string ids for equality
 		/// </summary>
-		/// <param name="Left">The first string id</param>
-		/// <param name="Right">Second string id</param>
+		/// <param name="left">The first string id</param>
+		/// <param name="right">Second string id</param>
 		/// <returns>True if the two string ids are equal</returns>
-		public static bool operator ==(StringId Left, StringId Right) => Left.Equals(Right);
+		public static bool operator ==(StringId left, StringId right) => left.Equals(right);
 
 		/// <summary>
 		/// Compares two string ids for inequality
 		/// </summary>
-		/// <param name="Left">The first string id</param>
-		/// <param name="Right">Second string id</param>
+		/// <param name="left">The first string id</param>
+		/// <param name="right">Second string id</param>
 		/// <returns>True if the two string ids are not equal</returns>
-		public static bool operator !=(StringId Left, StringId Right) => !Left.Equals(Right);
+		public static bool operator !=(StringId left, StringId right) => !left.Equals(right);
 	}
 }
