Fix issue where FColor default values were not working in C# UHT

#rnx

diff --git a/Engine/Source/Programs/Shared/EpicGames.UHT/Types/Properties/UhtStructProperty.cs b/Engine/Source/Programs/Shared/EpicGames.UHT/Types/Properties/UhtStructProperty.cs
--- a/Engine/Source/Programs/Shared/EpicGames.UHT/Types/Properties/UhtStructProperty.cs
+++ b/Engine/Source/Programs/Shared/EpicGames.UHT/Types/Properties/UhtStructProperty.cs
@@ -367,7 +367,7 @@ namespace EpicGames.UHT.Types
 		{
 			string Format = "(R={0},G={1},B={2},A={3})";
 
-			DefaultValueReader.Require("FLinearColor");
+			DefaultValueReader.Require("FColor");
 			if (DefaultValueReader.TryOptional("::"))
 			{
 				switch (DefaultValueReader.GetIdentifier().Value.ToString())
