Added a sample UHT exporter to dump stats.
Add the ability to disable all default exporters.

#rnx

diff --git a/Engine/Source/Programs/Shared/EpicGames.UHT/Exporters/Stats/UhtStatsExporter.cs b/Engine/Source/Programs/Shared/EpicGames.UHT/Exporters/Stats/UhtStatsExporter.cs
new file mode 100644
--- /dev/null
+++ b/Engine/Source/Programs/Shared/EpicGames.UHT/Exporters/Stats/UhtStatsExporter.cs
@@ -0,0 +1,50 @@
+using EpicGames.Core;
+using EpicGames.UHT.Tables;
+using EpicGames.UHT.Types;
+using EpicGames.UHT.Utils;
+using System;
+using System.Collections.Generic;
+using System.Text;
+
+namespace EpicGames.UHT.Exporters.Stats
+{
+	[UnrealHeaderTool]
+	internal class UhtStatsExporter
+	{
+		[UhtExporter(Name = "Stats", Description = "Type Stats", Options = UhtExporterOptions.None)]
+		private static void StatsExporter(IUhtExportFactory Factory)
+		{
+			SortedDictionary<string, int> CountByType = new SortedDictionary<string, int>();
+			foreach (UhtType Type in Factory.Session.Packages)
+			{
+				Collect(CountByType, Type);
+			}
+
+			Log.TraceInformation("Counts by type:");
+
+			foreach (KeyValuePair<string, int> Kvp in CountByType)
+			{
+				Log.TraceInformation($"{Kvp.Key} {Kvp.Value}");
+			}
+			Log.TraceInformation("");
+		}
+
+		private static void Collect(SortedDictionary<string, int> CountByType, UhtType Type)
+		{
+			int Count = 1;
+			if (CountByType.TryGetValue(Type.EngineClassName, out Count))
+			{
+				CountByType[Type.EngineClassName] = Count + 1;
+			}
+			else
+			{
+				CountByType[Type.EngineClassName] = 1;
+			}
+
+			foreach (UhtType Child in Type.Children)
+			{
+				Collect(CountByType, Child);
+			}
+		}
+	}
+}
diff --git a/Engine/Source/Programs/Shared/EpicGames.UHT/Utils/UhtSession.cs b/Engine/Source/Programs/Shared/EpicGames.UHT/Utils/UhtSession.cs
--- a/Engine/Source/Programs/Shared/EpicGames.UHT/Utils/UhtSession.cs
+++ b/Engine/Source/Programs/Shared/EpicGames.UHT/Utils/UhtSession.cs
@@ -847,6 +847,11 @@ namespace EpicGames.UHT.Utils
 		/// </summary>
 		public bool bIncludeDebugOutput = false;
 
+		/// <summary>
+		/// If true, disable all exporters which would normally be run by default
+		/// </summary>
+		public bool bNoDefaultExporters = false;
+
 		/// <summary>
 		/// If true, cache any error messages until the end of processing.  This is used by the testing
 		/// harness to generate more stable console output.
@@ -2369,7 +2374,8 @@ namespace EpicGames.UHT.Utils
 					bool Run = false;
 					if (!this.ExporterStates.TryGetValue(Exporter.Name, out Run))
 					{
-						Run = UhtConfig.Instance.IsExporterEnabled(Exporter.Name) || Exporter.Options.HasAnyFlags(UhtExporterOptions.Default);
+						Run = UhtConfig.Instance.IsExporterEnabled(Exporter.Name) || 
+							(Exporter.Options.HasAnyFlags(UhtExporterOptions.Default) && !this.bNoDefaultExporters);
 					}
 
 					if (Run)
diff --git a/Engine/Source/Programs/UnrealBuildTool/Modes/UnrealHeaderToolMode.cs b/Engine/Source/Programs/UnrealBuildTool/Modes/UnrealHeaderToolMode.cs
--- a/Engine/Source/Programs/UnrealBuildTool/Modes/UnrealHeaderToolMode.cs
+++ b/Engine/Source/Programs/UnrealBuildTool/Modes/UnrealHeaderToolMode.cs
@@ -400,6 +400,9 @@ namespace UnrealBuildTool.Modes
 		[CommandLine("-IncludeDebugOutput", Description = "Include extra content in generated output to assist with debugging")]
 		public bool bIncludeDebugOutput = false;
 
+		[CommandLine("-NoDefaultExporters", Description = "Disable all default exporters.  Useful for when a specific exporter is to be run")]
+		public bool bNoDefaultExporters = false;
+
 		/// <summary>
 		/// Initialize the options with the given command line arguments
 		/// </summary>
@@ -1102,6 +1105,7 @@ namespace UnrealBuildTool.Modes
 					bFailIfGeneratedCodeChanges = Options.bFailIfGeneratedCodeChanges,
 					bNoOutput = Options.bNoOutput,
 					bIncludeDebugOutput = Options.bIncludeDebugOutput,
+					bNoDefaultExporters = Options.bNoDefaultExporters,
 				};
 
 				if (Options.bWriteRef)
