Fixing missing copyright

#rnx

diff --git a/Engine/Source/Programs/Shared/EpicGames.UHT/Exporters/Stats/UhtStatsExporter.cs b/Engine/Source/Programs/Shared/EpicGames.UHT/Exporters/Stats/UhtStatsExporter.cs
--- a/Engine/Source/Programs/Shared/EpicGames.UHT/Exporters/Stats/UhtStatsExporter.cs
+++ b/Engine/Source/Programs/Shared/EpicGames.UHT/Exporters/Stats/UhtStatsExporter.cs
@@ -1,3 +1,5 @@
+// Copyright Epic Games, Inc. All Rights Reserved.
+
 using EpicGames.Core;
 using EpicGames.UHT.Tables;
 using EpicGames.UHT.Types;
