Reformat EpicGames.IoHash according to coding conventions.

diff --git a/Engine/Source/Programs/Shared/EpicGames.IoHash/Blake3Hash.cs b/Engine/Source/Programs/Shared/EpicGames.IoHash/Blake3Hash.cs
--- a/Engine/Source/Programs/Shared/EpicGames.IoHash/Blake3Hash.cs
+++ b/Engine/Source/Programs/Shared/EpicGames.IoHash/Blake3Hash.cs
@@ -1,12 +1,7 @@
 // Copyright Epic Games, Inc. All Rights Reserved.
 
 using System;
-using System.Collections.Generic;
-using System.Buffers;
-using System.Security.Cryptography;
-using System.Text;
 using System.Buffers.Binary;
-using System.Diagnostics.CodeAnalysis;
 
 namespace EpicGames.Core
 {
@@ -28,7 +23,7 @@ namespace EpicGames.Core
 		/// <summary>
 		/// Memory storing the digest data
 		/// </summary>
-		public ReadOnlyMemory<byte> Memory;
+		public ReadOnlyMemory<byte> Memory { get; }
 
 		/// <summary>
 		/// Span for the underlying memory
@@ -43,61 +38,61 @@ namespace EpicGames.Core
 		/// <summary>
 		/// Constructor
 		/// </summary>
-		/// <param name="Memory">Memory to construct from</param>
-		public Blake3Hash(ReadOnlyMemory<byte> Memory)
+		/// <param name="memory">Memory to construct from</param>
+		public Blake3Hash(ReadOnlyMemory<byte> memory)
 		{
-			if (Memory.Length != NumBytes)
+			if (memory.Length != NumBytes)
 			{
 				throw new ArgumentException($"Blake3Hash must be {NumBytes} bytes long");
 			}
 
-			this.Memory = Memory;
+			Memory = memory;
 		}
 
 		/// <summary>
 		/// Creates a content hash for a block of data, using a given algorithm.
 		/// </summary>
-		/// <param name="Data">Data to compute the hash for</param>
+		/// <param name="data">Data to compute the hash for</param>
 		/// <returns>New content hash instance containing the hash of the data</returns>
-		public static Blake3Hash Compute(ReadOnlySpan<byte> Data)
+		public static Blake3Hash Compute(ReadOnlySpan<byte> data)
 		{
-			byte[] Output = new byte[32];
-			Blake3.Hasher.Hash(Data, Output);
-			return new Blake3Hash(Output);
+			byte[] output = new byte[32];
+			Blake3.Hasher.Hash(data, output);
+			return new Blake3Hash(output);
 		}
 
 		/// <summary>
 		/// Parses a digest from the given hex string
 		/// </summary>
-		/// <param name="Text"></param>
+		/// <param name="text"></param>
 		/// <returns></returns>
-		public static Blake3Hash Parse(string Text)
+		public static Blake3Hash Parse(string text)
 		{
-			return new Blake3Hash(StringUtils.ParseHexString(Text));
+			return new Blake3Hash(StringUtils.ParseHexString(text));
 		}
 
 		/// <inheritdoc cref="IComparable{T}.CompareTo(T)"/>
-		public int CompareTo(Blake3Hash Other)
+		public int CompareTo(Blake3Hash other)
 		{
-			ReadOnlySpan<byte> A = Span;
-			ReadOnlySpan<byte> B = Other.Span;
+			ReadOnlySpan<byte> a = Span;
+			ReadOnlySpan<byte> b = other.Span;
 
-			for (int Idx = 0; Idx < A.Length && Idx < B.Length; Idx++)
+			for (int idx = 0; idx < a.Length && idx < b.Length; idx++)
 			{
-				int Compare = A[Idx] - B[Idx];
-				if (Compare != 0)
+				int compare = a[idx] - b[idx];
+				if (compare != 0)
 				{
-					return Compare;
+					return compare;
 				}
 			}
-			return A.Length - B.Length;
+			return a.Length - b.Length;
 		}
 
 		/// <inheritdoc/>
-		public bool Equals(Blake3Hash Other) => Span.SequenceEqual(Other.Span);
+		public bool Equals(Blake3Hash other) => Span.SequenceEqual(other.Span);
 
 		/// <inheritdoc/>
-		public override bool Equals(object? Obj) => (Obj is Blake3Hash Hash) && Hash.Span.SequenceEqual(Span);
+		public override bool Equals(object? obj) => (obj is Blake3Hash hash) && hash.Span.SequenceEqual(Span);
 
 		/// <inheritdoc/>
 		public override int GetHashCode() => BinaryPrimitives.ReadInt32LittleEndian(Span);
@@ -108,32 +103,32 @@ namespace EpicGames.Core
 		/// <summary>
 		/// Test two hash values for equality
 		/// </summary>
-		public static bool operator ==(Blake3Hash A, Blake3Hash B) => A.Span.SequenceEqual(B.Span);
+		public static bool operator ==(Blake3Hash a, Blake3Hash b) => a.Span.SequenceEqual(b.Span);
 
 		/// <summary>
 		/// Test two hash values for equality
 		/// </summary>
-		public static bool operator !=(Blake3Hash A, Blake3Hash B) => !(A == B);
+		public static bool operator !=(Blake3Hash a, Blake3Hash b) => !(a == b);
 
 		/// <summary>
 		/// Tests whether A > B
 		/// </summary>
-		public static bool operator >(Blake3Hash A, Blake3Hash B) => A.CompareTo(B) > 0;
+		public static bool operator >(Blake3Hash a, Blake3Hash b) => a.CompareTo(b) > 0;
 
 		/// <summary>
 		/// Tests whether A is less than B
 		/// </summary>
-		public static bool operator <(Blake3Hash A, Blake3Hash B) => A.CompareTo(B) < 0;
+		public static bool operator <(Blake3Hash a, Blake3Hash b) => a.CompareTo(b) < 0;
 
 		/// <summary>
 		/// Tests whether A is greater than or equal to B
 		/// </summary>
-		public static bool operator >=(Blake3Hash A, Blake3Hash B) => A.CompareTo(B) >= 0;
+		public static bool operator >=(Blake3Hash a, Blake3Hash b) => a.CompareTo(b) >= 0;
 
 		/// <summary>
 		/// Tests whether A is less than or equal to B
 		/// </summary>
-		public static bool operator <=(Blake3Hash A, Blake3Hash B) => A.CompareTo(B) <= 0;
+		public static bool operator <=(Blake3Hash a, Blake3Hash b) => a.CompareTo(b) <= 0;
 	}
 
 	/// <summary>
@@ -144,21 +139,21 @@ namespace EpicGames.Core
 		/// <summary>
 		/// Read an <see cref="Blake3Hash"/> from a memory reader
 		/// </summary>
-		/// <param name="Reader"></param>
+		/// <param name="reader"></param>
 		/// <returns></returns>
-		public static Blake3Hash ReadBlake3Hash(this MemoryReader Reader)
+		public static Blake3Hash ReadBlake3Hash(this MemoryReader reader)
 		{
-			return new Blake3Hash(Reader.ReadFixedLengthBytes(Blake3Hash.NumBytes));
+			return new Blake3Hash(reader.ReadFixedLengthBytes(Blake3Hash.NumBytes));
 		}
 
 		/// <summary>
 		/// Write an <see cref="Blake3Hash"/> to a memory writer
 		/// </summary>
-		/// <param name="Writer"></param>
-		/// <param name="Hash"></param>
-		public static void WriteBlake3Hash(this MemoryWriter Writer, Blake3Hash Hash)
+		/// <param name="writer"></param>
+		/// <param name="hash"></param>
+		public static void WriteBlake3Hash(this MemoryWriter writer, Blake3Hash hash)
 		{
-			Writer.WriteFixedLengthBytes(Hash.Span);
+			writer.WriteFixedLengthBytes(hash.Span);
 		}
 	}
 }
diff --git a/Engine/Source/Programs/Shared/EpicGames.IoHash/IoHash.cs b/Engine/Source/Programs/Shared/EpicGames.IoHash/IoHash.cs
--- a/Engine/Source/Programs/Shared/EpicGames.IoHash/IoHash.cs
+++ b/Engine/Source/Programs/Shared/EpicGames.IoHash/IoHash.cs
@@ -1,12 +1,8 @@
 // Copyright Epic Games, Inc. All Rights Reserved.
 
 using System;
-using System.Collections.Generic;
 using System.Buffers;
-using System.Security.Cryptography;
-using System.Text;
 using System.Buffers.Binary;
-using System.Diagnostics.CodeAnalysis;
 using System.ComponentModel;
 using System.Globalization;
 using System.Text.Json.Serialization;
@@ -31,9 +27,9 @@ namespace EpicGames.Core
 		/// </summary>
 		public const int NumBits = NumBytes * 8;
 
-		ulong A;
-		ulong B;
-		uint C;
+		readonly ulong _a;
+		readonly ulong _b;
+		readonly uint _c;
 
 		/// <summary>
 		/// Hash consisting of zeroes
@@ -43,105 +39,104 @@ namespace EpicGames.Core
 		/// <summary>
 		/// Constructor
 		/// </summary>
-		/// <param name="Memory">Memory to construct from</param>
-		public IoHash(ReadOnlySpan<byte> Span)
-			: this(BinaryPrimitives.ReadUInt64BigEndian(Span), BinaryPrimitives.ReadUInt64BigEndian(Span.Slice(8)), BinaryPrimitives.ReadUInt32BigEndian(Span.Slice(16)))
+		/// <param name="span">Memory to construct from</param>
+		public IoHash(ReadOnlySpan<byte> span)
+			: this(BinaryPrimitives.ReadUInt64BigEndian(span), BinaryPrimitives.ReadUInt64BigEndian(span.Slice(8)), BinaryPrimitives.ReadUInt32BigEndian(span.Slice(16)))
 		{
 		}
 
 		/// <summary>
 		/// Constructor
 		/// </summary>
-		/// <param name="Memory">Memory to construct from</param>
-		public IoHash(ulong A, ulong B, uint C)
+		public IoHash(ulong a, ulong b, uint c)
 		{
-			this.A = A;
-			this.B = B;
-			this.C = C;
+			_a = a;
+			_b = b;
+			_c = c;
 		}
 
 		/// <summary>
 		/// Construct 
 		/// </summary>
-		/// <param name="Hasher">The hasher to construct from</param>
-		public static IoHash FromBlake3(Blake3.Hasher Hasher)
+		/// <param name="hasher">The hasher to construct from</param>
+		public static IoHash FromBlake3(Blake3.Hasher hasher)
 		{
-			Span<byte> Output = stackalloc byte[32];
-			Hasher.Finalize(Output);
-			return new IoHash(Output);
+			Span<byte> output = stackalloc byte[32];
+			hasher.Finalize(output);
+			return new IoHash(output);
 		}
 
 		/// <summary>
 		/// Creates the IoHash for a block of data.
 		/// </summary>
-		/// <param name="Data">Data to compute the hash for</param>
+		/// <param name="data">Data to compute the hash for</param>
 		/// <returns>New hash instance containing the hash of the data</returns>
-		public static IoHash Compute(ReadOnlySpan<byte> Data)
+		public static IoHash Compute(ReadOnlySpan<byte> data)
 		{
-			Span<byte> Output = stackalloc byte[32];
-			Blake3.Hasher.Hash(Data, Output);
-			return new IoHash(Output);
+			Span<byte> output = stackalloc byte[32];
+			Blake3.Hasher.Hash(data, output);
+			return new IoHash(output);
 		}
 
 		/// <summary>
 		/// Creates the IoHash for a block of data.
 		/// </summary>
-		/// <param name="Sequence">Data to compute the hash for</param>
+		/// <param name="sequence">Data to compute the hash for</param>
 		/// <returns>New hash instance containing the hash of the data</returns>
-		public static IoHash Compute(ReadOnlySequence<byte> Sequence)
+		public static IoHash Compute(ReadOnlySequence<byte> sequence)
 		{
-			if (Sequence.IsSingleSegment)
+			if (sequence.IsSingleSegment)
 			{
-				return Compute(Sequence.FirstSpan);
+				return Compute(sequence.FirstSpan);
 			}
 
-			using (Blake3.Hasher Hasher = Blake3.Hasher.New())
+			using (Blake3.Hasher hasher = Blake3.Hasher.New())
 			{
-				foreach (ReadOnlyMemory<byte> Segment in Sequence)
+				foreach (ReadOnlyMemory<byte> segment in sequence)
 				{
-					Hasher.Update(Segment.Span);
+					hasher.Update(segment.Span);
 				}
-				return FromBlake3(Hasher);
+				return FromBlake3(hasher);
 			}
 		}
 
 		/// <summary>
 		/// Parses a digest from the given hex string
 		/// </summary>
-		/// <param name="Text"></param>
+		/// <param name="text"></param>
 		/// <returns></returns>
-		public static IoHash Parse(string Text)
+		public static IoHash Parse(string text)
 		{
-			return new IoHash(StringUtils.ParseHexString(Text));
+			return new IoHash(StringUtils.ParseHexString(text));
 		}
 
 		/// <summary>
 		/// Parses a digest from the given hex string
 		/// </summary>
-		/// <param name="Text"></param>
+		/// <param name="text"></param>
 		/// <returns></returns>
-		public static IoHash Parse(ReadOnlySpan<byte> Text)
+		public static IoHash Parse(ReadOnlySpan<byte> text)
 		{
-			return new IoHash(StringUtils.ParseHexString(Text));
+			return new IoHash(StringUtils.ParseHexString(text));
 		}
 
 		/// <summary>
 		/// Parses a digest from the given hex string
 		/// </summary>
-		/// <param name="Text"></param>
-		/// <param name="Hash">Receives the hash on success</param>
+		/// <param name="text"></param>
+		/// <param name="hash">Receives the hash on success</param>
 		/// <returns></returns>
-		public static bool TryParse(string Text, out IoHash Hash)
+		public static bool TryParse(string text, out IoHash hash)
 		{
-			byte[]? Bytes;
-			if (StringUtils.TryParseHexString(Text, out Bytes) && Bytes.Length == IoHash.NumBytes)
+			byte[]? bytes;
+			if (StringUtils.TryParseHexString(text, out bytes) && bytes.Length == IoHash.NumBytes)
 			{
-				Hash = new IoHash(Bytes);
+				hash = new IoHash(bytes);
 				return true;
 			}
 			else
 			{
-				Hash = default;
+				hash = default;
 				return false;
 			}
 		}
@@ -149,49 +144,49 @@ namespace EpicGames.Core
 		/// <summary>
 		/// Parses a digest from the given hex string
 		/// </summary>
-		/// <param name="Text"></param>
-		/// <param name="Hash">Receives the hash on success</param>
+		/// <param name="text"></param>
+		/// <param name="hash">Receives the hash on success</param>
 		/// <returns></returns>
-		public static bool TryParse(ReadOnlySpan<byte> Text, out IoHash Hash)
+		public static bool TryParse(ReadOnlySpan<byte> text, out IoHash hash)
 		{
-			byte[]? Bytes;
-			if (StringUtils.TryParseHexString(Text, out Bytes) && Bytes.Length == IoHash.NumBytes)
+			byte[]? bytes;
+			if (StringUtils.TryParseHexString(text, out bytes) && bytes.Length == IoHash.NumBytes)
 			{
-				Hash = new IoHash(Bytes);
+				hash = new IoHash(bytes);
 				return true;
 			}
 			else
 			{
-				Hash = default;
+				hash = default;
 				return false;
 			}
 		}
 
 		/// <inheritdoc cref="IComparable{T}.CompareTo(T)"/>
-		public int CompareTo(IoHash Other)
+		public int CompareTo(IoHash other)
 		{
-			if (A != Other.A)
+			if (_a != other._a)
 			{
-				return (A < Other.A) ? -1 : +1;
+				return (_a < other._a) ? -1 : +1;
 			}
-			else if (B != Other.B)
+			else if (_b != other._b)
 			{
-				return (B < Other.B) ? -1 : +1;
+				return (_b < other._b) ? -1 : +1;
 			}
 			else
 			{
-				return (C < Other.C) ? -1 : +1;
+				return (_c < other._c) ? -1 : +1;
 			}
 		}
 
 		/// <inheritdoc/>
-		public bool Equals(IoHash Other) => A == Other.A && B == Other.B && C == Other.C;
+		public bool Equals(IoHash other) => _a == other._a && _b == other._b && _c == other._c;
 
 		/// <inheritdoc/>
-		public override bool Equals(object? Obj) => (Obj is IoHash Hash) && Equals(Hash);
+		public override bool Equals(object? obj) => (obj is IoHash hash) && Equals(hash);
 
 		/// <inheritdoc/>
-		public override int GetHashCode() => (int)A;
+		public override int GetHashCode() => (int)_a;
 
 		/// <inheritdoc/>
 		public Utf8String ToUtf8String() => StringUtils.FormatUtf8HexString(ToByteArray());
@@ -201,59 +196,59 @@ namespace EpicGames.Core
 
 		public byte[] ToByteArray()
 		{
-			byte[] Data = new byte[NumBytes];
-			CopyTo(Data);
-			return Data;
+			byte[] data = new byte[NumBytes];
+			CopyTo(data);
+			return data;
 		}
 
 		/// <summary>
 		/// Copies this hash into a span
 		/// </summary>
-		/// <param name="Span"></param>
-		public void CopyTo(Span<byte> Span)
+		/// <param name="span"></param>
+		public void CopyTo(Span<byte> span)
 		{
-			BinaryPrimitives.WriteUInt64BigEndian(Span, A);
-			BinaryPrimitives.WriteUInt64BigEndian(Span[8..], B);
-			BinaryPrimitives.WriteUInt32BigEndian(Span[16..], C);
+			BinaryPrimitives.WriteUInt64BigEndian(span, _a);
+			BinaryPrimitives.WriteUInt64BigEndian(span[8..], _b);
+			BinaryPrimitives.WriteUInt32BigEndian(span[16..], _c);
 		}
 
 		/// <summary>
 		/// Test two hash values for equality
 		/// </summary>
-		public static bool operator ==(IoHash A, IoHash B) => A.Equals(B);
+		public static bool operator ==(IoHash a, IoHash b) => a.Equals(b);
 
 		/// <summary>
 		/// Test two hash values for equality
 		/// </summary>
-		public static bool operator !=(IoHash A, IoHash B) => !(A == B);
+		public static bool operator !=(IoHash a, IoHash b) => !(a == b);
 
 		/// <summary>
 		/// Tests whether A > B
 		/// </summary>
-		public static bool operator >(IoHash A, IoHash B) => A.CompareTo(B) > 0;
+		public static bool operator >(IoHash a, IoHash b) => a.CompareTo(b) > 0;
 
 		/// <summary>
 		/// Tests whether A is less than B
 		/// </summary>
-		public static bool operator <(IoHash A, IoHash B) => A.CompareTo(B) < 0;
+		public static bool operator <(IoHash a, IoHash b) => a.CompareTo(b) < 0;
 
 		/// <summary>
 		/// Tests whether A is greater than or equal to B
 		/// </summary>
-		public static bool operator >=(IoHash A, IoHash B) => A.CompareTo(B) >= 0;
+		public static bool operator >=(IoHash a, IoHash b) => a.CompareTo(b) >= 0;
 
 		/// <summary>
 		/// Tests whether A is less than or equal to B
 		/// </summary>
-		public static bool operator <=(IoHash A, IoHash B) => A.CompareTo(B) <= 0;
+		public static bool operator <=(IoHash a, IoHash b) => a.CompareTo(b) <= 0;
 
 		/// <summary>
 		/// Convert a Blake3Hash to an IoHash
 		/// </summary>
-		/// <param name="Hash"></param>
-		public static implicit operator IoHash(Blake3Hash Hash)
+		/// <param name="hash"></param>
+		public static implicit operator IoHash(Blake3Hash hash)
 		{
-			return new IoHash(Hash.Span.Slice(0, NumBytes));
+			return new IoHash(hash.Span.Slice(0, NumBytes));
 		}
 	}
 
@@ -265,21 +260,21 @@ namespace EpicGames.Core
 		/// <summary>
 		/// Read an <see cref="IoHash"/> from a memory reader
 		/// </summary>
-		/// <param name="Reader"></param>
+		/// <param name="reader"></param>
 		/// <returns></returns>
-		public static IoHash ReadIoHash(this MemoryReader Reader)
+		public static IoHash ReadIoHash(this MemoryReader reader)
 		{
-			return new IoHash(Reader.ReadFixedLengthBytes(IoHash.NumBytes).Span);
+			return new IoHash(reader.ReadFixedLengthBytes(IoHash.NumBytes).Span);
 		}
 
 		/// <summary>
 		/// Write an <see cref="IoHash"/> to a memory writer
 		/// </summary>
-		/// <param name="Writer"></param>
-		/// <param name="Hash"></param>
-		public static void WriteIoHash(this MemoryWriter Writer, IoHash Hash)
+		/// <param name="writer"></param>
+		/// <param name="hash"></param>
+		public static void WriteIoHash(this MemoryWriter writer, IoHash hash)
 		{
-			Hash.CopyTo(Writer.AllocateSpan(IoHash.NumBytes));
+			hash.CopyTo(writer.AllocateSpan(IoHash.NumBytes));
 		}
 	}
 
@@ -289,10 +284,10 @@ namespace EpicGames.Core
 	sealed class IoHashJsonConverter : JsonConverter<IoHash>
 	{
 		/// <inheritdoc/>
-		public override IoHash Read(ref Utf8JsonReader Reader, Type TypeToConvert, JsonSerializerOptions Options) => IoHash.Parse(Reader.ValueSpan);
+		public override IoHash Read(ref Utf8JsonReader reader, Type typeToConvert, JsonSerializerOptions options) => IoHash.Parse(reader.ValueSpan);
 
 		/// <inheritdoc/>
-		public override void Write(Utf8JsonWriter Writer, IoHash Value, JsonSerializerOptions Options) => Writer.WriteStringValue(Value.ToUtf8String().Span);
+		public override void Write(Utf8JsonWriter writer, IoHash value, JsonSerializerOptions options) => writer.WriteStringValue(value.ToUtf8String().Span);
 	}
 
 	/// <summary>
@@ -301,15 +296,15 @@ namespace EpicGames.Core
 	sealed class IoHashTypeConverter : TypeConverter
 	{
 		/// <inheritdoc/>
-		public override bool CanConvertFrom(ITypeDescriptorContext Context, Type SourceType)
+		public override bool CanConvertFrom(ITypeDescriptorContext context, Type sourceType)
 		{
-			return SourceType == typeof(string);
+			return sourceType == typeof(string);
 		}
 
 		/// <inheritdoc/>
-		public override object ConvertFrom(ITypeDescriptorContext Context, CultureInfo Culture, object Value)
+		public override object ConvertFrom(ITypeDescriptorContext context, CultureInfo culture, object value)
 		{
-			return IoHash.Parse((string)Value);
+			return IoHash.Parse((string)value);
 		}
 	}
 }
