Change EpicGames.Serialization library to match standard coding style.

diff --git a/Engine/Source/Programs/Shared/EpicGames.Serialization.Tests/EngineTests.cs b/Engine/Source/Programs/Shared/EpicGames.Serialization.Tests/EngineTests.cs
--- a/Engine/Source/Programs/Shared/EpicGames.Serialization.Tests/EngineTests.cs
+++ b/Engine/Source/Programs/Shared/EpicGames.Serialization.Tests/EngineTests.cs
@@ -3,9 +3,7 @@
 using EpicGames.Core;
 using Microsoft.VisualStudio.TestTools.UnitTesting;
 using System;
-using System.Buffers.Binary;
 using System.Collections.Generic;
-using System.Text;
 
 namespace EpicGames.Serialization.Tests
 {
@@ -14,97 +12,97 @@ namespace EpicGames.Serialization.Tests
 	{
 		class CbFieldAccessors
 		{
-			public object Default;
-			public Func<CbField, bool> IsType;
-			public Func<CbField, object> AsType;
-			public Func<CbField, object, object> AsTypeWithDefault;
-			public Func<object, object, bool> Comparer = (x, y) => x.Equals(y);
+			public object _defaultValue;
+			public Func<CbField, bool> _isType;
+			public Func<CbField, object> _asType;
+			public Func<CbField, object, object> _asTypeWithDefault;
+			public Func<object, object, bool> _comparer = (x, y) => x.Equals(y);
 
-			public CbFieldAccessors(object Default, Func<CbField, bool> IsType, Func<CbField, object> AsType, Func<CbField, object, object> AsTypeWithDefault)
+			public CbFieldAccessors(object defaultValue, Func<CbField, bool> isType, Func<CbField, object> asType, Func<CbField, object, object> asTypeWithDefault)
 			{
-				this.Default = Default;
-				this.IsType = IsType;
-				this.AsType = AsType;
-				this.AsTypeWithDefault = AsTypeWithDefault;
+				_defaultValue = defaultValue;
+				_isType = isType;
+				_asType = asType;
+				_asTypeWithDefault = asTypeWithDefault;
 			}
 
-			public static CbFieldAccessors FromStruct<T>(Func<CbField, bool> IsType, Func<CbField, T, object> AsTypeWithDefault) where T : struct
+			public static CbFieldAccessors FromStruct<T>(Func<CbField, bool> isType, Func<CbField, T, object> asTypeWithDefault) where T : struct
 			{
-				return new CbFieldAccessors(new T(), IsType, x => AsTypeWithDefault(x, new T()), (x, y) => AsTypeWithDefault(x, (T)y));
+				return new CbFieldAccessors(new T(), isType, x => asTypeWithDefault(x, new T()), (x, y) => asTypeWithDefault(x, (T)y));
 			}
 
-			public static CbFieldAccessors FromStruct<T>(T Default, Func<CbField, bool> IsType, Func<CbField, T, object> AsTypeWithDefault) where T : struct
+			public static CbFieldAccessors FromStruct<T>(T @default, Func<CbField, bool> isType, Func<CbField, T, object> asTypeWithDefault) where T : struct
 			{
-				return new CbFieldAccessors(Default, IsType, x => AsTypeWithDefault(x, new T()), (x, y) => AsTypeWithDefault(x, (T)y));
+				return new CbFieldAccessors(@default, isType, x => asTypeWithDefault(x, new T()), (x, y) => asTypeWithDefault(x, (T)y));
 			}
 
-			public static CbFieldAccessors FromStruct<T>(Func<CbField, bool> IsType, Func<CbField, T, object> AsTypeWithDefault, Func<T, T, bool> Comparer) where T : struct
+			public static CbFieldAccessors FromStruct<T>(Func<CbField, bool> isType, Func<CbField, T, object> asTypeWithDefault, Func<T, T, bool> comparer) where T : struct
 			{
-				return new CbFieldAccessors(new T(), IsType, x => AsTypeWithDefault(x, new T()), (x, y) => AsTypeWithDefault(x, (T)y)) { Comparer = (x, y) => Comparer((T)x, (T)y) };
+				return new CbFieldAccessors(new T(), isType, x => asTypeWithDefault(x, new T()), (x, y) => asTypeWithDefault(x, (T)y)) { _comparer = (x, y) => comparer((T)x, (T)y) };
 			}
 		}
 
-		Dictionary<CbFieldType, CbFieldAccessors> TypeAccessors = new Dictionary<CbFieldType, CbFieldAccessors>
+		readonly Dictionary<CbFieldType, CbFieldAccessors> _typeAccessors = new Dictionary<CbFieldType, CbFieldAccessors>
 		{
 			[CbFieldType.Object] = new CbFieldAccessors(CbObject.Empty, x => x.IsObject(), x => x.AsObject(), (x, y) => x.AsObject()),
 			[CbFieldType.UniformObject] = new CbFieldAccessors(CbObject.Empty, x => x.IsObject(), x => x.AsObject(), (x, y) => x.AsObject()),
 			[CbFieldType.Array] = new CbFieldAccessors(CbArray.Empty, x => x.IsArray(), x => x.AsArray(), (x, y) => x.AsArray()),
 			[CbFieldType.UniformArray] = new CbFieldAccessors(CbArray.Empty, x => x.IsArray(), x => x.AsArray(), (x, y) => x.AsArray()),
 			[CbFieldType.Binary] = CbFieldAccessors.FromStruct<ReadOnlyMemory<byte>>(x => x.IsBinary(), (x, y) => x.AsBinary(y), (x, y) => x.Span.SequenceEqual(y.Span)),
-			[CbFieldType.String] = new CbFieldAccessors(Utf8String.Empty, x => x.IsString(), x => x.AsUtf8String(), (x, Default) => x.AsUtf8String((Utf8String)Default)),
+			[CbFieldType.String] = new CbFieldAccessors(Utf8String.Empty, x => x.IsString(), x => x.AsUtf8String(), (x, @default) => x.AsUtf8String((Utf8String)@default)),
 			[CbFieldType.IntegerPositive] = CbFieldAccessors.FromStruct<ulong>(x => x.IsInteger(), (x, y) => x.AsUInt64(y)),
 			[CbFieldType.IntegerNegative] = CbFieldAccessors.FromStruct<long>(x => x.IsInteger(), (x, y) => x.AsInt64(y)),
 			[CbFieldType.Float32] = CbFieldAccessors.FromStruct<float>(x => x.IsFloat(), (x, y) => x.AsFloat(y)),
 			[CbFieldType.Float64] = CbFieldAccessors.FromStruct<double>(x => x.IsFloat(), (x, y) => x.AsDouble(y)),
 			[CbFieldType.BoolTrue] = CbFieldAccessors.FromStruct<bool>(x => x.IsBool(), (x, y) => x.AsBool(y)),
 			[CbFieldType.BoolFalse] = CbFieldAccessors.FromStruct<bool>(x => x.IsBool(), (x, y) => x.AsBool(y)),
-			[CbFieldType.ObjectAttachment] = new CbFieldAccessors(new CbObjectAttachment(IoHash.Zero), x => x.IsObjectAttachment(), x => x.AsObjectAttachment(), (x, Default) => x.AsObjectAttachment((CbObjectAttachment)Default)),
-			[CbFieldType.BinaryAttachment] = new CbFieldAccessors(new CbBinaryAttachment(IoHash.Zero), x => x.IsBinaryAttachment(), x => x.AsBinaryAttachment(), (x, Default) => x.AsBinaryAttachment((CbBinaryAttachment)Default)),
-			[CbFieldType.Hash] = new CbFieldAccessors(IoHash.Zero, x => x.IsHash(), x => x.AsHash(), (x, Default) => x.AsHash((IoHash)Default)),
+			[CbFieldType.ObjectAttachment] = new CbFieldAccessors(new CbObjectAttachment(IoHash.Zero), x => x.IsObjectAttachment(), x => x.AsObjectAttachment(), (x, @default) => x.AsObjectAttachment((CbObjectAttachment)@default)),
+			[CbFieldType.BinaryAttachment] = new CbFieldAccessors(new CbBinaryAttachment(IoHash.Zero), x => x.IsBinaryAttachment(), x => x.AsBinaryAttachment(), (x, @default) => x.AsBinaryAttachment((CbBinaryAttachment)@default)),
+			[CbFieldType.Hash] = new CbFieldAccessors(IoHash.Zero, x => x.IsHash(), x => x.AsHash(), (x, @default) => x.AsHash((IoHash)@default)),
 			[CbFieldType.Uuid] = CbFieldAccessors.FromStruct<Guid>(x => x.IsUuid(), (x, y) => x.AsUuid(y)),
 			[CbFieldType.DateTime] = CbFieldAccessors.FromStruct<DateTime>(new DateTime(0, DateTimeKind.Utc), x => x.IsDateTime(), (x, y) => x.AsDateTime(y)),
 			[CbFieldType.TimeSpan] = CbFieldAccessors.FromStruct<TimeSpan>(x => x.IsTimeSpan(), (x, y) => x.AsTimeSpan(y)),
 		};
 
-		void TestField(CbFieldType FieldType, CbField Field, object? ExpectedValue = null, object? DefaultValue = null, CbFieldError ExpectedError = CbFieldError.None, CbFieldAccessors? Accessors = null)
+		void TestField(CbFieldType fieldType, CbField field, object? expectedValue = null, object? defaultValue = null, CbFieldError expectedError = CbFieldError.None, CbFieldAccessors? accessors = null)
 		{
-			Accessors ??= TypeAccessors[FieldType];
-			ExpectedValue ??= Accessors.Default;
-			DefaultValue ??= Accessors.Default;
+			accessors ??= _typeAccessors[fieldType];
+			expectedValue ??= accessors._defaultValue;
+			defaultValue ??= accessors._defaultValue;
 
-			Assert.AreEqual(Accessors.IsType(Field), ExpectedError != CbFieldError.TypeError);
-			if (ExpectedError == CbFieldError.None && !Field.IsBool())
+			Assert.AreEqual(accessors._isType(field), expectedError != CbFieldError.TypeError);
+			if (expectedError == CbFieldError.None && !field.IsBool())
 			{
-				Assert.IsFalse(Field.AsBool());
-				Assert.IsTrue(Field.HasError());
-				Assert.AreEqual(Field.GetError(), CbFieldError.TypeError);
+				Assert.IsFalse(field.AsBool());
+				Assert.IsTrue(field.HasError());
+				Assert.AreEqual(field.GetError(), CbFieldError.TypeError);
 			}
 
-			object Value = Accessors.AsTypeWithDefault(Field, DefaultValue);
-			Assert.IsTrue(Accessors.Comparer(Accessors.AsTypeWithDefault(Field, DefaultValue), ExpectedValue));
-			Assert.AreEqual(Field.HasError(), ExpectedError != CbFieldError.None);
-			Assert.AreEqual(Field.GetError(), ExpectedError);
+			object value = accessors._asTypeWithDefault(field, defaultValue);
+			Assert.IsTrue(accessors._comparer(accessors._asTypeWithDefault(field, defaultValue), expectedValue));
+			Assert.AreEqual(field.HasError(), expectedError != CbFieldError.None);
+			Assert.AreEqual(field.GetError(), expectedError);
 		}
 
-		void TestField(CbFieldType FieldType, byte[] Payload, object? ExpectedValue = null, object? DefaultValue = null, CbFieldError ExpectedError = CbFieldError.None, CbFieldAccessors? Accessors = null)
+		void TestField(CbFieldType fieldType, byte[] payload, object? expectedValue = null, object? defaultValue = null, CbFieldError expectedError = CbFieldError.None, CbFieldAccessors? accessors = null)
 		{
-			CbField Field = new CbField(Payload, FieldType);
-			Assert.AreEqual(Field.GetSize(), Payload.Length + (CbFieldUtils.HasFieldType(FieldType) ? 0 : 1));
-			Assert.IsTrue(Field.HasValue());
-			Assert.IsFalse(Field.HasError());
-			Assert.AreEqual(Field.GetError(), CbFieldError.None);
-			TestField(FieldType, Field, ExpectedValue, DefaultValue, ExpectedError, Accessors);
+			CbField field = new CbField(payload, fieldType);
+			Assert.AreEqual(field.GetSize(), payload.Length + (CbFieldUtils.HasFieldType(fieldType) ? 0 : 1));
+			Assert.IsTrue(field.HasValue());
+			Assert.IsFalse(field.HasError());
+			Assert.AreEqual(field.GetError(), CbFieldError.None);
+			TestField(fieldType, field, expectedValue, defaultValue, expectedError, accessors);
 		}
 
-		void TestFieldError(CbFieldType FieldType, CbField Field, CbFieldError ExpectedError, object? ExpectedValue = null, CbFieldAccessors? Accessors = null)
+		void TestFieldError(CbFieldType fieldType, CbField field, CbFieldError expectedError, object? expectedValue = null, CbFieldAccessors? accessors = null)
 		{
-			TestField(FieldType, Field, ExpectedValue, ExpectedValue, ExpectedError, Accessors);
+			TestField(fieldType, field, expectedValue, expectedValue, expectedError, accessors);
 		}
 
-		void TestFieldError(CbFieldType FieldType, ReadOnlyMemory<byte> Payload, CbFieldError ExpectedError, object? ExpectedValue = null, CbFieldAccessors? Accessors = null)
+		void TestFieldError(CbFieldType fieldType, ReadOnlyMemory<byte> payload, CbFieldError expectedError, object? expectedValue = null, CbFieldAccessors? accessors = null)
 		{
-			CbField Field = new CbField(Payload, FieldType);
-			TestFieldError(FieldType, Field, ExpectedError, ExpectedValue, Accessors);
+			CbField field = new CbField(payload, fieldType);
+			TestFieldError(fieldType, field, expectedError, expectedValue, accessors);
 		}
 
 		[TestMethod]
@@ -112,98 +110,94 @@ namespace EpicGames.Serialization.Tests
 		{
 			// Test CbField()
 			{
-				CbField DefaultField = new CbField();
-				Assert.IsFalse(DefaultField.HasName());
-				Assert.IsFalse(DefaultField.HasValue());
-				Assert.IsFalse(DefaultField.HasError());
-				Assert.IsTrue(DefaultField.GetError() == CbFieldError.None);
-				Assert.AreEqual(DefaultField.GetSize(), 1);
-				Assert.AreEqual(DefaultField.GetName().Length, 0);
-				Assert.IsFalse(DefaultField.HasName());
-				Assert.IsFalse(DefaultField.HasValue());
-				Assert.IsFalse(DefaultField.HasError());
-				Assert.AreEqual(DefaultField.GetError(), CbFieldError.None);
-				Assert.AreEqual(DefaultField.GetHash(), Blake3Hash.Compute(new byte[] { (byte)CbFieldType.None }));
-				ReadOnlyMemory<byte> View;
-				Assert.IsFalse(DefaultField.TryGetView(out View));
+				CbField defaultField = new CbField();
+				Assert.IsFalse(defaultField.HasName());
+				Assert.IsFalse(defaultField.HasValue());
+				Assert.IsFalse(defaultField.HasError());
+				Assert.IsTrue(defaultField.GetError() == CbFieldError.None);
+				Assert.AreEqual(defaultField.GetSize(), 1);
+				Assert.AreEqual(defaultField.GetName().Length, 0);
+				Assert.IsFalse(defaultField.HasName());
+				Assert.IsFalse(defaultField.HasValue());
+				Assert.IsFalse(defaultField.HasError());
+				Assert.AreEqual(defaultField.GetError(), CbFieldError.None);
+				Assert.AreEqual(defaultField.GetHash(), Blake3Hash.Compute(new byte[] { (byte)CbFieldType.None }));
+				Assert.IsFalse(defaultField.TryGetView(out _));
 			}
 
 			// Test CbField(None)
 			{
-				CbField NoneField = new CbField(ReadOnlyMemory<byte>.Empty, CbFieldType.None);
-				Assert.AreEqual(NoneField.GetSize(), 1);
-				Assert.AreEqual(NoneField.GetName().Length, 0);
-				Assert.IsFalse(NoneField.HasName());
-				Assert.IsFalse(NoneField.HasValue());
-				Assert.IsFalse(NoneField.HasError());
-				Assert.AreEqual(NoneField.GetError(), CbFieldError.None);
-				Assert.AreEqual(NoneField.GetHash(), new CbField().GetHash());
-				ReadOnlyMemory<byte> View;
-				Assert.IsFalse(NoneField.TryGetView(out View));
+				CbField noneField = new CbField(ReadOnlyMemory<byte>.Empty, CbFieldType.None);
+				Assert.AreEqual(noneField.GetSize(), 1);
+				Assert.AreEqual(noneField.GetName().Length, 0);
+				Assert.IsFalse(noneField.HasName());
+				Assert.IsFalse(noneField.HasValue());
+				Assert.IsFalse(noneField.HasError());
+				Assert.AreEqual(noneField.GetError(), CbFieldError.None);
+				Assert.AreEqual(noneField.GetHash(), new CbField().GetHash());
+				Assert.IsFalse(noneField.TryGetView(out _));
 			}
 
 			// Test CbField(None|Type|Name)
 			{
-				CbFieldType FieldType = CbFieldType.None | CbFieldType.HasFieldName;
-				byte[] NoneBytes = { (byte)FieldType, 4, (byte)'N', (byte)'a', (byte)'m', (byte)'e' };
-				CbField NoneField = new CbField(NoneBytes);
-				Assert.AreEqual(NoneField.GetSize(), NoneBytes.Length);
-				Assert.AreEqual(NoneField.GetName(), "Name");
-				Assert.IsTrue(NoneField.HasName());
-				Assert.IsFalse(NoneField.HasValue());
-				Assert.AreEqual(NoneField.GetHash(), Blake3Hash.Compute(NoneBytes));
-				ReadOnlyMemory<byte> View;
-				Assert.IsTrue(NoneField.TryGetView(out View) && View.Span.SequenceEqual(NoneBytes));
+				CbFieldType fieldType = CbFieldType.None | CbFieldType.HasFieldName;
+				byte[] noneBytes = { (byte)fieldType, 4, (byte)'N', (byte)'a', (byte)'m', (byte)'e' };
+				CbField noneField = new CbField(noneBytes);
+				Assert.AreEqual(noneField.GetSize(), noneBytes.Length);
+				Assert.AreEqual(noneField.GetName(), "Name");
+				Assert.IsTrue(noneField.HasName());
+				Assert.IsFalse(noneField.HasValue());
+				Assert.AreEqual(noneField.GetHash(), Blake3Hash.Compute(noneBytes));
+				ReadOnlyMemory<byte> view;
+				Assert.IsTrue(noneField.TryGetView(out view) && view.Span.SequenceEqual(noneBytes));
 
-				byte[] CopyBytes = new byte[NoneBytes.Length];
-				NoneField.CopyTo(CopyBytes);
-				Assert.IsTrue(NoneBytes.AsSpan().SequenceEqual(CopyBytes));
+				byte[] copyBytes = new byte[noneBytes.Length];
+				noneField.CopyTo(copyBytes);
+				Assert.IsTrue(noneBytes.AsSpan().SequenceEqual(copyBytes));
 			}
 
 			// Test CbField(None|Type)
 			{
-				CbFieldType FieldType = CbFieldType.None;
-				byte[] NoneBytes = { (byte)FieldType };
-				CbField NoneField = new CbField(NoneBytes);
-				Assert.AreEqual(NoneField.GetSize(), NoneBytes.Length);
-				Assert.AreEqual(NoneField.GetName().Length, 0);
-				Assert.IsFalse(NoneField.HasName());
-				Assert.IsFalse(NoneField.HasValue());
-				Assert.AreEqual(NoneField.GetHash(), new CbField().GetHash());
-				ReadOnlyMemory<byte> View;
-				Assert.IsTrue(NoneField.TryGetView(out View) && View.Span.SequenceEqual(NoneBytes));
+				CbFieldType fieldType = CbFieldType.None;
+				byte[] noneBytes = { (byte)fieldType };
+				CbField noneField = new CbField(noneBytes);
+				Assert.AreEqual(noneField.GetSize(), noneBytes.Length);
+				Assert.AreEqual(noneField.GetName().Length, 0);
+				Assert.IsFalse(noneField.HasName());
+				Assert.IsFalse(noneField.HasValue());
+				Assert.AreEqual(noneField.GetHash(), new CbField().GetHash());
+				ReadOnlyMemory<byte> view;
+				Assert.IsTrue(noneField.TryGetView(out view) && view.Span.SequenceEqual(noneBytes));
 			}
 
 			// Test CbField(None|Name)
 			{
-				CbFieldType FieldType = CbFieldType.None | CbFieldType.HasFieldName;
-				byte[] NoneBytes = { (byte)FieldType, 4, (byte)'N', (byte)'a', (byte)'m', (byte)'e' };
-				CbField NoneField = new CbField(NoneBytes.AsMemory(1), FieldType);
-				Assert.AreEqual(NoneField.GetSize(), NoneBytes.Length);
-				Assert.AreEqual(NoneField.GetName(), "Name");
-				Assert.IsTrue(NoneField.HasName());
-				Assert.IsFalse(NoneField.HasValue());
-				Assert.AreEqual(NoneField.GetHash(), Blake3Hash.Compute(NoneBytes));
-				ReadOnlyMemory<byte> View;
-				Assert.IsFalse(NoneField.TryGetView(out View));
+				CbFieldType fieldType = CbFieldType.None | CbFieldType.HasFieldName;
+				byte[] noneBytes = { (byte)fieldType, 4, (byte)'N', (byte)'a', (byte)'m', (byte)'e' };
+				CbField noneField = new CbField(noneBytes.AsMemory(1), fieldType);
+				Assert.AreEqual(noneField.GetSize(), noneBytes.Length);
+				Assert.AreEqual(noneField.GetName(), "Name");
+				Assert.IsTrue(noneField.HasName());
+				Assert.IsFalse(noneField.HasValue());
+				Assert.AreEqual(noneField.GetHash(), Blake3Hash.Compute(noneBytes));
+				Assert.IsFalse(noneField.TryGetView(out _));
 
-				byte[] CopyBytes = new byte[NoneBytes.Length];
-				NoneField.CopyTo(CopyBytes);
-				Assert.IsTrue(NoneBytes.AsSpan().SequenceEqual(CopyBytes));
+				byte[] copyBytes = new byte[noneBytes.Length];
+				noneField.CopyTo(copyBytes);
+				Assert.IsTrue(noneBytes.AsSpan().SequenceEqual(copyBytes));
 			}
 
 			// Test CbField(None|EmptyName)
 			{
-				CbFieldType FieldType = CbFieldType.None | CbFieldType.HasFieldName;
-				byte[] NoneBytes = { (byte)FieldType, 0 };
-				CbField NoneField = new CbField(NoneBytes.AsMemory(1), FieldType);
-				Assert.AreEqual(NoneField.GetSize(), NoneBytes.Length);
-				Assert.AreEqual(NoneField.GetName(), "");
-				Assert.IsTrue(NoneField.HasName());
-				Assert.IsFalse(NoneField.HasValue());
-				Assert.AreEqual(NoneField.GetHash(), Blake3Hash.Compute(NoneBytes));
-				ReadOnlyMemory<byte> View;
-				Assert.IsFalse(NoneField.TryGetView(out View));
+				CbFieldType fieldType = CbFieldType.None | CbFieldType.HasFieldName;
+				byte[] noneBytes = { (byte)fieldType, 0 };
+				CbField noneField = new CbField(noneBytes.AsMemory(1), fieldType);
+				Assert.AreEqual(noneField.GetSize(), noneBytes.Length);
+				Assert.AreEqual(noneField.GetName(), "");
+				Assert.IsTrue(noneField.HasName());
+				Assert.IsFalse(noneField.HasValue());
+				Assert.AreEqual(noneField.GetHash(), Blake3Hash.Compute(noneBytes));
+				Assert.IsFalse(noneField.TryGetView(out _));
 			}
 		}
 
@@ -212,264 +206,262 @@ namespace EpicGames.Serialization.Tests
 		{
 			// Test CbField(Null)
 			{
-				CbField NullField = new CbField(ReadOnlyMemory<byte>.Empty, CbFieldType.Null);
-				Assert.AreEqual(NullField.GetSize(), 1);
-				Assert.IsTrue(NullField.IsNull());
-				Assert.IsTrue(NullField.HasValue());
-				Assert.IsFalse(NullField.HasError());
-				Assert.AreEqual(NullField.GetError(), CbFieldError.None);
-				Assert.AreEqual(NullField.GetHash(), Blake3Hash.Compute(new byte[] { (byte)CbFieldType.Null }));
+				CbField nullField = new CbField(ReadOnlyMemory<byte>.Empty, CbFieldType.Null);
+				Assert.AreEqual(nullField.GetSize(), 1);
+				Assert.IsTrue(nullField.IsNull());
+				Assert.IsTrue(nullField.HasValue());
+				Assert.IsFalse(nullField.HasError());
+				Assert.AreEqual(nullField.GetError(), CbFieldError.None);
+				Assert.AreEqual(nullField.GetHash(), Blake3Hash.Compute(new byte[] { (byte)CbFieldType.Null }));
 			}
 
 			// Test CbField(None) as Null
 			{
-				CbField Field = new CbField();
-				Assert.IsFalse(Field.IsNull());
+				CbField field = new CbField();
+				Assert.IsFalse(field.IsNull());
 			}
 		}
 
 		[TestMethod]
 		public void CbFieldObjectTest()
 		{
-			Action<CbObject, int, int> TestIntObject = (CbObject Object, int ExpectedNum, int ExpectedPayloadSize) =>
+			static void TestIntObject(CbObject obj, int expectedNum, int expectedPayloadSize)
 			{
-				Assert.AreEqual(Object.GetSize(), ExpectedPayloadSize + sizeof(CbFieldType));
+				Assert.AreEqual(obj.GetSize(), expectedPayloadSize + sizeof(CbFieldType));
 
-				int ActualNum = 0;
-				foreach (CbField Field in Object)
+				int actualNum = 0;
+				foreach (CbField field in obj)
 				{
-					++ActualNum;
-					Assert.AreNotEqual(Field.GetName().Length, 0);
-					Assert.AreEqual(Field.AsInt32(), ActualNum);
+					++actualNum;
+					Assert.AreNotEqual(field.GetName().Length, 0);
+					Assert.AreEqual(field.AsInt32(), actualNum);
 				}
-				Assert.AreEqual(ActualNum, ExpectedNum);
-			};
+				Assert.AreEqual(actualNum, expectedNum);
+			}
 
 			// Test CbField(Object, Empty)
 			TestField(CbFieldType.Object, new byte[1]);
 
 			// Test CbField(Object, Empty)
 			{
-				CbObject Object = CbObject.Empty;
-				TestIntObject(Object, 0, 1);
+				CbObject @object = CbObject.Empty;
+				TestIntObject(@object, 0, 1);
 
 				// Find fields that do not exist.
-				Assert.IsFalse(Object.Find("Field").HasValue());
-				Assert.IsFalse(Object.FindIgnoreCase("Field").HasValue());
-				Assert.IsFalse(Object["Field"].HasValue());
+				Assert.IsFalse(@object.Find("Field").HasValue());
+				Assert.IsFalse(@object.FindIgnoreCase("Field").HasValue());
+				Assert.IsFalse(@object["Field"].HasValue());
 
 				// Advance an iterator past the last field.
-				CbFieldIterator It = Object.CreateIterator();
-				Assert.IsFalse((bool)It);
-				Assert.IsTrue(!It);
-				for (int Count = 16; Count > 0; --Count)
+				CbFieldIterator it = @object.CreateIterator();
+				Assert.IsFalse((bool)it);
+				Assert.IsTrue(!it);
+				for (int count = 16; count > 0; --count)
 				{
-					++It;
-					It.Current.AsInt32();
+					++it;
+					it.Current.AsInt32();
 				}
-				Assert.IsFalse((bool)It);
-				Assert.IsTrue(!It);
+				Assert.IsFalse((bool)it);
+				Assert.IsTrue(!it);
 			}
 
 			// Test CbField(Object, NotEmpty)
 			{
-				byte IntType = (byte)(CbFieldType.HasFieldName | CbFieldType.IntegerPositive);
-				byte[] Payload = { 12, IntType, 1, (byte)'A', 1, IntType, 1, (byte)'B', 2, IntType, 1, (byte)'C', 3 };
-				CbField Field = new CbField(Payload, CbFieldType.Object);
-				TestField(CbFieldType.Object, Field, new CbObject(Payload, CbFieldType.Object));
-				CbObject Object = CbObject.Clone(Field.AsObject());
-				TestIntObject(Object, 3, Payload.Length);
-				TestIntObject(Field.AsObject(), 3, Payload.Length);
-				Assert.IsTrue(Object.Equals(Field.AsObject()));
-				Assert.AreEqual(Object.Find("B").AsInt32(), 2);
-				Assert.AreEqual(Object.Find("b").AsInt32(4), 4);
-				Assert.AreEqual(Object.FindIgnoreCase("B").AsInt32(), 2);
-				Assert.AreEqual(Object.FindIgnoreCase("b").AsInt32(), 2);
-				Assert.AreEqual(Object["B"].AsInt32(), 2);
-				Assert.AreEqual(Object["b"].AsInt32(4), 4);
+				byte intType = (byte)(CbFieldType.HasFieldName | CbFieldType.IntegerPositive);
+				byte[] payload = { 12, intType, 1, (byte)'A', 1, intType, 1, (byte)'B', 2, intType, 1, (byte)'C', 3 };
+				CbField field = new CbField(payload, CbFieldType.Object);
+				TestField(CbFieldType.Object, field, new CbObject(payload, CbFieldType.Object));
+				CbObject @object = CbObject.Clone(field.AsObject());
+				TestIntObject(@object, 3, payload.Length);
+				TestIntObject(field.AsObject(), 3, payload.Length);
+				Assert.IsTrue(@object.Equals(field.AsObject()));
+				Assert.AreEqual(@object.Find("B").AsInt32(), 2);
+				Assert.AreEqual(@object.Find("b").AsInt32(4), 4);
+				Assert.AreEqual(@object.FindIgnoreCase("B").AsInt32(), 2);
+				Assert.AreEqual(@object.FindIgnoreCase("b").AsInt32(), 2);
+				Assert.AreEqual(@object["B"].AsInt32(), 2);
+				Assert.AreEqual(@object["b"].AsInt32(4), 4);
 			}
 
 			// Test CbField(UniformObject, NotEmpty)
 			{
-				byte IntType = (byte)(CbFieldType.HasFieldName | CbFieldType.IntegerPositive);
-				byte[] Payload = { 10, IntType, 1, (byte)'A', 1, 1, (byte)'B', 2, 1, (byte)'C', 3 };
-				CbField Field = new CbField(Payload, CbFieldType.UniformObject);
-				TestField(CbFieldType.UniformObject, Field, new CbObject(Payload, CbFieldType.UniformObject));
-				CbObject Object = CbObject.Clone(Field.AsObject());
-				TestIntObject(Object, 3, Payload.Length);
-				TestIntObject(Field.AsObject(), 3, Payload.Length);
-				Assert.IsTrue(Object.Equals(Field.AsObject()));
-				Assert.AreEqual(Object.Find("B").AsInt32(), 2);
-				Assert.AreEqual(Object.Find("B").AsInt32(), 2);
-				Assert.AreEqual(Object.Find("b").AsInt32(4), 4);
-				Assert.AreEqual(Object.Find("b").AsInt32(4), 4);
-				Assert.AreEqual(Object.FindIgnoreCase("B").AsInt32(), 2);
-				Assert.AreEqual(Object.FindIgnoreCase("B").AsInt32(), 2);
-				Assert.AreEqual(Object.FindIgnoreCase("b").AsInt32(), 2);
-				Assert.AreEqual(Object.FindIgnoreCase("b").AsInt32(), 2);
-				Assert.AreEqual(Object["B"].AsInt32(), 2);
-				Assert.AreEqual(Object["b"].AsInt32(4), 4);
+				byte intType = (byte)(CbFieldType.HasFieldName | CbFieldType.IntegerPositive);
+				byte[] payload = { 10, intType, 1, (byte)'A', 1, 1, (byte)'B', 2, 1, (byte)'C', 3 };
+				CbField field = new CbField(payload, CbFieldType.UniformObject);
+				TestField(CbFieldType.UniformObject, field, new CbObject(payload, CbFieldType.UniformObject));
+				CbObject @object = CbObject.Clone(field.AsObject());
+				TestIntObject(@object, 3, payload.Length);
+				TestIntObject(field.AsObject(), 3, payload.Length);
+				Assert.IsTrue(@object.Equals(field.AsObject()));
+				Assert.AreEqual(@object.Find("B").AsInt32(), 2);
+				Assert.AreEqual(@object.Find("B").AsInt32(), 2);
+				Assert.AreEqual(@object.Find("b").AsInt32(4), 4);
+				Assert.AreEqual(@object.Find("b").AsInt32(4), 4);
+				Assert.AreEqual(@object.FindIgnoreCase("B").AsInt32(), 2);
+				Assert.AreEqual(@object.FindIgnoreCase("B").AsInt32(), 2);
+				Assert.AreEqual(@object.FindIgnoreCase("b").AsInt32(), 2);
+				Assert.AreEqual(@object.FindIgnoreCase("b").AsInt32(), 2);
+				Assert.AreEqual(@object["B"].AsInt32(), 2);
+				Assert.AreEqual(@object["b"].AsInt32(4), 4);
 
 				// Equals
-				byte[] NamedPayload = { 1, (byte)'O', 10, IntType, 1, (byte)'A', 1, 1, (byte)'B', 2, 1, (byte)'C', 3 };
-				CbField NamedField = new CbField(NamedPayload, CbFieldType.UniformObject | CbFieldType.HasFieldName);
-				Assert.IsTrue(Field.AsObject().Equals(NamedField.AsObject()));
+				byte[] namedPayload = { 1, (byte)'O', 10, intType, 1, (byte)'A', 1, 1, (byte)'B', 2, 1, (byte)'C', 3 };
+				CbField namedField = new CbField(namedPayload, CbFieldType.UniformObject | CbFieldType.HasFieldName);
+				Assert.IsTrue(field.AsObject().Equals(namedField.AsObject()));
 
 				// CopyTo
-				byte[] CopyBytes = new byte[Payload.Length + 1];
-				Field.AsObject().CopyTo(CopyBytes);
-				Assert.IsTrue(Payload.AsSpan().SequenceEqual(CopyBytes.AsSpan(1)));
-				NamedField.AsObject().CopyTo(CopyBytes);
-				Assert.IsTrue(Payload.AsSpan().SequenceEqual(CopyBytes.AsSpan(1)));
+				byte[] copyBytes = new byte[payload.Length + 1];
+				field.AsObject().CopyTo(copyBytes);
+				Assert.IsTrue(payload.AsSpan().SequenceEqual(copyBytes.AsSpan(1)));
+				namedField.AsObject().CopyTo(copyBytes);
+				Assert.IsTrue(payload.AsSpan().SequenceEqual(copyBytes.AsSpan(1)));
 
 				// TryGetView
-				ReadOnlyMemory<byte> View;
-				Assert.IsFalse(Field.AsObject().TryGetView(out View));
-				Assert.IsFalse(NamedField.AsObject().TryGetView(out View));
+				Assert.IsFalse(field.AsObject().TryGetView(out _));
+				Assert.IsFalse(namedField.AsObject().TryGetView(out _));
 			}
 
 			// Test CbField(None) as Object
 			{
-				CbField Field = CbField.Empty;
-				TestFieldError(CbFieldType.Object, Field, CbFieldError.TypeError);
-				CbField.MakeView(Field).AsObject();
+				CbField field = CbField.Empty;
+				TestFieldError(CbFieldType.Object, field, CbFieldError.TypeError);
+				CbField.MakeView(field).AsObject();
 			}
 
 			// Test FCbObjectView(ObjectWithName) and CreateIterator
 			{
-				byte ObjectType = (byte)(CbFieldType.Object | CbFieldType.HasFieldName);
-				byte[] Buffer = { ObjectType, 3, (byte)'K', (byte)'e', (byte)'y', 4, (byte)(CbFieldType.HasFieldName | CbFieldType.IntegerPositive), 1, (byte)'F', 8 };
-				CbObject Object = new CbObject(Buffer);
-				Assert.AreEqual(Object.GetSize(), 6);
-				CbObject ObjectClone = CbObject.Clone(Object);
-				Assert.AreEqual(ObjectClone.GetSize(), 6);
-				Assert.IsTrue(Object.Equals(ObjectClone));
-				Assert.AreEqual(ObjectClone.GetHash(), Object.GetHash());
-				for (CbFieldIterator It = ObjectClone.CreateIterator(); It; ++It)
+				byte objectType = (byte)(CbFieldType.Object | CbFieldType.HasFieldName);
+				byte[] buffer = { objectType, 3, (byte)'K', (byte)'e', (byte)'y', 4, (byte)(CbFieldType.HasFieldName | CbFieldType.IntegerPositive), 1, (byte)'F', 8 };
+				CbObject @object = new CbObject(buffer);
+				Assert.AreEqual(@object.GetSize(), 6);
+				CbObject objectClone = CbObject.Clone(@object);
+				Assert.AreEqual(objectClone.GetSize(), 6);
+				Assert.IsTrue(@object.Equals(objectClone));
+				Assert.AreEqual(objectClone.GetHash(), @object.GetHash());
+				for (CbFieldIterator it = objectClone.CreateIterator(); it; ++it)
 				{
-					CbField Field = It.Current;
-					Assert.AreEqual(Field.GetName(), "F");
-					Assert.AreEqual(Field.AsInt32(), 8);
+					CbField field = it.Current;
+					Assert.AreEqual(field.GetName(), "F");
+					Assert.AreEqual(field.AsInt32(), 8);
 				}
-				for (CbFieldIterator It = ObjectClone.CreateIterator(), End = new CbFieldIterator(); It != End; ++It)
+				for (CbFieldIterator it = objectClone.CreateIterator(), end = new CbFieldIterator(); it != end; ++it)
 				{
 				}
-				foreach (CbField Field in ObjectClone)
+				foreach (CbField _ in objectClone)
 				{
 				}
 			}
 
 			// Test FCbObjectView as CbFieldIterator
 			{
-				int Count = 0;
-				CbObject Object = CbObject.Empty;
-				for (CbFieldIterator It = CbFieldIterator.MakeSingle(Object.AsField()); It; ++It)
+				int count = 0;
+				CbObject @object = CbObject.Empty;
+				for (CbFieldIterator it = CbFieldIterator.MakeSingle(@object.AsField()); it; ++it)
 				{
-					CbField Field = It.Current;
-					Assert.IsTrue(Field.IsObject());
-					++Count;
+					CbField field = it.Current;
+					Assert.IsTrue(field.IsObject());
+					++count;
 				}
-				Assert.AreEqual(Count, 1);
+				Assert.AreEqual(count, 1);
 			}
 		}
 
 		public void CbFieldArrayTest()
 		{
-			Action<CbArray, int, int> TestIntArray = (CbArray Array, int ExpectedNum, int ExpectedPayloadSize) =>
+			static void TestIntArray(CbArray array, int expectedNum, int expectedPayloadSize)
 			{
-				Assert.AreEqual(Array.GetSize(), ExpectedPayloadSize + sizeof(CbFieldType));
-				Assert.AreEqual(Array.Count, ExpectedNum);
+				Assert.AreEqual(array.GetSize(), expectedPayloadSize + sizeof(CbFieldType));
+				Assert.AreEqual(array.Count, expectedNum);
 
-				int ActualNum = 0;
-				for (CbFieldIterator It = Array.CreateIterator(); It; ++It)
+				int actualNum = 0;
+				for (CbFieldIterator it = array.CreateIterator(); it; ++it)
 				{
-					++ActualNum;
-					Assert.AreEqual(It.Current.AsInt32(), ActualNum);
+					++actualNum;
+					Assert.AreEqual(it.Current.AsInt32(), actualNum);
 				}
-				Assert.AreEqual(ActualNum, ExpectedNum);
+				Assert.AreEqual(actualNum, expectedNum);
 
-				ActualNum = 0;
-				foreach (CbField Field in Array)
+				actualNum = 0;
+				foreach (CbField field in array)
 				{
-					++ActualNum;
-					Assert.AreEqual(Field.AsInt32(), ActualNum);
+					++actualNum;
+					Assert.AreEqual(field.AsInt32(), actualNum);
 				}
-				Assert.AreEqual(ActualNum, ExpectedNum);
+				Assert.AreEqual(actualNum, expectedNum);
 
-				ActualNum = 0;
-				foreach (CbField Field in Array.AsField())
+				actualNum = 0;
+				foreach (CbField field in array.AsField())
 				{
-					++ActualNum;
-					Assert.AreEqual(Field.AsInt32(), ActualNum);
+					++actualNum;
+					Assert.AreEqual(field.AsInt32(), actualNum);
 				}
-				Assert.AreEqual(ActualNum, ExpectedNum);
-			};
+				Assert.AreEqual(actualNum, expectedNum);
+			}
 
 			// Test CbField(Array, Empty)
 			TestField(CbFieldType.Array, new byte[]{1, 0});
 
 			// Test CbField(Array, Empty)
 			{
-				CbArray Array = new CbArray();
-				TestIntArray(Array, 0, 2);
+				CbArray array = new CbArray();
+				TestIntArray(array, 0, 2);
 
 				// Advance an iterator past the last field.
-				CbFieldIterator It = Array.CreateIterator();
-				Assert.IsFalse((bool)It);
-				Assert.IsTrue(!It);
-				for (int Count = 16; Count > 0; --Count)
+				CbFieldIterator it = array.CreateIterator();
+				Assert.IsFalse((bool)it);
+				Assert.IsTrue(!it);
+				for (int count = 16; count > 0; --count)
 				{
-					++It;
-					It.Current.AsInt32();
+					++it;
+					it.Current.AsInt32();
 				}
-				Assert.IsFalse((bool)It);
-				Assert.IsTrue(!It);
+				Assert.IsFalse((bool)it);
+				Assert.IsTrue(!it);
 			}
 
 			// Test CbField(Array, NotEmpty)
 			{
-				byte IntType = (byte)CbFieldType.IntegerPositive;
-				byte[] Payload = new byte[]{ 7, 3, IntType, 1, IntType, 2, IntType, 3 };
-				CbField Field = new CbField(Payload, CbFieldType.Array);
-				TestField(CbFieldType.Array, Field, new CbArray(Payload, CbFieldType.Array));
-				CbArray Array = CbArray.Clone(Field.AsArray());
-				TestIntArray(Array, 3, Payload.Length);
-				TestIntArray(Field.AsArray(), 3, Payload.Length);
-				Assert.IsTrue(Array.Equals(Field.AsArray()));
+				byte intType = (byte)CbFieldType.IntegerPositive;
+				byte[] payload = new byte[]{ 7, 3, intType, 1, intType, 2, intType, 3 };
+				CbField field = new CbField(payload, CbFieldType.Array);
+				TestField(CbFieldType.Array, field, new CbArray(payload, CbFieldType.Array));
+				CbArray array = CbArray.Clone(field.AsArray());
+				TestIntArray(array, 3, payload.Length);
+				TestIntArray(field.AsArray(), 3, payload.Length);
+				Assert.IsTrue(array.Equals(field.AsArray()));
 			}
 
 			// Test CbField(UniformArray)
 			{
-				byte IntType = (byte)(CbFieldType.IntegerPositive);
-				byte[] Payload = new byte[]{ 5, 3, IntType, 1, 2, 3 };
-				CbField Field = new CbField(Payload, CbFieldType.UniformArray);
-				TestField(CbFieldType.UniformArray, Field, new CbArray(Payload, CbFieldType.UniformArray));
-				CbArray Array = CbArray.Clone(Field.AsArray());
-				TestIntArray(Array, 3, Payload.Length);
-				TestIntArray(Field.AsArray(), 3, Payload.Length);
-				Assert.IsTrue(Array.Equals(Field.AsArray()));
+				byte intType = (byte)(CbFieldType.IntegerPositive);
+				byte[] payload = new byte[]{ 5, 3, intType, 1, 2, 3 };
+				CbField field = new CbField(payload, CbFieldType.UniformArray);
+				TestField(CbFieldType.UniformArray, field, new CbArray(payload, CbFieldType.UniformArray));
+				CbArray array = CbArray.Clone(field.AsArray());
+				TestIntArray(array, 3, payload.Length);
+				TestIntArray(field.AsArray(), 3, payload.Length);
+				Assert.IsTrue(array.Equals(field.AsArray()));
 
 //				Assert.IsTrue(Array.GetOuterBuffer() == Array.AsField().AsArray().GetOuterBuffer());
 
 				// Equals
-				byte[] NamedPayload = new byte[]{ 1, (byte)'A', 5, 3, IntType, 1, 2, 3 };
-				CbField NamedField = new CbField(NamedPayload, CbFieldType.UniformArray | CbFieldType.HasFieldName);
-				Assert.IsTrue(Field.AsArray().Equals(NamedField.AsArray()));
-				Assert.IsTrue(Field.Equals(Field.AsArray().AsField()));
-				Assert.IsTrue(NamedField.Equals(NamedField.AsArray().AsField()));
+				byte[] namedPayload = new byte[]{ 1, (byte)'A', 5, 3, intType, 1, 2, 3 };
+				CbField namedField = new CbField(namedPayload, CbFieldType.UniformArray | CbFieldType.HasFieldName);
+				Assert.IsTrue(field.AsArray().Equals(namedField.AsArray()));
+				Assert.IsTrue(field.Equals(field.AsArray().AsField()));
+				Assert.IsTrue(namedField.Equals(namedField.AsArray().AsField()));
 
 				// CopyTo
-				byte[] CopyBytes = new byte[Payload.Length + 1];
-				Field.AsArray().CopyTo(CopyBytes);
-				Assert.IsTrue(Payload.AsSpan().SequenceEqual(CopyBytes.AsSpan(1)));
-				NamedField.AsArray().CopyTo(CopyBytes);
-				Assert.IsTrue(Payload.AsSpan().SequenceEqual(CopyBytes.AsSpan(1)));
+				byte[] copyBytes = new byte[payload.Length + 1];
+				field.AsArray().CopyTo(copyBytes);
+				Assert.IsTrue(payload.AsSpan().SequenceEqual(copyBytes.AsSpan(1)));
+				namedField.AsArray().CopyTo(copyBytes);
+				Assert.IsTrue(payload.AsSpan().SequenceEqual(copyBytes.AsSpan(1)));
 
 				// TryGetView
-				ReadOnlyMemory<byte> View;
 //				Assert.IsTrue(Array.TryGetView(out View) && View == Array.GetOuterBuffer().GetView());
-				Assert.IsFalse(Field.AsArray().TryGetView(out View));
-				Assert.IsFalse(NamedField.AsArray().TryGetView(out View));
+				Assert.IsFalse(field.AsArray().TryGetView(out _));
+				Assert.IsFalse(namedField.AsArray().TryGetView(out _));
 
 //				// GetBuffer
 //				Assert.IsTrue(Array.GetBuffer().Flatten().GetView() == Array.GetOuterBuffer().GetView());
@@ -479,39 +471,39 @@ namespace EpicGames.Serialization.Tests
 
 			// Test CbField(None) as Array
 			{
-				CbField Field = new CbField();
+				CbField field = new CbField();
 //				TestFieldError(CbFieldType.Array, Field, CbFieldError.TypeError);
-				CbField.MakeView(Field).AsArray();
+				CbField.MakeView(field).AsArray();
 			}
 
 			// Test CbArray(ArrayWithName) and CreateIterator
 			{
-				byte ArrayType = (byte)(CbFieldType.Array | CbFieldType.HasFieldName);
-				byte[] Buffer = new byte[] { ArrayType, 3, (byte)'K', (byte)'e', (byte)'y', 3, 1, (byte)(CbFieldType.IntegerPositive), 8 };
-				CbArray Array = new CbArray(Buffer);
-				Assert.AreEqual(Array.GetSize(), 5);
-				CbArray ArrayClone = CbArray.Clone(Array);
-				Assert.AreEqual(ArrayClone.GetSize(), 5);
-				Assert.IsTrue(Array.Equals(ArrayClone));
-				Assert.AreEqual(ArrayClone.GetHash(), Array.GetHash());
-				for (CbFieldIterator It = ArrayClone.CreateIterator(); It; ++It)
+				byte arrayType = (byte)(CbFieldType.Array | CbFieldType.HasFieldName);
+				byte[] buffer = new byte[] { arrayType, 3, (byte)'K', (byte)'e', (byte)'y', 3, 1, (byte)(CbFieldType.IntegerPositive), 8 };
+				CbArray array = new CbArray(buffer);
+				Assert.AreEqual(array.GetSize(), 5);
+				CbArray arrayClone = CbArray.Clone(array);
+				Assert.AreEqual(arrayClone.GetSize(), 5);
+				Assert.IsTrue(array.Equals(arrayClone));
+				Assert.AreEqual(arrayClone.GetHash(), array.GetHash());
+				for (CbFieldIterator it = arrayClone.CreateIterator(); it; ++it)
 				{
-					CbField Field = It.Current;
-					Assert.AreEqual(Field.AsInt32(), 8);
+					CbField field = it.Current;
+					Assert.AreEqual(field.AsInt32(), 8);
 //					Assert.IsTrue(Field.IsOwned());
 				}
-				for (CbFieldIterator It = ArrayClone.CreateIterator(), End = new CbFieldIterator(); It != End; ++It)
+				for (CbFieldIterator it = arrayClone.CreateIterator(), end = new CbFieldIterator(); it != end; ++it)
 				{
 				}
-				foreach (CbField Field in ArrayClone)
+				foreach (CbField _ in arrayClone)
 				{
 				}
 
 				// CopyTo
-				byte[] CopyBytes = new byte[5];
-				Array.CopyTo(CopyBytes);
+				byte[] copyBytes = new byte[5];
+				array.CopyTo(copyBytes);
 //				Assert.IsTrue(ArrayClone.GetOuterBuffer().GetView().Span.SequenceEqual(CopyBytes));
-				ArrayClone.CopyTo(CopyBytes);
+				arrayClone.CopyTo(copyBytes);
 //				Assert.IsTrue(ArrayClone.GetOuterBuffer().GetView().Span.SequenceEqual(CopyBytes));
 
 //				// GetBuffer
@@ -522,29 +514,29 @@ namespace EpicGames.Serialization.Tests
 
 			// Test CbArray as CbFieldIterator
 			{
-				uint Count = 0;
-				CbArray Array = new CbArray();
-				for (CbFieldIterator Iter = CbFieldIterator.MakeSingle(Array.AsField()); Iter; ++Iter)
+				uint count = 0;
+				CbArray array = new CbArray();
+				for (CbFieldIterator iter = CbFieldIterator.MakeSingle(array.AsField()); iter; ++iter)
 				{
-					CbField Field = Iter.Current;
-					Assert.IsTrue(Field.IsArray());
-					++Count;
+					CbField field = iter.Current;
+					Assert.IsTrue(field.IsArray());
+					++count;
 				}
-				Assert.AreEqual(Count, 1u);
+				Assert.AreEqual(count, 1u);
 			}
 
 			// Test CbArray as CbFieldIterator
 			{
-				uint Count = 0;
-				CbArray Array = new CbArray();
+				uint count = 0;
+				CbArray array = new CbArray();
 //				Array.MakeOwned();
-				for(CbFieldIterator Iter = CbFieldIterator.MakeSingle(Array.AsField()); Iter; ++Iter)
+				for(CbFieldIterator iter = CbFieldIterator.MakeSingle(array.AsField()); iter; ++iter)
 				{
-					CbField Field = Iter.Current;
-					Assert.IsTrue(Field.IsArray());
-					++Count;
+					CbField field = iter.Current;
+					Assert.IsTrue(field.IsArray());
+					++count;
 				}
-				Assert.AreEqual(Count, 1u);
+				Assert.AreEqual(count, 1u);
 			}
 		}
 
@@ -556,12 +548,12 @@ namespace EpicGames.Serialization.Tests
 
 			// Test CbField(Binary, Value)
 			{
-				byte[] Payload = { 3, 4, 5, 6 }; // Size: 3, Data: 4/5/6
-				CbField FieldView = new CbField(Payload, CbFieldType.Binary);
-				TestField(CbFieldType.Binary, FieldView, (ReadOnlyMemory<byte>)Payload.AsMemory(1, 3));
+				byte[] payload = { 3, 4, 5, 6 }; // Size: 3, Data: 4/5/6
+				CbField fieldView = new CbField(payload, CbFieldType.Binary);
+				TestField(CbFieldType.Binary, fieldView, (ReadOnlyMemory<byte>)payload.AsMemory(1, 3));
 
-				CbField Field = CbField.Clone(FieldView);
-				Field.AsBinary();
+				CbField field = CbField.Clone(fieldView);
+				field.AsBinary();
 //				Assert.IsFalse(Field.GetOuterBuffer().IsNull());
 //				MoveTemp(Field).AsBinary();
 //				Assert.IsTrue(Field.GetOuterBuffer().IsNull());
@@ -569,9 +561,9 @@ namespace EpicGames.Serialization.Tests
 
 			// Test CbField(None) as Binary
 			{
-				CbField FieldView = new CbField();
-				byte[] Default = { 1, 2, 3 };
-				TestFieldError(CbFieldType.Binary, FieldView, CbFieldError.TypeError, (ReadOnlyMemory<byte>)Default);
+				CbField fieldView = new CbField();
+				byte[] @default = { 1, 2, 3 };
+				TestFieldError(CbFieldType.Binary, fieldView, CbFieldError.TypeError, (ReadOnlyMemory<byte>)@default);
 
 //				CbField Field = CbField.Clone(FieldView);
 //				TestFieldError(CbFieldType.Binary, FSharedBuffer, Field, CbFieldError.TypeError, FSharedBuffer.MakeView(Default), FCbBinaryAccessors());
@@ -586,21 +578,21 @@ namespace EpicGames.Serialization.Tests
 
 			// Test CbField(String, Value)
 			{
-				byte[] Payload = { 3, (byte)'A', (byte)'B', (byte)'C' }; // Size: 3, Data: ABC
-				TestField(CbFieldType.String, Payload, new Utf8String(Payload.AsMemory(1, 3)));
+				byte[] payload = { 3, (byte)'A', (byte)'B', (byte)'C' }; // Size: 3, Data: ABC
+				TestField(CbFieldType.String, payload, new Utf8String(payload.AsMemory(1, 3)));
 			}
 
 			// Test CbField(String, OutOfRangeSize)
 			{
-				byte[] Payload = new byte[9];
-				VarInt.Write(Payload, (ulong)(1) << 31);
-				TestFieldError(CbFieldType.String, Payload, CbFieldError.RangeError, new Utf8String("ABC"));
+				byte[] payload = new byte[9];
+				VarInt.Write(payload, (ulong)(1) << 31);
+				TestFieldError(CbFieldType.String, payload, CbFieldError.RangeError, new Utf8String("ABC"));
 			}
 
 			// Test CbField(None) as String
 			{
-				CbField Field = new CbField();
-				TestFieldError(CbFieldType.String, Field, CbFieldError.TypeError, new Utf8String("ABC"));
+				CbField field = new CbField();
+				TestFieldError(CbFieldType.String, field, CbFieldError.TypeError, new Utf8String("ABC"));
 			}
 		}
 
@@ -632,40 +624,41 @@ namespace EpicGames.Serialization.Tests
 			Neg7  = Neg15 | Int8,
 		};
 
-		void TestIntegerField(CbFieldType FieldType, EIntType ExpectedMask, ulong Magnitude)
+		void TestIntegerField(CbFieldType fieldType, EIntType expectedMask, ulong magnitude)
 		{
-			byte[] Payload = new byte[9];
-			ulong Negative = (ulong)((byte)FieldType & 1);
-			VarInt.Write(Payload, Magnitude - Negative);
-			ulong DefaultValue = 8;
-			ulong ExpectedValue = (Negative != 0)? (ulong)(-(long)(Magnitude)) : Magnitude;
-			CbField Field = new CbField(Payload, FieldType);
+			byte[] payload = new byte[9];
+			ulong negative = (ulong)((byte)fieldType & 1);
+			VarInt.Write(payload, magnitude - negative);
+			ulong defaultValue = 8;
+			ulong expectedValue = (negative != 0)? (ulong)(-(long)(magnitude)) : magnitude;
+			CbField field = new CbField(payload, fieldType);
 
-			TestField(CbFieldType.IntegerNegative, Field, (sbyte)(((ExpectedMask & EIntType.Int8) != 0) ? ExpectedValue : DefaultValue),
-				(sbyte)(DefaultValue), ((ExpectedMask & EIntType.Int8) != 0)? CbFieldError.None : CbFieldError.RangeError, CbFieldAccessors.FromStruct<sbyte>(x => x.IsInteger(), (x, y) => x.AsInt8(y)));
+			TestField(CbFieldType.IntegerNegative, field, (sbyte)(((expectedMask & EIntType.Int8) != 0) ? expectedValue : defaultValue),
+				(sbyte)(defaultValue), ((expectedMask & EIntType.Int8) != 0)? CbFieldError.None : CbFieldError.RangeError, CbFieldAccessors.FromStruct<sbyte>(x => x.IsInteger(), (x, y) => x.AsInt8(y)));
 
-			TestField(CbFieldType.IntegerNegative, Field, (short)(((ExpectedMask & EIntType.Int16) != 0)? ExpectedValue : DefaultValue),
-				(short)(DefaultValue), ((ExpectedMask & EIntType.Int16) != 0)? CbFieldError.None : CbFieldError.RangeError, CbFieldAccessors.FromStruct<short>(x => x.IsInteger(), (x, y) => x.AsInt16(y)));
+			TestField(CbFieldType.IntegerNegative, field, (short)(((expectedMask & EIntType.Int16) != 0)? expectedValue : defaultValue),
+				(short)(defaultValue), ((expectedMask & EIntType.Int16) != 0)? CbFieldError.None : CbFieldError.RangeError, CbFieldAccessors.FromStruct<short>(x => x.IsInteger(), (x, y) => x.AsInt16(y)));
 
-			TestField(CbFieldType.IntegerNegative, Field, (int)(((ExpectedMask & EIntType.Int32) != 0)? ExpectedValue : DefaultValue),
-				(int)(DefaultValue), ((ExpectedMask & EIntType.Int32) != 0)? CbFieldError.None : CbFieldError.RangeError, CbFieldAccessors.FromStruct<int>(x => x.IsInteger(), (x, y) => x.AsInt32(y)));
+			TestField(CbFieldType.IntegerNegative, field, (int)(((expectedMask & EIntType.Int32) != 0)? expectedValue : defaultValue),
+				(int)(defaultValue), ((expectedMask & EIntType.Int32) != 0)? CbFieldError.None : CbFieldError.RangeError, CbFieldAccessors.FromStruct<int>(x => x.IsInteger(), (x, y) => x.AsInt32(y)));
 
-			TestField(CbFieldType.IntegerNegative, Field, (long)(((ExpectedMask & EIntType.Int64) != 0)? ExpectedValue : DefaultValue),
-				(long)(DefaultValue), ((ExpectedMask & EIntType.Int64) != 0)? CbFieldError.None : CbFieldError.RangeError, CbFieldAccessors.FromStruct<long>(x => x.IsInteger(), (x, y) => x.AsInt64(y)));
+			TestField(CbFieldType.IntegerNegative, field, (long)(((expectedMask & EIntType.Int64) != 0)? expectedValue : defaultValue),
+				(long)(defaultValue), ((expectedMask & EIntType.Int64) != 0)? CbFieldError.None : CbFieldError.RangeError, CbFieldAccessors.FromStruct<long>(x => x.IsInteger(), (x, y) => x.AsInt64(y)));
 
-			TestField(CbFieldType.IntegerPositive, Field, (byte)(((ExpectedMask & EIntType.UInt8) != 0) ? ExpectedValue : DefaultValue),
-				(byte)(DefaultValue), ((ExpectedMask & EIntType.UInt8) != 0) ? CbFieldError.None : CbFieldError.RangeError, CbFieldAccessors.FromStruct<byte>(x => x.IsInteger(), (x, y) => x.AsUInt8(y)));
+			TestField(CbFieldType.IntegerPositive, field, (byte)(((expectedMask & EIntType.UInt8) != 0) ? expectedValue : defaultValue),
+				(byte)(defaultValue), ((expectedMask & EIntType.UInt8) != 0) ? CbFieldError.None : CbFieldError.RangeError, CbFieldAccessors.FromStruct<byte>(x => x.IsInteger(), (x, y) => x.AsUInt8(y)));
 
-			TestField(CbFieldType.IntegerPositive, Field, (ushort)(((ExpectedMask & EIntType.UInt16) != 0)? ExpectedValue : DefaultValue),
-				(ushort)(DefaultValue), ((ExpectedMask & EIntType.UInt16) != 0) ? CbFieldError.None : CbFieldError.RangeError, CbFieldAccessors.FromStruct<ushort>(x => x.IsInteger(), (x, y) => x.AsUInt16(y)));
+			TestField(CbFieldType.IntegerPositive, field, (ushort)(((expectedMask & EIntType.UInt16) != 0)? expectedValue : defaultValue),
+				(ushort)(defaultValue), ((expectedMask & EIntType.UInt16) != 0) ? CbFieldError.None : CbFieldError.RangeError, CbFieldAccessors.FromStruct<ushort>(x => x.IsInteger(), (x, y) => x.AsUInt16(y)));
 
-			TestField(CbFieldType.IntegerPositive, Field, (uint)(((ExpectedMask & EIntType.UInt32) != 0) ? ExpectedValue : DefaultValue),
-				(uint)(DefaultValue), ((ExpectedMask & EIntType.UInt32) != 0) ? CbFieldError.None : CbFieldError.RangeError, CbFieldAccessors.FromStruct<uint>(x => x.IsInteger(), (x, y) => x.AsUInt32(y)));
-			TestField(CbFieldType.IntegerPositive, Field, (ulong)(((ExpectedMask & EIntType.UInt64) != 0) ? ExpectedValue : DefaultValue),
-				(ulong)(DefaultValue), ((ExpectedMask & EIntType.UInt64) != 0)? CbFieldError.None : CbFieldError.RangeError, CbFieldAccessors.FromStruct<ulong>(x => x.IsInteger(), (x, y) => x.AsUInt64(y)));
+			TestField(CbFieldType.IntegerPositive, field, (uint)(((expectedMask & EIntType.UInt32) != 0) ? expectedValue : defaultValue),
+				(uint)(defaultValue), ((expectedMask & EIntType.UInt32) != 0) ? CbFieldError.None : CbFieldError.RangeError, CbFieldAccessors.FromStruct<uint>(x => x.IsInteger(), (x, y) => x.AsUInt32(y)));
+			TestField(CbFieldType.IntegerPositive, field, (ulong)(((expectedMask & EIntType.UInt64) != 0) ? expectedValue : defaultValue),
+				(ulong)(defaultValue), ((expectedMask & EIntType.UInt64) != 0)? CbFieldError.None : CbFieldError.RangeError, CbFieldAccessors.FromStruct<ulong>(x => x.IsInteger(), (x, y) => x.AsUInt64(y)));
 		}
 
-		void CbFieldIntegerTest()
+		[TestMethod]
+		public void CbFieldIntegerTest()
 		{
 			// Test CbField(IntegerPositive)
 			TestIntegerField(CbFieldType.IntegerPositive, EIntType.Pos7,  0x00);
@@ -699,9 +692,9 @@ namespace EpicGames.Serialization.Tests
 
 			// Test CbField(None) as Integer
 			{
-				CbField Field = new CbField();
-				TestFieldError(CbFieldType.IntegerPositive, Field, CbFieldError.TypeError, (ulong)(8));
-				TestFieldError(CbFieldType.IntegerNegative, Field, CbFieldError.TypeError, (long)(8));
+				CbField field = new CbField();
+				TestFieldError(CbFieldType.IntegerPositive, field, CbFieldError.TypeError, (ulong)(8));
+				TestFieldError(CbFieldType.IntegerNegative, field, CbFieldError.TypeError, (long)(8));
 			}
 		}
 
@@ -710,108 +703,108 @@ namespace EpicGames.Serialization.Tests
 		{
 			// Test CbField(Float, 32-bit)
 			{
-				byte[] Payload = new byte[]{ 0xc0, 0x12, 0x34, 0x56 }; // -2.28444433f
-				TestField(CbFieldType.Float32, Payload, -2.28444433f);
+				byte[] payload = new byte[]{ 0xc0, 0x12, 0x34, 0x56 }; // -2.28444433f
+				TestField(CbFieldType.Float32, payload, -2.28444433f);
 
-				CbField Field = new CbField(Payload, CbFieldType.Float32);
-				TestField(CbFieldType.Float64, Field, (double)-2.28444433f);
+				CbField field = new CbField(payload, CbFieldType.Float32);
+				TestField(CbFieldType.Float64, field, (double)-2.28444433f);
 			}
 
 			// Test CbField(Float, 64-bit)
 			{
-				byte[] Payload = new byte[]{ 0xc1, 0x23, 0x45, 0x67, 0x89, 0xab, 0xcd, 0xef }; // -631475.76888888876
-				TestField(CbFieldType.Float64, Payload, -631475.76888888876);
+				byte[] payload = new byte[]{ 0xc1, 0x23, 0x45, 0x67, 0x89, 0xab, 0xcd, 0xef }; // -631475.76888888876
+				TestField(CbFieldType.Float64, payload, -631475.76888888876);
 
-				CbField Field = new CbField(Payload, CbFieldType.Float64);
-				TestFieldError(CbFieldType.Float32, Field, CbFieldError.RangeError, 8.0f);
+				CbField field = new CbField(payload, CbFieldType.Float64);
+				TestFieldError(CbFieldType.Float32, field, CbFieldError.RangeError, 8.0f);
 			}
 
 			// Test CbField(Integer+, MaxBinary32) as Float
 			{
-				byte[] Payload = new byte[9];
-				VarInt.Write(Payload, ((ulong)(1) << 24) - 1); // 16,777,215
-				CbField Field = new CbField(Payload, CbFieldType.IntegerPositive);
-				TestField(CbFieldType.Float32, Field, 16_777_215.0f);
-				TestField(CbFieldType.Float64, Field, 16_777_215.0);
+				byte[] payload = new byte[9];
+				VarInt.Write(payload, ((ulong)(1) << 24) - 1); // 16,777,215
+				CbField field = new CbField(payload, CbFieldType.IntegerPositive);
+				TestField(CbFieldType.Float32, field, 16_777_215.0f);
+				TestField(CbFieldType.Float64, field, 16_777_215.0);
 			}
 
 			// Test CbField(Integer+, MaxBinary32+1) as Float
 			{
-				byte[] Payload = new byte[9];
-				VarInt.Write(Payload, (ulong)(1) << 24); // 16,777,216
-				CbField Field = new CbField(Payload, CbFieldType.IntegerPositive);
-				TestFieldError(CbFieldType.Float32, Field, CbFieldError.RangeError, 8.0f);
-				TestField(CbFieldType.Float64, Field, 16_777_216.0);
+				byte[] payload = new byte[9];
+				VarInt.Write(payload, (ulong)(1) << 24); // 16,777,216
+				CbField field = new CbField(payload, CbFieldType.IntegerPositive);
+				TestFieldError(CbFieldType.Float32, field, CbFieldError.RangeError, 8.0f);
+				TestField(CbFieldType.Float64, field, 16_777_216.0);
 			}
 
 			// Test CbField(Integer+, MaxBinary64) as Float
 			{
-				byte[] Payload = new byte[9];
-				VarInt.Write(Payload, ((ulong)(1) << 53) - 1); // 9,007,199,254,740,991
-				CbField Field = new CbField(Payload, CbFieldType.IntegerPositive);
-				TestFieldError(CbFieldType.Float32, Field, CbFieldError.RangeError, 8.0f);
-				TestField(CbFieldType.Float64, Field, 9_007_199_254_740_991.0);
+				byte[] payload = new byte[9];
+				VarInt.Write(payload, ((ulong)(1) << 53) - 1); // 9,007,199,254,740,991
+				CbField field = new CbField(payload, CbFieldType.IntegerPositive);
+				TestFieldError(CbFieldType.Float32, field, CbFieldError.RangeError, 8.0f);
+				TestField(CbFieldType.Float64, field, 9_007_199_254_740_991.0);
 			}
 
 			// Test CbField(Integer+, MaxBinary64+1) as Float
 			{
-				byte[] Payload = new byte[9];
-				VarInt.Write(Payload, (ulong)(1) << 53); // 9,007,199,254,740,992
-				CbField Field = new CbField(Payload, CbFieldType.IntegerPositive);
-				TestFieldError(CbFieldType.Float32, Field, CbFieldError.RangeError, 8.0f);
-				TestFieldError(CbFieldType.Float64, Field, CbFieldError.RangeError, 8.0);
+				byte[] payload = new byte[9];
+				VarInt.Write(payload, (ulong)(1) << 53); // 9,007,199,254,740,992
+				CbField field = new CbField(payload, CbFieldType.IntegerPositive);
+				TestFieldError(CbFieldType.Float32, field, CbFieldError.RangeError, 8.0f);
+				TestFieldError(CbFieldType.Float64, field, CbFieldError.RangeError, 8.0);
 			}
 
 			// Test CbField(Integer+, MaxUInt64) as Float
 			{
-				byte[] Payload = new byte[9];
-				VarInt.Write(Payload, ~(ulong)0); // Max uint64
-				CbField Field = new CbField(Payload, CbFieldType.IntegerPositive);
-				TestFieldError(CbFieldType.Float32, Field, CbFieldError.RangeError, 8.0f);
-				TestFieldError(CbFieldType.Float64, Field, CbFieldError.RangeError, 8.0);
+				byte[] payload = new byte[9];
+				VarInt.Write(payload, ~(ulong)0); // Max uint64
+				CbField field = new CbField(payload, CbFieldType.IntegerPositive);
+				TestFieldError(CbFieldType.Float32, field, CbFieldError.RangeError, 8.0f);
+				TestFieldError(CbFieldType.Float64, field, CbFieldError.RangeError, 8.0);
 			}
 
 			// Test CbField(Integer-, MaxBinary32) as Float
 			{
-				byte[] Payload = new byte[9];
-				VarInt.Write(Payload, ((ulong)(1) << 24) - 2); // -16,777,215
-				CbField Field = new CbField(Payload, CbFieldType.IntegerNegative);
-				TestField(CbFieldType.Float32, Field, -16_777_215.0f);
-				TestField(CbFieldType.Float64, Field, -16_777_215.0);
+				byte[] payload = new byte[9];
+				VarInt.Write(payload, ((ulong)(1) << 24) - 2); // -16,777,215
+				CbField field = new CbField(payload, CbFieldType.IntegerNegative);
+				TestField(CbFieldType.Float32, field, -16_777_215.0f);
+				TestField(CbFieldType.Float64, field, -16_777_215.0);
 			}
 
 			// Test CbField(Integer-, MaxBinary32+1) as Float
 			{
-				byte[] Payload = new byte[9];
-				VarInt.Write(Payload, ((ulong)(1) << 24) - 1); // -16,777,216
-				CbField Field = new CbField(Payload, CbFieldType.IntegerNegative);
-				TestFieldError(CbFieldType.Float32, Field, CbFieldError.RangeError, 8.0f);
-				TestField(CbFieldType.Float64, Field, -16_777_216.0);
+				byte[] payload = new byte[9];
+				VarInt.Write(payload, ((ulong)(1) << 24) - 1); // -16,777,216
+				CbField field = new CbField(payload, CbFieldType.IntegerNegative);
+				TestFieldError(CbFieldType.Float32, field, CbFieldError.RangeError, 8.0f);
+				TestField(CbFieldType.Float64, field, -16_777_216.0);
 			}
 
 			// Test CbField(Integer-, MaxBinary64) as Float
 			{
-				byte[] Payload = new byte[9];
-				VarInt.Write(Payload, ((ulong)(1) << 53) - 2); // -9,007,199,254,740,991
-				CbField Field = new CbField(Payload, CbFieldType.IntegerNegative);
-				TestFieldError(CbFieldType.Float32, Field, CbFieldError.RangeError, 8.0f);
-				TestField(CbFieldType.Float64, Field, -9_007_199_254_740_991.0);
+				byte[] payload = new byte[9];
+				VarInt.Write(payload, ((ulong)(1) << 53) - 2); // -9,007,199,254,740,991
+				CbField field = new CbField(payload, CbFieldType.IntegerNegative);
+				TestFieldError(CbFieldType.Float32, field, CbFieldError.RangeError, 8.0f);
+				TestField(CbFieldType.Float64, field, -9_007_199_254_740_991.0);
 			}
 
 			// Test CbField(Integer-, MaxBinary64+1) as Float
 			{
-				byte[] Payload = new byte[9];
-				VarInt.Write(Payload, ((ulong)(1) << 53) - 1); // -9,007,199,254,740,992
-				CbField Field = new CbField(Payload, CbFieldType.IntegerNegative);
-				TestFieldError(CbFieldType.Float32, Field, CbFieldError.RangeError, 8.0f);
-				TestFieldError(CbFieldType.Float64, Field, CbFieldError.RangeError, 8.0);
+				byte[] payload = new byte[9];
+				VarInt.Write(payload, ((ulong)(1) << 53) - 1); // -9,007,199,254,740,992
+				CbField field = new CbField(payload, CbFieldType.IntegerNegative);
+				TestFieldError(CbFieldType.Float32, field, CbFieldError.RangeError, 8.0f);
+				TestFieldError(CbFieldType.Float64, field, CbFieldError.RangeError, 8.0);
 			}
 
 			// Test CbField(None) as Float
 			{
-				CbField Field = new CbField();
-				TestFieldError(CbFieldType.Float32, Field, CbFieldError.TypeError, 8.0f);
-				TestFieldError(CbFieldType.Float64, Field, CbFieldError.TypeError, 8.0);
+				CbField field = new CbField();
+				TestFieldError(CbFieldType.Float32, field, CbFieldError.TypeError, 8.0f);
+				TestFieldError(CbFieldType.Float64, field, CbFieldError.TypeError, 8.0);
 			}
 		}
 
@@ -826,110 +819,110 @@ namespace EpicGames.Serialization.Tests
 
 			// Test CbField(None) as Bool
 			{
-				CbField DefaultField = new CbField();
-				TestFieldError(CbFieldType.BoolFalse, DefaultField, CbFieldError.TypeError, false);
-				TestFieldError(CbFieldType.BoolTrue, DefaultField, CbFieldError.TypeError, true);
+				CbField defaultField = new CbField();
+				TestFieldError(CbFieldType.BoolFalse, defaultField, CbFieldError.TypeError, false);
+				TestFieldError(CbFieldType.BoolTrue, defaultField, CbFieldError.TypeError, true);
 			}
 		}
 
 		[TestMethod]
 		public void CbFieldObjectAttachmentTest()
 		{
-			byte[] ZeroBytes = new byte[20];
-			byte[] SequentialBytes = new byte[]{1, 2, 3, 4, 5, 6, 7, 8, 9, 10, 11, 12, 13, 14, 15, 16, 17, 18, 19, 20};
+			byte[] zeroBytes = new byte[20];
+			byte[] sequentialBytes = new byte[]{1, 2, 3, 4, 5, 6, 7, 8, 9, 10, 11, 12, 13, 14, 15, 16, 17, 18, 19, 20};
 
 			// Test CbField(ObjectAttachment, Zero)
-			TestField(CbFieldType.ObjectAttachment, ZeroBytes);
+			TestField(CbFieldType.ObjectAttachment, zeroBytes);
 
 			// Test CbField(ObjectAttachment, NonZero)
-			TestField(CbFieldType.ObjectAttachment, SequentialBytes, (CbObjectAttachment)new IoHash(SequentialBytes));
+			TestField(CbFieldType.ObjectAttachment, sequentialBytes, (CbObjectAttachment)new IoHash(sequentialBytes));
 
 			// Test CbField(ObjectAttachment, NonZero) AsAttachment
 			{
-				CbField Field = new CbField(SequentialBytes, CbFieldType.ObjectAttachment);
-				TestField(CbFieldType.ObjectAttachment, Field, (CbObjectAttachment)new IoHash(SequentialBytes), (CbObjectAttachment)new IoHash(), CbFieldError.None);
+				CbField field = new CbField(sequentialBytes, CbFieldType.ObjectAttachment);
+				TestField(CbFieldType.ObjectAttachment, field, (CbObjectAttachment)new IoHash(sequentialBytes), (CbObjectAttachment)new IoHash(), CbFieldError.None);
 			}
 
 			// Test CbField(None) as ObjectAttachment
 			{
-				CbField DefaultField = new CbField();
-				TestFieldError(CbFieldType.ObjectAttachment, DefaultField, CbFieldError.TypeError, (CbObjectAttachment)new IoHash(SequentialBytes));
+				CbField defaultField = new CbField();
+				TestFieldError(CbFieldType.ObjectAttachment, defaultField, CbFieldError.TypeError, (CbObjectAttachment)new IoHash(sequentialBytes));
 			}
 		}
 
 		[TestMethod]
 		public void CbFieldBinaryAttachmentTest()
 		{
-			byte[] ZeroBytes = new byte[20];
-			byte[] SequentialBytes = new byte[]{1, 2, 3, 4, 5, 6, 7, 8, 9, 10, 11, 12, 13, 14, 15, 16, 17, 18, 19, 20};
+			byte[] zeroBytes = new byte[20];
+			byte[] sequentialBytes = new byte[]{1, 2, 3, 4, 5, 6, 7, 8, 9, 10, 11, 12, 13, 14, 15, 16, 17, 18, 19, 20};
 
 			// Test CbField(BinaryAttachment, Zero)
-			TestField(CbFieldType.BinaryAttachment, ZeroBytes);
+			TestField(CbFieldType.BinaryAttachment, zeroBytes);
 
 			// Test CbField(BinaryAttachment, NonZero)
-			TestField(CbFieldType.BinaryAttachment, SequentialBytes, (CbBinaryAttachment)new IoHash(SequentialBytes));
+			TestField(CbFieldType.BinaryAttachment, sequentialBytes, (CbBinaryAttachment)new IoHash(sequentialBytes));
 
 			// Test CbField(BinaryAttachment, NonZero) AsAttachment
 			{
-				CbField Field = new CbField(SequentialBytes, CbFieldType.BinaryAttachment);
-				TestField(CbFieldType.BinaryAttachment, Field, (CbBinaryAttachment)new IoHash(SequentialBytes), (CbBinaryAttachment)new IoHash(), CbFieldError.None);
+				CbField field = new CbField(sequentialBytes, CbFieldType.BinaryAttachment);
+				TestField(CbFieldType.BinaryAttachment, field, (CbBinaryAttachment)new IoHash(sequentialBytes), (CbBinaryAttachment)new IoHash(), CbFieldError.None);
 			}
 
 			// Test CbField(None) as BinaryAttachment
 			{
-				CbField DefaultField = new CbField();
-				TestFieldError(CbFieldType.BinaryAttachment, DefaultField, CbFieldError.TypeError, (CbBinaryAttachment)new IoHash(SequentialBytes));
+				CbField defaultField = new CbField();
+				TestFieldError(CbFieldType.BinaryAttachment, defaultField, CbFieldError.TypeError, (CbBinaryAttachment)new IoHash(sequentialBytes));
 			}
 		}
 
 		[TestMethod]
 		public void CbFieldHashTest()
 		{
-			byte[] ZeroBytes = new byte[20];
-			byte[] SequentialBytes = new byte[]{1, 2, 3, 4, 5, 6, 7, 8, 9, 10, 11, 12, 13, 14, 15, 16, 17, 18, 19, 20};
+			byte[] zeroBytes = new byte[20];
+			byte[] sequentialBytes = new byte[]{1, 2, 3, 4, 5, 6, 7, 8, 9, 10, 11, 12, 13, 14, 15, 16, 17, 18, 19, 20};
 
 			// Test CbField(Hash, Zero)
-			TestField(CbFieldType.Hash, ZeroBytes);
+			TestField(CbFieldType.Hash, zeroBytes);
 
 			// Test CbField(Hash, NonZero)
-			TestField(CbFieldType.Hash, SequentialBytes, new IoHash(SequentialBytes));
+			TestField(CbFieldType.Hash, sequentialBytes, new IoHash(sequentialBytes));
 
 			// Test CbField(None) as Hash
 			{
-				CbField DefaultField = new CbField();
-				TestFieldError(CbFieldType.Hash, DefaultField, CbFieldError.TypeError, new IoHash(SequentialBytes));
+				CbField defaultField = new CbField();
+				TestFieldError(CbFieldType.Hash, defaultField, CbFieldError.TypeError, new IoHash(sequentialBytes));
 			}
 
 			// Test CbField(ObjectAttachment) as Hash
 			{
-				CbField Field = new CbField(SequentialBytes, CbFieldType.ObjectAttachment);
-				TestField(CbFieldType.Hash, Field, new IoHash(SequentialBytes));
+				CbField field = new CbField(sequentialBytes, CbFieldType.ObjectAttachment);
+				TestField(CbFieldType.Hash, field, new IoHash(sequentialBytes));
 			}
 
 			// Test CbField(BinaryAttachment) as Hash
 			{
-				CbField Field = new CbField(SequentialBytes, CbFieldType.BinaryAttachment);
-				TestField(CbFieldType.Hash, Field, new IoHash(SequentialBytes));
+				CbField field = new CbField(sequentialBytes, CbFieldType.BinaryAttachment);
+				TestField(CbFieldType.Hash, field, new IoHash(sequentialBytes));
 			}
 		}
 
 		[TestMethod]
 		public void CbFieldUuidTest()
 		{
-			byte[] ZeroBytes = new byte[]{0, 0, 0, 0, 0, 0, 0, 0, 0, 0, 0, 0, 0, 0, 0, 0};
-			byte[] SequentialBytes = new byte[]{0, 1, 2, 3, 4, 5, 6, 7, 8, 9, 10, 11, 12, 13, 14, 15};
-			Guid SequentialGuid = Guid.Parse("00010203-0405-0607-0809-0a0b0c0d0e0f");
+			byte[] zeroBytes = new byte[]{0, 0, 0, 0, 0, 0, 0, 0, 0, 0, 0, 0, 0, 0, 0, 0};
+			byte[] sequentialBytes = new byte[]{0, 1, 2, 3, 4, 5, 6, 7, 8, 9, 10, 11, 12, 13, 14, 15};
+			Guid sequentialGuid = Guid.Parse("00010203-0405-0607-0809-0a0b0c0d0e0f");
 
 			// Test CbField(Uuid, Zero)
-			TestField(CbFieldType.Uuid, ZeroBytes, new Guid(), SequentialGuid);
+			TestField(CbFieldType.Uuid, zeroBytes, new Guid(), sequentialGuid);
 
 			// Test CbField(Uuid, NonZero)
-			TestField(CbFieldType.Uuid, SequentialBytes, SequentialGuid, new Guid());
+			TestField(CbFieldType.Uuid, sequentialBytes, sequentialGuid, new Guid());
 
 			// Test CbField(None) as Uuid
 			{
-				CbField DefaultField = new CbField();
-				TestFieldError(CbFieldType.Uuid, DefaultField, CbFieldError.TypeError, Guid.NewGuid());
+				CbField defaultField = new CbField();
+				TestFieldError(CbFieldType.Uuid, defaultField, CbFieldError.TypeError, Guid.NewGuid());
 			}
 		}
 
@@ -944,17 +937,17 @@ namespace EpicGames.Serialization.Tests
 
 			// Test CbField(DateTime, Zero) as FDateTime
 			{
-				byte[] Payload = new byte[]{0, 0, 0, 0, 0, 0, 0, 0};
-				CbField Field = new CbField(Payload, CbFieldType.DateTime);
-				Assert.AreEqual(Field.AsDateTime(), new DateTime(0));
+				byte[] payload = new byte[]{0, 0, 0, 0, 0, 0, 0, 0};
+				CbField field = new CbField(payload, CbFieldType.DateTime);
+				Assert.AreEqual(field.AsDateTime(), new DateTime(0));
 			}
 
 			// Test CbField(None) as DateTime
 			{
-				CbField DefaultField = new CbField();
-				TestFieldError(CbFieldType.DateTime, DefaultField, CbFieldError.TypeError);
-				DateTime DefaultValue = new DateTime(0x1020_3040_5060_7080L, DateTimeKind.Utc);
-				Assert.AreEqual(DefaultField.AsDateTime(DefaultValue), DefaultValue);
+				CbField defaultField = new CbField();
+				TestFieldError(CbFieldType.DateTime, defaultField, CbFieldError.TypeError);
+				DateTime defaultValue = new DateTime(0x1020_3040_5060_7080L, DateTimeKind.Utc);
+				Assert.AreEqual(defaultField.AsDateTime(defaultValue), defaultValue);
 			}
 		}
 
@@ -969,17 +962,17 @@ namespace EpicGames.Serialization.Tests
 
 			// Test CbField(TimeSpan, Zero) as FTimeSpan
 			{
-				byte[] Payload = new byte[]{0, 0, 0, 0, 0, 0, 0, 0};
-				CbField Field = new CbField(Payload, CbFieldType.TimeSpan);
-				Assert.AreEqual(Field.AsTimeSpan(), new TimeSpan(0));
+				byte[] payload = new byte[]{0, 0, 0, 0, 0, 0, 0, 0};
+				CbField field = new CbField(payload, CbFieldType.TimeSpan);
+				Assert.AreEqual(field.AsTimeSpan(), new TimeSpan(0));
 			}
 
 			// Test CbField(None) as TimeSpan
 			{
-				CbField DefaultField = new CbField();
-				TestFieldError(CbFieldType.TimeSpan, DefaultField, CbFieldError.TypeError);
-				TimeSpan DefaultValue = new TimeSpan(0x1020_3040_5060_7080L);
-				Assert.AreEqual(DefaultField.AsTimeSpan(DefaultValue), DefaultValue);
+				CbField defaultField = new CbField();
+				TestFieldError(CbFieldType.TimeSpan, defaultField, CbFieldError.TypeError);
+				TimeSpan defaultValue = new TimeSpan(0x1020_3040_5060_7080L);
+				Assert.AreEqual(defaultField.AsTimeSpan(defaultValue), defaultValue);
 			}
 		}
 
@@ -1615,7 +1608,6 @@ namespace EpicGames.Serialization.Tests
 			return true;
 		}
 #endif
-		delegate void ParseObjectType(CbObject Object, ref uint A, ref uint B, ref uint C, ref uint D);
 
 		[TestMethod]
 		public void CbFieldParseTest()
@@ -1624,54 +1616,54 @@ namespace EpicGames.Serialization.Tests
 			// Under ideal conditions, when the fields are in the expected order and there are no extra fields,
 			// the loop will execute once and only one comparison will be performed for each field name. Either
 			// way, each field will only be visited once even if the loop needs to execute several times.
-			ParseObjectType ParseObject = (CbObject Object, ref uint A, ref uint B, ref uint C, ref uint D) =>
+			static void ParseObject(CbObject obj, ref uint a, ref uint b, ref uint c, ref uint d)
 			{
-				for (CbFieldIterator It = Object.CreateIterator(); It;)
+				for (CbFieldIterator it = obj.CreateIterator(); it;)
 				{
-					CbFieldIterator Last = It;
-					if (It.Current.GetName().Equals("A"))
+					CbFieldIterator last = it;
+					if (it.Current.GetName().Equals("A"))
 					{
-						A = It.Current.AsUInt32();
-						++It;
+						a = it.Current.AsUInt32();
+						++it;
 					}
-					if (It.Current.GetName().Equals("B"))
+					if (it.Current.GetName().Equals("B"))
 					{
-						B = It.Current.AsUInt32();
-						++It;
+						b = it.Current.AsUInt32();
+						++it;
 					}
-					if (It.Current.GetName().Equals("C"))
+					if (it.Current.GetName().Equals("C"))
 					{
-						C = It.Current.AsUInt32();
-						++It;
+						c = it.Current.AsUInt32();
+						++it;
 					}
-					if (It.Current.GetName().Equals("D"))
+					if (it.Current.GetName().Equals("D"))
 					{
-						D = It.Current.AsUInt32();
-						++It;
+						d = it.Current.AsUInt32();
+						++it;
 					}
-					if (Last == It)
+					if (last == it)
 					{
-						++It;
+						++it;
 					}
 				}
-			};
+			}
 
-			Func<byte[], uint, uint, uint, uint, bool> TestParseObject = (byte[] Data, uint A, uint B, uint C, uint D) =>
+			static bool TestParseObject(byte[] data, uint a, uint b, uint c, uint d)
 			{
-				uint ParsedA = 0, ParsedB = 0, ParsedC = 0, ParsedD = 0;
-				ParseObject(new CbObject(Data, CbFieldType.Object), ref ParsedA, ref ParsedB, ref ParsedC, ref ParsedD);
-				return A == ParsedA && B == ParsedB && C == ParsedC && D == ParsedD;
-			};
+				uint parsedA = 0, parsedB = 0, parsedC = 0, parsedD = 0;
+				ParseObject(new CbObject(data, CbFieldType.Object), ref parsedA, ref parsedB, ref parsedC, ref parsedD);
+				return a == parsedA && b == parsedB && c == parsedC && d == parsedD;
+			}
 
-			byte T = (byte)(CbFieldType.IntegerPositive | CbFieldType.HasFieldName);
+			byte t = (byte)(CbFieldType.IntegerPositive | CbFieldType.HasFieldName);
 			Assert.IsTrue(TestParseObject(new byte[]{0}, 0, 0, 0, 0));
-			Assert.IsTrue(TestParseObject(new byte[] { 16, T, 1, (byte)'A', 1, T, 1, (byte)'B', 2, T, 1, (byte)'C', 3, T, 1, (byte)'D', 4}, 1, 2, 3, 4));
-			Assert.IsTrue(TestParseObject(new byte[]{16, T, 1, (byte)'B', 2, T, 1, (byte)'C', 3, T, 1, (byte)'D', 4, T, 1, (byte)'A', 1}, 1, 2, 3, 4));
-			Assert.IsTrue(TestParseObject(new byte[]{12, T, 1, (byte)'B', 2, T, 1, (byte)'C', 3, T, 1, (byte)'D', 4}, 0, 2, 3, 4));
-			Assert.IsTrue(TestParseObject(new byte[]{8, T, 1, (byte)'B', 2, T, 1, (byte)'C', 3}, 0, 2, 3, 0));
-			Assert.IsTrue(TestParseObject(new byte[]{20, T, 1, (byte)'A', 1, T, 1, (byte)'B', 2, T, 1, (byte)'C', 3, T, 1, (byte)'D', 4, T, 1, (byte)'E', 5}, 1, 2, 3, 4));
-			Assert.IsTrue(TestParseObject(new byte[]{20, T, 1, (byte)'E', 5, T, 1, (byte)'A', 1, T, 1, (byte)'B', 2, T, 1, (byte)'C', 3, T, 1, (byte)'D', 4}, 1, 2, 3, 4));
-			Assert.IsTrue(TestParseObject(new byte[] { 16, T, 1, (byte)'D', 4, T, 1, (byte)'C', 3, T, 1, (byte)'B', 2, T, 1, (byte)'A', 1}, 1, 2, 3, 4));
+			Assert.IsTrue(TestParseObject(new byte[] { 16, t, 1, (byte)'A', 1, t, 1, (byte)'B', 2, t, 1, (byte)'C', 3, t, 1, (byte)'D', 4}, 1, 2, 3, 4));
+			Assert.IsTrue(TestParseObject(new byte[]{16, t, 1, (byte)'B', 2, t, 1, (byte)'C', 3, t, 1, (byte)'D', 4, t, 1, (byte)'A', 1}, 1, 2, 3, 4));
+			Assert.IsTrue(TestParseObject(new byte[]{12, t, 1, (byte)'B', 2, t, 1, (byte)'C', 3, t, 1, (byte)'D', 4}, 0, 2, 3, 4));
+			Assert.IsTrue(TestParseObject(new byte[]{8, t, 1, (byte)'B', 2, t, 1, (byte)'C', 3}, 0, 2, 3, 0));
+			Assert.IsTrue(TestParseObject(new byte[]{20, t, 1, (byte)'A', 1, t, 1, (byte)'B', 2, t, 1, (byte)'C', 3, t, 1, (byte)'D', 4, t, 1, (byte)'E', 5}, 1, 2, 3, 4));
+			Assert.IsTrue(TestParseObject(new byte[]{20, t, 1, (byte)'E', 5, t, 1, (byte)'A', 1, t, 1, (byte)'B', 2, t, 1, (byte)'C', 3, t, 1, (byte)'D', 4}, 1, 2, 3, 4));
+			Assert.IsTrue(TestParseObject(new byte[] { 16, t, 1, (byte)'D', 4, t, 1, (byte)'C', 3, t, 1, (byte)'B', 2, t, 1, (byte)'A', 1}, 1, 2, 3, 4));
 		}
 	}
 }
diff --git a/Engine/Source/Programs/Shared/EpicGames.Serialization.Tests/JupiterTests.cs b/Engine/Source/Programs/Shared/EpicGames.Serialization.Tests/JupiterTests.cs
--- a/Engine/Source/Programs/Shared/EpicGames.Serialization.Tests/JupiterTests.cs
+++ b/Engine/Source/Programs/Shared/EpicGames.Serialization.Tests/JupiterTests.cs
@@ -4,7 +4,6 @@ using System;
 using System.Collections.Generic;
 using System.IO;
 using System.Linq;
-using System.Numerics;
 using EpicGames.Core;
 using Microsoft.VisualStudio.TestTools.UnitTesting;
 
@@ -16,36 +15,36 @@ namespace EpicGames.Serialization.Tests
 		[TestMethod]
 		public void OverflowTest()
 		{
-			CbWriter Writer = new CbWriter();
+			CbWriter writer = new CbWriter();
 
-			Writer.BeginArray();
-			for (int Idx = 0; Idx < 128; Idx++)
+			writer.BeginArray();
+			for (int idx = 0; idx < 128; idx++)
 			{
-				Writer.WriteNullValue();
+				writer.WriteNullValue();
 			}
-			Writer.EndArray();
+			writer.EndArray();
 
-			byte[] Data = Writer.ToByteArray();
+			writer.ToByteArray();
 		}
 
 		[TestMethod]
 		public void EmbeddedObject()
 		{
-			CbWriter Writer1 = new CbWriter();
-			Writer1.BeginObject();
-			Writer1.WriteInteger("a", 1);
-			Writer1.WriteUtf8String("b", "hello");
-			Writer1.EndObject();
+			CbWriter writer1 = new CbWriter();
+			writer1.BeginObject();
+			writer1.WriteInteger("a", 1);
+			writer1.WriteUtf8String("b", "hello");
+			writer1.EndObject();
 
-			CbObject Object1 = Writer1.ToObject();
+			CbObject object1 = writer1.ToObject();
 
-			CbWriter Writer2 = new CbWriter();
-			Writer2.BeginObject();
-			Writer2.WriteObject("ref", Object1);
-			Writer2.EndObject();
+			CbWriter writer2 = new CbWriter();
+			writer2.BeginObject();
+			writer2.WriteObject("ref", object1);
+			writer2.EndObject();
 
-			CbObject Object2 = Writer2.ToObject();
-			Assert.AreEqual(21, Object2.GetSize());
+			CbObject object2 = writer2.ToObject();
+			Assert.AreEqual(21, object2.GetSize());
 		}
 
 		[TestMethod]
@@ -78,7 +77,6 @@ namespace EpicGames.Serialization.Tests
 			Assert.AreEqual(IoHash.Parse("f855382171a0b1e5a1c653aa6c5121a05cbf4ba0"), inputsFields[0].AsHash());
 		}
 
-
 		[TestMethod]
 		public void ReferenceOutput()
 		{
@@ -100,7 +98,7 @@ namespace EpicGames.Serialization.Tests
 		}
 
 		[TestMethod]
-		public void compact_binary()
+		public void CompactBinary()
 		{
 			byte[] bytes = File.ReadAllBytes("CompactBinaryObjects/compact_binary");
 
@@ -118,8 +116,8 @@ namespace EpicGames.Serialization.Tests
 		[TestMethod]
 		public void WriteArray()
 		{
-			var hash1 = IoHash.Parse("5d8a6dc277c968f0d027c98f879c955c1905c293");
-			var hash2 = IoHash.Parse("313f0d0d334100d83aeb1ee2c42794fd087cb0ae");
+			IoHash hash1 = IoHash.Parse("5d8a6dc277c968f0d027c98f879c955c1905c293");
+			IoHash hash2 = IoHash.Parse("313f0d0d334100d83aeb1ee2c42794fd087cb0ae");
 
 			CbWriter writer = new CbWriter();
 			writer.BeginObject();
@@ -143,11 +141,10 @@ namespace EpicGames.Serialization.Tests
 			CollectionAssert.AreEqual(new IoHash[] { hash1, hash2 }, blobs);
 		}
 
-
 		[TestMethod]
 		public void WriteObject()
 		{
-			var hash1 = IoHash.Parse("5d8a6dc277c968f0d027c98f879c955c1905c293");
+			IoHash hash1 = IoHash.Parse("5d8a6dc277c968f0d027c98f879c955c1905c293");
 
 			CbWriter writer = new CbWriter();
 			writer.BeginObject();
diff --git a/Engine/Source/Programs/Shared/EpicGames.Serialization.Tests/SerializerTests.cs b/Engine/Source/Programs/Shared/EpicGames.Serialization.Tests/SerializerTests.cs
--- a/Engine/Source/Programs/Shared/EpicGames.Serialization.Tests/SerializerTests.cs
+++ b/Engine/Source/Programs/Shared/EpicGames.Serialization.Tests/SerializerTests.cs
@@ -7,7 +7,6 @@ using System.Collections.Generic;
 using System.IO;
 using System.Linq;
 using System.Text;
-using System.Threading.Tasks;
 
 namespace EpicGames.Serialization.Tests
 {
@@ -28,9 +27,9 @@ namespace EpicGames.Serialization.Tests
 			[CbField]
 			public List<TestObject> Children { get; set; } = new List<TestObject>();
 
-			public override bool Equals(object? obj) => obj is TestObject Other && Equals(Other);
+			public override bool Equals(object? obj) => obj is TestObject other && Equals(other);
 
-			public bool Equals(TestObject? Other) => Other != null && A == Other.A && B == Other.B && C == Other.C && Enumerable.SequenceEqual(Children, Other.Children);
+			public bool Equals(TestObject? other) => other != null && A == other.A && B == other.B && C == other.C && Enumerable.SequenceEqual(Children, other.Children);
 
 			public override int GetHashCode() => HashCode.Combine(A, B, C);
 		}
@@ -38,57 +37,57 @@ namespace EpicGames.Serialization.Tests
 		[TestMethod]
 		public void RoundTrip()
 		{
-			TestObject Obj1 = CreateTestObject(new Random(123));
-			CbObject CbObj1 = CbSerializer.Serialize(Obj1);
-			ReadOnlyMemory<byte> Mem1 = CbObj1.GetView();
+			TestObject obj1 = CreateTestObject(new Random(123));
+			CbObject cbObj1 = CbSerializer.Serialize(obj1);
+			ReadOnlyMemory<byte> mem1 = cbObj1.GetView();
 
-			TestObject Obj2 = CbSerializer.Deserialize<TestObject>(Mem1);
-			CbObject CbObj2 = CbSerializer.Serialize(Obj2);
-			ReadOnlyMemory<byte> Mem2 = CbObj1.GetView();
+			TestObject obj2 = CbSerializer.Deserialize<TestObject>(mem1);
+			CbObject cbObj2 = CbSerializer.Serialize(obj2);
+			ReadOnlyMemory<byte> mem2 = cbObj2.GetView();
 
-			Assert.AreEqual(Obj1, Obj2);
-			Assert.AreEqual(Mem1, Mem2);
+			Assert.AreEqual(obj1, obj2);
+			Assert.AreEqual(mem1, mem2);
 		}
 
 		[TestMethod]
 		public void StreamTest()
 		{
-			TestObject Obj1 = CreateTestObject(new Random(123));
+			TestObject obj1 = CreateTestObject(new Random(123));
 
-			CbWriter Writer = new CbWriter();
-			CbSerializer.Serialize(Writer, Obj1);
+			CbWriter writer = new CbWriter();
+			CbSerializer.Serialize(writer, obj1);
 
-			byte[] Data1 = Writer.ToByteArray();
-			byte[] Data2;
-			using (Stream Stream = Writer.AsStream())
+			byte[] data1 = writer.ToByteArray();
+			byte[] data2;
+			using (Stream stream = writer.AsStream())
 			{
-				Data2 = new byte[Stream.Length];
+				data2 = new byte[stream.Length];
 
-				int ReadLength = Stream.Read(Data2, 0, Data2.Length);
-				Assert.AreEqual(Data2.Length, ReadLength);
+				int readLength = stream.Read(data2, 0, data2.Length);
+				Assert.AreEqual(data2.Length, readLength);
 
-				ReadLength = Stream.Read(new byte[10], 0, 10);
-				Assert.AreEqual(0, ReadLength);
+				readLength = stream.Read(new byte[10], 0, 10);
+				Assert.AreEqual(0, readLength);
 			}
-			Assert.IsTrue(Data1.AsSpan().SequenceEqual(Data2));
+			Assert.IsTrue(data1.AsSpan().SequenceEqual(data2));
 
-			TestObject Obj2 = CbSerializer.Deserialize<TestObject>(Data1);
-			Assert.AreEqual(Obj1, Obj2);
+			TestObject obj2 = CbSerializer.Deserialize<TestObject>(data1);
+			Assert.AreEqual(obj1, obj2);
 		}
 
-		static TestObject CreateTestObject(Random Random)
+		static TestObject CreateTestObject(Random random)
 		{
-			TestObject Obj = new TestObject();
-			Obj.A = Random.Next();
-			Obj.B = Random.NextInt64().ToString();
-			Obj.C = IoHash.Compute(Encoding.UTF8.GetBytes(Obj.B));
+			TestObject obj = new TestObject();
+			obj.A = random.Next();
+			obj.B = random.NextInt64().ToString();
+			obj.C = IoHash.Compute(Encoding.UTF8.GetBytes(obj.B));
 
-			for (int Idx = 0; Idx < Random.Next(0, 3); Idx++)
+			for (int idx = 0; idx < random.Next(0, 3); idx++)
 			{
-				Obj.Children.Add(CreateTestObject(Random));
+				obj.Children.Add(CreateTestObject(random));
 			}
 
-			return Obj;
+			return obj;
 		}
 	}
 }
diff --git a/Engine/Source/Programs/Shared/EpicGames.Serialization.Tests/VarIntTests.cs b/Engine/Source/Programs/Shared/EpicGames.Serialization.Tests/VarIntTests.cs
--- a/Engine/Source/Programs/Shared/EpicGames.Serialization.Tests/VarIntTests.cs
+++ b/Engine/Source/Programs/Shared/EpicGames.Serialization.Tests/VarIntTests.cs
@@ -1,11 +1,5 @@
 // Copyright Epic Games, Inc. All Rights Reserved.
 
-using System;
-using System.Collections.Generic;
-using System.IO;
-using System.Linq;
-using System.Numerics;
-using EpicGames.Core;
 using Microsoft.VisualStudio.TestTools.UnitTesting;
 
 namespace EpicGames.Serialization.Tests
@@ -16,17 +10,17 @@ namespace EpicGames.Serialization.Tests
 		[TestMethod]
 		public void TestVarInt()
 		{
-			byte[] Buffer = new byte[20];
+			byte[] buffer = new byte[20];
 
-			int Length = VarInt.Write(Buffer, -1);
-			Assert.AreEqual(9, Length);
+			int length = VarInt.Write(buffer, -1);
+			Assert.AreEqual(9, length);
 			Assert.AreEqual(9, VarInt.Measure(-1));
 
-			Assert.AreEqual(9, VarInt.Measure(Buffer));
-			int Value = (int)(long)VarInt.Read(Buffer, out int BytesRead);
-			Assert.AreEqual(9, BytesRead);
+			Assert.AreEqual(9, VarInt.Measure(buffer));
+			int value = (int)(long)VarInt.Read(buffer, out int bytesRead);
+			Assert.AreEqual(9, bytesRead);
 
-			Assert.AreEqual(-1, Value);
+			Assert.AreEqual(-1, value);
 		}
 	}
 }
diff --git a/Engine/Source/Programs/Shared/EpicGames.Serialization/CbConverter.cs b/Engine/Source/Programs/Shared/EpicGames.Serialization/CbConverter.cs
--- a/Engine/Source/Programs/Shared/EpicGames.Serialization/CbConverter.cs
+++ b/Engine/Source/Programs/Shared/EpicGames.Serialization/CbConverter.cs
@@ -25,10 +25,10 @@ namespace EpicGames.Serialization
 		/// <summary>
 		/// Constructor
 		/// </summary>
-		/// <param name="ConverterType"></param>
-		public CbConverterAttribute(Type ConverterType)
+		/// <param name="converterType"></param>
+		public CbConverterAttribute(Type converterType)
 		{
-			this.ConverterType = ConverterType;
+			this.ConverterType = converterType;
 		}
 	}
 
@@ -40,24 +40,24 @@ namespace EpicGames.Serialization
 		/// <summary>
 		/// Reads an object from a field
 		/// </summary>
-		/// <param name="Field"></param>
+		/// <param name="field"></param>
 		/// <returns></returns>
-		object? ReadObject(CbField Field);
+		object? ReadObject(CbField field);
 
 		/// <summary>
 		/// Writes an object to a field
 		/// </summary>
-		/// <param name="Writer"></param>
-		/// <param name="Value"></param>
-		void WriteObject(CbWriter Writer, object? Value);
+		/// <param name="writer"></param>
+		/// <param name="value"></param>
+		void WriteObject(CbWriter writer, object? value);
 
 		/// <summary>
 		/// Writes an object to a named field, if not equal to the default value
 		/// </summary>
-		/// <param name="Writer"></param>
-		/// <param name="Name"></param>
-		/// <param name="Value"></param>
-		void WriteNamedObject(CbWriter Writer, Utf8String Name, object? Value);
+		/// <param name="writer"></param>
+		/// <param name="name"></param>
+		/// <param name="value"></param>
+		void WriteNamedObject(CbWriter writer, Utf8String name, object? value);
 	}
 
 	/// <summary>
@@ -68,24 +68,24 @@ namespace EpicGames.Serialization
 		/// <summary>
 		/// Reads an object from a field
 		/// </summary>
-		/// <param name="Field"></param>
+		/// <param name="field"></param>
 		/// <returns></returns>
-		T Read(CbField Field);
+		T Read(CbField field);
 
 		/// <summary>
 		/// Writes an object to a field
 		/// </summary>
-		/// <param name="Writer"></param>
-		/// <param name="Value"></param>
-		void Write(CbWriter Writer, T Value);
+		/// <param name="writer"></param>
+		/// <param name="value"></param>
+		void Write(CbWriter writer, T value);
 
 		/// <summary>
 		/// Writes an object to a named field, if not equal to the default value
 		/// </summary>
-		/// <param name="Writer"></param>
-		/// <param name="Name"></param>
-		/// <param name="Value"></param>
-		void WriteNamed(CbWriter Writer, Utf8String Name, T Value);
+		/// <param name="writer"></param>
+		/// <param name="name"></param>
+		/// <param name="value"></param>
+		void WriteNamed(CbWriter writer, Utf8String name, T value);
 	}
 
 	/// <summary>
@@ -116,34 +116,34 @@ namespace EpicGames.Serialization
 	{
 		class CbConverterMethodsWrapper<TConverter, T> : ICbConverterMethods where TConverter : class, ICbConverter<T>
 		{
-			static TConverter StaticConverter = null!;
+			static TConverter s_staticConverter = null!;
 
-			static T Read(CbField Field) => StaticConverter.Read(Field);
-			static void Write(CbWriter Writer, T Value) => StaticConverter.Write(Writer, Value);
-			static void WriteNamed(CbWriter Writer, Utf8String Name, T Value) => StaticConverter.WriteNamed(Writer, Name, Value);
+			static T Read(CbField field) => s_staticConverter.Read(field);
+			static void Write(CbWriter writer, T value) => s_staticConverter.Write(writer, value);
+			static void WriteNamed(CbWriter writer, Utf8String name, T value) => s_staticConverter.WriteNamed(writer, name, value);
 
 			public MethodInfo ReadMethod { get; } = GetMethodInfo(() => Read(null!));
 			public MethodInfo WriteMethod { get; } = GetMethodInfo(() => Write(null!, default!));
 			public MethodInfo WriteNamedMethod { get; } = GetMethodInfo(() => WriteNamed(null!, null!, default!));
 
-			public CbConverterMethodsWrapper(TConverter Converter)
+			public CbConverterMethodsWrapper(TConverter converter)
 			{
-				StaticConverter = Converter;
+				s_staticConverter = converter;
 			}
 
-			static MethodInfo GetMethodInfo(Expression<Action> Expr)
+			static MethodInfo GetMethodInfo(Expression<Action> expr)
 			{
-				return ((MethodCallExpression)Expr.Body).Method;
+				return ((MethodCallExpression)expr.Body).Method;
 			}
 		}
 
-		static Dictionary<PropertyInfo, ICbConverterMethods> PropertyToMethods = new Dictionary<PropertyInfo, ICbConverterMethods>();
-		static Dictionary<Type, ICbConverterMethods> TypeToMethods = new Dictionary<Type, ICbConverterMethods>();
+		static Dictionary<PropertyInfo, ICbConverterMethods> s_propertyToMethods = new Dictionary<PropertyInfo, ICbConverterMethods>();
+		static Dictionary<Type, ICbConverterMethods> s_typeToMethods = new Dictionary<Type, ICbConverterMethods>();
 
-		static ICbConverterMethods CreateWrapper(Type Type, ICbConverter Converter)
+		static ICbConverterMethods CreateWrapper(Type type, ICbConverter converter)
 		{
-			Type ConverterMethodsType = typeof(CbConverterMethodsWrapper<,>).MakeGenericType(Converter.GetType(), Type);
-			return (ICbConverterMethods)Activator.CreateInstance(ConverterMethodsType, new object[] { Converter })!;
+			Type converterMethodsType = typeof(CbConverterMethodsWrapper<,>).MakeGenericType(converter.GetType(), type);
+			return (ICbConverterMethods)Activator.CreateInstance(converterMethodsType, new object[] { converter })!;
 		}
 
 		/// <summary>
@@ -151,33 +151,33 @@ namespace EpicGames.Serialization
 		/// </summary>
 		/// <param name="Type"></param>
 		/// <returns></returns>
-		public static ICbConverterMethods Get(PropertyInfo Property)
+		public static ICbConverterMethods Get(PropertyInfo property)
 		{
-			ICbConverterMethods? Methods;
-			if (!PropertyToMethods.TryGetValue(Property, out Methods))
+			ICbConverterMethods? methods;
+			if (!s_propertyToMethods.TryGetValue(property, out methods))
 			{
-				ICbConverter Converter = CbConverter.GetConverter(Property);
-				Methods = (Converter as ICbConverterMethods) ?? CreateWrapper(Property.PropertyType, Converter);
-				PropertyToMethods.Add(Property, Methods);
+				ICbConverter converter = CbConverter.GetConverter(property);
+				methods = (converter as ICbConverterMethods) ?? CreateWrapper(property.PropertyType, converter);
+				s_propertyToMethods.Add(property, methods);
 			}
-			return Methods;
+			return methods;
 		}
 
 		/// <summary>
 		/// Gets a <see cref="ICbConverterMethods"/> interface for the given type
 		/// </summary>
-		/// <param name="Type"></param>
+		/// <param name="type"></param>
 		/// <returns></returns>
-		public static ICbConverterMethods Get(Type Type)
+		public static ICbConverterMethods Get(Type type)
 		{
-			ICbConverterMethods? Methods;
-			if (!TypeToMethods.TryGetValue(Type, out Methods))
+			ICbConverterMethods? methods;
+			if (!s_typeToMethods.TryGetValue(type, out methods))
 			{
-				ICbConverter Converter = CbConverter.GetConverter(Type);
-				Methods = (Converter as ICbConverterMethods) ?? CreateWrapper(Type, Converter);
-				TypeToMethods.Add(Type, Methods);
+				ICbConverter converter = CbConverter.GetConverter(type);
+				methods = (converter as ICbConverterMethods) ?? CreateWrapper(type, converter);
+				s_typeToMethods.Add(type, methods);
 			}
-			return Methods;
+			return methods;
 		}
 	}
 
@@ -188,22 +188,22 @@ namespace EpicGames.Serialization
 	public abstract class CbConverterBase<T> : ICbConverter<T>
 	{
 		/// <inheritdoc/>
-		public object? ReadObject(CbField Field) => Read(Field);
+		public object? ReadObject(CbField field) => Read(field);
 
 		/// <inheritdoc/>
-		public void WriteObject(CbWriter Writer, object? Value) => Write(Writer, (T)Value!);
+		public void WriteObject(CbWriter writer, object? value) => Write(writer, (T)value!);
 
 		/// <inheritdoc/>
-		public void WriteNamedObject(CbWriter Writer, Utf8String Name, object? Value) => WriteNamed(Writer, Name, (T)Value!);
+		public void WriteNamedObject(CbWriter writer, Utf8String name, object? value) => WriteNamed(writer, name, (T)value!);
 
 		/// <inheritdoc/>
-		public abstract T Read(CbField Field);
+		public abstract T Read(CbField field);
 
 		/// <inheritdoc/>
-		public abstract void Write(CbWriter Writer, T Value);
+		public abstract void Write(CbWriter writer, T value);
 
 		/// <inheritdoc/>
-		public abstract void WriteNamed(CbWriter Writer, Utf8String Name, T Value);
+		public abstract void WriteNamed(CbWriter writer, Utf8String name, T value);
 	}
 
 	/// <summary>
@@ -214,9 +214,9 @@ namespace EpicGames.Serialization
 		/// <summary>
 		/// Create a converter for the given type
 		/// </summary>
-		/// <param name="Type">The type to create a converter for</param>
+		/// <param name="type">The type to create a converter for</param>
 		/// <returns>The converter instance, or null if this factory does not support the given type</returns>
-		public abstract ICbConverter? CreateConverter(Type Type);
+		public abstract ICbConverter? CreateConverter(Type type);
 	}
 
 	/// <summary>
@@ -227,12 +227,12 @@ namespace EpicGames.Serialization
 		/// <summary>
 		/// Object used for locking access to shared objects
 		/// </summary>
-		static object LockObject = new object();
+		static object s_lockObject = new object();
 
 		/// <summary>
 		/// Cache of property to converter
 		/// </summary>
-		static Dictionary<PropertyInfo, ICbConverter> PropertyToConverter = new Dictionary<PropertyInfo, ICbConverter>();
+		static Dictionary<PropertyInfo, ICbConverter> s_propertyToConverter = new Dictionary<PropertyInfo, ICbConverter>();
 
 		/// <summary>
 		/// Cache of type to converter
@@ -287,8 +287,8 @@ namespace EpicGames.Serialization
 			/// <returns></returns>
 			static ICbConverter<T> CreateConverter()
 			{
-				ICbConverter Converter = GetConverter(typeof(T));
-				return (Converter as ICbConverter<T>) ?? new CbConverterWrapper(Converter);
+				ICbConverter converter = GetConverter(typeof(T));
+				return (converter as ICbConverter<T>) ?? new CbConverterWrapper(converter);
 			}
 
 			/// <summary>
@@ -296,84 +296,84 @@ namespace EpicGames.Serialization
 			/// </summary>
 			class CbConverterWrapper : CbConverterBase<T>
 			{
-				ICbConverter Inner;
+				ICbConverter _inner;
 
-				public CbConverterWrapper(ICbConverter Inner) => this.Inner = Inner;
+				public CbConverterWrapper(ICbConverter inner) => this._inner = inner;
 
 				/// <inheritdoc/>
-				public override T Read(CbField Field) => (T)Inner.ReadObject(Field)!;
+				public override T Read(CbField field) => (T)_inner.ReadObject(field)!;
 
 				/// <inheritdoc/>
-				public override void Write(CbWriter Writer, T Value) => Inner.WriteObject(Writer, Value);
+				public override void Write(CbWriter writer, T value) => _inner.WriteObject(writer, value);
 
 				/// <inheritdoc/>
-				public override void WriteNamed(CbWriter Writer, Utf8String Name, T Value) => Inner.WriteNamedObject(Writer, Name, Value);
+				public override void WriteNamed(CbWriter writer, Utf8String name, T value) => _inner.WriteNamedObject(writer, name, value);
 			}
 		}
 
 		/// <summary>
 		/// Gets the converter for a particular property
 		/// </summary>
-		/// <param name="Property"></param>
+		/// <param name="property"></param>
 		/// <returns></returns>
-		public static ICbConverter GetConverter(PropertyInfo Property)
+		public static ICbConverter GetConverter(PropertyInfo property)
 		{
-			CbConverterAttribute? ConverterAttribute = Property.GetCustomAttribute<CbConverterAttribute>();
-			if (ConverterAttribute != null)
+			CbConverterAttribute? converterAttribute = property.GetCustomAttribute<CbConverterAttribute>();
+			if (converterAttribute != null)
 			{
-				Type ConverterType = ConverterAttribute.ConverterType;
-				lock (LockObject)
+				Type converterType = converterAttribute.ConverterType;
+				lock (s_lockObject)
 				{
-					ICbConverter? Converter;
-					if (!PropertyToConverter.TryGetValue(Property, out Converter))
+					ICbConverter? converter;
+					if (!s_propertyToConverter.TryGetValue(property, out converter))
 					{
-						Converter = (ICbConverter?)Activator.CreateInstance(ConverterType)!;
-						PropertyToConverter.Add(Property, Converter);
+						converter = (ICbConverter?)Activator.CreateInstance(converterType)!;
+						s_propertyToConverter.Add(property, converter);
 					}
-					return Converter;
+					return converter;
 				}
 			}
-			return GetConverter(Property.PropertyType);
+			return GetConverter(property.PropertyType);
 		}
 
 		/// <summary>
 		/// Gets the converter for a particular type
 		/// </summary>
-		/// <param name="Type"></param>
+		/// <param name="type"></param>
 		/// <returns></returns>
-		public static ICbConverter GetConverter(Type Type)
+		public static ICbConverter GetConverter(Type type)
 		{
-			ICbConverter? Converter;
-			lock (LockObject)
+			ICbConverter? converter;
+			lock (s_lockObject)
 			{
-				if (!TypeToConverter.TryGetValue(Type, out Converter))
+				if (!TypeToConverter.TryGetValue(type, out converter))
 				{
-					CbConverterAttribute? ConverterAttribute = Type.GetCustomAttribute<CbConverterAttribute>();
-					if (ConverterAttribute != null)
+					CbConverterAttribute? converterAttribute = type.GetCustomAttribute<CbConverterAttribute>();
+					if (converterAttribute != null)
 					{
-						Type ConverterType = ConverterAttribute.ConverterType;
-						if (Type.IsGenericType && ConverterType.IsGenericTypeDefinition)
+						Type converterType = converterAttribute.ConverterType;
+						if (type.IsGenericType && converterType.IsGenericTypeDefinition)
 						{
-							ConverterType = ConverterType.MakeGenericType(Type.GetGenericArguments());
+							converterType = converterType.MakeGenericType(type.GetGenericArguments());
 						}
-						Converter = (ICbConverter?)Activator.CreateInstance(ConverterType)!;
+						converter = (ICbConverter?)Activator.CreateInstance(converterType)!;
 					}
 					else
 					{
-						for (int Idx = ConverterFactories.Count - 1; Idx >= 0 && Converter == null; Idx--)
+						for (int idx = ConverterFactories.Count - 1; idx >= 0 && converter == null; idx--)
 						{
-							Converter = ConverterFactories[Idx].CreateConverter(Type);
+							converter = ConverterFactories[idx].CreateConverter(type);
 						}
 
-						if (Converter == null)
+						if (converter == null)
 						{
-							throw new CbException($"Unable to create converter for {Type.Name}");
+							throw new CbException($"Unable to create converter for {type.Name}");
 						}
 					}
-					TypeToConverter.Add(Type, Converter!);
+					TypeToConverter.Add(type, converter!);
 				}
 			}
-			return Converter;
+			return converter;
 		}
 
 		/// <summary>
diff --git a/Engine/Source/Programs/Shared/EpicGames.Serialization/CbField.cs b/Engine/Source/Programs/Shared/EpicGames.Serialization/CbField.cs
--- a/Engine/Source/Programs/Shared/EpicGames.Serialization/CbField.cs
+++ b/Engine/Source/Programs/Shared/EpicGames.Serialization/CbField.cs
@@ -10,7 +10,6 @@ using System.ComponentModel;
 using System.Diagnostics;
 using System.Globalization;
 using System.Linq;
-using System.Numerics;
 using System.Text;
 using System.Text.Json;
 using System.Text.Json.Serialization;
@@ -194,7 +193,7 @@ namespace EpicGames.Serialization
 	/// <summary>
 	/// A binary attachment, referenced by <see cref="IoHash"/>
 	/// </summary>
-	[DebuggerDisplay("{Hash}")]
+	[DebuggerDisplay("{_hash}")]
 	[JsonConverter(typeof(CbBinaryAttachmentJsonConverter))]
 	[TypeConverter(typeof(CbBinaryAttachmentTypeConverter))]
 	public struct CbBinaryAttachment
@@ -207,43 +206,43 @@ namespace EpicGames.Serialization
 		/// <summary>
 		/// Hash of the referenced object
 		/// </summary>
-		public IoHash Hash;
+		public IoHash Hash { get; }
 
 		/// <summary>
 		/// Constructor
 		/// </summary>
-		/// <param name="Hash">Hash of the referenced object</param>
-		public CbBinaryAttachment(IoHash Hash)
+		/// <param name="hash">Hash of the referenced object</param>
+		public CbBinaryAttachment(IoHash hash)
 		{
-			this.Hash = Hash;
+			Hash = hash;
 		}
 
 		/// <inheritdoc/>
 		public override string ToString() => Hash.ToString();
 
 		/// <inheritdoc/>
-		public override bool Equals(object? Obj) => Obj is CbBinaryAttachment Other && Other.Hash == Hash;
+		public override bool Equals(object? obj) => obj is CbBinaryAttachment other && other.Hash == Hash;
 
 		/// <inheritdoc/>
 		public override int GetHashCode() => Hash.GetHashCode();
 
 		/// <inheritdoc/>
-		public static bool operator ==(CbBinaryAttachment Lhs, CbBinaryAttachment Rhs) => Lhs.Hash == Rhs.Hash;
+		public static bool operator ==(CbBinaryAttachment lhs, CbBinaryAttachment rhs) => lhs.Hash == rhs.Hash;
 
 		/// <inheritdoc/>
-		public static bool operator !=(CbBinaryAttachment Lhs, CbBinaryAttachment Rhs) => Lhs.Hash != Rhs.Hash;
+		public static bool operator !=(CbBinaryAttachment lhs, CbBinaryAttachment rhs) => lhs.Hash != rhs.Hash;
 
 		/// <summary>
 		/// Convert a hash to a binary attachment 
 		/// </summary>
-		/// <param name="Hash">The attachment to convert</param>
-		public static implicit operator CbBinaryAttachment(IoHash Hash) => new CbBinaryAttachment(Hash);
+		/// <param name="hash">The attachment to convert</param>
+		public static implicit operator CbBinaryAttachment(IoHash hash) => new CbBinaryAttachment(hash);
 
 		/// <summary>
 		/// Use a binary attachment as a hash
 		/// </summary>
-		/// <param name="Attachment">The attachment to convert</param>
-		public static implicit operator IoHash(CbBinaryAttachment Attachment) => Attachment.Hash;
+		/// <param name="attachment">The attachment to convert</param>
+		public static implicit operator IoHash(CbBinaryAttachment attachment) => attachment.Hash;
 	}
 
 	/// <summary>
@@ -252,10 +251,10 @@ namespace EpicGames.Serialization
 	sealed class CbBinaryAttachmentJsonConverter : JsonConverter<CbBinaryAttachment>
 	{
 		/// <inheritdoc/>
-		public override CbBinaryAttachment Read(ref Utf8JsonReader Reader, Type TypeToConvert, JsonSerializerOptions Options) => IoHash.Parse(Reader.ValueSpan);
+		public override CbBinaryAttachment Read(ref Utf8JsonReader reader, Type typeToConvert, JsonSerializerOptions options) => IoHash.Parse(reader.ValueSpan);
 
 		/// <inheritdoc/>
-		public override void Write(Utf8JsonWriter Writer, CbBinaryAttachment Value, JsonSerializerOptions Options) => Writer.WriteStringValue(Value.Hash.ToUtf8String().Span);
+		public override void Write(Utf8JsonWriter writer, CbBinaryAttachment value, JsonSerializerOptions options) => writer.WriteStringValue(value.Hash.ToUtf8String().Span);
 	}
 
 	/// <summary>
@@ -264,16 +263,16 @@ namespace EpicGames.Serialization
 	sealed class CbBinaryAttachmentTypeConverter : TypeConverter
 	{
 		/// <inheritdoc/>
-		public override bool CanConvertFrom(ITypeDescriptorContext Context, Type SourceType) => SourceType == typeof(string);
+		public override bool CanConvertFrom(ITypeDescriptorContext context, Type sourceType) => sourceType == typeof(string);
 
 		/// <inheritdoc/>
-		public override object ConvertFrom(ITypeDescriptorContext Context, CultureInfo Culture, object Value) => new CbBinaryAttachment(IoHash.Parse((string)Value));
+		public override object ConvertFrom(ITypeDescriptorContext context, CultureInfo culture, object value) => new CbBinaryAttachment(IoHash.Parse((string)value));
 	}
 
 	/// <summary>
 	/// An object attachment, referenced by <see cref="IoHash"/>
 	/// </summary>
-	[DebuggerDisplay("{Hash}")]
+	[DebuggerDisplay("{_hash}")]
 	[JsonConverter(typeof(CbObjectAttachmentJsonConverter))]
 	[TypeConverter(typeof(CbObjectAttachmentTypeConverter))]
 	public struct CbObjectAttachment
@@ -286,43 +285,43 @@ namespace EpicGames.Serialization
 		/// <summary>
 		/// Hash of the referenced object
 		/// </summary>
-		public IoHash Hash;
+		public IoHash Hash { get; }
 
 		/// <summary>
 		/// Constructor
 		/// </summary>
-		/// <param name="Hash">Hash of the referenced object</param>
-		public CbObjectAttachment(IoHash Hash)
+		/// <param name="hash">Hash of the referenced object</param>
+		public CbObjectAttachment(IoHash hash)
 		{
-			this.Hash = Hash;
+			Hash = hash;
 		}
 
 		/// <inheritdoc/>
 		public override string ToString() => Hash.ToString();
 
 		/// <inheritdoc/>
-		public override bool Equals(object? Obj) => Obj is CbObjectAttachment Other && Other.Hash == Hash;
+		public override bool Equals(object? obj) => obj is CbObjectAttachment other && other.Hash == Hash;
 
 		/// <inheritdoc/>
 		public override int GetHashCode() => Hash.GetHashCode();
 
 		/// <inheritdoc/>
-		public static bool operator ==(CbObjectAttachment Lhs, CbObjectAttachment Rhs) => Lhs.Hash == Rhs.Hash;
+		public static bool operator ==(CbObjectAttachment lhs, CbObjectAttachment rhs) => lhs.Hash == rhs.Hash;
 
 		/// <inheritdoc/>
-		public static bool operator !=(CbObjectAttachment Lhs, CbObjectAttachment Rhs) => Lhs.Hash != Rhs.Hash;
+		public static bool operator !=(CbObjectAttachment lhs, CbObjectAttachment rhs) => lhs.Hash != rhs.Hash;
 
 		/// <summary>
 		/// Use an object attachment as a hash
 		/// </summary>
 		/// <param name="Attachment">The attachment to convert</param>
-		public static implicit operator CbObjectAttachment(IoHash Hash) => new CbObjectAttachment(Hash);
+		public static implicit operator CbObjectAttachment(IoHash hash) => new CbObjectAttachment(hash);
 
 		/// <summary>
 		/// Use an object attachment as a hash
 		/// </summary>
-		/// <param name="Attachment">The attachment to convert</param>
-		public static implicit operator IoHash(CbObjectAttachment Attachment) => Attachment.Hash;
+		/// <param name="attachment">The attachment to convert</param>
+		public static implicit operator IoHash(CbObjectAttachment attachment) => attachment.Hash;
 	}
 
 	/// <summary>
@@ -331,10 +330,10 @@ namespace EpicGames.Serialization
 	sealed class CbObjectAttachmentJsonConverter : JsonConverter<CbObjectAttachment>
 	{
 		/// <inheritdoc/>
-		public override CbObjectAttachment Read(ref Utf8JsonReader Reader, Type TypeToConvert, JsonSerializerOptions Options) => IoHash.Parse(Reader.ValueSpan);
+		public override CbObjectAttachment Read(ref Utf8JsonReader reader, Type typeToConvert, JsonSerializerOptions options) => IoHash.Parse(reader.ValueSpan);
 
 		/// <inheritdoc/>
-		public override void Write(Utf8JsonWriter Writer, CbObjectAttachment Value, JsonSerializerOptions Options) => Writer.WriteStringValue(Value.Hash.ToUtf8String().Span);
+		public override void Write(Utf8JsonWriter writer, CbObjectAttachment value, JsonSerializerOptions options) => writer.WriteStringValue(value.Hash.ToUtf8String().Span);
 	}
 
 	/// <summary>
@@ -343,10 +342,10 @@ namespace EpicGames.Serialization
 	sealed class CbObjectAttachmentTypeConverter : TypeConverter
 	{
 		/// <inheritdoc/>
-		public override bool CanConvertFrom(ITypeDescriptorContext Context, Type SourceType) => SourceType == typeof(string);
+		public override bool CanConvertFrom(ITypeDescriptorContext context, Type sourceType) => sourceType == typeof(string);
 
 		/// <inheritdoc/>
-		public override object ConvertFrom(ITypeDescriptorContext Context, CultureInfo Culture, object Value) => new CbObjectAttachment(IoHash.Parse((string)Value));
+		public override object ConvertFrom(ITypeDescriptorContext context, CultureInfo culture, object value) => new CbObjectAttachment(IoHash.Parse((string)value));
 	}
 
 	/// <summary>
@@ -378,81 +377,81 @@ namespace EpicGames.Serialization
 		/// <summary>
 		/// Removes flags from the given type
 		/// </summary>
-		/// <param name="Type">Type to check</param>
+		/// <param name="type">Type to check</param>
 		/// <returns>Type without flag fields</returns>
-		public static CbFieldType GetType(CbFieldType Type)
+		public static CbFieldType GetType(CbFieldType type)
 		{
-			return Type & TypeMask;
+			return type & TypeMask;
 		}
 
 		/// <summary>
 		/// Gets the serialized type
 		/// </summary>
-		/// <param name="Type">Type to check</param>
+		/// <param name="type">Type to check</param>
 		/// <returns>Type without flag fields</returns>
-		public static CbFieldType GetSerializedType(CbFieldType Type)
+		public static CbFieldType GetSerializedType(CbFieldType type)
 		{
-			return Type & SerializedTypeMask;
+			return type & SerializedTypeMask;
 		}
 
 		/// <summary>
 		/// Tests if the given field has a type
 		/// </summary>
-		/// <param name="Type">Type to check</param>
+		/// <param name="type">Type to check</param>
 		/// <returns>True if the field has a type</returns>
-		public static bool HasFieldType(CbFieldType Type)
+		public static bool HasFieldType(CbFieldType type)
 		{
-			return (Type & CbFieldType.HasFieldType) != 0;
+			return (type & CbFieldType.HasFieldType) != 0;
 		}
 
 		/// <summary>
 		/// Tests if the given field has a name
 		/// </summary>
-		/// <param name="Type">Type to check</param>
+		/// <param name="type">Type to check</param>
 		/// <returns>True if the field has a name</returns>
-		public static bool HasFieldName(CbFieldType Type)
+		public static bool HasFieldName(CbFieldType type)
 		{
-			return (Type & CbFieldType.HasFieldName) != 0;
+			return (type & CbFieldType.HasFieldName) != 0;
 		}
 
 		/// <summary>
 		/// Tests if the given field type is none
 		/// </summary>
-		/// <param name="Type">Type to check</param>
+		/// <param name="type">Type to check</param>
 		/// <returns>True if the field is none</returns>
-		public static bool IsNone(CbFieldType Type)
+		public static bool IsNone(CbFieldType type)
 		{
-			return GetType(Type) == CbFieldType.None;
+			return GetType(type) == CbFieldType.None;
 		}
 
 		/// <summary>
 		/// Tests if the given field type is a null value
 		/// </summary>
-		/// <param name="Type">Type to check</param>
+		/// <param name="type">Type to check</param>
 		/// <returns>True if the field is a null</returns>
-		public static bool IsNull(CbFieldType Type)
+		public static bool IsNull(CbFieldType type)
 		{
-			return GetType(Type) == CbFieldType.Null;
+			return GetType(type) == CbFieldType.Null;
 		}
 
 		/// <summary>
 		/// Tests if the given field type is an object
 		/// </summary>
-		/// <param name="Type">Type to check</param>
+		/// <param name="type">Type to check</param>
 		/// <returns>True if the field is an object type</returns>
-		public static bool IsObject(CbFieldType Type)
+		public static bool IsObject(CbFieldType type)
 		{
-			return (Type & ObjectMask) == ObjectBase;
+			return (type & ObjectMask) == ObjectBase;
 		}
 
 		/// <summary>
 		/// Tests if the given field type is an array
 		/// </summary>
-		/// <param name="Type">Type to check</param>
+		/// <param name="type">Type to check</param>
 		/// <returns>True if the field is an array type</returns>
-		public static bool IsArray(CbFieldType Type)
+		public static bool IsArray(CbFieldType type)
 		{
-			return (Type & ArrayMask) == ArrayBase;
+			return (type & ArrayMask) == ArrayBase;
 		}
 
 		/// <summary>
@@ -468,41 +467,41 @@ namespace EpicGames.Serialization
 		/// <summary>
 		/// Tests if the given field type is a string
 		/// </summary>
-		/// <param name="Type">Type to check</param>
+		/// <param name="type">Type to check</param>
 		/// <returns>True if the field is an array type</returns>
-		public static bool IsString(CbFieldType Type)
+		public static bool IsString(CbFieldType type)
 		{
-			return GetType(Type) == CbFieldType.String;
+			return GetType(type) == CbFieldType.String;
 		}
 
 		/// <summary>
 		/// Tests if the given field type is an integer
 		/// </summary>
-		/// <param name="Type">Type to check</param>
+		/// <param name="type">Type to check</param>
 		/// <returns>True if the field is an integer type</returns>
-		public static bool IsInteger(CbFieldType Type)
+		public static bool IsInteger(CbFieldType type)
 		{
-			return (Type & IntegerMask) == IntegerBase;
+			return (type & IntegerMask) == IntegerBase;
 		}
 
 		/// <summary>
 		/// Tests if the given field type is a float (or integer, due to implicit conversion)
 		/// </summary>
-		/// <param name="Type">Type to check</param>
+		/// <param name="type">Type to check</param>
 		/// <returns>True if the field is a float type</returns>
-		public static bool IsFloat(CbFieldType Type)
+		public static bool IsFloat(CbFieldType type)
 		{
-			return (Type & FloatMask) == FloatBase;
+			return (type & FloatMask) == FloatBase;
 		}
 
 		/// <summary>
 		/// Tests if the given field type is a boolean
 		/// </summary>
-		/// <param name="Type">Type to check</param>
+		/// <param name="type">Type to check</param>
 		/// <returns>True if the field is an bool type</returns>
-		public static bool IsBool(CbFieldType Type)
+		public static bool IsBool(CbFieldType type)
 		{
-			return (Type & BoolMask) == BoolBase;
+			return (type & BoolMask) == BoolBase;
 		}
 
 		/// <summary>
@@ -528,21 +527,21 @@ namespace EpicGames.Serialization
 		/// <summary>
 		/// Tests if the given field type is an attachment
 		/// </summary>
-		/// <param name="Type">Type to check</param>
+		/// <param name="type">Type to check</param>
 		/// <returns>True if the field is an attachment type</returns>
-		public static bool IsAttachment(CbFieldType Type)
+		public static bool IsAttachment(CbFieldType type)
 		{
-			return (Type & AttachmentMask) == AttachmentBase;
+			return (type & AttachmentMask) == AttachmentBase;
 		}
 
 		/// <summary>
 		/// Tests if the given field type is a hash
 		/// </summary>
-		/// <param name="Type">Type to check</param>
+		/// <param name="type">Type to check</param>
 		/// <returns>True if the field is a hash</returns>
-		public static bool IsHash(CbFieldType Type)
+		public static bool IsHash(CbFieldType type)
 		{
-			return GetType(Type) == CbFieldType.Hash || IsAttachment(Type);
+			return GetType(type) == CbFieldType.Hash || IsAttachment(type);
 		}
 
 		/// <summary>
@@ -585,35 +584,34 @@ namespace EpicGames.Serialization
 			return GetType(type) == CbFieldType.ObjectId;
 		}
 
-
 		/// <summary>
 		/// Tests if the given field type has fields
 		/// </summary>
-		/// <param name="Type">Type to check</param>
+		/// <param name="type">Type to check</param>
 		/// <returns>True if the field has fields</returns>
-		public static bool HasFields(CbFieldType Type)
+		public static bool HasFields(CbFieldType type)
 		{
-			CbFieldType NoFlags = GetType(Type);
-			return NoFlags >= CbFieldType.Object && NoFlags <= CbFieldType.UniformArray;
+			CbFieldType noFlags = GetType(type);
+			return noFlags >= CbFieldType.Object && noFlags <= CbFieldType.UniformArray;
 		}
 
 		/// <summary>
 		/// Tests if the given field type has uniform fields (array/object)
 		/// </summary>
-		/// <param name="Type">Type to check</param>
+		/// <param name="type">Type to check</param>
 		/// <returns>True if the field has uniform fields</returns>
-		public static bool HasUniformFields(CbFieldType Type)
+		public static bool HasUniformFields(CbFieldType type)
 		{
-			CbFieldType LocalType = GetType(Type);
-			return LocalType == CbFieldType.UniformObject || LocalType == CbFieldType.UniformArray;
+			CbFieldType localType = GetType(type);
+			return localType == CbFieldType.UniformObject || localType == CbFieldType.UniformArray;
 		}
 
 		/// <summary>
 		/// Tests if the type is or may contain fields of any attachment type.
 		/// </summary>
-		public static bool MayContainAttachments(CbFieldType Type)
+		public static bool MayContainAttachments(CbFieldType type)
 		{
-			return IsObject(Type) | IsArray(Type) | IsAttachment(Type);
+			return IsObject(type) | IsArray(type) | IsAttachment(type);
 		}
 	}
 
@@ -643,8 +641,8 @@ namespace EpicGames.Serialization
 	/// </summary>
 	class CbFieldWithNameDebugView
 	{
-		public string? Name;
-		public object? Value;
+		public string? Name { get; set; }
+		public object? Value { get; set; }
 	}
 
 	/// <summary>
@@ -652,9 +650,9 @@ namespace EpicGames.Serialization
 	/// </summary>
 	class CbFieldDebugView
 	{
-		public CbFieldDebugView(CbField Field) => Value = Field.HasName()
-				? new CbFieldWithNameDebugView { Name = Field.Name.ToString(), Value = Field.Value } 
-				: Field.Value;
+		public CbFieldDebugView(CbField field) => Value = field.HasName()
+				? new CbFieldWithNameDebugView { Name = field.Name.ToString(), Value = field.Value } 
+				: field.Value;
 
 		[DebuggerBrowsable(DebuggerBrowsableState.RootHidden)]
 		public object? Value { get; }
@@ -718,12 +716,12 @@ namespace EpicGames.Serialization
 		/// <summary>
 		/// Offset of the name with the memory
 		/// </summary>
-		public int NameLen;
+		public int _nameLen;
 
 		/// <summary>
 		/// Offset of the payload within the memory
 		/// </summary>
-		public int PayloadOffset;
+		public int _payloadOffset;
 
 		/// <summary>
 		/// Error for parsing the current field type
@@ -741,14 +739,14 @@ namespace EpicGames.Serialization
 		/// <summary>
 		/// Copy constructor
 		/// </summary>
-		/// <param name="Other"></param>
-		public CbField(CbField Other)
+		/// <param name="other"></param>
+		public CbField(CbField other)
 		{
-			this.TypeWithFlags = Other.TypeWithFlags;
-			this.Memory = Other.Memory;
-			this.NameLen = Other.NameLen;
-			this.PayloadOffset = Other.PayloadOffset;
-			this.Error = Other.Error;
+			TypeWithFlags = other.TypeWithFlags;
+			Memory = other.Memory;
+			_nameLen = other._nameLen;
+			_payloadOffset = other._payloadOffset;
+			Error = other.Error;
 		}
 
 		/// <summary>
@@ -756,32 +754,32 @@ namespace EpicGames.Serialization
 		/// </summary>
 		/// <param>Data Pointer to the start of the field data.</param>
 		/// <param>Type HasFieldType means that Data contains the type. Otherwise, use the given type.</param>
-		public CbField(ReadOnlyMemory<byte> Data, CbFieldType Type = CbFieldType.HasFieldType)
+		public CbField(ReadOnlyMemory<byte> data, CbFieldType type = CbFieldType.HasFieldType)
 		{
-			int Offset = 0;
-			if (CbFieldUtils.HasFieldType(Type))
+			int offset = 0;
+			if (CbFieldUtils.HasFieldType(type))
 			{
-				Type = (CbFieldType)Data.Span[Offset] | CbFieldType.HasFieldType;
-				Offset++;
+				type = (CbFieldType)data.Span[offset] | CbFieldType.HasFieldType;
+				offset++;
 			}
 
-			if (CbFieldUtils.HasFieldName(Type))
+			if (CbFieldUtils.HasFieldName(type))
 			{
-				NameLen = (int)VarInt.Read(Data.Slice(Offset).Span, out int NameLenByteCount);
-				Offset += NameLenByteCount + NameLen;
+				_nameLen = (int)VarInt.Read(data.Slice(offset).Span, out int nameLenByteCount);
+				offset += nameLenByteCount + _nameLen;
 			}
 
-			this.Memory = Data;
-			this.TypeWithFlags = Type;
-			this.PayloadOffset = Offset;
+			Memory = data;
+			TypeWithFlags = type;
+			_payloadOffset = offset;
 
-			Memory = Memory.Slice(0, (int)Math.Min((ulong)Memory.Length, (ulong)PayloadOffset + GetPayloadSize()));
+			Memory = Memory.Slice(0, (int)Math.Min((ulong)Memory.Length, (ulong)_payloadOffset + GetPayloadSize()));
 		}
 
 		/// <summary>
 		/// Returns the name of the field if it has a name, otherwise an empty view.
 		/// </summary>
-		public Utf8String Name => new Utf8String(Memory.Slice(PayloadOffset - NameLen, NameLen));
+		public Utf8String Name => new Utf8String(Memory.Slice(_payloadOffset - _nameLen, _nameLen));
 
 		/// <summary>
 		/// Gets the value of this field
@@ -790,8 +788,8 @@ namespace EpicGames.Serialization
 		{
 			get
 			{
-				CbFieldType FieldType = CbFieldUtils.GetType(TypeWithFlags);
-				switch (FieldType)
+				CbFieldType fieldType = CbFieldUtils.GetType(TypeWithFlags);
+				switch (fieldType)
 				{
 					case CbFieldType.None:
 						return None;
@@ -834,7 +832,7 @@ namespace EpicGames.Serialization
 					case CbFieldType.ObjectId:
 						return AsObjectId();
 					default:
-						throw new NotImplementedException($"Unknown field type ({FieldType})");
+						throw new NotImplementedException($"Unknown field type ({fieldType})");
 				}
 			}
 		}
@@ -891,19 +889,19 @@ namespace EpicGames.Serialization
 		/// Access the field as binary data.
 		/// </summary>
 		/// <returns></returns>
-		public ReadOnlyMemory<byte> AsBinary(ReadOnlyMemory<byte> Default)
+		public ReadOnlyMemory<byte> AsBinary(ReadOnlyMemory<byte> defaultValue)
 		{
 			if (CbFieldUtils.IsBinary(TypeWithFlags))
 			{
 				Error = CbFieldError.None;
 
-				ulong Length = VarInt.Read(Payload.Span, out int BytesRead);
-				return Payload.Slice(BytesRead, (int)Length);
+				ulong length = VarInt.Read(Payload.Span, out int bytesRead);
+				return Payload.Slice(bytesRead, (int)length);
 			}
 			else
 			{
 				Error = CbFieldError.TypeError;
-				return Default;
+				return defaultValue;
 			}
 		}
 
@@ -920,9 +918,9 @@ namespace EpicGames.Serialization
 		/// Access the field as binary data.
 		/// </summary>
 		/// <returns></returns>
-		public byte[] AsBinaryArray(byte[] Default)
+		public byte[] AsBinaryArray(byte[] defaultValue)
 		{
-			return AsBinary(Default).ToArray();
+			return AsBinary(defaultValue).ToArray();
 		}
 
 		/// <summary>
@@ -935,7 +933,7 @@ namespace EpicGames.Serialization
 		/// Access the field as a UTF-8 string.
 		/// </summary>
 		/// <returns></returns>
-		public string AsString(string Default) => AsUtf8String(Default).ToString();
+		public string AsString(string defaultValue) => AsUtf8String(defaultValue).ToString();
 
 		/// <summary>
 		/// Access the field as a UTF-8 string.
@@ -949,45 +947,45 @@ namespace EpicGames.Serialization
 		/// <summary>
 		/// Access the field as a UTF-8 string. Returns the provided default on error.
 		/// </summary>
-		/// <param name="Default">Default value to return</param>
+		/// <param name="defaultValue">Default value to return</param>
 		/// <returns></returns>
-		public Utf8String AsUtf8String(Utf8String Default)
+		public Utf8String AsUtf8String(Utf8String defaultValue)
 		{
 			if (CbFieldUtils.IsString(TypeWithFlags))
 			{
-				ulong ValueSize = VarInt.Read(Payload.Span, out int ValueSizeByteCount);
-				if (ValueSize >= (1UL << 31))
+				ulong valueSize = VarInt.Read(Payload.Span, out int valueSizeByteCount);
+				if (valueSize >= (1UL << 31))
 				{
 					Error = CbFieldError.RangeError;
-					return Default;
+					return defaultValue;
 				}
 				else
 				{
 					Error = CbFieldError.None;
-					return new Utf8String(Payload.Slice(ValueSizeByteCount, (int)ValueSize));
+					return new Utf8String(Payload.Slice(valueSizeByteCount, (int)valueSize));
 				}
 			}
 			else
 			{
 				Error = CbFieldError.TypeError;
-				return Default;
+				return defaultValue;
 			}
 		}
 
 		/// <summary>
 		/// Access the field as an int8. Returns the provided default on error.
 		/// </summary>
-		public sbyte AsInt8(sbyte Default = 0)
+		public sbyte AsInt8(sbyte defaultValue = 0)
 		{
-			return (sbyte)AsInteger((ulong)Default, 7, true);
+			return (sbyte)AsInteger((ulong)defaultValue, 7, true);
 		}
 
 		/// <summary>
 		/// Access the field as an int16. Returns the provided default on error.
 		/// </summary>
-		public short AsInt16(short Default = 0)
+		public short AsInt16(short defaultValue = 0)
 		{
-			return (short)AsInteger((ulong)Default, 15, true);
+			return (short)AsInteger((ulong)defaultValue, 15, true);
 		}
 
 		/// <summary>
@@ -1001,9 +999,9 @@ namespace EpicGames.Serialization
 		/// <summary>
 		/// Access the field as an int32. Returns the provided default on error.
 		/// </summary>
-		public int AsInt32(int Default)
+		public int AsInt32(int defaultValue)
 		{
-			return (int)AsInteger((ulong)Default, 31, true);
+			return (int)AsInteger((ulong)defaultValue, 31, true);
 		}
 
 		/// <summary>
@@ -1017,25 +1015,25 @@ namespace EpicGames.Serialization
 		/// <summary>
 		/// Access the field as an int64. Returns the provided default on error.
 		/// </summary>
-		public long AsInt64(long Default)
+		public long AsInt64(long defaultValue)
 		{
-			return (long)AsInteger((ulong)Default, 63, true);
+			return (long)AsInteger((ulong)defaultValue, 63, true);
 		}
 
 		/// <summary>
 		/// Access the field as an int8. Returns the provided default on error.
 		/// </summary>
-		public byte AsUInt8(byte Default = 0)
+		public byte AsUInt8(byte defaultValue = 0)
 		{
-			return (byte)AsInteger(Default, 8, false);
+			return (byte)AsInteger(defaultValue, 8, false);
 		}
 
 		/// <summary>
 		/// Access the field as an int16. Returns the provided default on error.
 		/// </summary>
-		public ushort AsUInt16(ushort Default = 0)
+		public ushort AsUInt16(ushort defaultValue = 0)
 		{
-			return (ushort)AsInteger(Default, 16, false);
+			return (ushort)AsInteger(defaultValue, 16, false);
 		}
 
 		/// <summary>
@@ -1049,9 +1047,9 @@ namespace EpicGames.Serialization
 		/// <summary>
 		/// Access the field as an int32. Returns the provided default on error.
 		/// </summary>
-		public uint AsUInt32(uint Default)
+		public uint AsUInt32(uint defaultValue)
 		{
-			return (uint)AsInteger(Default, 32, false);
+			return (uint)AsInteger(defaultValue, 32, false);
 		}
 
 		/// <summary>
@@ -1065,78 +1063,76 @@ namespace EpicGames.Serialization
 		/// <summary>
 		/// Access the field as an int64. Returns the provided default on error.
 		/// </summary>
-		public ulong AsUInt64(ulong Default)
+		public ulong AsUInt64(ulong defaultValue)
 		{
-			return (ulong)AsInteger(Default, 64, false);
+			return (ulong)AsInteger(defaultValue, 64, false);
 		}
 
 		/// <summary>
 		/// Access the field as an integer, checking that it's in the correct range
 		/// </summary>
-		/// <param name="Default"></param>
-		/// <param name="MagnitudeBits"></param>
-		/// <param name="IsSigned"></param>
+		/// <param name="default"></param>
+		/// <param name="magnitudeBits"></param>
+		/// <param name="isSigned"></param>
 		/// <returns></returns>
-		private ulong AsInteger(ulong Default, int MagnitudeBits, bool IsSigned)
+		private ulong AsInteger(ulong defaultValue, int magnitudeBits, bool isSigned)
 		{
 			if (CbFieldUtils.IsInteger(TypeWithFlags))
 			{
 				// A shift of a 64-bit value by 64 is undefined so shift by one less because magnitude is never zero.
-				ulong OutOfRangeMask = ~(ulong)1 << (MagnitudeBits - 1);
-				ulong IsNegative = (ulong)(byte)(TypeWithFlags) & 1;
+				ulong outOfRangeMask = ~(ulong)1 << (magnitudeBits - 1);
+				ulong isNegative = (ulong)(byte)(TypeWithFlags) & 1;
 
-				int MagnitudeByteCount;
-				ulong Magnitude = VarInt.Read(Payload.Span, out MagnitudeByteCount);
-				ulong Value = Magnitude ^ (ulong)-(long)(IsNegative);
+				ulong magnitude = VarInt.Read(Payload.Span, out _);
+				ulong value = magnitude ^ (ulong)-(long)(isNegative);
 
-				if ((Magnitude & OutOfRangeMask) == 0 && (IsNegative == 0 || IsSigned))
+				if ((magnitude & outOfRangeMask) == 0 && (isNegative == 0 || isSigned))
 				{
 					Error = CbFieldError.None;
-					return Value;
+					return value;
 				}
 				else
 				{
 					Error = CbFieldError.RangeError;
-					return Default;
+					return defaultValue;
 				}
 			}
 			else
 			{
 				Error = CbFieldError.TypeError;
-				return Default;
+				return defaultValue;
 			}
 		}
 
 		/// <summary>
 		/// Access the field as a float. Returns the provided default on error.
 		/// </summary>
-		/// <param name="Default">Default value</param>
+		/// <param name="default">Default value</param>
 		/// <returns>Value of the field</returns>
-		public float AsFloat(float Default = 0.0f)
+		public float AsFloat(float defaultValue = 0.0f)
 		{
 			switch (GetType())
 			{
 				case CbFieldType.IntegerPositive:
 				case CbFieldType.IntegerNegative:
 					{
-						ulong IsNegative = (ulong)TypeWithFlags & 1;
-						ulong OutOfRangeMask = ~((1UL << /*FLT_MANT_DIG*/ 24) - 1);
-
-						int MagnitudeByteCount;
-						ulong Magnitude = VarInt.Read(Payload.Span, out MagnitudeByteCount) + IsNegative;
-						bool IsInRange = (Magnitude & OutOfRangeMask) == 0;
-						Error = IsInRange ? CbFieldError.None : CbFieldError.RangeError;
-						return IsInRange ? (float)((IsNegative != 0) ? (float)-(long)Magnitude : (float)Magnitude) : Default;
+						ulong isNegative = (ulong)TypeWithFlags & 1;
+						ulong outOfRangeMask = ~((1UL << /*FLT_MANT_DIG*/ 24) - 1);
+
+						ulong magnitude = VarInt.Read(Payload.Span, out _) + isNegative;
+						bool isInRange = (magnitude & outOfRangeMask) == 0;
+						Error = isInRange ? CbFieldError.None : CbFieldError.RangeError;
+						return isInRange ? (float)((isNegative != 0) ? (float)-(long)magnitude : (float)magnitude) : defaultValue;
 					}
 				case CbFieldType.Float32:
 					Error = CbFieldError.None;
 					return BitConverter.Int32BitsToSingle(BinaryPrimitives.ReadInt32BigEndian(Payload.Span));
 				case CbFieldType.Float64:
 					Error = CbFieldError.RangeError;
-					return Default;
+					return defaultValue;
 				default:
 					Error = CbFieldError.TypeError;
-					return Default;
+					return defaultValue;
 			}
 		}
 
@@ -1149,23 +1145,22 @@ namespace EpicGames.Serialization
 		/// <summary>
 		/// Access the field as a double. Returns the provided default on error.
 		/// </summary>
-		/// <param name="Default">Default value</param>
+		/// <param name="default">Default value</param>
 		/// <returns>Value of the field</returns>
-		public double AsDouble(double Default)
+		public double AsDouble(double defaultValue)
 		{
 			switch (GetType())
 			{
 				case CbFieldType.IntegerPositive:
 				case CbFieldType.IntegerNegative:
 					{
-						ulong IsNegative = (ulong)TypeWithFlags & 1;
-						ulong OutOfRangeMask = ~((1UL << /*DBL_MANT_DIG*/ 53) - 1);
-
-						int MagnitudeByteCount;
-						ulong Magnitude = VarInt.Read(Payload.Span, out MagnitudeByteCount) + IsNegative;
-						bool IsInRange = (Magnitude & OutOfRangeMask) == 0;
-						Error = IsInRange ? CbFieldError.None : CbFieldError.RangeError;
-						return IsInRange ? (double)((IsNegative != 0) ? (double)-(long)Magnitude : (double)Magnitude) : Default;
+						ulong isNegative = (ulong)TypeWithFlags & 1;
+						ulong outOfRangeMask = ~((1UL << /*DBL_MANT_DIG*/ 53) - 1);
+
+						ulong magnitude = VarInt.Read(Payload.Span, out _) + isNegative;
+						bool isInRange = (magnitude & outOfRangeMask) == 0;
+						Error = isInRange ? CbFieldError.None : CbFieldError.RangeError;
+						return isInRange ? (double)((isNegative != 0) ? (double)-(long)magnitude : (double)magnitude) : defaultValue;
 					}
 				case CbFieldType.Float32:
 					{
@@ -1179,7 +1174,7 @@ namespace EpicGames.Serialization
 					}
 				default:
 					Error = CbFieldError.TypeError;
-					return Default;
+					return defaultValue;
 			}
 		}
 
@@ -1192,9 +1187,9 @@ namespace EpicGames.Serialization
 		/// <summary>
 		/// Access the field as a bool. Returns the provided default on error.
 		/// </summary>
-		/// <param name="Default">Default value</param>
+		/// <param name="default">Default value</param>
 		/// <returns>Value of the field</returns>
-		public bool AsBool(bool Default)
+		public bool AsBool(bool defaultValue)
 		{
 			switch (GetType())
 			{
@@ -1206,7 +1201,7 @@ namespace EpicGames.Serialization
 					return false;
 				default:
 					Error = CbFieldError.TypeError;
-					return Default;
+					return defaultValue;
 			}
 		}
 
@@ -1219,9 +1214,9 @@ namespace EpicGames.Serialization
 		/// <summary>
 		/// Access the field as a hash referencing an object attachment. Returns the provided default on error.
 		/// </summary>
-		/// <param name="Default">Default value</param>
+		/// <param name="default">Default value</param>
 		/// <returns>Value of the field</returns>
-		public CbObjectAttachment AsObjectAttachment(CbObjectAttachment Default)
+		public CbObjectAttachment AsObjectAttachment(CbObjectAttachment defaultValue)
 		{
 			if (CbFieldUtils.IsObjectAttachment(TypeWithFlags))
 			{
@@ -1231,7 +1226,7 @@ namespace EpicGames.Serialization
 			else
 			{
 				Error = CbFieldError.TypeError;
-				return Default;
+				return defaultValue;
 			}
 		}
 
@@ -1245,9 +1240,9 @@ namespace EpicGames.Serialization
 		/// <summary>
 		/// Access the field as a hash referencing a binary attachment. Returns the provided default on error.
 		/// </summary>
-		/// <param name="Default">Default value</param>
+		/// <param name="default">Default value</param>
 		/// <returns>Value of the field</returns>
-		public CbBinaryAttachment AsBinaryAttachment(CbBinaryAttachment Default)
+		public CbBinaryAttachment AsBinaryAttachment(CbBinaryAttachment defaultValue)
 		{
 			if (CbFieldUtils.IsBinaryAttachment(TypeWithFlags))
 			{
@@ -1257,7 +1252,7 @@ namespace EpicGames.Serialization
 			else
 			{
 				Error = CbFieldError.TypeError;
-				return Default;
+				return defaultValue;
 			}
 		}
 
@@ -1271,9 +1266,9 @@ namespace EpicGames.Serialization
 		/// <summary>
 		/// Access the field as a hash referencing an attachment. Returns the provided default on error.
 		/// </summary>
-		/// <param name="Default">Default value</param>
+		/// <param name="default">Default value</param>
 		/// <returns>Value of the field</returns>
-		public IoHash AsAttachment(IoHash Default)
+		public IoHash AsAttachment(IoHash defaultValue)
 		{
 			if (CbFieldUtils.IsAttachment(TypeWithFlags))
 			{
@@ -1283,7 +1278,7 @@ namespace EpicGames.Serialization
 			else
 			{
 				Error = CbFieldError.TypeError;
-				return Default;
+				return defaultValue;
 			}
 		}
 
@@ -1296,9 +1291,9 @@ namespace EpicGames.Serialization
 		/// <summary>
 		/// Access the field as a hash referencing an attachment. Returns the provided default on error.
 		/// </summary>
-		/// <param name="Default">Default value</param>
+		/// <param name="default">Default value</param>
 		/// <returns>Value of the field</returns>
-		public IoHash AsHash(IoHash Default)
+		public IoHash AsHash(IoHash defaultValue)
 		{
 			if (CbFieldUtils.IsHash(TypeWithFlags))
 			{
@@ -1308,41 +1303,41 @@ namespace EpicGames.Serialization
 			else
 			{
 				Error = CbFieldError.TypeError;
-				return Default;
+				return defaultValue;
 			}
 		}
 
 		/// <summary>
 		/// Access the field as a UUID. Returns a nil UUID on error.
 		/// </summary>
-		/// <param name="Default">Default value</param>
+		/// <param name="default">Default value</param>
 		/// <returns>Value of the field</returns>
-		public Guid AsUuid(Guid Default = default(Guid))
+		public Guid AsUuid(Guid defaultValue = default)
 		{
 			if (CbFieldUtils.IsUuid(TypeWithFlags))
 			{
 				Error = CbFieldError.None;
 
-				ReadOnlySpan<byte> Span = Payload.Span;
-				uint A = BinaryPrimitives.ReadUInt32BigEndian(Span);
-				ushort B = BinaryPrimitives.ReadUInt16BigEndian(Span.Slice(4));
-				ushort C = BinaryPrimitives.ReadUInt16BigEndian(Span.Slice(6));
+				ReadOnlySpan<byte> span = Payload.Span;
+				uint a = BinaryPrimitives.ReadUInt32BigEndian(span);
+				ushort b = BinaryPrimitives.ReadUInt16BigEndian(span.Slice(4));
+				ushort c = BinaryPrimitives.ReadUInt16BigEndian(span.Slice(6));
 
-				return new Guid(A, B, C, Span[8], Span[9], Span[10], Span[11], Span[12], Span[13], Span[14], Span[15]);
+				return new Guid(a, b, c, span[8], span[9], span[10], span[11], span[12], span[13], span[14], span[15]);
 			}
 			else
 			{
 				Error = CbFieldError.TypeError;
-				return Default;
+				return defaultValue;
 			}
 		}
 
 		/// <summary>
 		/// Reads a date time as number of ticks from the stream
 		/// </summary>
-		/// <param name="Default"></param>
+		/// <param name="default"></param>
 		/// <returns></returns>
-		public long AsDateTimeTicks(long Default = 0)
+		public long AsDateTimeTicks(long defaultValue = 0)
 		{
 			if (CbFieldUtils.IsDateTime(TypeWithFlags))
 			{
@@ -1352,7 +1347,7 @@ namespace EpicGames.Serialization
 			else
 			{
 				Error = CbFieldError.TypeError;
-				return Default;
+				return defaultValue;
 			}
 		}
 
@@ -1368,19 +1363,19 @@ namespace EpicGames.Serialization
 		/// <summary>
 		/// Access the field as a DateTime.
 		/// </summary>
-		/// <param name="Default"></param>
+		/// <param name="default"></param>
 		/// <returns></returns>
-		public DateTime AsDateTime(DateTime Default)
+		public DateTime AsDateTime(DateTime defaultValue)
 		{
-			return new DateTime(AsDateTimeTicks(Default.ToUniversalTime().Ticks), DateTimeKind.Utc);
+			return new DateTime(AsDateTimeTicks(defaultValue.ToUniversalTime().Ticks), DateTimeKind.Utc);
 		}
 
 		/// <summary>
 		/// Reads a timespan as number of ticks from the stream
 		/// </summary>
-		/// <param name="Default"></param>
+		/// <param name="default"></param>
 		/// <returns></returns>
-		public long AsTimeSpanTicks(long Default = 0)
+		public long AsTimeSpanTicks(long defaultValue = 0)
 		{
 			if (CbFieldUtils.IsTimeSpan(TypeWithFlags))
 			{
@@ -1390,24 +1385,23 @@ namespace EpicGames.Serialization
 			else
 			{
 				Error = CbFieldError.TypeError;
-				return Default;
+				return defaultValue;
 			}
 		}
 
 		/// <summary>
 		/// Reads a timespan as number of ticks from the stream
 		/// </summary>
-		/// <param name="Default"></param>
+		/// <param name="default"></param>
 		/// <returns></returns>
-		public TimeSpan AsTimeSpan(TimeSpan Default = default) => new TimeSpan(AsTimeSpanTicks(Default.Ticks));
-
+		public TimeSpan AsTimeSpan(TimeSpan defaultValue = default) => new TimeSpan(AsTimeSpanTicks(defaultValue.Ticks));
 
 		/// <summary>
 		/// Access the field as a object id
 		/// </summary>
-		/// <param name="Default"></param>
+		/// <param name="default"></param>
 		/// <returns></returns>
-		public ReadOnlyMemory<byte> AsObjectId(ReadOnlyMemory<byte> Default = default)
+		public ReadOnlyMemory<byte> AsObjectId(ReadOnlyMemory<byte> defaultValue = default)
 		{
 			if (CbFieldUtils.IsObjectId(TypeWithFlags))
 			{
@@ -1418,7 +1412,7 @@ namespace EpicGames.Serialization
 			else
 			{
 				Error = CbFieldError.TypeError;
-				return Default;
+				return defaultValue;
 			}
 		}
 
@@ -1477,8 +1471,8 @@ namespace EpicGames.Serialization
 		/// <summary>
 		/// Whether the field has a value
 		/// </summary>
-		/// <param name="Field"></param>
-		public static explicit operator bool(CbField Field) => Field.HasValue();
+		/// <param name="field"></param>
+		public static explicit operator bool(CbField field) => field.HasValue();
 
 		/// <summary>
 		/// Whether the field has a value.
@@ -1508,28 +1502,28 @@ namespace EpicGames.Serialization
 		/// <returns></returns>
 		public Blake3Hash GetHash()
 		{
-			using (Blake3.Hasher Hasher = Blake3.Hasher.New())
+			using (Blake3.Hasher hasher = Blake3.Hasher.New())
 			{
-				AppendHash(Hasher);
+				AppendHash(hasher);
 
-				byte[] Hash = new byte[32];
-				Hasher.Finalize(Hash);
+				byte[] hash = new byte[32];
+				hasher.Finalize(hash);
 
-				return new Blake3Hash(Hash);
+				return new Blake3Hash(hash);
 			}
 		}
 
 		/// <summary>
 		/// Append the hash of the field, including the type and name
 		/// </summary>
-		/// <param name="Hasher"></param>
-		void AppendHash(Blake3.Hasher Hasher)
+		/// <param name="hasher"></param>
+		void AppendHash(Blake3.Hasher hasher)
 		{
-			Span<byte> Data = stackalloc byte[1];
-			Data[0] = (byte)CbFieldUtils.GetSerializedType(TypeWithFlags);
-			Hasher.Update(Data);
+			Span<byte> data = stackalloc byte[1];
+			data[0] = (byte)CbFieldUtils.GetSerializedType(TypeWithFlags);
+			hasher.Update(data);
 
-			Hasher.Update(GetViewNoType().Span);
+			hasher.Update(GetViewNoType().Span);
 		}
 
 		/// <summary>
@@ -1541,42 +1535,42 @@ namespace EpicGames.Serialization
 		/// these assumptions do not hold, this may return false for equivalent inputs. Validation can
 		/// be performed with ValidateCompactBinary, except for field order and field name case.
 		/// </summary>
-		/// <param name="Other"></param>
+		/// <param name="other"></param>
 		/// <returns></returns>
-		public bool Equals(CbField? Other)
+		public bool Equals(CbField? other)
 		{
-			return Other != null && CbFieldUtils.GetSerializedType(TypeWithFlags) == CbFieldUtils.GetSerializedType(Other.TypeWithFlags) && GetViewNoType().Span.SequenceEqual(Other.GetViewNoType().Span);
+			return other != null && CbFieldUtils.GetSerializedType(TypeWithFlags) == CbFieldUtils.GetSerializedType(other.TypeWithFlags) && GetViewNoType().Span.SequenceEqual(other.GetViewNoType().Span);
 		}
 
 		/// <summary>
 		/// Copy the field into a buffer of exactly GetSize() bytes, including the type and name.
 		/// </summary>
-		/// <param name="Buffer"></param>
-		public void CopyTo(Span<byte> Buffer)
+		/// <param name="buffer"></param>
+		public void CopyTo(Span<byte> buffer)
 		{
-			Buffer[0] = (byte)CbFieldUtils.GetSerializedType(TypeWithFlags);
-			GetViewNoType().Span.CopyTo(Buffer.Slice(1));
+			buffer[0] = (byte)CbFieldUtils.GetSerializedType(TypeWithFlags);
+			GetViewNoType().Span.CopyTo(buffer.Slice(1));
 		}
 
 		/// <summary>
 		/// Invoke the visitor for every attachment in the field.
 		/// </summary>
-		/// <param name="Visitor"></param>
-		public void IterateAttachments(Action<CbField> Visitor)
+		/// <param name="visitor"></param>
+		public void IterateAttachments(Action<CbField> visitor)
 		{
 			switch (GetType())
 			{
 				case CbFieldType.Object:
 				case CbFieldType.UniformObject:
-					CbObject.FromFieldNoCheck(this).IterateAttachments(Visitor);
+					CbObject.FromFieldNoCheck(this).IterateAttachments(visitor);
 					break;
 				case CbFieldType.Array:
 				case CbFieldType.UniformArray:
-					CbArray.FromFieldNoCheck(this).IterateAttachments(Visitor);
+					CbArray.FromFieldNoCheck(this).IterateAttachments(visitor);
 					break;
 				case CbFieldType.ObjectAttachment:
 				case CbFieldType.BinaryAttachment:
-					Visitor(this);
+					visitor(this);
 					break;
 			}
 		}
@@ -1587,18 +1581,18 @@ namespace EpicGames.Serialization
 		/// A view is available if the field contains its type. Access the equivalent for other fields
 		/// through FCbField::GetBuffer, FCbField::Clone, or CopyTo.
 		/// </summary>
-		/// <param name="OutView"></param>
+		/// <param name="outView"></param>
 		/// <returns></returns>
-		public bool TryGetView(out ReadOnlyMemory<byte> OutView)
+		public bool TryGetView(out ReadOnlyMemory<byte> outView)
 		{
 			if (CbFieldUtils.HasFieldType(TypeWithFlags))
 			{
-				OutView = Memory;
+				outView = Memory;
 				return true;
 			}
 			else
 			{
-				OutView = ReadOnlyMemory<byte>.Empty;
+				outView = ReadOnlyMemory<byte>.Empty;
 				return false;
 			}
 		}
@@ -1606,12 +1600,9 @@ namespace EpicGames.Serialization
 		/// <summary>
 		/// Find a field of an object by case-sensitive name comparison, otherwise a field with no value.
 		/// </summary>
-		/// <param name="Name"></param>
+		/// <param name="name"></param>
 		/// <returns></returns>
-		public CbField this[Utf8String Name]
-		{
-			get { return this.FirstOrDefault(Field => Field.Name == Name) ?? CbField.Empty; }
-		}
+		public CbField this[Utf8String name] => this.FirstOrDefault(field => field.Name == name) ?? CbField.Empty;
 
 		/// <summary>
 		/// Create an iterator for the fields of an array or object, otherwise an empty iterator.
@@ -1619,26 +1610,26 @@ namespace EpicGames.Serialization
 		/// <returns></returns>
 		public CbFieldIterator CreateIterator()
 		{
-			CbFieldType LocalTypeWithFlags = TypeWithFlags;
-			if (CbFieldUtils.HasFields(LocalTypeWithFlags))
+			CbFieldType localTypeWithFlags = TypeWithFlags;
+			if (CbFieldUtils.HasFields(localTypeWithFlags))
 			{
-				ReadOnlyMemory<byte> PayloadBytes = Payload;
-				int PayloadSizeByteCount;
-				int PayloadSize = (int)VarInt.Read(PayloadBytes.Span, out PayloadSizeByteCount);
-				PayloadBytes = PayloadBytes.Slice(PayloadSizeByteCount);
-				int NumByteCount = CbFieldUtils.IsArray(LocalTypeWithFlags) ? (int)VarInt.Measure(PayloadBytes.Span) : 0;
-				if (PayloadSize > NumByteCount)
+				ReadOnlyMemory<byte> payloadBytes = Payload;
+				int payloadSizeByteCount;
+				int payloadSize = (int)VarInt.Read(payloadBytes.Span, out payloadSizeByteCount);
+				payloadBytes = payloadBytes.Slice(payloadSizeByteCount);
+				int numByteCount = CbFieldUtils.IsArray(localTypeWithFlags) ? (int)VarInt.Measure(payloadBytes.Span) : 0;
+				if (payloadSize > numByteCount)
 				{
-					PayloadBytes = PayloadBytes.Slice(NumByteCount);
+					payloadBytes = payloadBytes.Slice(numByteCount);
 
-					CbFieldType UniformType = CbFieldType.HasFieldType;
+					CbFieldType uniformType = CbFieldType.HasFieldType;
 					if (CbFieldUtils.HasUniformFields(TypeWithFlags))
 					{
-						UniformType = (CbFieldType)PayloadBytes.Span[0];
-						PayloadBytes = PayloadBytes.Slice(1);
+						uniformType = (CbFieldType)payloadBytes.Span[0];
+						payloadBytes = payloadBytes.Slice(1);
 					}
 
-					return new CbFieldIterator(PayloadBytes, UniformType);
+					return new CbFieldIterator(payloadBytes, uniformType);
 				}
 			}
 			return new CbFieldIterator(ReadOnlyMemory<byte>.Empty, CbFieldType.HasFieldType);
@@ -1647,9 +1638,9 @@ namespace EpicGames.Serialization
 		/// <inheritdoc/>
 		public IEnumerator<CbField> GetEnumerator()
 		{
-			for (CbFieldIterator Iter = CreateIterator(); Iter; Iter.MoveNext())
+			for (CbFieldIterator iter = CreateIterator(); iter; iter.MoveNext())
 			{
-				yield return Iter.Current;
+				yield return iter.Current;
 			}
 		}
 
@@ -1662,20 +1653,20 @@ namespace EpicGames.Serialization
 		/// <returns></returns>
 		private ReadOnlyMemory<byte> GetViewNoType()
 		{
-			int NameSize = CbFieldUtils.HasFieldName(TypeWithFlags) ? NameLen + (int)VarInt.Measure((uint)NameLen) : 0;
-			return Memory.Slice(PayloadOffset - NameSize);
+			int nameSize = CbFieldUtils.HasFieldName(TypeWithFlags) ? _nameLen + (int)VarInt.Measure((uint)_nameLen) : 0;
+			return Memory.Slice(_payloadOffset - nameSize);
 		}
 
 		/// <summary>
 		/// Accessor for the payload
 		/// </summary>
-		internal ReadOnlyMemory<byte> Payload => Memory.Slice(PayloadOffset);
+		internal ReadOnlyMemory<byte> Payload => Memory.Slice(_payloadOffset);
 
 		/// <summary>
 		/// Returns a view of the value payload, which excludes the type and name.
 		/// </summary>
 		/// <returns></returns>
-		internal ReadOnlyMemory<byte> GetPayloadView() => Memory.Slice(PayloadOffset);
+		internal ReadOnlyMemory<byte> GetPayloadView() => Memory.Slice(_payloadOffset);
 
 		/// <summary>
 		/// Returns the type of the field excluding flags.
@@ -1707,8 +1698,8 @@ namespace EpicGames.Serialization
 				case CbFieldType.Binary:
 				case CbFieldType.String:
 					{
-						ulong PayloadSize = VarInt.Read(Payload.Span, out int BytesRead);
-						return PayloadSize + (ulong)BytesRead;
+						ulong payloadSize = VarInt.Read(Payload.Span, out int bytesRead);
+						return payloadSize + (ulong)bytesRead;
 					}
 				case CbFieldType.IntegerPositive:
 				case CbFieldType.IntegerNegative:
@@ -1740,10 +1731,10 @@ namespace EpicGames.Serialization
 
 		#region Mimic inheritance from TCbBufferFactory
 
-		public static CbField Clone(ReadOnlyMemory<byte> Data) => Clone(new CbField(Data));
-		public static CbField Clone(CbField Other) => Other;
-		public static CbField MakeView(ReadOnlyMemory<byte> Data) => new CbField(Data);
-		public static CbField MakeView(CbField Other) => Other;
+		public static CbField Clone(ReadOnlyMemory<byte> data) => Clone(new CbField(data));
+		public static CbField Clone(CbField other) => other;
+		public static CbField MakeView(ReadOnlyMemory<byte> data) => new CbField(data);
+		public static CbField MakeView(CbField other) => other;
 
 		#endregion
 	}
@@ -1760,156 +1751,156 @@ namespace EpicGames.Serialization
 		}
 
 		/// <inheritdoc/>
-		public override void Write(Utf8JsonWriter Writer, CbField Field, JsonSerializerOptions Options)
+		public override void Write(Utf8JsonWriter writer, CbField field, JsonSerializerOptions options)
 		{
-			switch (Field.GetType())
+			switch (field.GetType())
 			{
 				case CbFieldType.Null:
-					if (Field.HasName())
+					if (field.HasName())
 					{
-						Writer.WriteNull(Field.Name.Span);
+						writer.WriteNull(field.Name.Span);
 					}
 					else
 					{
-						Writer.WriteNullValue();
+						writer.WriteNullValue();
 					}
 					break;
 				case CbFieldType.Object:
 				case CbFieldType.UniformObject:
-					if (Field.HasName())
+					if (field.HasName())
 					{
-						Writer.WriteStartObject(Field.Name.Span);
+						writer.WriteStartObject(field.Name.Span);
 					}
 					else
 					{
-						Writer.WriteStartObject();
+						writer.WriteStartObject();
 					}
 
-					foreach (CbField Member in Field.AsObject())
+					foreach (CbField member in field.AsObject())
 					{
-						Write(Writer, Member, Options);
+						Write(writer, member, options);
 					}
-					Writer.WriteEndObject();
+					writer.WriteEndObject();
 					break;
 				case CbFieldType.Array:
 				case CbFieldType.UniformArray:
-					if (Field.HasName())
+					if (field.HasName())
 					{
-						Writer.WriteStartArray(Field.Name.Span);
+						writer.WriteStartArray(field.Name.Span);
 					}
 					else
 					{
-						Writer.WriteStartArray();
+						writer.WriteStartArray();
 					}
 
-					foreach (CbField Element in Field.AsArray())
+					foreach (CbField element in field.AsArray())
 					{
-						Write(Writer, Element, Options);
+						Write(writer, element, options);
 					}
-					Writer.WriteEndArray();
+					writer.WriteEndArray();
 					break;
 				case CbFieldType.Binary:
-					if (Field.HasName())
+					if (field.HasName())
 					{
-						Writer.WriteBase64String(Field.Name.Span, Field.AsBinary().Span);
+						writer.WriteBase64String(field.Name.Span, field.AsBinary().Span);
 					}
 					else
 					{
-						Writer.WriteBase64StringValue(Field.AsBinary().Span);
+						writer.WriteBase64StringValue(field.AsBinary().Span);
 					}
 					break;
 				case CbFieldType.String:
-					if (Field.HasName())
+					if (field.HasName())
 					{
-						Writer.WriteString(Field.Name.Span, Field.AsUtf8String().Span);
+						writer.WriteString(field.Name.Span, field.AsUtf8String().Span);
 					}
 					else
 					{
-						Writer.WriteStringValue(Field.AsUtf8String().Span);
+						writer.WriteStringValue(field.AsUtf8String().Span);
 					}
 					break;
 				case CbFieldType.IntegerPositive:
-					if (Field.HasName())
+					if (field.HasName())
 					{
-						Writer.WriteNumber(Field.Name.Span, Field.AsUInt64());
+						writer.WriteNumber(field.Name.Span, field.AsUInt64());
 					}
 					else
 					{
-						Writer.WriteNumberValue(Field.AsUInt64());
+						writer.WriteNumberValue(field.AsUInt64());
 					}
 					break;
 				case CbFieldType.IntegerNegative:
-					if (Field.HasName())
+					if (field.HasName())
 					{
-						Writer.WriteNumber(Field.Name.Span, Field.AsInt64());
+						writer.WriteNumber(field.Name.Span, field.AsInt64());
 					}
 					else
 					{
-						Writer.WriteNumberValue(Field.AsInt64());
+						writer.WriteNumberValue(field.AsInt64());
 					}
 					break;
 				case CbFieldType.Float32:
 				case CbFieldType.Float64:
-					if (Field.HasName())
+					if (field.HasName())
 					{
-						Writer.WriteNumber(Field.Name.Span, Field.AsDouble());
+						writer.WriteNumber(field.Name.Span, field.AsDouble());
 					}
 					else
 					{
-						Writer.WriteNumberValue(Field.AsDouble());
+						writer.WriteNumberValue(field.AsDouble());
 					}
 					break;
 				case CbFieldType.BoolFalse:
 				case CbFieldType.BoolTrue:
-					if (Field.HasName())
+					if (field.HasName())
 					{
-						Writer.WriteBoolean(Field.Name.Span, Field.AsBool());
+						writer.WriteBoolean(field.Name.Span, field.AsBool());
 					}
 					else
 					{
-						Writer.WriteBooleanValue(Field.AsBool());
+						writer.WriteBooleanValue(field.AsBool());
 					}
 					break;
 				case CbFieldType.ObjectAttachment:
 				case CbFieldType.BinaryAttachment:
 				case CbFieldType.Hash:
-					if (Field.HasName())
+					if (field.HasName())
 					{
-						Writer.WriteString(Field.Name.Span, Field.AsHash().ToString());
+						writer.WriteString(field.Name.Span, field.AsHash().ToString());
 					}
 					else
 					{
-						Writer.WriteStringValue(Field.AsHash().ToString());
+						writer.WriteStringValue(field.AsHash().ToString());
 					}
 					break;
 				case CbFieldType.Uuid:
-					if (Field.HasName())
+					if (field.HasName())
 					{
-						Writer.WriteString(Field.Name.Span, Field.AsUuid().ToString());
+						writer.WriteString(field.Name.Span, field.AsUuid().ToString());
 					}
 					else
 					{
-						Writer.WriteStringValue(Field.AsUuid().ToString());
+						writer.WriteStringValue(field.AsUuid().ToString());
 					}
 					break;
 				case CbFieldType.DateTime:
-					if (Field.HasName())
+					if (field.HasName())
 					{
-						Writer.WriteNumber(Field.Name.Span, Field.AsDateTimeTicks());
+						writer.WriteNumber(field.Name.Span, field.AsDateTimeTicks());
 					}
 					else
 					{
-						Writer.WriteNumberValue(Field.AsDateTimeTicks());
+						writer.WriteNumberValue(field.AsDateTimeTicks());
 					}
 					break;
 				case CbFieldType.TimeSpan:
-					if (Field.HasName())
+					if (field.HasName())
 					{
-						Writer.WriteNumber(Field.Name.Span, Field.AsTimeSpanTicks());
+						writer.WriteNumber(field.Name.Span, field.AsTimeSpanTicks());
 					}
 					else
 					{
-						Writer.WriteNumberValue(Field.AsTimeSpanTicks());
+						writer.WriteNumberValue(field.AsTimeSpanTicks());
 					}
 					break;
 				default:
@@ -1923,7 +1914,7 @@ namespace EpicGames.Serialization
 		/// <summary>
 		/// The underlying buffer
 		/// </summary>
-		ReadOnlyMemory<byte> Data;
+		readonly ReadOnlyMemory<byte> _data;
 
 		/// <summary>
 		/// Type for all fields
@@ -1939,12 +1930,12 @@ namespace EpicGames.Serialization
 		/// <summary>
 		/// Constructor
 		/// </summary>
-		/// <param name="Data"></param>
-		/// <param name="UniformType"></param>
-		public CbFieldEnumerator(ReadOnlyMemory<byte> Data, CbFieldType UniformType)
+		/// <param name="data"></param>
+		/// <param name="uniformType"></param>
+		public CbFieldEnumerator(ReadOnlyMemory<byte> data, CbFieldType uniformType)
 		{
-			this.Data = Data;
-			this.UniformType = UniformType;
+			_data = data;
+			UniformType = uniformType;
 		}
 
 		/// <inheritdoc/>
@@ -1961,9 +1952,9 @@ namespace EpicGames.Serialization
 		/// <inheritdoc/>
 		public bool MoveNext()
 		{
-			if (Data.Length > 0)
+			if (_data.Length > 0)
 			{
-				Current = new CbField(Data, UniformType);
+				Current = new CbField(_data, UniformType);
 				return true;
 			}
 			else
@@ -1979,7 +1970,7 @@ namespace EpicGames.Serialization
 		/// <returns></returns>
 		public CbFieldEnumerator Clone()
 		{
-			return new CbFieldEnumerator(Data, UniformType);
+			return new CbFieldEnumerator(_data, UniformType);
 		}
 	}
 
@@ -1991,12 +1982,12 @@ namespace EpicGames.Serialization
 		/// <summary>
 		/// The underlying buffer
 		/// </summary>
-		ReadOnlyMemory<byte> NextData;
+		ReadOnlyMemory<byte> _nextData;
 
 		/// <summary>
 		/// Type for all fields
 		/// </summary>
-		CbFieldType UniformType;
+		readonly CbFieldType _uniformType;
 
 		/// <summary>
 		/// The current iterator
@@ -2014,22 +2005,22 @@ namespace EpicGames.Serialization
 		/// <summary>
 		/// Constructor for single field iterator
 		/// </summary>
-		/// <param name="Field"></param>
-		private CbFieldIterator(CbField Field)
+		/// <param name="field"></param>
+		private CbFieldIterator(CbField field)
 		{
-			NextData = ReadOnlyMemory<byte>.Empty;
-			Current = Field;
+			_nextData = ReadOnlyMemory<byte>.Empty;
+			Current = field;
 		}
 
 		/// <summary>
 		/// Constructor
 		/// </summary>
-		/// <param name="Data"></param>
-		/// <param name="UniformType"></param>
-		public CbFieldIterator(ReadOnlyMemory<byte> Data, CbFieldType UniformType)
+		/// <param name="data"></param>
+		/// <param name="uniformType"></param>
+		public CbFieldIterator(ReadOnlyMemory<byte> data, CbFieldType uniformType)
 		{
-			this.NextData = Data;
-			this.UniformType = UniformType;
+			_nextData = data;
+			_uniformType = uniformType;
 
 			MoveNext();
 		}
@@ -2037,33 +2028,33 @@ namespace EpicGames.Serialization
 		/// <summary>
 		/// Copy constructor
 		/// </summary>
-		/// <param name="Other"></param>
-		public CbFieldIterator(CbFieldIterator Other)
+		/// <param name="other"></param>
+		public CbFieldIterator(CbFieldIterator other)
 		{
-			this.NextData = Other.NextData;
-			this.UniformType = Other.UniformType;
-			this.Current = Other.Current;
+			_nextData = other._nextData;
+			_uniformType = other._uniformType;
+			Current = other.Current;
 		}
 
 		/// <summary>
 		/// Construct a field range that contains exactly one field.
 		/// </summary>
-		/// <param name="Field"></param>
+		/// <param name="field"></param>
 		/// <returns></returns>
-		public static CbFieldIterator MakeSingle(CbField Field)
+		public static CbFieldIterator MakeSingle(CbField field)
 		{
-			return new CbFieldIterator(Field);
+			return new CbFieldIterator(field);
 		}
 
 		/// <summary>
 		/// Construct a field range from a buffer containing zero or more valid fields.
 		/// </summary>
-		/// <param name="View">A buffer containing zero or more valid fields.</param>
-		/// <param name="Type">HasFieldType means that View contains the type.Otherwise, use the given type.</param>
+		/// <param name="view">A buffer containing zero or more valid fields.</param>
+		/// <param name="type">HasFieldType means that View contains the type.Otherwise, use the given type.</param>
 		/// <returns></returns>
-		public static CbFieldIterator MakeRange(ReadOnlyMemory<byte> View, CbFieldType Type = CbFieldType.HasFieldType)
+		public static CbFieldIterator MakeRange(ReadOnlyMemory<byte> view, CbFieldType type = CbFieldType.HasFieldType)
 		{
-			return new CbFieldIterator(View, Type);
+			return new CbFieldIterator(view, type);
 		}
 
 		/// <summary>
@@ -2087,21 +2078,21 @@ namespace EpicGames.Serialization
 		/// <summary>
 		/// Copy the field range into a buffer of exactly GetRangeSize() bytes.
 		/// </summary>
-		/// <param name="Buffer"></param>
-		public void CopyRangeTo(Span<byte> Buffer)
+		/// <param name="buffer"></param>
+		public void CopyRangeTo(Span<byte> buffer)
 		{
-			ReadOnlyMemory<byte> Source;
-			if (TryGetRangeView(out Source))
+			ReadOnlyMemory<byte> source;
+			if (TryGetRangeView(out source))
 			{
-				Source.Span.CopyTo(Buffer);
+				source.Span.CopyTo(buffer);
 			}
 			else
 			{
-				for (CbFieldIterator It = new CbFieldIterator(this); It; It.MoveNext())
+				for (CbFieldIterator it = new CbFieldIterator(this); it; it.MoveNext())
 				{
-					int Size = It.Current.GetSize();
-					It.Current.CopyTo(Buffer);
-					Buffer = Buffer.Slice(Size);
+					int size = it.Current.GetSize();
+					it.Current.CopyTo(buffer);
+					buffer = buffer.Slice(size);
 				}
 			}
 		}
@@ -2109,17 +2100,17 @@ namespace EpicGames.Serialization
 		/// <summary>
 		/// Invoke the visitor for every attachment in the field range.
 		/// </summary>
-		/// <param name="Visitor"></param>
-		public void IterateRangeAttachments(Action<CbField> Visitor)
+		/// <param name="visitor"></param>
+		public void IterateRangeAttachments(Action<CbField> visitor)
 		{
 			// Always iterate over non-uniform ranges because we do not know if they contain an attachment.
 			if (CbFieldUtils.HasFieldType(Current.GetTypeWithFlags()))
 			{
-				for (CbFieldIterator It = new CbFieldIterator(this); It; ++It)
+				for (CbFieldIterator it = new CbFieldIterator(this); it; ++it)
 				{
-					if (CbFieldUtils.MayContainAttachments(It.Current.GetTypeWithFlags()))
+					if (CbFieldUtils.MayContainAttachments(it.Current.GetTypeWithFlags()))
 					{
-						It.Current.IterateAttachments(Visitor);
+						it.Current.IterateAttachments(visitor);
 					}
 				}
 			}
@@ -2128,9 +2119,9 @@ namespace EpicGames.Serialization
 			{
 				if (CbFieldUtils.MayContainAttachments(Current.GetTypeWithFlags()))
 				{
-					for (CbFieldIterator It = new CbFieldIterator(this); It; ++It)
+					for (CbFieldIterator it = new CbFieldIterator(this); it; ++it)
 					{
-						It.Current.IterateAttachments(Visitor);
+						it.Current.IterateAttachments(visitor);
 					}
 				}
 			}
@@ -2143,7 +2134,7 @@ namespace EpicGames.Serialization
 		/// ranges through FCbFieldIterator::CloneRange or CopyRangeTo.
 		/// </summary>
 		/// <returns></returns>
-		bool TryGetRangeView(out ReadOnlyMemory<byte> OutView)
+		bool TryGetRangeView(out ReadOnlyMemory<byte> outView)
 		{
 			throw new NotImplementedException();
 /*			FMemoryView View;
@@ -2161,10 +2152,10 @@ namespace EpicGames.Serialization
 		/// <returns></returns>
 		public bool MoveNext()
 		{
-			if (NextData.Length > 0)
+			if (_nextData.Length > 0)
 			{
-				Current = new CbField(NextData, UniformType);
-				NextData = NextData.Slice(Current.Memory.Length);
+				Current = new CbField(_nextData, _uniformType);
+				_nextData = _nextData.Slice(Current.Memory.Length);
 				return true;
 			}
 			else
@@ -2177,20 +2168,20 @@ namespace EpicGames.Serialization
 		/// <summary>
 		/// Test whether the iterator is valid
 		/// </summary>
-		/// <param name="Iterator"></param>
-		public static implicit operator bool(CbFieldIterator Iterator)
+		/// <param name="iterator"></param>
+		public static implicit operator bool(CbFieldIterator iterator)
 		{
-			return Iterator.IsValid();
+			return iterator.IsValid();
 		}
 
 		/// <summary>
 		/// Move to the next item
 		/// </summary>
-		/// <param name="Iterator"></param>
+		/// <param name="iterator"></param>
 		/// <returns></returns>
-		public static CbFieldIterator operator ++(CbFieldIterator Iterator)
+		public static CbFieldIterator operator ++(CbFieldIterator iterator)
 		{
-			return new CbFieldIterator(Iterator.NextData, Iterator.UniformType);
+			return new CbFieldIterator(iterator._nextData, iterator._uniformType);
 		}
 
 		public override bool Equals(object? obj)
@@ -2203,13 +2194,13 @@ namespace EpicGames.Serialization
 			throw new NotImplementedException();
 		}
 
-		public static bool operator ==(CbFieldIterator A, CbFieldIterator B)
+		public static bool operator ==(CbFieldIterator a, CbFieldIterator b)
 		{
-			return A.Current.Equals(B.Current);
+			return a.Current.Equals(b.Current);
 		}
-		public static bool operator !=(CbFieldIterator A, CbFieldIterator B)
+		public static bool operator !=(CbFieldIterator a, CbFieldIterator b)
 		{
-			return !A.Current.Equals(B.Current);
+			return !a.Current.Equals(b.Current);
 		}
 	}
 
@@ -2218,12 +2209,12 @@ namespace EpicGames.Serialization
 	/// </summary>
 	class CbArrayDebugView
 	{
-		CbArray Array;
+		readonly CbArray _array;
 
-		public CbArrayDebugView(CbArray Array) => this.Array = Array;
+		public CbArrayDebugView(CbArray array) => _array = array;
 
 		[DebuggerBrowsable(DebuggerBrowsableState.RootHidden)]
-		public object?[] Value => Array.Select(x => x.Value).ToArray();
+		public object?[] Value => _array.Select(x => x.Value).ToArray();
 	}
 
 	/// <summary>
@@ -2239,7 +2230,7 @@ namespace EpicGames.Serialization
 		/// <summary>
 		/// The field containing this array
 		/// </summary>
-		readonly CbField InnerField;
+		readonly CbField _innerField;
 
 		/// <summary>
 		/// Empty array constant
@@ -2251,26 +2242,26 @@ namespace EpicGames.Serialization
 		/// </summary>
 		public CbArray()
 		{
-			InnerField = Empty.InnerField;
+			_innerField = Empty._innerField;
 		}
 
 		/// <summary>
 		/// Constructor
 		/// </summary>
-		/// <param name="Field"></param>
-		private CbArray(CbField Field)
+		/// <param name="field"></param>
+		private CbArray(CbField field)
 		{
-			InnerField = Field;
+			_innerField = field;
 		}
 
 		/// <summary>
 		/// Constructor
 		/// </summary>
-		/// <param name="Data"></param>
-		/// <param name="Type"></param>
-		public CbArray(ReadOnlyMemory<byte> Data, CbFieldType Type = CbFieldType.HasFieldType)
+		/// <param name="data"></param>
+		/// <param name="type"></param>
+		public CbArray(ReadOnlyMemory<byte> data, CbFieldType type = CbFieldType.HasFieldType)
 		{
-			InnerField = new CbField(Data, Type);
+			_innerField = new CbField(data, type);
 		}
 
 		/// <summary>
@@ -2280,9 +2271,9 @@ namespace EpicGames.Serialization
 		{
 			get
 			{
-				ReadOnlyMemory<byte> PayloadBytes = InnerField.Payload;
-				PayloadBytes = PayloadBytes.Slice(VarInt.Measure(PayloadBytes.Span));
-				return (int)VarInt.Read(PayloadBytes.Span, out int NumByteCount);
+				ReadOnlyMemory<byte> payloadBytes = _innerField.Payload;
+				payloadBytes = payloadBytes.Slice(VarInt.Measure(payloadBytes.Span));
+				return (int)VarInt.Read(payloadBytes.Span, out int _);
 			}
 		}
 
@@ -2290,14 +2281,14 @@ namespace EpicGames.Serialization
 		/// Access the array as an array field.
 		/// </summary>
 		/// <returns></returns>
-		public CbField AsField() => InnerField;
+		public CbField AsField() => _innerField;
 
 		/// <summary>
 		/// Construct an array from an array field. No type check is performed!
 		/// </summary>
-		/// <param name="Field"></param>
+		/// <param name="field"></param>
 		/// <returns></returns>
-		public static CbArray FromFieldNoCheck(CbField Field) => new CbArray(Field);
+		public static CbArray FromFieldNoCheck(CbField field) => new CbArray(field);
 
 		/// <summary>
 		/// Returns the size of the array in bytes if serialized by itself with no name.
@@ -2305,7 +2296,7 @@ namespace EpicGames.Serialization
 		/// <returns></returns>
 		public int GetSize()
 		{
-			return (int)Math.Min((ulong)sizeof(CbFieldType) + InnerField.GetPayloadSize(), int.MaxValue);
+			return (int)Math.Min((ulong)sizeof(CbFieldType) + _innerField.GetPayloadSize(), Int32.MaxValue);
 		}
 
 		/// <summary>
@@ -2314,29 +2305,29 @@ namespace EpicGames.Serialization
 		/// <returns></returns>
 		public Blake3Hash GetHash()
 		{
-			using (Blake3.Hasher Hasher = Blake3.Hasher.New())
+			using (Blake3.Hasher hasher = Blake3.Hasher.New())
 			{
-				AppendHash(Hasher);
+				AppendHash(hasher);
 
-				byte[] Result = new byte[Blake3Hash.NumBytes];
-				Hasher.Finalize(Result);
+				byte[] result = new byte[Blake3Hash.NumBytes];
+				hasher.Finalize(result);
 
-				return new Blake3Hash(Result);
+				return new Blake3Hash(result);
 			}
 		}
 
 		/// <summary>
 		/// Append the hash of the array if serialized by itself with no name.
 		/// </summary>
-		public void AppendHash(Blake3.Hasher Hasher)
+		public void AppendHash(Blake3.Hasher hasher)
 		{
-			byte[] SerializedType = new byte[] { (byte)InnerField.GetType() };
-			Hasher.Update(SerializedType);
-			Hasher.Update(InnerField.Payload.Span);
+			byte[] serializedType = new byte[] { (byte)_innerField.GetType() };
+			hasher.Update(serializedType);
+			hasher.Update(_innerField.Payload.Span);
 		}
 
 		/// <inheritdoc/>
-		public override bool Equals(object? Obj) => Equals(Obj as CbArray);
+		public override bool Equals(object? obj) => Equals(obj as CbArray);
 
 		/// <inheritdoc/>
 		public override int GetHashCode() => BinaryPrimitives.ReadInt32BigEndian(GetHash().Span);
@@ -2350,25 +2341,25 @@ namespace EpicGames.Serialization
 		/// these assumptions do not hold, this may return false for equivalent inputs. Validation can
 		/// be done with the All mode to check these assumptions about the format of the inputs.
 		/// </summary>
-		/// <param name="Other"></param>
+		/// <param name="other"></param>
 		/// <returns></returns>
-		public bool Equals(CbArray? Other)
+		public bool Equals(CbArray? other)
 		{
-			return Other != null && GetType() == Other.GetType() && GetPayloadView().Span.SequenceEqual(Other.GetPayloadView().Span);
+			return other != null && GetType() == other.GetType() && GetPayloadView().Span.SequenceEqual(other.GetPayloadView().Span);
 		}
 
 		/// <summary>
 		/// Copy the array into a buffer of exactly GetSize() bytes, with no name.
 		/// </summary>
-		/// <param name="Buffer"></param>
-		public void CopyTo(Span<byte> Buffer)
+		/// <param name="buffer"></param>
+		public void CopyTo(Span<byte> buffer)
 		{
-			Buffer[0] = (byte)GetType();
-			GetPayloadView().Span.CopyTo(Buffer.Slice(1));
+			buffer[0] = (byte)GetType();
+			GetPayloadView().Span.CopyTo(buffer.Slice(1));
 		}
 
 		/** Invoke the visitor for every attachment in the array. */
-		public void IterateAttachments(Action<CbField> Visitor) => CreateIterator().IterateRangeAttachments(Visitor);
+		public void IterateAttachments(Action<CbField> visitor) => CreateIterator().IterateRangeAttachments(visitor);
 
 		/// <summary>
 		/// Try to get a view of the array as it would be serialized, such as by CopyTo.
@@ -2376,21 +2367,21 @@ namespace EpicGames.Serialization
 		/// A view is available if the array contains its type and has no name. Access the equivalent
 		/// for other arrays through FCbArray::GetBuffer, FCbArray::Clone, or CopyTo.
 		/// </summary>
-		public bool TryGetView(out ReadOnlyMemory<byte> OutView)
+		public bool TryGetView(out ReadOnlyMemory<byte> outView)
 		{
-			if(InnerField.HasName())
+			if(_innerField.HasName())
 			{
-				OutView = ReadOnlyMemory<byte>.Empty;
+				outView = ReadOnlyMemory<byte>.Empty;
 				return false;
 			}
-			return InnerField.TryGetView(out OutView);
+			return _innerField.TryGetView(out outView);
 		}
 
 		/// <inheritdoc cref="CbField.CreateIterator"/>
-		public CbFieldIterator CreateIterator() => InnerField.CreateIterator();
+		public CbFieldIterator CreateIterator() => _innerField.CreateIterator();
 
 		/// <inheritdoc/>
-		public IEnumerator<CbField> GetEnumerator() => InnerField.GetEnumerator();
+		public IEnumerator<CbField> GetEnumerator() => _innerField.GetEnumerator();
 
 		/// <inheritdoc/>
 		IEnumerator IEnumerable.GetEnumerator() => GetEnumerator();
@@ -2398,19 +2389,19 @@ namespace EpicGames.Serialization
 		#region Mimic inheritance from CbField
 
 		/// <inheritdoc cref="CbField.GetType"/>
-		internal new CbFieldType GetType() => InnerField.GetType();
+		internal new CbFieldType GetType() => _innerField.GetType();
 
 		/// <inheritdoc cref="CbField.GetPayloadView"/>
-		internal ReadOnlyMemory<byte> GetPayloadView() => InnerField.GetPayloadView();
+		internal ReadOnlyMemory<byte> GetPayloadView() => _innerField.GetPayloadView();
 
 		#endregion
 
 		#region Mimic inheritance from TCbBufferFactory
 
-		public static CbArray Clone(ReadOnlyMemory<byte> Data) => Clone(new CbArray(Data));
-		public static CbArray Clone(CbArray Other) => Other;
-		public static CbArray MakeView(ReadOnlyMemory<byte> Data) => new CbArray(Data);
-		public static CbArray MakeView(CbArray Other) => Other;
+		public static CbArray Clone(ReadOnlyMemory<byte> data) => Clone(new CbArray(data));
+		public static CbArray Clone(CbArray other) => other;
+		public static CbArray MakeView(ReadOnlyMemory<byte> data) => new CbArray(data);
+		public static CbArray MakeView(CbArray other) => other;
 
 		#endregion
 	}
@@ -2427,12 +2418,12 @@ namespace EpicGames.Serialization
 			public object? Value { get; set; }
 		}
 
-		CbObject Object;
+		readonly CbObject _object;
 
-		public CbObjectDebugView(CbObject Object) => this.Object = Object;
+		public CbObjectDebugView(CbObject obj) => _object = obj;
 
 		[DebuggerBrowsable(DebuggerBrowsableState.RootHidden)]
-		public Property[] Properties => Object.Select(x => new Property { Name = x.Name.ToString(), Value = x.Value }).ToArray();
+		public Property[] Properties => _object.Select(x => new Property { Name = x.Name.ToString(), Value = x.Value }).ToArray();
 	}
 
 	/// <summary>
@@ -2454,38 +2445,38 @@ namespace EpicGames.Serialization
 		/// <summary>
 		/// The inner field object
 		/// </summary>
-		private CbField InnerField;
+		private readonly CbField _innerField;
 
 		/// <summary>
 		/// Constructor
 		/// </summary>
 		/// <param name="Buffer"></param>
-		private CbObject(CbField Field)
+		private CbObject(CbField field)
 		{
-			InnerField = new CbField(Field.Memory, Field.TypeWithFlags);
+			_innerField = new CbField(field.Memory, field.TypeWithFlags);
 		}
 
 		/// <summary>
 		/// Constructor
 		/// </summary>
-		/// <param name="Buffer"></param>
-		public CbObject(ReadOnlyMemory<byte> Buffer, CbFieldType FieldType = CbFieldType.HasFieldType)
+		/// <param name="buffer"></param>
+		public CbObject(ReadOnlyMemory<byte> buffer, CbFieldType fieldType = CbFieldType.HasFieldType)
 		{
-			InnerField = new CbField(Buffer, FieldType);
+			_innerField = new CbField(buffer, fieldType);
 		}
 
 		/// <summary>
 		/// Builds an object by calling a delegate with a writer
 		/// </summary>
-		/// <param name="Build"></param>
+		/// <param name="build"></param>
 		/// <returns></returns>
-		public static CbObject Build(Action<CbWriter> Build)
+		public static CbObject Build(Action<CbWriter> build)
 		{
-			CbWriter Writer = new CbWriter();
-			Writer.BeginObject();
-			Build(Writer);
-			Writer.EndObject();
-			return new CbObject(Writer.ToByteArray());
+			CbWriter writer = new CbWriter();
+			writer.BeginObject();
+			build(writer);
+			writer.EndObject();
+			return new CbObject(writer.ToByteArray());
 		}
 
 		/// <summary>
@@ -2494,33 +2485,33 @@ namespace EpicGames.Serialization
 		/// The cost of this operation scales linearly with the number of fields in the object. Prefer to
 		/// iterate over the fields only once when consuming an object.
 		/// </summary>
-		/// <param name="Name">The name of the field.</param>
+		/// <param name="name">The name of the field.</param>
 		/// <returns>The matching field if found, otherwise a field with no value.</returns>
-		public CbField Find(Utf8String Name) => InnerField[Name];
+		public CbField Find(Utf8String name) => _innerField[name];
 
 		/// <summary>
 		/// Find a field by case-insensitive name comparison.
 		/// </summary>
-		/// <param name="Name">The name of the field.</param>
+		/// <param name="name">The name of the field.</param>
 		/// <returns>The matching field if found, otherwise a field with no value.</returns>
-		public CbField FindIgnoreCase(Utf8String Name) => InnerField.FirstOrDefault(Field => Utf8StringComparer.OrdinalIgnoreCase.Equals(Field.Name, Name)) ?? new CbField();
+		public CbField FindIgnoreCase(Utf8String name) => _innerField.FirstOrDefault(field => Utf8StringComparer.OrdinalIgnoreCase.Equals(field.Name, name)) ?? new CbField();
 
 		/// <summary>
 		/// Find a field by case-sensitive name comparison.
 		/// </summary>
-		/// <param name="Name">The name of the field.</param>
+		/// <param name="name">The name of the field.</param>
 		/// <returns>The matching field if found, otherwise a field with no value.</returns>
-		public CbField this[Utf8String Name] => InnerField[Name];
+		public CbField this[Utf8String name] => _innerField[name];
 
 		/// <inheritdoc cref="AsFieldView"/>
-		public CbField AsField() => InnerField;
+		public CbField AsField() => _innerField;
 
 		/// <summary>
 		/// Construct an object from an object field. No type check is performed!
 		/// </summary>
-		/// <param name="Field"></param>
+		/// <param name="field"></param>
 		/// <returns></returns>
-		public static CbObject FromFieldNoCheck(CbField Field) => new CbObject(Field);
+		public static CbObject FromFieldNoCheck(CbField field) => new CbObject(field);
 
 		/// <summary>
 		/// Returns the size of the object in bytes if serialized by itself with no name.
@@ -2528,7 +2519,7 @@ namespace EpicGames.Serialization
 		/// <returns></returns>
 		public int GetSize()
 		{
-			return sizeof(CbFieldType) + InnerField.Payload.Length;
+			return sizeof(CbFieldType) + _innerField.Payload.Length;
 		}
 
 		/// <summary>
@@ -2537,30 +2528,30 @@ namespace EpicGames.Serialization
 		/// <returns></returns>
 		public Blake3Hash GetHash()
 		{
-			using (Blake3.Hasher Hasher = Blake3.Hasher.New())
+			using (Blake3.Hasher hasher = Blake3.Hasher.New())
 			{
-				AppendHash(Hasher);
+				AppendHash(hasher);
 
-				byte[] Data = new byte[Blake3Hash.NumBytes];
-				Hasher.Finalize(Data);
+				byte[] data = new byte[Blake3Hash.NumBytes];
+				hasher.Finalize(data);
 
-				return new Blake3Hash(Data);
+				return new Blake3Hash(data);
 			}
 		}
 
 		/// <summary>
 		/// Append the hash of the object if serialized by itself with no name.
 		/// </summary>
-		/// <param name="Hasher"></param>
-		public void AppendHash(Blake3.Hasher Hasher)
+		/// <param name="hasher"></param>
+		public void AppendHash(Blake3.Hasher hasher)
 		{
-			byte[] Temp = new byte[] { (byte)InnerField.GetType() };
-			Hasher.Update(Temp);
-			Hasher.Update(InnerField.Payload.Span);
+			byte[] temp = new byte[] { (byte)_innerField.GetType() };
+			hasher.Update(temp);
+			hasher.Update(_innerField.Payload.Span);
 		}
 
 		/// <inheritdoc/>
-		public override bool Equals(object? Obj) => Equals(Obj as CbObject);
+		public override bool Equals(object? obj) => Equals(obj as CbObject);
 
 		/// <inheritdoc/>
 		public override int GetHashCode() => BinaryPrimitives.ReadInt32BigEndian(GetHash().Span);
@@ -2574,28 +2565,28 @@ namespace EpicGames.Serialization
 		/// these assumptions do not hold, this may return false for equivalent inputs. Validation can
 		/// be done with the All mode to check these assumptions about the format of the inputs.
 		/// </summary>
-		/// <param name="Other"></param>
+		/// <param name="other"></param>
 		/// <returns></returns>
-		public bool Equals(CbObject? Other)
+		public bool Equals(CbObject? other)
 		{
-			return Other != null && InnerField.GetType() == Other.InnerField.GetType() && InnerField.Payload.Span.SequenceEqual(Other.InnerField.Payload.Span);
+			return other != null && _innerField.GetType() == other._innerField.GetType() && _innerField.Payload.Span.SequenceEqual(other._innerField.Payload.Span);
 		}
 
 		/// <summary>
 		/// Copy the object into a buffer of exactly GetSize() bytes, with no name.
 		/// </summary>
-		/// <param name="Buffer"></param>
-		public void CopyTo(Span<byte> Buffer)
+		/// <param name="buffer"></param>
+		public void CopyTo(Span<byte> buffer)
 		{
-			Buffer[0] = (byte)InnerField.GetType();
-			InnerField.Payload.Span.CopyTo(Buffer.Slice(1));
+			buffer[0] = (byte)_innerField.GetType();
+			_innerField.Payload.Span.CopyTo(buffer.Slice(1));
 		}
 
 		/// <summary>
 		/// Invoke the visitor for every attachment in the object.
 		/// </summary>
-		/// <param name="Visitor"></param>
-		public void IterateAttachments(Action<CbField> Visitor) => CreateIterator().IterateRangeAttachments(Visitor);
+		/// <param name="visitor"></param>
+		public void IterateAttachments(Action<CbField> visitor) => CreateIterator().IterateRangeAttachments(visitor);
 
 		/// <summary>
 		/// Creates a view of the object, excluding the name
@@ -2603,14 +2594,14 @@ namespace EpicGames.Serialization
 		/// <returns></returns>
 		public ReadOnlyMemory<byte> GetView()
 		{
-			ReadOnlyMemory<byte> Memory;
-			if (!TryGetView(out Memory))
+			ReadOnlyMemory<byte> memory;
+			if (!TryGetView(out memory))
 			{
-				byte[] Data = new byte[GetSize()];
-				CopyTo(Data);
-				Memory = Data;
+				byte[] data = new byte[GetSize()];
+				CopyTo(data);
+				memory = data;
 			}
-			return Memory;
+			return memory;
 		}
 
 		/// <summary>
@@ -2619,33 +2610,33 @@ namespace EpicGames.Serialization
 		/// A view is available if the object contains its type and has no name. Access the equivalent
 		/// for other objects through FCbObject::GetBuffer, FCbObject::Clone, or CopyTo.
 		/// </summary>
-		/// <param name="OutView"></param>
+		/// <param name="outView"></param>
 		/// <returns></returns>
-		public bool TryGetView(out ReadOnlyMemory<byte> OutView)
+		public bool TryGetView(out ReadOnlyMemory<byte> outView)
 		{
-			if (InnerField.HasName())
+			if (_innerField.HasName())
 			{
-				OutView = ReadOnlyMemory<byte>.Empty;
+				outView = ReadOnlyMemory<byte>.Empty;
 				return false;
 			}
-			return InnerField.TryGetView(out OutView);
+			return _innerField.TryGetView(out outView);
 		}
 
 		/// <inheritdoc cref="CbField.CreateIterator"/>
-		public CbFieldIterator CreateIterator() => InnerField.CreateIterator();
+		public CbFieldIterator CreateIterator() => _innerField.CreateIterator();
 
 		/// <inheritdoc/>
-		public IEnumerator<CbField> GetEnumerator() => InnerField.GetEnumerator();
+		public IEnumerator<CbField> GetEnumerator() => _innerField.GetEnumerator();
 
 		/// <inheritdoc/>
-		IEnumerator IEnumerable.GetEnumerator() => InnerField.GetEnumerator();
+		IEnumerator IEnumerable.GetEnumerator() => _innerField.GetEnumerator();
 
 		/// <summary>
 		/// Clone this object
 		/// </summary>
-		/// <param name="Object"></param>
+		/// <param name="obj"></param>
 		/// <returns></returns>
-		public static CbObject Clone(CbObject Object) => Object;
+		public static CbObject Clone(CbObject obj) => obj;
 
 		#region Conversion to Json
 		/// <summary>
@@ -2654,96 +2645,101 @@ namespace EpicGames.Serialization
 		/// <returns></returns>
 		public string ToJson()
 		{
-			ArrayBufferWriter<byte> Buffer = new ArrayBufferWriter<byte>();
-			using (Utf8JsonWriter JsonWriter = new Utf8JsonWriter(Buffer))
+			ArrayBufferWriter<byte> buffer = new ArrayBufferWriter<byte>();
+			using (Utf8JsonWriter jsonWriter = new Utf8JsonWriter(buffer))
 			{
-				ToJson(JsonWriter);
+				ToJson(jsonWriter);
 			}
-			return Encoding.UTF8.GetString(Buffer.WrittenMemory.Span);
+			return Encoding.UTF8.GetString(buffer.WrittenMemory.Span);
 		}
 
 		/// <summary>
 		/// Write this object to JSON
 		/// </summary>
-		/// <param name="Writer"></param>
-		public void ToJson(Utf8JsonWriter Writer)
+		/// <param name="writer"></param>
+		public void ToJson(Utf8JsonWriter writer)
 		{
-			Writer.WriteStartObject();
-			foreach (CbField Field in InnerField)
+			writer.WriteStartObject();
+			foreach (CbField field in _innerField)
 			{
-				WriteField(Field, Writer);
+				WriteField(field, writer);
 			}
-			Writer.WriteEndObject();
+			writer.WriteEndObject();
 		}
 
 		/// <summary>
 		/// Write a single field to a writer
 		/// </summary>
-		/// <param name="Field"></param>
-		/// <param name="Writer"></param>
-		private static void WriteField(CbField Field, Utf8JsonWriter Writer)
+		/// <param name="field"></param>
+		/// <param name="writer"></param>
+		private static void WriteField(CbField field, Utf8JsonWriter writer)
 		{
-			if (Field.IsObject())
+			if (field.IsObject())
 			{
-				if (Field.NameLen != 0)
-					Writer.WriteStartObject(Field.Name.Span);
+				if (field._nameLen != 0)
+				{
+					writer.WriteStartObject(field.Name.Span);
+				}
 				else
-					Writer.WriteStartObject();
-				CbObject Object = Field.AsObject();
-				foreach (CbField ObjectField in Object.InnerField)
 				{
-					WriteField(ObjectField, Writer);
+					writer.WriteStartObject();
+				}
+
+				CbObject obj = field.AsObject();
+				foreach (CbField objField in obj._innerField)
+				{
+					WriteField(objField, writer);
 				}
-				Writer.WriteEndObject();
+				writer.WriteEndObject();
 			}
-			else if (Field.IsArray())
+			else if (field.IsArray())
 			{
-				Writer.WriteStartArray(Field.Name.Span);
-				CbArray Array = Field.AsArray();
-				foreach (CbField ObjectField in Array)
+				writer.WriteStartArray(field.Name.Span);
+				CbArray array = field.AsArray();
+				foreach (CbField objectField in array)
 				{
-					WriteField(ObjectField, Writer);
+					WriteField(objectField, writer);
 				}
-				Writer.WriteEndArray();
+				writer.WriteEndArray();
 			}
-			else if (Field.IsInteger())
+			else if (field.IsInteger())
 			{
-				if (Field.GetType() == CbFieldType.IntegerNegative)
+				if (field.GetType() == CbFieldType.IntegerNegative)
 				{
-					Writer.WriteNumber(Field.Name.Span, -Field.AsInt64());
+					writer.WriteNumber(field.Name.Span, -field.AsInt64());
 				}
 				else
 				{
-					Writer.WriteNumber(Field.Name.Span, Field.AsUInt64());
+					writer.WriteNumber(field.Name.Span, field.AsUInt64());
 				}
 			}
-			else if (Field.IsBool())
+			else if (field.IsBool())
 			{
-				Writer.WriteBoolean(Field.Name.Span, Field.AsBool());
+				writer.WriteBoolean(field.Name.Span, field.AsBool());
 			}
-			else if (Field.IsNull())
+			else if (field.IsNull())
 			{
-				Writer.WriteNullValue();
+				writer.WriteNullValue();
 			}
-			else if (Field.IsDateTime())
+			else if (field.IsDateTime())
 			{
-				Writer.WriteString(Field.Name.Span, Field.AsDateTime());
+				writer.WriteString(field.Name.Span, field.AsDateTime());
 			}
-			else if (Field.IsHash())
+			else if (field.IsHash())
 			{
-				Writer.WriteString(Field.Name.Span, Field.AsHash().ToUtf8String().Span);
+				writer.WriteString(field.Name.Span, field.AsHash().ToUtf8String().Span);
 			}
-			else if (Field.IsString())
+			else if (field.IsString())
 			{
-				Writer.WriteString(Field.Name.Span, Field.AsUtf8String().Span);
+				writer.WriteString(field.Name.Span, field.AsUtf8String().Span);
 			}
-			else if (Field.IsObjectId())
+			else if (field.IsObjectId())
 			{
-				Writer.WriteString(Field.Name.Span, StringUtils.FormatUtf8HexString(Field.AsObjectId().Span).Span);
+				writer.WriteString(field.Name.Span, StringUtils.FormatUtf8HexString(field.AsObjectId().Span).Span);
 			}
 			else
 			{
-				throw new NotImplementedException($"Unhandled type {Field.GetType()} when attempting to convert to json");
+				throw new NotImplementedException($"Unhandled type {field.GetType()} when attempting to convert to json");
 			}
 		}
 		#endregion
diff --git a/Engine/Source/Programs/Shared/EpicGames.Serialization/CbSerializer.cs b/Engine/Source/Programs/Shared/EpicGames.Serialization/CbSerializer.cs
--- a/Engine/Source/Programs/Shared/EpicGames.Serialization/CbSerializer.cs
+++ b/Engine/Source/Programs/Shared/EpicGames.Serialization/CbSerializer.cs
@@ -35,10 +35,10 @@ namespace EpicGames.Serialization
 		/// <summary>
 		/// Constructor
 		/// </summary>
-		/// <param name="Name"></param>
-		public CbFieldAttribute(string Name)
+		/// <param name="name"></param>
+		public CbFieldAttribute(string name)
 		{
-			this.Name = Name;
+			this.Name = name;
 		}
 	}
 
@@ -62,8 +62,8 @@ namespace EpicGames.Serialization
 		/// <summary>
 		/// Constructor
 		/// </summary>
-		/// <param name="Name">Name used to identify this class</param>
-		public CbDiscriminatorAttribute(string Name) => this.Name = Name;
+		/// <param name="name">Name used to identify this class</param>
+		public CbDiscriminatorAttribute(string name) => this.Name = name;
 	}
 
 	/// <summary>
@@ -72,12 +72,12 @@ namespace EpicGames.Serialization
 	public class CbException : Exception
 	{
 		/// <inheritdoc cref="Exception(string?)"/>
-		public CbException(string Message) : base(Message)
+		public CbException(string message) : base(message)
 		{
 		}
 
 		/// <inheritdoc cref="Exception(string?, Exception)"/>
-		public CbException(string Message, Exception Inner) : base(Message, Inner)
+		public CbException(string message, Exception inner) : base(message, inner)
 		{
 		}
 	}
@@ -90,88 +90,88 @@ namespace EpicGames.Serialization
 		/// <summary>
 		/// Serialize an object
 		/// </summary>
-		/// <param name="Type">Type of the object to serialize</param>
-		/// <param name="Value"></param>
+		/// <param name="type">Type of the object to serialize</param>
+		/// <param name="value"></param>
 		/// <returns></returns>
-		public static CbObject Serialize(Type Type, object Value)
+		public static CbObject Serialize(Type type, object value)
 		{
-			CbWriter Writer = new CbWriter();
-			CbConverter.GetConverter(Type).WriteObject(Writer, Value);
-			return Writer.ToObject();
+			CbWriter writer = new CbWriter();
+			CbConverter.GetConverter(type).WriteObject(writer, value);
+			return writer.ToObject();
 		}
 
 		/// <summary>
 		/// Serialize an object
 		/// </summary>
 		/// <typeparam name="T"></typeparam>
-		/// <param name="Value"></param>
+		/// <param name="value"></param>
 		/// <returns></returns>
-		public static CbObject Serialize<T>(T Value)
+		public static CbObject Serialize<T>(T value)
 		{
-			CbWriter Writer = new CbWriter();
-			CbConverter.GetConverter<T>().Write(Writer, Value);
-			return Writer.ToObject();
+			CbWriter writer = new CbWriter();
+			CbConverter.GetConverter<T>().Write(writer, value);
+			return writer.ToObject();
 		}
 
 		/// <summary>
 		/// Serialize a property to a given writer
 		/// </summary>
 		/// <typeparam name="T"></typeparam>
-		/// <param name="Writer"></param>
-		/// <param name="Value"></param>
-		public static void Serialize<T>(CbWriter Writer, T Value)
+		/// <param name="writer"></param>
+		/// <param name="value"></param>
+		public static void Serialize<T>(CbWriter writer, T value)
 		{
-			CbConverter.GetConverter<T>().Write(Writer, Value);
+			CbConverter.GetConverter<T>().Write(writer, value);
 		}
 
 		/// <summary>
 		/// Serialize a named property to the given writer
 		/// </summary>
 		/// <typeparam name="T"></typeparam>
-		/// <param name="Writer"></param>
-		/// <param name="Name"></param>
-		/// <param name="Value"></param>
-		public static void Serialize<T>(CbWriter Writer, Utf8String Name, T Value)
+		/// <param name="writer"></param>
+		/// <param name="name"></param>
+		/// <param name="value"></param>
+		public static void Serialize<T>(CbWriter writer, Utf8String name, T value)
 		{
-			CbConverter.GetConverter<T>().WriteNamed(Writer, Name, Value);
+			CbConverter.GetConverter<T>().WriteNamed(writer, name, value);
 		}
 
 		/// <summary>
 		/// Deserialize an object from a <see cref="CbObject"/>
 		/// </summary>
-		/// <param name="Field"></param>
-		/// <param name="Type">Type of the object to read</param>
+		/// <param name="field"></param>
+		/// <param name="type">Type of the object to read</param>
 		/// <returns></returns>
-		public static object? Deserialize(CbField Field, Type Type)
+		public static object? Deserialize(CbField field, Type type)
 		{
-			return CbConverter.GetConverter(Type).ReadObject(Field);
+			return CbConverter.GetConverter(type).ReadObject(field);
 		}
 
 		/// <summary>
 		/// Deserialize an object from a <see cref="CbField"/>
 		/// </summary>
 		/// <typeparam name="T"></typeparam>
-		/// <param name="Field"></param>
+		/// <param name="field"></param>
 		/// <returns></returns>
-		public static T Deserialize<T>(CbField Field)
+		public static T Deserialize<T>(CbField field)
 		{
-			return CbConverter.GetConverter<T>().Read(Field);
+			return CbConverter.GetConverter<T>().Read(field);
 		}
 
 		/// <summary>
 		/// Deserialize an object from a <see cref="CbObject"/>
 		/// </summary>
 		/// <typeparam name="T"></typeparam>
-		/// <param name="Object"></param>
+		/// <param name="obj"></param>
 		/// <returns></returns>
-		public static T Deserialize<T>(CbObject Object) => Deserialize<T>(Object.AsField());
+		public static T Deserialize<T>(CbObject obj) => Deserialize<T>(obj.AsField());
 
 		/// <summary>
 		/// Deserialize an object from a block of memory
 		/// </summary>
 		/// <typeparam name="T"></typeparam>
-		/// <param name="Data"></param>
+		/// <param name="data"></param>
 		/// <returns></returns>
-		public static T Deserialize<T>(ReadOnlyMemory<byte> Data) => Deserialize<T>(new CbField(Data));
+		public static T Deserialize<T>(ReadOnlyMemory<byte> data) => Deserialize<T>(new CbField(data));
 	}
 }
diff --git a/Engine/Source/Programs/Shared/EpicGames.Serialization/CbWriter.cs b/Engine/Source/Programs/Shared/EpicGames.Serialization/CbWriter.cs
--- a/Engine/Source/Programs/Shared/EpicGames.Serialization/CbWriter.cs
+++ b/Engine/Source/Programs/Shared/EpicGames.Serialization/CbWriter.cs
@@ -19,19 +19,19 @@ namespace EpicGames.Serialization
 		/// <summary>
 		/// Constructor
 		/// </summary>
-		/// <param name="Message"></param>
-		public CbWriterException(string Message)
-			: base(Message)
+		/// <param name="message"></param>
+		public CbWriterException(string message)
+			: base(message)
 		{
 		}
 
 		/// <summary>
 		/// Constructor
 		/// </summary>
-		/// <param name="Message"></param>
-		/// <param name="Ex"></param>
-		public CbWriterException(string Message, Exception? Ex)
-			: base(Message, Ex)
+		/// <param name="message"></param>
+		/// <param name="ex"></param>
+		public CbWriterException(string message, Exception? ex)
+			: base(message, ex)
 		{
 		}
 	}
@@ -47,28 +47,28 @@ namespace EpicGames.Serialization
 		/// </summary>
 		class Scope
 		{
-			public CbFieldType FieldType;
-			public CbFieldType UniformFieldType;
-			public int Offset; // Offset to insert the length/count
-			public int Length; // Excludes the size of this field's headers, and child fields' headers.
-			public int Count;
-			public List<Scope> Children = new List<Scope>();
-			public int SizeOfChildHeaders; // Sum of additional headers for child items, recursively.
+			public CbFieldType _fieldType;
+			public CbFieldType _uniformFieldType;
+			public int _offset; // Offset to insert the length/count
+			public int _length; // Excludes the size of this field's headers, and child fields' headers.
+			public int _count;
+			public List<Scope> _children = new List<Scope>();
+			public int _sizeOfChildHeaders; // Sum of additional headers for child items, recursively.
 
-			public Scope(CbFieldType FieldType, CbFieldType UniformFieldType, int Offset)
+			public Scope(CbFieldType fieldType, CbFieldType uniformFieldType, int offset)
 			{
-				Reset(FieldType, UniformFieldType, Offset);
+				Reset(fieldType, uniformFieldType, offset);
 			}
 
-			public void Reset(CbFieldType FieldType, CbFieldType UniformFieldType, int Offset)
+			public void Reset(CbFieldType fieldType, CbFieldType uniformFieldType, int offset)
 			{
-				this.FieldType = FieldType;
-				this.UniformFieldType = UniformFieldType;
-				this.Offset = Offset;
-				this.Length = 0;
-				this.Count = 0;
-				this.Children.Clear();
-				this.SizeOfChildHeaders = 0;
+				this._fieldType = fieldType;
+				this._uniformFieldType = uniformFieldType;
+				this._offset = offset;
+				this._length = 0;
+				this._count = 0;
+				this._children.Clear();
+				this._sizeOfChildHeaders = 0;
 			}
 		}
 
@@ -77,34 +77,34 @@ namespace EpicGames.Serialization
 		/// </summary>
 		class Chunk
 		{
-			public int Offset;
-			public int Length;
-			public byte[] Data;
-			public List<Scope> Scopes = new List<Scope>();
+			public int _offset;
+			public int _length;
+			public byte[] _data;
+			public List<Scope> _scopes = new List<Scope>();
 
-			public Chunk(int Offset, int MaxLength)
+			public Chunk(int offset, int maxLength)
 			{
-				this.Data = new byte[MaxLength];
-				Reset(Offset);
+				this._data = new byte[maxLength];
+				Reset(offset);
 			}
 
-			public void Reset(int Offset)
+			public void Reset(int offset)
 			{
-				this.Offset = Offset;
-				this.Length = 0;
-				this.Scopes.Clear();
+				this._offset = offset;
+				this._length = 0;
+				this._scopes.Clear();
 			}
 		}
 
 		const int DefaultChunkSize = 1024;
 
-		List<Chunk> Chunks = new List<Chunk>();
-		Stack<Scope> OpenScopes = new Stack<Scope>();
-		Chunk CurrentChunk => Chunks[Chunks.Count - 1];
-		Scope CurrentScope => OpenScopes.Peek();
-		int CurrentOffset;
-		List<Chunk> FreeChunks = new List<Chunk>();
-		List<Scope> FreeScopes = new List<Scope>();
+		List<Chunk> _chunks = new List<Chunk>();
+		Stack<Scope> _openScopes = new Stack<Scope>();
+		Chunk CurrentChunk => _chunks[_chunks.Count - 1];
+		Scope CurrentScope => _openScopes.Peek();
+		int _currentOffset;
+		List<Chunk> _freeChunks = new List<Chunk>();
+		List<Scope> _freeScopes = new List<Scope>();
 
 		/// <summary>
 		/// Constructor
@@ -117,11 +117,11 @@ namespace EpicGames.Serialization
 		/// <summary>
 		/// Constructor
 		/// </summary>
-		/// <param name="Reserve">Amount of data to reserve for output</param>
-		public CbWriter(int Reserve)
+		/// <param name="reserve">Amount of data to reserve for output</param>
+		public CbWriter(int reserve)
 		{
-			Chunks.Add(new Chunk(0, Reserve));
-			OpenScopes.Push(new Scope(CbFieldType.Array, CbFieldType.None, 0));
+			_chunks.Add(new Chunk(0, reserve));
+			_openScopes.Push(new Scope(CbFieldType.Array, CbFieldType.None, 0));
 		}
 
 		/// <summary>
@@ -129,57 +129,57 @@ namespace EpicGames.Serialization
 		/// </summary>
 		public void Clear()
 		{
-			foreach (Chunk Chunk in Chunks)
+			foreach (Chunk chunk in _chunks)
 			{
-				FreeChunk(Chunk);
+				FreeChunk(chunk);
 			}
 
-			CurrentOffset = 0;
+			_currentOffset = 0;
 
-			Chunks.Clear();
-			Chunks.Add(AllocChunk(0, DefaultChunkSize));
+			_chunks.Clear();
+			_chunks.Add(AllocChunk(0, DefaultChunkSize));
 
-			OpenScopes.Clear();
-			OpenScopes.Push(AllocScope(CbFieldType.Array, CbFieldType.None, 0));
+			_openScopes.Clear();
+			_openScopes.Push(AllocScope(CbFieldType.Array, CbFieldType.None, 0));
 		}
 
 		/// <summary>
 		/// Allocate a new chunk object
 		/// </summary>
-		/// <param name="Offset">Offset of the chunk</param>
-		/// <param name="MaxLength">Maximum length of the chunk</param>
+		/// <param name="offset">Offset of the chunk</param>
+		/// <param name="maxLength">Maximum length of the chunk</param>
 		/// <returns>New chunk object</returns>
-		Chunk AllocChunk(int Offset, int MaxLength)
+		Chunk AllocChunk(int offset, int maxLength)
 		{
-			for(int Idx = FreeChunks.Count - 1; Idx >= 0; Idx--)
+			for(int idx = _freeChunks.Count - 1; idx >= 0; idx--)
 			{
-				Chunk Chunk = FreeChunks[Idx];
-				if (Chunk.Data.Length >= MaxLength)
+				Chunk chunk = _freeChunks[idx];
+				if (chunk._data.Length >= maxLength)
 				{
-					FreeChunks.RemoveAt(Idx);
-					Chunk.Reset(Offset);
-					return Chunk;
+					_freeChunks.RemoveAt(idx);
+					chunk.Reset(offset);
+					return chunk;
 				}
 			}
-			return new Chunk(Offset, MaxLength);
+			return new Chunk(offset, maxLength);
 		}
 
 		/// <summary>
 		/// Adds a chunk to the free list
 		/// </summary>
-		/// <param name="Chunk"></param>
-		void FreeChunk(Chunk Chunk)
+		/// <param name="chunk"></param>
+		void FreeChunk(Chunk chunk)
 		{
 			// Add the scopes to the free list
-			FreeScopes.AddRange(Chunk.Scopes);
-			Chunk.Scopes.Clear();
+			_freeScopes.AddRange(chunk._scopes);
+			chunk._scopes.Clear();
 
 			// Insert it into the free list, sorted by descending size
-			for (int Idx = 0; ; Idx++)
+			for (int idx = 0; ; idx++)
 			{
-				if (Idx == FreeChunks.Count || Chunk.Data.Length >= FreeChunks[Idx].Data.Length)
+				if (idx == _freeChunks.Count || chunk._data.Length >= _freeChunks[idx]._data.Length)
 				{
-					FreeChunks.Insert(Idx, Chunk);
+					_freeChunks.Insert(idx, chunk);
 					break;
 				}
 			}
@@ -188,55 +188,55 @@ namespace EpicGames.Serialization
 		/// <summary>
 		/// Allocate a scope object
 		/// </summary>
-		/// <param name="FieldType"></param>
-		/// <param name="UniformFieldType"></param>
-		/// <param name="Offset"></param>
+		/// <param name="fieldType"></param>
+		/// <param name="uniformFieldType"></param>
+		/// <param name="offset"></param>
 		/// <returns></returns>
-		Scope AllocScope(CbFieldType FieldType, CbFieldType UniformFieldType, int Offset)
+		Scope AllocScope(CbFieldType fieldType, CbFieldType uniformFieldType, int offset)
 		{
-			if (FreeScopes.Count > 0)
+			if (_freeScopes.Count > 0)
 			{
-				Scope Scope = FreeScopes[FreeScopes.Count - 1];
-				Scope.Reset(FieldType, UniformFieldType, Offset);
-				FreeScopes.RemoveAt(FreeScopes.Count - 1);
-				return Scope;
+				Scope scope = _freeScopes[_freeScopes.Count - 1];
+				scope.Reset(fieldType, uniformFieldType, offset);
+				_freeScopes.RemoveAt(_freeScopes.Count - 1);
+				return scope;
 			}
-			return new Scope(FieldType, UniformFieldType, Offset);
+			return new Scope(fieldType, uniformFieldType, offset);
 		}
 
 		/// <summary>
 		/// Ensure that a block of contiguous memory of the given length is available in the output buffer
 		/// </summary>
-		/// <param name="Length"></param>
+		/// <param name="length"></param>
 		/// <returns>The allocated memory</returns>
-		Memory<byte> Allocate(int Length)
+		Memory<byte> Allocate(int length)
 		{
-			Chunk LastChunk = CurrentChunk;
-			if (LastChunk.Length + Length > LastChunk.Data.Length)
+			Chunk lastChunk = CurrentChunk;
+			if (lastChunk._length + length > lastChunk._data.Length)
 			{
-				int ChunkSize = Math.Max(Length, DefaultChunkSize);
-				LastChunk = AllocChunk(CurrentOffset, ChunkSize);
-				Chunks.Add(LastChunk);
+				int chunkSize = Math.Max(length, DefaultChunkSize);
+				lastChunk = AllocChunk(_currentOffset, chunkSize);
+				_chunks.Add(lastChunk);
 			}
 
-			Memory<byte> Buffer = LastChunk.Data.AsMemory(LastChunk.Length, Length);
-			LastChunk.Length += Length;
-			CurrentOffset += Length;
-			return Buffer;
+			Memory<byte> buffer = lastChunk._data.AsMemory(lastChunk._length, length);
+			lastChunk._length += length;
+			_currentOffset += length;
+			return buffer;
 		}
 
 		/// <summary>
 		/// Insert a new scope
 		/// </summary>
-		/// <param name="FieldType"></param>
-		/// <param name="UniformFieldType"></param>
-		void PushScope(CbFieldType FieldType, CbFieldType UniformFieldType)
+		/// <param name="fieldType"></param>
+		/// <param name="uniformFieldType"></param>
+		void PushScope(CbFieldType fieldType, CbFieldType uniformFieldType)
 		{
-			Scope NewScope = AllocScope(FieldType, UniformFieldType, CurrentOffset);
-			CurrentScope.Children.Add(NewScope);
-			OpenScopes.Push(NewScope);
+			Scope newScope = AllocScope(fieldType, uniformFieldType, _currentOffset);
+			CurrentScope._children.Add(newScope);
+			_openScopes.Push(newScope);
 
-			CurrentChunk.Scopes.Add(NewScope);
+			CurrentChunk._scopes.Add(newScope);
 		}
 
 		/// <summary>
@@ -244,92 +244,92 @@ namespace EpicGames.Serialization
 		/// </summary>
 		void PopScope()
 		{
-			Scope Scope = CurrentScope;
-			Scope.Length = CurrentOffset - Scope.Offset;
-			Scope.SizeOfChildHeaders = ComputeSizeOfChildHeaders(Scope);
-			OpenScopes.Pop();
+			Scope scope = CurrentScope;
+			scope._length = _currentOffset - scope._offset;
+			scope._sizeOfChildHeaders = ComputeSizeOfChildHeaders(scope);
+			_openScopes.Pop();
 		}
 
 		/// <summary>
 		/// Writes the header for an unnamed field
 		/// </summary>
-		/// <param name="Type"></param>
-		void WriteFieldHeader(CbFieldType Type)
+		/// <param name="type"></param>
+		void WriteFieldHeader(CbFieldType type)
 		{
-			Scope Scope = CurrentScope;
-			if (!CbFieldUtils.IsArray(Scope.FieldType))
+			Scope scope = CurrentScope;
+			if (!CbFieldUtils.IsArray(scope._fieldType))
 			{
-				throw new CbWriterException($"Anonymous fields are not allowed within fields of type {Scope.FieldType}");
+				throw new CbWriterException($"Anonymous fields are not allowed within fields of type {scope._fieldType}");
 			}
 			
-			if (Scope.UniformFieldType == CbFieldType.None)
+			if (scope._uniformFieldType == CbFieldType.None)
 			{
-				Allocate(1).Span[0] = (byte)Type;
+				Allocate(1).Span[0] = (byte)type;
 			}
-			else if (Scope.UniformFieldType != Type)
+			else if (scope._uniformFieldType != type)
 			{
-				throw new CbWriterException($"Mismatched type for uniform array - expected {Scope.UniformFieldType}, not {Type}");
+				throw new CbWriterException($"Mismatched type for uniform array - expected {scope._uniformFieldType}, not {type}");
 			}
-			Scope.Count++;
+			scope._count++;
 		}
 
 		/// <summary>
 		/// Writes the header for a named field
 		/// </summary>
-		/// <param name="Type"></param>
-		/// <param name="Name"></param>
-		void WriteFieldHeader(CbFieldType Type, Utf8String Name)
+		/// <param name="type"></param>
+		/// <param name="name"></param>
+		void WriteFieldHeader(CbFieldType type, Utf8String name)
 		{
-			Scope Scope = CurrentScope;
-			if (!CbFieldUtils.IsObject(Scope.FieldType))
+			Scope scope = CurrentScope;
+			if (!CbFieldUtils.IsObject(scope._fieldType))
 			{
-				throw new CbWriterException($"Named fields are not allowed within fields of type {Scope.FieldType}");
+				throw new CbWriterException($"Named fields are not allowed within fields of type {scope._fieldType}");
 			}
 
-			int NameVarIntLength = VarInt.Measure(Name.Length);
-			if (Scope.UniformFieldType == CbFieldType.None)
+			int nameVarIntLength = VarInt.Measure(name.Length);
+			if (scope._uniformFieldType == CbFieldType.None)
 			{
-				Span<byte> Buffer = Allocate(1 + NameVarIntLength + Name.Length).Span;
-				Buffer[0] = (byte)(Type | CbFieldType.HasFieldName);
-				WriteBinaryPayload(Buffer[1..], Name.Span);
+				Span<byte> buffer = Allocate(1 + nameVarIntLength + name.Length).Span;
+				buffer[0] = (byte)(type | CbFieldType.HasFieldName);
+				WriteBinaryPayload(buffer[1..], name.Span);
 			}
 			else
 			{
-				if (Scope.UniformFieldType != Type)
+				if (scope._uniformFieldType != type)
 				{
-					throw new CbWriterException($"Mismatched type for uniform object - expected {Scope.UniformFieldType}, not {Type}");
+					throw new CbWriterException($"Mismatched type for uniform object - expected {scope._uniformFieldType}, not {type}");
 				}
-				WriteBinaryPayload(Name.Span);
+				WriteBinaryPayload(name.Span);
 			}
-			Scope.Count++;
+			scope._count++;
 		}
 
 		/// <summary>
 		/// Copies an entire field value to the output
 		/// </summary>
-		/// <param name="Field"></param>
-		public void WriteFieldValue(CbField Field)
+		/// <param name="field"></param>
+		public void WriteFieldValue(CbField field)
 		{
-			WriteFieldHeader(Field.GetType());
-			int Size = (int)Field.GetPayloadSize();
-			Field.GetPayloadView().CopyTo(Allocate(Size));
+			WriteFieldHeader(field.GetType());
+			int size = (int)field.GetPayloadSize();
+			field.GetPayloadView().CopyTo(Allocate(size));
 		}
 
 		/// <summary>
 		/// Copies an entire field value to the output, using the name from the field
 		/// </summary>
-		/// <param name="Field"></param>
-		public void WriteField(CbField Field) => WriteField(Field.GetName(), Field);
+		/// <param name="field"></param>
+		public void WriteField(CbField field) => WriteField(field.GetName(), field);
 
 		/// <summary>
 		/// Copies an entire field value to the output
 		/// </summary>
-		/// <param name="Field"></param>
-		public void WriteField(Utf8String Name, CbField Field)
+		/// <param name="field"></param>
+		public void WriteField(Utf8String name, CbField field)
 		{
-			WriteFieldHeader(Field.GetType(), Name);
-			int Size = (int)Field.GetPayloadSize();
-			Field.GetPayloadView().CopyTo(Allocate(Size));
+			WriteFieldHeader(field.GetType(), name);
+			int size = (int)field.GetPayloadSize();
+			field.GetPayloadView().CopyTo(Allocate(size));
 		}
 
 		/// <summary>
@@ -344,10 +344,10 @@ namespace EpicGames.Serialization
 		/// <summary>
 		/// Begin writing an object field
 		/// </summary>
-		/// <param name="Name">Name of the field</param>
-		public void BeginObject(Utf8String Name)
+		/// <param name="name">Name of the field</param>
+		public void BeginObject(Utf8String name)
 		{
-			WriteFieldHeader(CbFieldType.Object, Name);
+			WriteFieldHeader(CbFieldType.Object, name);
 			PushScope(CbFieldType.Object, CbFieldType.None);
 		}
 
@@ -371,10 +371,10 @@ namespace EpicGames.Serialization
 		/// <summary>
 		/// Begin writing a named array field
 		/// </summary>
-		/// <param name="Name"></param>
-		public void BeginArray(Utf8String Name)
+		/// <param name="name"></param>
+		public void BeginArray(Utf8String name)
 		{
-			WriteFieldHeader(CbFieldType.Array, Name);
+			WriteFieldHeader(CbFieldType.Array, name);
 			PushScope(CbFieldType.Array, CbFieldType.None);
 		}
 
@@ -389,24 +389,24 @@ namespace EpicGames.Serialization
 		/// <summary>
 		/// Begin writing a uniform array field
 		/// </summary>
-		/// <param name="FieldType">The field type for elements in the array</param>
-		public void BeginUniformArray(CbFieldType FieldType)
+		/// <param name="fieldType">The field type for elements in the array</param>
+		public void BeginUniformArray(CbFieldType fieldType)
 		{
 			WriteFieldHeader(CbFieldType.UniformArray);
-			PushScope(CbFieldType.UniformArray, FieldType);
-			Allocate(1).Span[0] = (byte)FieldType;
+			PushScope(CbFieldType.UniformArray, fieldType);
+			Allocate(1).Span[0] = (byte)fieldType;
 		}
 
 		/// <summary>
 		/// Begin writing a named uniform array field
 		/// </summary>
-		/// <param name="Name">Name of the field</param>
-		/// <param name="FieldType">The field type for elements in the array</param>
-		public void BeginUniformArray(Utf8String Name, CbFieldType FieldType)
+		/// <param name="name">Name of the field</param>
+		/// <param name="fieldType">The field type for elements in the array</param>
+		public void BeginUniformArray(Utf8String name, CbFieldType fieldType)
 		{
-			WriteFieldHeader(CbFieldType.UniformArray, Name);
-			PushScope(CbFieldType.UniformArray, FieldType);
-			Allocate(1).Span[0] = (byte)FieldType;
+			WriteFieldHeader(CbFieldType.UniformArray, name);
+			PushScope(CbFieldType.UniformArray, fieldType);
+			Allocate(1).Span[0] = (byte)fieldType;
 		}
 
 		/// <summary>
@@ -428,411 +428,411 @@ namespace EpicGames.Serialization
 		/// <summary>
 		/// Write a named null field
 		/// </summary>
-		/// <param name="Name">Name of the field</param>
-		public void WriteNull(Utf8String Name)
+		/// <param name="name">Name of the field</param>
+		public void WriteNull(Utf8String name)
 		{
-			WriteFieldHeader(CbFieldType.Null, Name);
+			WriteFieldHeader(CbFieldType.Null, name);
 		}
 
 		/// <summary>
 		/// Writes a boolean value
 		/// </summary>
-		/// <param name="Value"></param>
-		public void WriteBoolValue(bool Value)
+		/// <param name="value"></param>
+		public void WriteBoolValue(bool value)
 		{
-			WriteFieldHeader(Value? CbFieldType.BoolTrue : CbFieldType.BoolFalse);
+			WriteFieldHeader(value? CbFieldType.BoolTrue : CbFieldType.BoolFalse);
 		}
 
 		/// <summary>
 		/// Writes a boolean value
 		/// </summary>
-		/// <param name="Name">Name of the field</param>
-		/// <param name="Value"></param>
-		public void WriteBool(Utf8String Name, bool Value)
+		/// <param name="name">Name of the field</param>
+		/// <param name="value"></param>
+		public void WriteBool(Utf8String name, bool value)
 		{
-			WriteFieldHeader(Value ? CbFieldType.BoolTrue : CbFieldType.BoolFalse, Name);
+			WriteFieldHeader(value ? CbFieldType.BoolTrue : CbFieldType.BoolFalse, name);
 		}
 
 		/// <summary>
 		/// Writes the payload for an integer
 		/// </summary>
-		/// <param name="Value">Value to write</param>
-		void WriteIntegerPayload(ulong Value)
+		/// <param name="value">Value to write</param>
+		void WriteIntegerPayload(ulong value)
 		{
-			int Length = VarInt.Measure(Value);
-			Span<byte> Buffer = Allocate(Length).Span;
-			VarInt.Write(Buffer, Value);
+			int length = VarInt.Measure(value);
+			Span<byte> buffer = Allocate(length).Span;
+			VarInt.Write(buffer, value);
 		}
 
 		/// <summary>
 		/// Writes an unnamed integer field
 		/// </summary>
-		/// <param name="Value">Value to be written</param>
-		public void WriteIntegerValue(int Value)
+		/// <param name="value">Value to be written</param>
+		public void WriteIntegerValue(int value)
 		{
-			WriteIntegerValue((long)Value);
+			WriteIntegerValue((long)value);
 		}
 
 		/// <summary>
 		/// Writes an unnamed integer field
 		/// </summary>
-		/// <param name="Value">Value to be written</param>
-		public void WriteIntegerValue(long Value)
+		/// <param name="value">Value to be written</param>
+		public void WriteIntegerValue(long value)
 		{
-			if (Value >= 0)
+			if (value >= 0)
 			{
 				WriteFieldHeader(CbFieldType.IntegerPositive);
-				WriteIntegerPayload((ulong)Value);
+				WriteIntegerPayload((ulong)value);
 			}
 			else
 			{
 				WriteFieldHeader(CbFieldType.IntegerNegative);
-				WriteIntegerPayload((ulong)-Value);
+				WriteIntegerPayload((ulong)-value);
 			}
 		}
 
 		/// <summary>
 		/// Writes an named integer field
 		/// </summary>
-		/// <param name="Name">Name of the field</param>
-		/// <param name="Value">Value to be written</param>
-		public void WriteInteger(Utf8String Name, int Value)
+		/// <param name="name">Name of the field</param>
+		/// <param name="value">Value to be written</param>
+		public void WriteInteger(Utf8String name, int value)
 		{
-			WriteInteger(Name, (long)Value);
+			WriteInteger(name, (long)value);
 		}
 
 		/// <summary>
 		/// Writes an named integer field
 		/// </summary>
-		/// <param name="Name">Name of the field</param>
-		/// <param name="Value">Value to be written</param>
-		public void WriteInteger(Utf8String Name, long Value)
+		/// <param name="name">Name of the field</param>
+		/// <param name="value">Value to be written</param>
+		public void WriteInteger(Utf8String name, long value)
 		{
-			if (Value >= 0)
+			if (value >= 0)
 			{
-				WriteFieldHeader(CbFieldType.IntegerPositive, Name);
-				WriteIntegerPayload((ulong)Value);
+				WriteFieldHeader(CbFieldType.IntegerPositive, name);
+				WriteIntegerPayload((ulong)value);
 			}
 			else
 			{
-				WriteFieldHeader(CbFieldType.IntegerNegative, Name);
-				WriteIntegerPayload((ulong)-Value);
+				WriteFieldHeader(CbFieldType.IntegerNegative, name);
+				WriteIntegerPayload((ulong)-value);
 			}
 		}
 
 		/// <summary>
 		/// Writes an unnamed integer field
 		/// </summary>
-		/// <param name="Value">Value to be written</param>
-		public void WriteIntegerValue(ulong Value)
+		/// <param name="value">Value to be written</param>
+		public void WriteIntegerValue(ulong value)
 		{
 			WriteFieldHeader(CbFieldType.IntegerPositive);
-			WriteIntegerPayload(Value);
+			WriteIntegerPayload(value);
 		}
 
 		/// <summary>
 		/// Writes a named integer field
 		/// </summary>
-		/// <param name="Name">Name of the field</param>
-		/// <param name="Value">Value to be written</param>
-		public void WriteInteger(Utf8String Name, ulong Value)
+		/// <param name="name">Name of the field</param>
+		/// <param name="value">Value to be written</param>
+		public void WriteInteger(Utf8String name, ulong value)
 		{
-			WriteFieldHeader(CbFieldType.IntegerPositive, Name);
-			WriteIntegerPayload(Value);
+			WriteFieldHeader(CbFieldType.IntegerPositive, name);
+			WriteIntegerPayload(value);
 		}
 
 		/// <summary>
 		/// Writes an unnamed double field
 		/// </summary>
-		/// <param name="Value">Value to be written</param>
-		public void WriteDoubleValue(double Value)
+		/// <param name="value">Value to be written</param>
+		public void WriteDoubleValue(double value)
 		{
 			WriteFieldHeader(CbFieldType.Float64);
-			BinaryPrimitives.WriteInt64BigEndian(Allocate(sizeof(double)).Span, BitConverter.DoubleToInt64Bits(Value));
+			BinaryPrimitives.WriteInt64BigEndian(Allocate(sizeof(double)).Span, BitConverter.DoubleToInt64Bits(value));
 		}
 
 		/// <summary>
 		/// Writes a named double field
 		/// </summary>
-		/// <param name="Name">Name of the field</param>
-		/// <param name="Value">Value to be written</param>
-		public void WriteDouble(Utf8String Name, double Value)
+		/// <param name="name">Name of the field</param>
+		/// <param name="value">Value to be written</param>
+		public void WriteDouble(Utf8String name, double value)
 		{
-			WriteFieldHeader(CbFieldType.Float64, Name);
-			BinaryPrimitives.WriteInt64BigEndian(Allocate(sizeof(double)).Span, BitConverter.DoubleToInt64Bits(Value));
+			WriteFieldHeader(CbFieldType.Float64, name);
+			BinaryPrimitives.WriteInt64BigEndian(Allocate(sizeof(double)).Span, BitConverter.DoubleToInt64Bits(value));
 		}
 
 		/// <summary>
 		/// Writes the payload for a <see cref="DateTime"/> value
 		/// </summary>
-		/// <param name="DateTime">The value to write</param>
-		void WriteDateTimePayload(DateTime DateTime)
+		/// <param name="dateTime">The value to write</param>
+		void WriteDateTimePayload(DateTime dateTime)
 		{
-			Span<byte> Buffer = Allocate(sizeof(long)).Span;
-			BinaryPrimitives.WriteInt64BigEndian(Buffer, DateTime.Ticks);
+			Span<byte> buffer = Allocate(sizeof(long)).Span;
+			BinaryPrimitives.WriteInt64BigEndian(buffer, dateTime.Ticks);
 		}
 
 		/// <summary>
 		/// Writes an unnamed <see cref="DateTime"/> field
 		/// </summary>
-		/// <param name="Value">Value to be written</param>
-		public void WriteDateTimeValue(DateTime Value)
+		/// <param name="value">Value to be written</param>
+		public void WriteDateTimeValue(DateTime value)
 		{
 			WriteFieldHeader(CbFieldType.DateTime);
-			WriteDateTimePayload(Value);
+			WriteDateTimePayload(value);
 		}
 
 		/// <summary>
 		/// Writes a named <see cref="DateTime"/> field
 		/// </summary>
-		/// <param name="Name">Name of the field</param>
-		/// <param name="Value">Value to be written</param>
-		public void WriteDateTime(Utf8String Name, DateTime Value)
+		/// <param name="name">Name of the field</param>
+		/// <param name="value">Value to be written</param>
+		public void WriteDateTime(Utf8String name, DateTime value)
 		{
-			WriteFieldHeader(CbFieldType.DateTime, Name);
-			WriteDateTimePayload(Value);
+			WriteFieldHeader(CbFieldType.DateTime, name);
+			WriteDateTimePayload(value);
 		}
 
 		/// <summary>
 		/// Writes the payload for a hash
 		/// </summary>
-		/// <param name="Hash"></param>
-		void WriteHashPayload(IoHash Hash)
+		/// <param name="hash"></param>
+		void WriteHashPayload(IoHash hash)
 		{
-			Span<byte> Buffer = Allocate(IoHash.NumBytes).Span;
-			Hash.CopyTo(Buffer);
+			Span<byte> buffer = Allocate(IoHash.NumBytes).Span;
+			hash.CopyTo(buffer);
 		}
 
 		/// <summary>
 		/// Writes an unnamed <see cref="IoHash"/> field
 		/// </summary>
-		/// <param name="Hash"></param>
-		public void WriteHashValue(IoHash Hash)
+		/// <param name="hash"></param>
+		public void WriteHashValue(IoHash hash)
 		{
 			WriteFieldHeader(CbFieldType.Hash);
-			WriteHashPayload(Hash);
+			WriteHashPayload(hash);
 		}
 
 		/// <summary>
 		/// Writes a named <see cref="IoHash"/> field
 		/// </summary>
-		/// <param name="Name">Name of the field</param>
-		/// <param name="Value">Value to be written</param>
-		public void WriteHash(Utf8String Name, IoHash Value)
+		/// <param name="name">Name of the field</param>
+		/// <param name="value">Value to be written</param>
+		public void WriteHash(Utf8String name, IoHash value)
 		{
-			WriteFieldHeader(CbFieldType.Hash, Name);
-			WriteHashPayload(Value);
+			WriteFieldHeader(CbFieldType.Hash, name);
+			WriteHashPayload(value);
 		}
 
 		/// <summary>
 		/// Writes an unnamed reference to a binary attachment
 		/// </summary>
-		/// <param name="Hash">Hash of the attachment</param>
-		public void WriteBinaryAttachmentValue(IoHash Hash)
+		/// <param name="hash">Hash of the attachment</param>
+		public void WriteBinaryAttachmentValue(IoHash hash)
 		{
 			WriteFieldHeader(CbFieldType.BinaryAttachment);
-			WriteHashPayload(Hash);
+			WriteHashPayload(hash);
 		}
 
 		/// <summary>
 		/// Writes a named reference to a binary attachment
 		/// </summary>
-		/// <param name="Name">Name of the field</param>
-		/// <param name="Hash">Hash of the attachment</param>
-		public void WriteBinaryAttachment(Utf8String Name, IoHash Hash)
+		/// <param name="name">Name of the field</param>
+		/// <param name="hash">Hash of the attachment</param>
+		public void WriteBinaryAttachment(Utf8String name, IoHash hash)
 		{
-			WriteFieldHeader(CbFieldType.BinaryAttachment, Name);
-			WriteHashPayload(Hash);
+			WriteFieldHeader(CbFieldType.BinaryAttachment, name);
+			WriteHashPayload(hash);
 		}
 
 		/// <summary>
 		/// Writes the payload for an object to the buffer
 		/// </summary>
-		/// <param name="Object"></param>
-		void WriteObjectPayload(CbObject Object)
+		/// <param name="obj"></param>
+		void WriteObjectPayload(CbObject obj)
 		{
-			CbField Field = Object.AsField();
-			Memory<byte> Buffer = Allocate(Field.Payload.Length);
-			Field.Payload.CopyTo(Buffer);
+			CbField field = obj.AsField();
+			Memory<byte> buffer = Allocate(field.Payload.Length);
+			field.Payload.CopyTo(buffer);
 		}
 
 		/// <summary>
 		/// Writes an object directly into the writer
 		/// </summary>
-		/// <param name="Object">Object to write</param>
-		public void WriteObject(CbObject Object)
+		/// <param name="obj">Object to write</param>
+		public void WriteObject(CbObject obj)
 		{
 			WriteFieldHeader(CbFieldType.Object);
-			WriteObjectPayload(Object);
+			WriteObjectPayload(obj);
 		}
 
 		/// <summary>
 		/// Writes an object directly into the writer
 		/// </summary>
-		/// <param name="Name">Name of the object</param>
-		/// <param name="Object">Object to write</param>
-		public void WriteObject(Utf8String Name, CbObject Object)
+		/// <param name="name">Name of the object</param>
+		/// <param name="obj">Object to write</param>
+		public void WriteObject(Utf8String name, CbObject obj)
 		{
-			WriteFieldHeader(CbFieldType.Object, Name);
-			WriteObjectPayload(Object);
+			WriteFieldHeader(CbFieldType.Object, name);
+			WriteObjectPayload(obj);
 		}
 
 		/// <summary>
 		/// Writes an unnamed reference to an object attachment
 		/// </summary>
-		/// <param name="Hash">Hash of the attachment</param>
-		public void WriteObjectAttachmentValue(IoHash Hash)
+		/// <param name="hash">Hash of the attachment</param>
+		public void WriteObjectAttachmentValue(IoHash hash)
 		{
 			WriteFieldHeader(CbFieldType.ObjectAttachment);
-			WriteHashPayload(Hash);
+			WriteHashPayload(hash);
 		}
 
 		/// <summary>
 		/// Writes a named reference to an object attachment
 		/// </summary>
-		/// <param name="Name">Name of the field</param>
-		/// <param name="Hash">Hash of the attachment</param>
-		public void WriteObjectAttachment(Utf8String Name, IoHash Hash)
+		/// <param name="name">Name of the field</param>
+		/// <param name="hash">Hash of the attachment</param>
+		public void WriteObjectAttachment(Utf8String name, IoHash hash)
 		{
-			WriteFieldHeader(CbFieldType.ObjectAttachment, Name);
-			WriteHashPayload(Hash);
+			WriteFieldHeader(CbFieldType.ObjectAttachment, name);
+			WriteHashPayload(hash);
 		}
 
 		/// <summary>
 		/// Writes the payload for a binary value
 		/// </summary>
-		/// <param name="Output">Output buffer</param>
-		/// <param name="Value">Value to be written</param>
-		static void WriteBinaryPayload(Span<byte> Output, ReadOnlySpan<byte> Value)
+		/// <param name="output">Output buffer</param>
+		/// <param name="value">Value to be written</param>
+		static void WriteBinaryPayload(Span<byte> output, ReadOnlySpan<byte> value)
 		{
-			int VarIntLength = VarInt.Write(Output, Value.Length);
-			Output = Output[VarIntLength..];
+			int varIntLength = VarInt.Write(output, value.Length);
+			output = output[varIntLength..];
 
-			Value.CopyTo(Output);
-			CheckSize(Output, Value.Length);
+			value.CopyTo(output);
+			CheckSize(output, value.Length);
 		}
 
 		/// <summary>
 		/// Writes the payload for a binary value
 		/// </summary>
-		/// <param name="Value">Value to be written</param>
-		void WriteBinaryPayload(ReadOnlySpan<byte> Value)
+		/// <param name="value">Value to be written</param>
+		void WriteBinaryPayload(ReadOnlySpan<byte> value)
 		{
-			int ValueVarIntLength = VarInt.Measure(Value.Length);
-			Span<byte> Buffer = Allocate(ValueVarIntLength + Value.Length).Span;
-			WriteBinaryPayload(Buffer, Value);
+			int valueVarIntLength = VarInt.Measure(value.Length);
+			Span<byte> buffer = Allocate(valueVarIntLength + value.Length).Span;
+			WriteBinaryPayload(buffer, value);
 		}
 
 		/// <summary>
 		/// Writes an unnamed string value
 		/// </summary>
-		/// <param name="Value">Value to be written</param>
-		public void WriteStringValue(string Value) => WriteUtf8StringValue(Value);
+		/// <param name="value">Value to be written</param>
+		public void WriteStringValue(string value) => WriteUtf8StringValue(value);
 
 		/// <summary>
 		/// Writes a named string value
 		/// </summary>
-		/// <param name="Name">Name of the field</param>
-		/// <param name="Value">Value to be written</param>
-		public void WriteString(Utf8String Name, string? Value)
+		/// <param name="name">Name of the field</param>
+		/// <param name="value">Value to be written</param>
+		public void WriteString(Utf8String name, string? value)
 		{
-			if(Value != null)
+			if(value != null)
 			{
-				WriteUtf8String(Name, Value);
+				WriteUtf8String(name, value);
 			}
 		}
 
 		/// <summary>
 		/// Writes an unnamed string value
 		/// </summary>
-		/// <param name="Value">Value to be written</param>
-		public void WriteUtf8StringValue(Utf8String Value)
+		/// <param name="value">Value to be written</param>
+		public void WriteUtf8StringValue(Utf8String value)
 		{
 			WriteFieldHeader(CbFieldType.String);
-			WriteBinaryPayload(Value.Span);
+			WriteBinaryPayload(value.Span);
 		}
 
 		/// <summary>
 		/// Writes a named string value
 		/// </summary>
-		/// <param name="Name">Name of the field</param>
-		/// <param name="Value">Value to be written</param>
-		public void WriteUtf8String(Utf8String Name, Utf8String Value)
+		/// <param name="name">Name of the field</param>
+		/// <param name="value">Value to be written</param>
+		public void WriteUtf8String(Utf8String name, Utf8String value)
 		{
-			if (Value.Length > 0)
+			if (value.Length > 0)
 			{
-				WriteFieldHeader(CbFieldType.String, Name);
-				WriteBinaryPayload(Value.Span);
+				WriteFieldHeader(CbFieldType.String, name);
+				WriteBinaryPayload(value.Span);
 			}
 		}
 
 		/// <summary>
 		/// Writes an unnamed binary value
 		/// </summary>
-		/// <param name="Value">Value to be written</param>
-		public void WriteBinarySpanValue(ReadOnlySpan<byte> Value)
+		/// <param name="value">Value to be written</param>
+		public void WriteBinarySpanValue(ReadOnlySpan<byte> value)
 		{
 			WriteFieldHeader(CbFieldType.Binary);
-			WriteBinaryPayload(Value);
+			WriteBinaryPayload(value);
 		}
 
 		/// <summary>
 		/// Writes a named binary value
 		/// </summary>
-		/// <param name="Name">Name of the field</param>
-		/// <param name="Value">Value to be written</param>
-		public void WriteBinarySpan(Utf8String Name, ReadOnlySpan<byte> Value)
+		/// <param name="name">Name of the field</param>
+		/// <param name="value">Value to be written</param>
+		public void WriteBinarySpan(Utf8String name, ReadOnlySpan<byte> value)
 		{
-			WriteFieldHeader(CbFieldType.Binary, Name);
-			WriteBinaryPayload(Value);
+			WriteFieldHeader(CbFieldType.Binary, name);
+			WriteBinaryPayload(value);
 		}
 
 		/// <summary>
 		/// Writes an unnamed binary value
 		/// </summary>
-		/// <param name="Value">Value to be written</param>
-		public void WriteBinaryValue(ReadOnlyMemory<byte> Value)
+		/// <param name="value">Value to be written</param>
+		public void WriteBinaryValue(ReadOnlyMemory<byte> value)
 		{
-			WriteBinarySpanValue(Value.Span);
+			WriteBinarySpanValue(value.Span);
 		}
 
 		/// <summary>
 		/// Writes a named binary value
 		/// </summary>
-		/// <param name="Name">Name of the field</param>
-		/// <param name="Value">Value to be written</param>
-		public void WriteBinary(Utf8String Name, ReadOnlyMemory<byte> Value)
+		/// <param name="name">Name of the field</param>
+		/// <param name="value">Value to be written</param>
+		public void WriteBinary(Utf8String name, ReadOnlyMemory<byte> value)
 		{
-			WriteBinarySpan(Name, Value.Span);
+			WriteBinarySpan(name, value.Span);
 		}
 
 		/// <summary>
 		/// Writes an unnamed binary value
 		/// </summary>
-		/// <param name="Value">Value to be written</param>
-		public void WriteBinaryArrayValue(byte[] Value)
+		/// <param name="value">Value to be written</param>
+		public void WriteBinaryArrayValue(byte[] value)
 		{
-			WriteBinarySpanValue(Value.AsSpan());
+			WriteBinarySpanValue(value.AsSpan());
 		}
 
 		/// <summary>
 		/// Writes a named binary value
 		/// </summary>
-		/// <param name="Name">Name of the field</param>
-		/// <param name="Value">Value to be written</param>
-		public void WriteBinaryArray(Utf8String Name, byte[] Value)
+		/// <param name="name">Name of the field</param>
+		/// <param name="value">Value to be written</param>
+		public void WriteBinaryArray(Utf8String name, byte[] value)
 		{
-			WriteBinarySpan(Name, Value.AsSpan());
+			WriteBinarySpan(name, value.AsSpan());
 		}
 
 		/// <summary>
 		/// Check that the given span is the required size
 		/// </summary>
-		/// <param name="Span"></param>
-		/// <param name="ExpectedSize"></param>
-		static void CheckSize(Span<byte> Span, int ExpectedSize)
+		/// <param name="span"></param>
+		/// <param name="expectedSize"></param>
+		static void CheckSize(Span<byte> span, int expectedSize)
 		{
-			if (Span.Length != ExpectedSize)
+			if (span.Length != expectedSize)
 			{
 				throw new Exception("Size of buffer is not correct");
 			}
@@ -844,13 +844,13 @@ namespace EpicGames.Serialization
 		/// <returns>Hash for the object</returns>
 		public IoHash ComputeHash()
 		{
-			using (Blake3.Hasher Hasher = Blake3.Hasher.New())
+			using (Blake3.Hasher hasher = Blake3.Hasher.New())
 			{
-				foreach (ReadOnlyMemory<byte> Segment in EnumerateSegments())
+				foreach (ReadOnlyMemory<byte> segment in EnumerateSegments())
 				{
-					Hasher.Update(Segment.Span);
+					hasher.Update(segment.Span);
 				}
-				return IoHash.FromBlake3(Hasher);
+				return IoHash.FromBlake3(hasher);
 			}
 		}
 
@@ -860,12 +860,12 @@ namespace EpicGames.Serialization
 		/// <returns></returns>
 		public int GetSize()
 		{
-			if (OpenScopes.Count != 1)
+			if (_openScopes.Count != 1)
 			{
 				throw new CbWriterException("Unfinished scope in writer");
 			}
 
-			return CurrentOffset + ComputeSizeOfChildHeaders(CurrentScope);
+			return _currentOffset + ComputeSizeOfChildHeaders(CurrentScope);
 		}
 
 		/// <summary>
@@ -876,68 +876,68 @@ namespace EpicGames.Serialization
 
 		private IEnumerable<ReadOnlyMemory<byte>> EnumerateSegments()
 		{
-			byte[] ScopeHeader = new byte[64];
+			byte[] scopeHeader = new byte[64];
 
-			int SourceOffset = 0;
-			foreach (Chunk Chunk in Chunks)
+			int sourceOffset = 0;
+			foreach (Chunk chunk in _chunks)
 			{
-				foreach (Scope Scope in Chunk.Scopes)
+				foreach (Scope scope in chunk._scopes)
 				{
-					ReadOnlyMemory<byte> SourceData = Chunk.Data.AsMemory(SourceOffset - Chunk.Offset, Scope.Offset - SourceOffset);
-					yield return SourceData;
+					ReadOnlyMemory<byte> sourceData = chunk._data.AsMemory(sourceOffset - chunk._offset, scope._offset - sourceOffset);
+					yield return sourceData;
 
-					SourceOffset += SourceData.Length;
+					sourceOffset += sourceData.Length;
 
-					int HeaderLength = WriteScopeHeader(ScopeHeader, Scope);
-					yield return ScopeHeader.AsMemory(0, HeaderLength);
+					int headerLength = WriteScopeHeader(scopeHeader, scope);
+					yield return scopeHeader.AsMemory(0, headerLength);
 				}
 
-				ReadOnlyMemory<byte> LastSourceData = Chunk.Data.AsMemory(SourceOffset - Chunk.Offset, (Chunk.Offset + Chunk.Length) - SourceOffset);
-				yield return LastSourceData;
+				ReadOnlyMemory<byte> lastSourceData = chunk._data.AsMemory(sourceOffset - chunk._offset, (chunk._offset + chunk._length) - sourceOffset);
+				yield return lastSourceData;
 
-				SourceOffset += LastSourceData.Length;
+				sourceOffset += lastSourceData.Length;
 			}
 		}
 
 		/// <summary>
 		/// Copy the data from this writer to a buffer
 		/// </summary>
-		/// <param name="Buffer"></param>
-		public void CopyTo(Span<byte> Buffer)
+		/// <param name="buffer"></param>
+		public void CopyTo(Span<byte> buffer)
 		{
-			int BufferOffset = 0;
+			int bufferOffset = 0;
 
-			int SourceOffset = 0;
-			foreach (Chunk Chunk in Chunks)
+			int sourceOffset = 0;
+			foreach (Chunk chunk in _chunks)
 			{
-				foreach (Scope Scope in Chunk.Scopes)
+				foreach (Scope scope in chunk._scopes)
 				{
-					ReadOnlySpan<byte> SourceData = Chunk.Data.AsSpan(SourceOffset - Chunk.Offset, Scope.Offset - SourceOffset);
-					SourceData.CopyTo(Buffer.Slice(BufferOffset));
+					ReadOnlySpan<byte> sourceData = chunk._data.AsSpan(sourceOffset - chunk._offset, scope._offset - sourceOffset);
+					sourceData.CopyTo(buffer.Slice(bufferOffset));
 
-					BufferOffset += SourceData.Length;
-					SourceOffset += SourceData.Length;
+					bufferOffset += sourceData.Length;
+					sourceOffset += sourceData.Length;
 
-					BufferOffset += WriteScopeHeader(Buffer.Slice(BufferOffset), Scope);
+					bufferOffset += WriteScopeHeader(buffer.Slice(bufferOffset), scope);
 				}
 
-				ReadOnlySpan<byte> LastSourceData = Chunk.Data.AsSpan(SourceOffset - Chunk.Offset, (Chunk.Offset + Chunk.Length) - SourceOffset);
-				LastSourceData.CopyTo(Buffer.Slice(BufferOffset));
-				BufferOffset += LastSourceData.Length;
-				SourceOffset += LastSourceData.Length;
+				ReadOnlySpan<byte> lastSourceData = chunk._data.AsSpan(sourceOffset - chunk._offset, (chunk._offset + chunk._length) - sourceOffset);
+				lastSourceData.CopyTo(buffer.Slice(bufferOffset));
+				bufferOffset += lastSourceData.Length;
+				sourceOffset += lastSourceData.Length;
 			}
 		}
 
 		class ReadStream : Stream
 		{
-			IEnumerator<ReadOnlyMemory<byte>> Enumerator;
-			ReadOnlyMemory<byte> Segment;
-			long PositionInternal;
+			IEnumerator<ReadOnlyMemory<byte>> _enumerator;
+			ReadOnlyMemory<byte> _segment;
+			long _positionInternal;
 
-			public ReadStream(IEnumerator<ReadOnlyMemory<byte>> Enumerator, long Length)
+			public ReadStream(IEnumerator<ReadOnlyMemory<byte>> enumerator, long length)
 			{
-				this.Enumerator = Enumerator;
-				this.Length = Length;
+				this._enumerator = enumerator;
+				this.Length = length;
 			}
 
 			/// <inheritdoc/>
@@ -955,7 +955,7 @@ namespace EpicGames.Serialization
 			/// <inheritdoc/>
 			public override long Position
 			{
-				get => PositionInternal;
+				get => _positionInternal;
 				set => throw new NotSupportedException();
 			}
 
@@ -963,42 +963,42 @@ namespace EpicGames.Serialization
 			public override void Flush() { }
 
 			/// <inheritdoc/>
-			public override int Read(Span<byte> Buffer)
+			public override int Read(Span<byte> buffer)
 			{
-				int ReadLength = 0;
-				while (ReadLength < Buffer.Length)
+				int readLength = 0;
+				while (readLength < buffer.Length)
 				{
-					while (Segment.Length == 0)
+					while (_segment.Length == 0)
 					{
-						if (!Enumerator.MoveNext())
+						if (!_enumerator.MoveNext())
 						{
-							return ReadLength;
+							return readLength;
 						}
-						Segment = Enumerator.Current;
+						_segment = _enumerator.Current;
 					}
 
-					int CopyLength = Math.Min(Segment.Length, Buffer.Length);
+					int copyLength = Math.Min(_segment.Length, buffer.Length);
 
-					Segment.Span.Slice(0, CopyLength).CopyTo(Buffer.Slice(ReadLength));
-					Segment = Segment.Slice(CopyLength);
+					_segment.Span.Slice(0, copyLength).CopyTo(buffer.Slice(readLength));
+					_segment = _segment.Slice(copyLength);
 
-					PositionInternal += CopyLength;
-					ReadLength += CopyLength;
+					_positionInternal += copyLength;
+					readLength += copyLength;
 				}
-				return ReadLength;
+				return readLength;
 			}
 
 			/// <inheritdoc/>
-			public override int Read(byte[] Buffer, int Offset, int Count) => Read(Buffer.AsSpan(Offset, Count));
+			public override int Read(byte[] buffer, int offset, int count) => Read(buffer.AsSpan(offset, count));
 
 			/// <inheritdoc/>
-			public override long Seek(long Offset, SeekOrigin Origin) => throw new NotSupportedException();
+			public override long Seek(long offset, SeekOrigin origin) => throw new NotSupportedException();
 
 			/// <inheritdoc/>
-			public override void SetLength(long Value) => throw new NotSupportedException();
+			public override void SetLength(long value) => throw new NotSupportedException();
 
 			/// <inheritdoc/>
-			public override void Write(byte[] Buffer, int Offset, int Count) => throw new NotSupportedException();
+			public override void Write(byte[] buffer, int offset, int count) => throw new NotSupportedException();
 		}
 
 		/// <summary>
@@ -1016,56 +1016,56 @@ namespace EpicGames.Serialization
 		/// <returns></returns>
 		public byte[] ToByteArray()
 		{
-			byte[] Buffer = new byte[GetSize()];
-			CopyTo(Buffer);
-			return Buffer;
+			byte[] buffer = new byte[GetSize()];
+			CopyTo(buffer);
+			return buffer;
 		}
 
 		/// <summary>
 		/// Comptues the size of any child headers
 		/// </summary>
-		/// <param name="Scope"></param>
-		static int ComputeSizeOfChildHeaders(Scope Scope)
+		/// <param name="scope"></param>
+		static int ComputeSizeOfChildHeaders(Scope scope)
 		{
-			int SizeOfChildHeaders = 0;
-			foreach (Scope ChildScope in Scope.Children)
+			int sizeOfChildHeaders = 0;
+			foreach (Scope childScope in scope._children)
 			{
-				switch (ChildScope.FieldType)
+				switch (childScope._fieldType)
 				{
 					case CbFieldType.Object:
 					case CbFieldType.UniformObject:
-						SizeOfChildHeaders += ChildScope.SizeOfChildHeaders + VarInt.Measure(ChildScope.Length + ChildScope.SizeOfChildHeaders);
+						sizeOfChildHeaders += childScope._sizeOfChildHeaders + VarInt.Measure(childScope._length + childScope._sizeOfChildHeaders);
 						break;
 					case CbFieldType.Array:
 					case CbFieldType.UniformArray:
-						int ArrayCountLength = VarInt.Measure(ChildScope.Count);
-						SizeOfChildHeaders += ChildScope.SizeOfChildHeaders + VarInt.Measure(ChildScope.Length + ChildScope.SizeOfChildHeaders + ArrayCountLength) + ArrayCountLength;
+						int arrayCountLength = VarInt.Measure(childScope._count);
+						sizeOfChildHeaders += childScope._sizeOfChildHeaders + VarInt.Measure(childScope._length + childScope._sizeOfChildHeaders + arrayCountLength) + arrayCountLength;
 						break;
 					default:
 						throw new InvalidOperationException();
 				}
 			}
-			return SizeOfChildHeaders;
+			return sizeOfChildHeaders;
 		}
 
 		/// <summary>
 		/// Writes the header for a particular scope
 		/// </summary>
-		/// <param name="Span"></param>
-		/// <param name="Scope"></param>
+		/// <param name="span"></param>
+		/// <param name="scope"></param>
 		/// <returns></returns>
-		static int WriteScopeHeader(Span<byte> Span, Scope Scope)
+		static int WriteScopeHeader(Span<byte> span, Scope scope)
 		{
-			switch (Scope.FieldType)
+			switch (scope._fieldType)
 			{
 				case CbFieldType.Object:
 				case CbFieldType.UniformObject:
-					return VarInt.Write(Span, Scope.Length + Scope.SizeOfChildHeaders);
+					return VarInt.Write(span, scope._length + scope._sizeOfChildHeaders);
 				case CbFieldType.Array:
 				case CbFieldType.UniformArray:
-					int NumItemsLength = VarInt.Measure(Scope.Count);
-					int Offset = VarInt.Write(Span, Scope.Length + Scope.SizeOfChildHeaders + NumItemsLength);
-					return Offset + VarInt.Write(Span.Slice(Offset), Scope.Count);
+					int numItemsLength = VarInt.Measure(scope._count);
+					int offset = VarInt.Write(span, scope._length + scope._sizeOfChildHeaders + numItemsLength);
+					return offset + VarInt.Write(span.Slice(offset), scope._count);
 				default:
 					throw new InvalidOperationException();
 			}
diff --git a/Engine/Source/Programs/Shared/EpicGames.Serialization/Converters/CbArrayConverter.cs b/Engine/Source/Programs/Shared/EpicGames.Serialization/Converters/CbArrayConverter.cs
--- a/Engine/Source/Programs/Shared/EpicGames.Serialization/Converters/CbArrayConverter.cs
+++ b/Engine/Source/Programs/Shared/EpicGames.Serialization/Converters/CbArrayConverter.cs
@@ -20,36 +20,36 @@ namespace EpicGames.Serialization.Converters
 	class CbArrayConverter<T> : CbConverterBase<T[]>
 	{
 		/// <inheritdoc/>
-		public override T[] Read(CbField Field)
+		public override T[] Read(CbField field)
 		{
-			List<T> List = new List<T>();
-			foreach (CbField ElementField in Field)
+			List<T> list = new List<T>();
+			foreach (CbField elementField in field)
 			{
-				List.Add(CbSerializer.Deserialize<T>(ElementField));
+				list.Add(CbSerializer.Deserialize<T>(elementField));
 			}
-			return List.ToArray();
+			return list.ToArray();
 		}
 
 		/// <inheritdoc/>
-		public override void Write(CbWriter Writer, T[] List)
+		public override void Write(CbWriter writer, T[] list)
 		{
-			Writer.BeginArray();
-			foreach (T Element in List)
+			writer.BeginArray();
+			foreach (T element in list)
 			{
-				CbSerializer.Serialize<T>(Writer, Element);
+				CbSerializer.Serialize<T>(writer, element);
 			}
-			Writer.EndArray();
+			writer.EndArray();
 		}
 
 		/// <inheritdoc/>
-		public override void WriteNamed(CbWriter Writer, Utf8String Name, T[] Array)
+		public override void WriteNamed(CbWriter writer, Utf8String name, T[] array)
 		{
-			Writer.BeginArray(Name);
-			foreach (T Element in Array)
+			writer.BeginArray(name);
+			foreach (T element in array)
 			{
-				CbSerializer.Serialize<T>(Writer, Element);
+				CbSerializer.Serialize<T>(writer, element);
 			}
-			Writer.EndArray();
+			writer.EndArray();
 		}
 	}
 	
@@ -59,13 +59,13 @@ namespace EpicGames.Serialization.Converters
 	class CbArrayConverterFactory : CbConverterFactory
 	{
 		/// <inheritdoc/>
-		public override ICbConverter? CreateConverter(Type Type)
+		public override ICbConverter? CreateConverter(Type type)
 		{
-			if (Type.IsArray)
+			if (type.IsArray)
 			{
-				Type ElementType = Type.GetElementType()!;
-				Type ConverterType = typeof(CbArrayConverter<>).MakeGenericType(ElementType);
-				return (ICbConverter)Activator.CreateInstance(ConverterType)!;
+				Type elementType = type.GetElementType()!;
+				Type converterType = typeof(CbArrayConverter<>).MakeGenericType(elementType);
+				return (ICbConverter)Activator.CreateInstance(converterType)!;
 			}
 			return null;
 		}
diff --git a/Engine/Source/Programs/Shared/EpicGames.Serialization/Converters/CbClassConverter.cs b/Engine/Source/Programs/Shared/EpicGames.Serialization/Converters/CbClassConverter.cs
--- a/Engine/Source/Programs/Shared/EpicGames.Serialization/Converters/CbClassConverter.cs
+++ b/Engine/Source/Programs/Shared/EpicGames.Serialization/Converters/CbClassConverter.cs
@@ -18,13 +18,13 @@ namespace EpicGames.Serialization.Converters
 			public static Utf8String[]? Names = null;
 			public static PropertyInfo[]? Properties = null;
 
-			public static bool MatchName(CbField Field, int Idx)
+			public static bool MatchName(CbField field, int idx)
 			{
-				return Field.Name == Names![Idx];
+				return field.Name == Names![idx];
 			}
 		}
 
-		static Utf8String DiscriminatorKey = "_t";
+		static Utf8String s_discriminatorKey = "_t";
 
 		public Type ClassType { get; }
 		public bool IsPolymorphic { get; }
@@ -38,36 +38,36 @@ namespace EpicGames.Serialization.Converters
 		public DynamicMethod ReadConcreteMethod { get; }
 		public DynamicMethod WriteConcreteContentsMethod { get; }
 
-		static Dictionary<PropertyInfo, CbClassConverterMethods> PropertyToMethods = new Dictionary<PropertyInfo, CbClassConverterMethods>();
-		static Dictionary<Type, CbClassConverterMethods> TypeToMethods = new Dictionary<Type, CbClassConverterMethods>();
+		static Dictionary<PropertyInfo, CbClassConverterMethods> s_propertyToMethods = new Dictionary<PropertyInfo, CbClassConverterMethods>();
+		static Dictionary<Type, CbClassConverterMethods> s_typeToMethods = new Dictionary<Type, CbClassConverterMethods>();
 
-		public static CbClassConverterMethods Create(Type ClassType)
+		public static CbClassConverterMethods Create(Type classType)
 		{
-			CbClassConverterMethods? Methods;
-			if (!TypeToMethods.TryGetValue(ClassType, out Methods))
+			CbClassConverterMethods? methods;
+			if (!s_typeToMethods.TryGetValue(classType, out methods))
 			{
-				Methods = new CbClassConverterMethods(ClassType);
-				TypeToMethods.Add(ClassType, Methods);
-				Methods.GenerateBytecode();
+				methods = new CbClassConverterMethods(classType);
+				s_typeToMethods.Add(classType, methods);
+				methods.GenerateBytecode();
 			}
-			return Methods;
+			return methods;
 		}
 
-		private CbClassConverterMethods(Type ClassType)
+		private CbClassConverterMethods(Type classType)
 		{
-			this.ClassType = ClassType;
-			this.IsPolymorphic = (ClassType.GetCustomAttribute<CbPolymorphicAttribute>(true) != null);
+			this.ClassType = classType;
+			this.IsPolymorphic = (classType.GetCustomAttribute<CbPolymorphicAttribute>(true) != null);
 
-			ReadConcreteMethod = new DynamicMethod($"ReadConcrete_{ClassType.Name}", ClassType, new Type[] { typeof(CbField) });
-			WriteConcreteContentsMethod = new DynamicMethod($"WriteConcreteContents_{ClassType.Name}", null, new Type[] { typeof(CbWriter), ClassType });
+			ReadConcreteMethod = new DynamicMethod($"ReadConcrete_{classType.Name}", classType, new Type[] { typeof(CbField) });
+			WriteConcreteContentsMethod = new DynamicMethod($"WriteConcreteContents_{classType.Name}", null, new Type[] { typeof(CbWriter), classType });
 
-			WriteMethod = new DynamicMethod($"Write_{ClassType.Name}", null, new Type[] { typeof(CbWriter), ClassType });
-			WriteNamedMethod = new DynamicMethod($"WriteNamed_{ClassType.Name}", null, new Type[] { typeof(CbWriter), typeof(Utf8String), ClassType });
+			WriteMethod = new DynamicMethod($"Write_{classType.Name}", null, new Type[] { typeof(CbWriter), classType });
+			WriteNamedMethod = new DynamicMethod($"WriteNamed_{classType.Name}", null, new Type[] { typeof(CbWriter), typeof(Utf8String), classType });
 
 			if (IsPolymorphic)
 			{
-				ReadMethod = new DynamicMethod($"Read_{ClassType.Name}", ClassType, new Type[] { typeof(CbField) });
-				WriteContentsMethod = new DynamicMethod($"WriteContents_{ClassType.Name}", null, new Type[] { typeof(CbWriter), ClassType });
+				ReadMethod = new DynamicMethod($"Read_{classType.Name}", classType, new Type[] { typeof(CbField) });
+				WriteContentsMethod = new DynamicMethod($"WriteContents_{classType.Name}", null, new Type[] { typeof(CbWriter), classType });
 			}
 			else
 			{
@@ -89,194 +89,194 @@ namespace EpicGames.Serialization.Converters
 			if (IsPolymorphic)
 			{
 				// Create the dispatch type
-				Type DispatchType = typeof(CbPolymorphicDispatch<>).MakeGenericType(ClassType);
+				Type dispatchType = typeof(CbPolymorphicDispatch<>).MakeGenericType(ClassType);
 
 				// Create the read dispatch method
 				{
-					ILGenerator Generator = ReadMethod.GetILGenerator();
-					Generator.Emit(OpCodes.Ldarg_0);
-					Generator.Emit(OpCodes.Call, DispatchType.GetMethod(nameof(CbPolymorphicDispatch<object>.Read))!);
-					Generator.Emit(OpCodes.Ret);
+					ILGenerator generator = ReadMethod.GetILGenerator();
+					generator.Emit(OpCodes.Ldarg_0);
+					generator.Emit(OpCodes.Call, dispatchType.GetMethod(nameof(CbPolymorphicDispatch<object>.Read))!);
+					generator.Emit(OpCodes.Ret);
 				}
 
 				// Create the write dispatch method
 				{
-					ILGenerator Generator = WriteContentsMethod.GetILGenerator();
-					Generator.Emit(OpCodes.Ldarg_0);
-					Generator.Emit(OpCodes.Ldarg_1);
-					Generator.Emit(OpCodes.Call, DispatchType.GetMethod(nameof(CbPolymorphicDispatch<object>.WriteContents))!);
-					Generator.Emit(OpCodes.Ret);
+					ILGenerator generator = WriteContentsMethod.GetILGenerator();
+					generator.Emit(OpCodes.Ldarg_0);
+					generator.Emit(OpCodes.Ldarg_1);
+					generator.Emit(OpCodes.Call, dispatchType.GetMethod(nameof(CbPolymorphicDispatch<object>.WriteContents))!);
+					generator.Emit(OpCodes.Ret);
 				}
 
 				// Finally, update the dispatch type with all the methods. We should be safe on the recursive path now.
-				PopulateDispatchType(ClassType, DispatchType);
+				PopulateDispatchType(ClassType, dispatchType);
 			}
 		}
 
-		static void PopulateDispatchType(Type ClassType, Type DispatchType)
+		static void PopulateDispatchType(Type classType, Type dispatchType)
 		{
-			Dictionary<Utf8String, Type> DiscriminatorToKnownType = new Dictionary<Utf8String, Type>();
+			Dictionary<Utf8String, Type> discriminatorToKnownType = new Dictionary<Utf8String, Type>();
 
-			Type[] KnownTypes = ClassType.Assembly.GetTypes();
-			foreach (Type KnownType in KnownTypes)
+			Type[] knownTypes = classType.Assembly.GetTypes();
+			foreach (Type knownType in knownTypes)
 			{
-				if (KnownType.IsClass && !KnownType.IsAbstract)
+				if (knownType.IsClass && !knownType.IsAbstract)
 				{
-					for (Type? BaseType = KnownType; BaseType != null; BaseType = BaseType.BaseType)
+					for (Type? baseType = knownType; baseType != null; baseType = baseType.BaseType)
 					{
-						if (BaseType == ClassType)
+						if (baseType == classType)
 						{
-							CbDiscriminatorAttribute? Discriminator = KnownType.GetCustomAttribute<CbDiscriminatorAttribute>();
-							if (Discriminator == null)
+							CbDiscriminatorAttribute? discriminator = knownType.GetCustomAttribute<CbDiscriminatorAttribute>();
+							if (discriminator == null)
 							{
 								throw new NotSupportedException();
 							}
-							DiscriminatorToKnownType[Discriminator.Name] = KnownType;
+							discriminatorToKnownType[discriminator.Name] = knownType;
 						}
 					}
 				}
 			}
 
 			// Populate the dictionary
-			Dictionary<Utf8String, Func<CbField, object>> NameToReadFunc = (Dictionary<Utf8String, Func<CbField, object>>)DispatchType.GetField(nameof(CbPolymorphicDispatch<object>.NameToReadFunc))!.GetValue(null)!;
-			Dictionary<Type, Action<CbWriter, object>> TypeToWriteContentsFunc = (Dictionary<Type, Action<CbWriter, object>>)DispatchType.GetField(nameof(CbPolymorphicDispatch<object>.TypeToWriteContentsFunc))!.GetValue(null)!;
+			Dictionary<Utf8String, Func<CbField, object>> nameToReadFunc = (Dictionary<Utf8String, Func<CbField, object>>)dispatchType.GetField(nameof(CbPolymorphicDispatch<object>.NameToReadFunc))!.GetValue(null)!;
+			Dictionary<Type, Action<CbWriter, object>> typeToWriteContentsFunc = (Dictionary<Type, Action<CbWriter, object>>)dispatchType.GetField(nameof(CbPolymorphicDispatch<object>.TypeToWriteContentsFunc))!.GetValue(null)!;
 
-			foreach ((Utf8String Name, Type KnownType) in DiscriminatorToKnownType)
+			foreach ((Utf8String name, Type knownType) in discriminatorToKnownType)
 			{
-				CbClassConverterMethods Methods = Create(KnownType);
+				CbClassConverterMethods methods = Create(knownType);
 
 				{
-					DynamicMethod DynamicMethod = new DynamicMethod("_", typeof(object), new Type[] { typeof(CbField) });
-					ILGenerator Generator = DynamicMethod.GetILGenerator();
-					Generator.Emit(OpCodes.Ldarg_0);
-					Generator.Emit(OpCodes.Call, Methods.ReadConcreteMethod);
-					Generator.Emit(OpCodes.Castclass, typeof(object));
-					Generator.Emit(OpCodes.Ret);
-					NameToReadFunc[Name] = CreateDelegate<Func<CbField, object>>(DynamicMethod);
+					DynamicMethod dynamicMethod = new DynamicMethod("_", typeof(object), new Type[] { typeof(CbField) });
+					ILGenerator generator = dynamicMethod.GetILGenerator();
+					generator.Emit(OpCodes.Ldarg_0);
+					generator.Emit(OpCodes.Call, methods.ReadConcreteMethod);
+					generator.Emit(OpCodes.Castclass, typeof(object));
+					generator.Emit(OpCodes.Ret);
+					nameToReadFunc[name] = CreateDelegate<Func<CbField, object>>(dynamicMethod);
 				}
 
 				{
-					DynamicMethod DynamicMethod = new DynamicMethod("_", null, new Type[] { typeof(CbWriter), typeof(object) });
-					ILGenerator Generator = DynamicMethod.GetILGenerator();
-					Generator.Emit(OpCodes.Ldarg_0);
-					Generator.Emit(OpCodes.Ldarg_1);
-					Generator.Emit(OpCodes.Castclass, KnownType);
-					Generator.Emit(OpCodes.Call, Methods.WriteConcreteContentsMethod);
-					Generator.Emit(OpCodes.Ret);
-					TypeToWriteContentsFunc[KnownType] = CreateDelegate<Action<CbWriter, object>>(DynamicMethod);
+					DynamicMethod dynamicMethod = new DynamicMethod("_", null, new Type[] { typeof(CbWriter), typeof(object) });
+					ILGenerator generator = dynamicMethod.GetILGenerator();
+					generator.Emit(OpCodes.Ldarg_0);
+					generator.Emit(OpCodes.Ldarg_1);
+					generator.Emit(OpCodes.Castclass, knownType);
+					generator.Emit(OpCodes.Call, methods.WriteConcreteContentsMethod);
+					generator.Emit(OpCodes.Ret);
+					typeToWriteContentsFunc[knownType] = CreateDelegate<Action<CbWriter, object>>(dynamicMethod);
 				}
 			}
 		}
 
-		static void CreateObjectWriter(Type Type, ILGenerator Generator, DynamicMethod ContentsWriter)
+		static void CreateObjectWriter(Type type, ILGenerator generator, DynamicMethod contentsWriter)
 		{
-			Generator.Emit(OpCodes.Ldarg_1);
+			generator.Emit(OpCodes.Ldarg_1);
 
-			Label SkipLabel = Generator.DefineLabel();
-			Generator.Emit(OpCodes.Brfalse, SkipLabel);
+			Label skipLabel = generator.DefineLabel();
+			generator.Emit(OpCodes.Brfalse, skipLabel);
 
-			Generator.Emit(OpCodes.Ldarg_0);
-			Generator.EmitCall(OpCodes.Call, GetMethodInfo<CbWriter>(x => x.BeginObject()), null);
+			generator.Emit(OpCodes.Ldarg_0);
+			generator.EmitCall(OpCodes.Call, GetMethodInfo<CbWriter>(x => x.BeginObject()), null);
 
-			Generator.Emit(OpCodes.Ldarg_0);
-			Generator.Emit(OpCodes.Ldarg_1);
-			Generator.EmitCall(OpCodes.Call, ContentsWriter, null);
+			generator.Emit(OpCodes.Ldarg_0);
+			generator.Emit(OpCodes.Ldarg_1);
+			generator.EmitCall(OpCodes.Call, contentsWriter, null);
 
-			Generator.Emit(OpCodes.Ldarg_0);
-			Generator.EmitCall(OpCodes.Call, GetMethodInfo<CbWriter>(x => x.EndObject()), null);
+			generator.Emit(OpCodes.Ldarg_0);
+			generator.EmitCall(OpCodes.Call, GetMethodInfo<CbWriter>(x => x.EndObject()), null);
 
-			Generator.MarkLabel(SkipLabel);
-			Generator.Emit(OpCodes.Ret);
+			generator.MarkLabel(skipLabel);
+			generator.Emit(OpCodes.Ret);
 		}
 
-		static void CreateNamedObjectWriter(Type Type, ILGenerator Generator, DynamicMethod ContentsWriter)
+		static void CreateNamedObjectWriter(Type type, ILGenerator generator, DynamicMethod contentsWriter)
 		{
-			Generator.Emit(OpCodes.Ldarg_2);
+			generator.Emit(OpCodes.Ldarg_2);
 
-			Label SkipLabel = Generator.DefineLabel();
-			Generator.Emit(OpCodes.Brfalse, SkipLabel);
+			Label skipLabel = generator.DefineLabel();
+			generator.Emit(OpCodes.Brfalse, skipLabel);
 
-			Generator.Emit(OpCodes.Ldarg_0);
-			Generator.Emit(OpCodes.Ldarg_1);
-			Generator.EmitCall(OpCodes.Call, GetMethodInfo<CbWriter>(x => x.BeginObject(null!)), null);
+			generator.Emit(OpCodes.Ldarg_0);
+			generator.Emit(OpCodes.Ldarg_1);
+			generator.EmitCall(OpCodes.Call, GetMethodInfo<CbWriter>(x => x.BeginObject(null!)), null);
 
-			Generator.Emit(OpCodes.Ldarg_0);
-			Generator.Emit(OpCodes.Ldarg_2);
-			Generator.EmitCall(OpCodes.Call, ContentsWriter, null);
+			generator.Emit(OpCodes.Ldarg_0);
+			generator.Emit(OpCodes.Ldarg_2);
+			generator.EmitCall(OpCodes.Call, contentsWriter, null);
 
-			Generator.Emit(OpCodes.Ldarg_0);
-			Generator.EmitCall(OpCodes.Call, GetMethodInfo<CbWriter>(x => x.EndObject()), null);
+			generator.Emit(OpCodes.Ldarg_0);
+			generator.EmitCall(OpCodes.Call, GetMethodInfo<CbWriter>(x => x.EndObject()), null);
 
-			Generator.MarkLabel(SkipLabel);
-			Generator.Emit(OpCodes.Ret);
+			generator.MarkLabel(skipLabel);
+			generator.Emit(OpCodes.Ret);
 		}
 
-		static void CreateConcreteObjectContentsWriter(Type Type, ILGenerator Generator)
+		static void CreateConcreteObjectContentsWriter(Type type, ILGenerator generator)
 		{
 			// Find the reflected properties from this type
-			(Utf8String Name, PropertyInfo Property)[] Properties = GetProperties(Type);
+			(Utf8String Name, PropertyInfo Property)[] properties = GetProperties(type);
 
 			// Create a static type with the required reflection data
-			Type ReflectedType = typeof(CbReflectedTypeInfo<>).MakeGenericType(Type);
-			FieldInfo NamesField = ReflectedType.GetField(nameof(CbReflectedTypeInfo<object>.Names))!;
-			NamesField.SetValue(null, Properties.Select(x => x.Name).ToArray());
+			Type reflectedType = typeof(CbReflectedTypeInfo<>).MakeGenericType(type);
+			FieldInfo namesField = reflectedType.GetField(nameof(CbReflectedTypeInfo<object>.Names))!;
+			namesField.SetValue(null, properties.Select(x => x.Name).ToArray());
 
 			// Write the discriminator
-			CbDiscriminatorAttribute? Discriminator = Type.GetCustomAttribute<CbDiscriminatorAttribute>();
-			if (Discriminator != null)
+			CbDiscriminatorAttribute? discriminator = type.GetCustomAttribute<CbDiscriminatorAttribute>();
+			if (discriminator != null)
 			{
-				FieldInfo DiscriminatorKeyField = GetFieldInfo(() => DiscriminatorKey);
-				Generator.Emit(OpCodes.Ldarg_0);
-				Generator.Emit(OpCodes.Ldsfld, DiscriminatorKeyField);
-				Generator.Emit(OpCodes.Ldstr, Discriminator.Name);
-				Generator.Emit(OpCodes.Call, GetMethodInfo<CbWriter>(x => x.WriteString(default, null!)));
+				FieldInfo discriminatorKeyField = GetFieldInfo(() => s_discriminatorKey);
+				generator.Emit(OpCodes.Ldarg_0);
+				generator.Emit(OpCodes.Ldsfld, discriminatorKeyField);
+				generator.Emit(OpCodes.Ldstr, discriminator.Name);
+				generator.Emit(OpCodes.Call, GetMethodInfo<CbWriter>(x => x.WriteString(default, null!)));
 			}
 
 			// Write all the remaining properties
-			for (int Idx = 0; Idx < Properties.Length; Idx++)
+			for (int idx = 0; idx < properties.Length; idx++)
 			{
-				PropertyInfo Property = Properties[Idx].Property;
-				Type PropertyType = Property.PropertyType;
+				PropertyInfo property = properties[idx].Property;
+				Type propertyType = property.PropertyType;
 
 				// Get the field value
-				Generator.Emit(OpCodes.Ldarg_1);
-				Generator.EmitCall(OpCodes.Call, Property.GetMethod!, null);
+				generator.Emit(OpCodes.Ldarg_1);
+				generator.EmitCall(OpCodes.Call, property.GetMethod!, null);
 
-				Label SkipLabel = Generator.DefineLabel();
+				Label skipLabel = generator.DefineLabel();
 
-				MethodInfo WriteMethod;
-				if (TypeToMethods.TryGetValue(PropertyType, out CbClassConverterMethods? DynamicMethods))
+				MethodInfo writeMethod;
+				if (s_typeToMethods.TryGetValue(propertyType, out CbClassConverterMethods? dynamicMethods))
 				{
-					Generator.Emit(OpCodes.Dup);
-					Generator.Emit(OpCodes.Brfalse, SkipLabel);
-					WriteMethod = DynamicMethods.WriteNamedMethod;
+					generator.Emit(OpCodes.Dup);
+					generator.Emit(OpCodes.Brfalse, skipLabel);
+					writeMethod = dynamicMethods.WriteNamedMethod;
 				}
 				else
 				{
-					ICbConverterMethods Methods = CbConverterMethods.Get(Property);
-					WriteMethod = Methods.WriteNamedMethod;
+					ICbConverterMethods methods = CbConverterMethods.Get(property);
+					writeMethod = methods.WriteNamedMethod;
 				}
 
 				// Store the variable in a local
-				LocalBuilder Local = Generator.DeclareLocal(PropertyType);
-				Generator.Emit(OpCodes.Dup);
-				Generator.Emit(OpCodes.Stloc, Local);
+				LocalBuilder local = generator.DeclareLocal(propertyType);
+				generator.Emit(OpCodes.Dup);
+				generator.Emit(OpCodes.Stloc, local);
 
 				// Call the writer
-				Generator.Emit(OpCodes.Ldarg_0);
+				generator.Emit(OpCodes.Ldarg_0);
 
-				Generator.Emit(OpCodes.Ldsfld, NamesField);
-				Generator.Emit(OpCodes.Ldc_I4, Idx);
-				Generator.Emit(OpCodes.Ldelem, typeof(Utf8String));
+				generator.Emit(OpCodes.Ldsfld, namesField);
+				generator.Emit(OpCodes.Ldc_I4, idx);
+				generator.Emit(OpCodes.Ldelem, typeof(Utf8String));
 
-				Generator.Emit(OpCodes.Ldloc, Local);
-				Generator.EmitCall(OpCodes.Call, WriteMethod, null);
+				generator.Emit(OpCodes.Ldloc, local);
+				generator.EmitCall(OpCodes.Call, writeMethod, null);
 
 				// Remove the duplicated value from the top of the stack
-				Generator.MarkLabel(SkipLabel);
-				Generator.Emit(OpCodes.Pop);
+				generator.MarkLabel(skipLabel);
+				generator.Emit(OpCodes.Pop);
 			}
-			Generator.Emit(OpCodes.Ret);
+			generator.Emit(OpCodes.Ret);
 		}
 
 		class CbPolymorphicDispatch<T>
@@ -284,221 +284,221 @@ namespace EpicGames.Serialization.Converters
 			public static Dictionary<Utf8String, Func<CbField, object>> NameToReadFunc = new Dictionary<Utf8String, Func<CbField, object>>();
 			public static Dictionary<Type, Action<CbWriter, object>> TypeToWriteContentsFunc = new Dictionary<Type, Action<CbWriter, object>>();
 
-			public static object Read(CbField Field)
+			public static object Read(CbField field)
 			{
-				Utf8String Name = Field.AsObject().Find(DiscriminatorKey).AsUtf8String();
-				return NameToReadFunc[Name](Field);
+				Utf8String name = field.AsObject().Find(s_discriminatorKey).AsUtf8String();
+				return NameToReadFunc[name](field);
 			}
 
-			public static void WriteContents(CbWriter Writer, object Value)
+			public static void WriteContents(CbWriter writer, object value)
 			{
-				Type Type = Value!.GetType();
-				TypeToWriteContentsFunc[Type](Writer, Value);
+				Type type = value!.GetType();
+				TypeToWriteContentsFunc[type](writer, value);
 			}
 		}
 
-		static T CreateDelegate<T>(DynamicMethod Method) where T : Delegate
+		static T CreateDelegate<T>(DynamicMethod method) where T : Delegate
 		{
-			return (T)Method.CreateDelegate(typeof(T));
+			return (T)method.CreateDelegate(typeof(T));
 		}
 
-		static void CreateConcreteObjectReader(Type Type, ILGenerator Generator)
+		static void CreateConcreteObjectReader(Type type, ILGenerator generator)
 		{
 			// Construct the object
-			ConstructorInfo? Constructor = Type.GetConstructor(BindingFlags.Public | BindingFlags.NonPublic | BindingFlags.Instance, null, Type.EmptyTypes, null);
-			if (Constructor == null)
+			ConstructorInfo? constructor = type.GetConstructor(BindingFlags.Public | BindingFlags.NonPublic | BindingFlags.Instance, null, Type.EmptyTypes, null);
+			if (constructor == null)
 			{
-				throw new CbException($"Unable to find default constructor for {Type}");
+				throw new CbException($"Unable to find default constructor for {type}");
 			}
 
 			// Find the reflected properties from this type
-			(Utf8String Name, PropertyInfo Property)[] Properties = GetProperties(Type);
+			(Utf8String Name, PropertyInfo Property)[] properties = GetProperties(type);
 
 			// Create a static type with the required reflection data
-			Type ReflectedType = typeof(CbReflectedTypeInfo<>).MakeGenericType(Type);
-			FieldInfo NamesField = ReflectedType.GetField(nameof(CbReflectedTypeInfo<object>.Names))!;
-			NamesField.SetValue(null, Properties.Select(x => x.Name).ToArray());
-			MethodInfo MatchNameMethod = ReflectedType.GetMethod(nameof(CbReflectedTypeInfo<object>.MatchName))!;
+			Type reflectedType = typeof(CbReflectedTypeInfo<>).MakeGenericType(type);
+			FieldInfo namesField = reflectedType.GetField(nameof(CbReflectedTypeInfo<object>.Names))!;
+			namesField.SetValue(null, properties.Select(x => x.Name).ToArray());
+			MethodInfo matchNameMethod = reflectedType.GetMethod(nameof(CbReflectedTypeInfo<object>.MatchName))!;
 
 			// NewObjectLocal = new Type()
-			LocalBuilder NewObjectLocal = Generator.DeclareLocal(typeof(object));
-			Generator.Emit(OpCodes.Newobj, Constructor);
-			Generator.Emit(OpCodes.Stloc, NewObjectLocal);
+			LocalBuilder newObjectLocal = generator.DeclareLocal(typeof(object));
+			generator.Emit(OpCodes.Newobj, constructor);
+			generator.Emit(OpCodes.Stloc, newObjectLocal);
 
 			// Stack(0) = CbField.CreateIterator()
-			Generator.Emit(OpCodes.Ldarg_0);
-			Generator.EmitCall(OpCodes.Call, GetMethodInfo<CbField>(x => x.CreateIterator()), null);
+			generator.Emit(OpCodes.Ldarg_0);
+			generator.EmitCall(OpCodes.Call, GetMethodInfo<CbField>(x => x.CreateIterator()), null);
 
 			// CbFieldIterator IteratorLocal = Stack(0)
-			LocalBuilder IteratorLocal = Generator.DeclareLocal(typeof(CbFieldIterator));
-			Generator.Emit(OpCodes.Dup);
-			Generator.Emit(OpCodes.Stloc, IteratorLocal);
+			LocalBuilder iteratorLocal = generator.DeclareLocal(typeof(CbFieldIterator));
+			generator.Emit(OpCodes.Dup);
+			generator.Emit(OpCodes.Stloc, iteratorLocal);
 
 			// if(!Stack.Pop().IsValid()) goto ReturnLabel
-			Generator.EmitCall(OpCodes.Call, GetMethodInfo<CbFieldIterator>(x => x.IsValid()), null);
-			Label ReturnLabel = Generator.DefineLabel();
-			Generator.Emit(OpCodes.Brfalse, ReturnLabel);
+			generator.EmitCall(OpCodes.Call, GetMethodInfo<CbFieldIterator>(x => x.IsValid()), null);
+			Label returnLabel = generator.DefineLabel();
+			generator.Emit(OpCodes.Brfalse, returnLabel);
 
 			// NamesLocal = CbReflectedTypeInfo<Type>.Names
-			LocalBuilder NamesLocal = Generator.DeclareLocal(typeof(Utf8String[]));
-			Generator.Emit(OpCodes.Ldsfld, NamesField);
-			Generator.Emit(OpCodes.Stloc, NamesLocal);
+			LocalBuilder namesLocal = generator.DeclareLocal(typeof(Utf8String[]));
+			generator.Emit(OpCodes.Ldsfld, namesField);
+			generator.Emit(OpCodes.Stloc, namesLocal);
 
 			// IterationLoopLabel:
-			Label IterationLoopLabel = Generator.DefineLabel();
-			Generator.MarkLabel(IterationLoopLabel);
+			Label iterationLoopLabel = generator.DefineLabel();
+			generator.MarkLabel(iterationLoopLabel);
 
 			// bool MatchLocal = false
-			LocalBuilder MatchLocal = Generator.DeclareLocal(typeof(bool));
-			Generator.Emit(OpCodes.Ldc_I4_0);
-			Generator.Emit(OpCodes.Stloc, MatchLocal);
+			LocalBuilder matchLocal = generator.DeclareLocal(typeof(bool));
+			generator.Emit(OpCodes.Ldc_I4_0);
+			generator.Emit(OpCodes.Stloc, matchLocal);
 
 			// Stack(0) = IteratorLocal.GetCurrent()
-			Generator.Emit(OpCodes.Ldloc, IteratorLocal);
-			Generator.EmitCall(OpCodes.Call, GetMethodInfo<CbFieldIterator>(x => x.GetCurrent()), null);
+			generator.Emit(OpCodes.Ldloc, iteratorLocal);
+			generator.EmitCall(OpCodes.Call, GetMethodInfo<CbFieldIterator>(x => x.GetCurrent()), null);
 
 			// Try to parse each of the properties in order. If fields are ordered correctly, we will parse the object in a single pass. Otherwise we can loop and start again.
-			LocalBuilder FieldLocal = Generator.DeclareLocal(typeof(CbField));
-			for (int Idx = 0; Idx < Properties.Length; Idx++)
+			LocalBuilder fieldLocal = generator.DeclareLocal(typeof(CbField));
+			for (int idx = 0; idx < properties.Length; idx++)
 			{
-				PropertyInfo Property = Properties[Idx].Property;
+				PropertyInfo property = properties[idx].Property;
 
 				// Get the read method for this property type
-				MethodInfo ReadMethod;
-				if (TypeToMethods.TryGetValue(Property.PropertyType, out CbClassConverterMethods? DynamicMethods))
+				MethodInfo readMethod;
+				if (s_typeToMethods.TryGetValue(property.PropertyType, out CbClassConverterMethods? dynamicMethods))
 				{
-					ReadMethod = DynamicMethods.ReadMethod;
+					readMethod = dynamicMethods.ReadMethod;
 				}
 				else
 				{
-					ReadMethod = CbConverterMethods.Get(Property).ReadMethod;
+					readMethod = CbConverterMethods.Get(property).ReadMethod;
 				}
 
 				// if(!CbReflectedTypeInfo<Type>.MatchName(Stack(0), Idx)) goto SkipPropertyLabel
-				Label SkipPropertyLabel = Generator.DefineLabel();
-				Generator.Emit(OpCodes.Dup); // Current CbField
-				Generator.Emit(OpCodes.Ldc_I4, Idx);
-				Generator.Emit(OpCodes.Call, MatchNameMethod);
-				Generator.Emit(OpCodes.Brfalse, SkipPropertyLabel);
+				Label skipPropertyLabel = generator.DefineLabel();
+				generator.Emit(OpCodes.Dup); // Current CbField
+				generator.Emit(OpCodes.Ldc_I4, idx);
+				generator.Emit(OpCodes.Call, matchNameMethod);
+				generator.Emit(OpCodes.Brfalse, skipPropertyLabel);
 
 				// FieldLocal = Stack.Pop()
-				Generator.Emit(OpCodes.Stloc, FieldLocal);
+				generator.Emit(OpCodes.Stloc, fieldLocal);
 
 				// Property.SetMethod(NewObjectLocal, ReadMethod(FieldLocal))
-				Generator.Emit(OpCodes.Ldloc, NewObjectLocal);
-				Generator.Emit(OpCodes.Ldloc, FieldLocal);
-				Generator.EmitCall(OpCodes.Call, ReadMethod, null);
-				Generator.EmitCall(OpCodes.Call, Property.SetMethod!, null);
+				generator.Emit(OpCodes.Ldloc, newObjectLocal);
+				generator.Emit(OpCodes.Ldloc, fieldLocal);
+				generator.EmitCall(OpCodes.Call, readMethod, null);
+				generator.EmitCall(OpCodes.Call, property.SetMethod!, null);
 
 				// if(!IteratorLocal.MoveNext()) goto ReturnLabel
-				Generator.Emit(OpCodes.Ldloc, IteratorLocal);
-				Generator.EmitCall(OpCodes.Call, GetMethodInfo<CbFieldIterator>(x => x.MoveNext()), null);
-				Generator.Emit(OpCodes.Brfalse, ReturnLabel);
+				generator.Emit(OpCodes.Ldloc, iteratorLocal);
+				generator.EmitCall(OpCodes.Call, GetMethodInfo<CbFieldIterator>(x => x.MoveNext()), null);
+				generator.Emit(OpCodes.Brfalse, returnLabel);
 
 				// MatchLocal = true
-				Generator.Emit(OpCodes.Ldc_I4_1);
-				Generator.Emit(OpCodes.Stloc, MatchLocal);
+				generator.Emit(OpCodes.Ldc_I4_1);
+				generator.Emit(OpCodes.Stloc, matchLocal);
 
 				// Stack(0) = IteratorLocal.GetCurrent()
-				Generator.Emit(OpCodes.Ldloc, IteratorLocal);
-				Generator.EmitCall(OpCodes.Call, GetMethodInfo<CbFieldIterator>(x => x.GetCurrent()), null);
+				generator.Emit(OpCodes.Ldloc, iteratorLocal);
+				generator.EmitCall(OpCodes.Call, GetMethodInfo<CbFieldIterator>(x => x.GetCurrent()), null);
 
 				// SkipPropertyLabel:
-				Generator.MarkLabel(SkipPropertyLabel);
+				generator.MarkLabel(skipPropertyLabel);
 			}
 
 			// Stack.Pop()
-			Generator.Emit(OpCodes.Pop); // Current CbField
+			generator.Emit(OpCodes.Pop); // Current CbField
 
 			// if(MatchLocal) goto IterationLoopLabel
-			Generator.Emit(OpCodes.Ldloc, MatchLocal);
-			Generator.Emit(OpCodes.Brtrue, IterationLoopLabel);
+			generator.Emit(OpCodes.Ldloc, matchLocal);
+			generator.Emit(OpCodes.Brtrue, iterationLoopLabel);
 
 			// if(IteratorLocal.MoveNext()) goto IterationLoopLabel
-			Generator.Emit(OpCodes.Ldloc, IteratorLocal);
-			Generator.EmitCall(OpCodes.Call, GetMethodInfo<CbFieldIterator>(x => x.MoveNext()), null);
-			Generator.Emit(OpCodes.Brtrue, IterationLoopLabel);
+			generator.Emit(OpCodes.Ldloc, iteratorLocal);
+			generator.EmitCall(OpCodes.Call, GetMethodInfo<CbFieldIterator>(x => x.MoveNext()), null);
+			generator.Emit(OpCodes.Brtrue, iterationLoopLabel);
 
 			// return NewObjectLocal
-			Generator.MarkLabel(ReturnLabel);
-			Generator.Emit(OpCodes.Ldloc, NewObjectLocal);
-			Generator.Emit(OpCodes.Ret);
+			generator.MarkLabel(returnLabel);
+			generator.Emit(OpCodes.Ldloc, newObjectLocal);
+			generator.Emit(OpCodes.Ret);
 		}
 
-		static (Utf8String, PropertyInfo)[] GetProperties(Type Type)
+		static (Utf8String, PropertyInfo)[] GetProperties(Type type)
 		{
-			List<(Utf8String, PropertyInfo)> PropertyList = new List<(Utf8String, PropertyInfo)>();
-			foreach (PropertyInfo Property in Type.GetProperties(BindingFlags.Public | BindingFlags.NonPublic | BindingFlags.Instance))
+			List<(Utf8String, PropertyInfo)> propertyList = new List<(Utf8String, PropertyInfo)>();
+			foreach (PropertyInfo property in type.GetProperties(BindingFlags.Public | BindingFlags.NonPublic | BindingFlags.Instance))
 			{
-				CbFieldAttribute? Attribute = Property.GetCustomAttribute<CbFieldAttribute>();
-				if (Attribute != null)
+				CbFieldAttribute? attribute = property.GetCustomAttribute<CbFieldAttribute>();
+				if (attribute != null)
 				{
-					Utf8String Name = Attribute.Name ?? Property.Name;
-					PropertyList.Add((Name, Property));
+					Utf8String name = attribute.Name ?? property.Name;
+					propertyList.Add((name, property));
 				}
 			}
-			return PropertyList.ToArray();
+			return propertyList.ToArray();
 		}
 
-		static FieldInfo GetFieldInfo<T>(Expression<Func<T>> Expr)
+		static FieldInfo GetFieldInfo<T>(Expression<Func<T>> expr)
 		{
-			return (FieldInfo)((MemberExpression)Expr.Body).Member;
+			return (FieldInfo)((MemberExpression)expr.Body).Member;
 		}
 
-		static MethodInfo GetMethodInfo(Expression<Action> Expr)
+		static MethodInfo GetMethodInfo(Expression<Action> expr)
 		{
-			return ((MethodCallExpression)Expr.Body).Method;
+			return ((MethodCallExpression)expr.Body).Method;
 		}
 
-		static MethodInfo GetMethodInfo<T>(Expression<Action<T>> Expr)
+		static MethodInfo GetMethodInfo<T>(Expression<Action<T>> expr)
 		{
-			return ((MethodCallExpression)Expr.Body).Method;
+			return ((MethodCallExpression)expr.Body).Method;
 		}
 	}
 
 	class CbClassConverter<T> : CbConverterBase<T>, ICbConverterMethods where T : class
 	{
-		CbClassConverterMethods Methods;
+		CbClassConverterMethods _methods;
 
-		Func<CbField, T> ReadFunc;
-		Action<CbWriter, T> WriteFunc;
-		Action<CbWriter, Utf8String, T> WriteNamedFunc;
+		Func<CbField, T> _readFunc;
+		Action<CbWriter, T> _writeFunc;
+		Action<CbWriter, Utf8String, T> _writeNamedFunc;
 
 		public CbClassConverter()
 		{
-			Methods = CbClassConverterMethods.Create(typeof(T));
+			_methods = CbClassConverterMethods.Create(typeof(T));
 
-			ReadFunc = CreateDelegate<Func<CbField, T>>(Methods.ReadMethod);
-			WriteFunc = CreateDelegate<Action<CbWriter, T>>(Methods.WriteMethod);
-			WriteNamedFunc = CreateDelegate<Action<CbWriter, Utf8String, T>>(Methods.WriteNamedMethod);
+			_readFunc = CreateDelegate<Func<CbField, T>>(_methods.ReadMethod);
+			_writeFunc = CreateDelegate<Action<CbWriter, T>>(_methods.WriteMethod);
+			_writeNamedFunc = CreateDelegate<Action<CbWriter, Utf8String, T>>(_methods.WriteNamedMethod);
 		}
 
-		public MethodInfo ReadMethod => Methods.ReadMethod;
+		public MethodInfo ReadMethod => _methods.ReadMethod;
 
-		public MethodInfo WriteMethod => Methods.WriteMethod;
+		public MethodInfo WriteMethod => _methods.WriteMethod;
 
-		public MethodInfo WriteNamedMethod => Methods.WriteNamedMethod;
+		public MethodInfo WriteNamedMethod => _methods.WriteNamedMethod;
 
-		static TDelegate CreateDelegate<TDelegate>(DynamicMethod Method) where TDelegate : Delegate => (TDelegate)Method.CreateDelegate(typeof(TDelegate));
+		static TDelegate CreateDelegate<TDelegate>(DynamicMethod method) where TDelegate : Delegate => (TDelegate)method.CreateDelegate(typeof(TDelegate));
 
-		public override T Read(CbField Field) => ReadFunc(Field);
+		public override T Read(CbField field) => _readFunc(field);
 
-		public override void Write(CbWriter Writer, T Value) => WriteFunc(Writer, Value);
+		public override void Write(CbWriter writer, T value) => _writeFunc(writer, value);
 
-		public override void WriteNamed(CbWriter Writer, Utf8String Name, T Value) => WriteNamedFunc(Writer, Name, Value);
+		public override void WriteNamed(CbWriter writer, Utf8String name, T value) => _writeNamedFunc(writer, name, value);
 	}
 
 	class CbClassConverterFactory : CbConverterFactory
 	{
-		public override ICbConverter? CreateConverter(Type Type)
+		public override ICbConverter? CreateConverter(Type type)
 		{
-			ICbConverter? Converter = null;
-			if (Type.IsClass)
+			ICbConverter? converter = null;
+			if (type.IsClass)
 			{
-				Type ConverterType = typeof(CbClassConverter<>).MakeGenericType(Type);
-				Converter = (ICbConverter?)Activator.CreateInstance(ConverterType);
+				Type converterType = typeof(CbClassConverter<>).MakeGenericType(type);
+				converter = (ICbConverter?)Activator.CreateInstance(converterType);
 			}
-			return Converter;
+			return converter;
 		}
 	}
 }
diff --git a/Engine/Source/Programs/Shared/EpicGames.Serialization/Converters/CbDictionaryConverter.cs b/Engine/Source/Programs/Shared/EpicGames.Serialization/Converters/CbDictionaryConverter.cs
--- a/Engine/Source/Programs/Shared/EpicGames.Serialization/Converters/CbDictionaryConverter.cs
+++ b/Engine/Source/Programs/Shared/EpicGames.Serialization/Converters/CbDictionaryConverter.cs
@@ -19,61 +19,61 @@ namespace EpicGames.Serialization.Converters
 	class CbDictionaryConverter<TKey, TValue> : CbConverterBase<Dictionary<TKey, TValue>> where TKey : notnull
 	{
 		/// <inheritdoc/>
-		public override Dictionary<TKey, TValue> Read(CbField Field)
+		public override Dictionary<TKey, TValue> Read(CbField field)
 		{
-			Dictionary<TKey, TValue> Dictionary = new Dictionary<TKey, TValue>();
-			foreach (CbField Element in Field)
+			Dictionary<TKey, TValue> dictionary = new Dictionary<TKey, TValue>();
+			foreach (CbField element in field)
 			{
-				IEnumerator<CbField> Enumerator = Element.AsArray().GetEnumerator();
+				IEnumerator<CbField> enumerator = element.AsArray().GetEnumerator();
 
-				if (!Enumerator.MoveNext())
+				if (!enumerator.MoveNext())
 				{
 					throw new CbException("Missing key for dictionary entry");
 				}
-				TKey Key = CbSerializer.Deserialize<TKey>(Enumerator.Current);
+				TKey key = CbSerializer.Deserialize<TKey>(enumerator.Current);
 
-				if (!Enumerator.MoveNext())
+				if (!enumerator.MoveNext())
 				{
 					throw new CbException("Missing value for dictionary entry");
 				}
-				TValue Value = CbSerializer.Deserialize<TValue>(Enumerator.Current);
+				TValue value = CbSerializer.Deserialize<TValue>(enumerator.Current);
 
-				Dictionary.Add(Key, Value);
+				dictionary.Add(key, value);
 			}
-			return Dictionary;
+			return dictionary;
 		}
 
 		/// <inheritdoc/>
-		public override void Write(CbWriter Writer, Dictionary<TKey, TValue> Value)
+		public override void Write(CbWriter writer, Dictionary<TKey, TValue> value)
 		{
-			if (Value.Count > 0)
+			if (value.Count > 0)
 			{
-				Writer.BeginUniformArray(CbFieldType.Array);
-				foreach (KeyValuePair<TKey, TValue> Pair in Value)
+				writer.BeginUniformArray(CbFieldType.Array);
+				foreach (KeyValuePair<TKey, TValue> pair in value)
 				{
-					Writer.BeginArray();
-					CbSerializer.Serialize(Writer, Pair.Key);
-					CbSerializer.Serialize(Writer, Pair.Value);
-					Writer.EndArray();
+					writer.BeginArray();
+					CbSerializer.Serialize(writer, pair.Key);
+					CbSerializer.Serialize(writer, pair.Value);
+					writer.EndArray();
 				}
-				Writer.EndUniformArray();
+				writer.EndUniformArray();
 			}
 		}
 
 		/// <inheritdoc/>
-		public override void WriteNamed(CbWriter Writer, Utf8String Name, Dictionary<TKey, TValue> Value)
+		public override void WriteNamed(CbWriter writer, Utf8String name, Dictionary<TKey, TValue> value)
 		{
-			if (Value.Count > 0)
+			if (value.Count > 0)
 			{
-				Writer.BeginUniformArray(Name, CbFieldType.Array);
-				foreach (KeyValuePair<TKey, TValue> Pair in Value)
+				writer.BeginUniformArray(name, CbFieldType.Array);
+				foreach (KeyValuePair<TKey, TValue> pair in value)
 				{
-					Writer.BeginArray();
-					CbSerializer.Serialize(Writer, Pair.Key);
-					CbSerializer.Serialize(Writer, Pair.Value);
-					Writer.EndArray();
+					writer.BeginArray();
+					CbSerializer.Serialize(writer, pair.Key);
+					CbSerializer.Serialize(writer, pair.Value);
+					writer.EndArray();
 				}
-				Writer.EndUniformArray();
+				writer.EndUniformArray();
 			}
 		}
 	}
@@ -84,12 +84,12 @@ namespace EpicGames.Serialization.Converters
 	class CbDictionaryConverterFactory : CbConverterFactory
 	{
 		/// <inheritdoc/>
-		public override ICbConverter? CreateConverter(Type Type)
+		public override ICbConverter? CreateConverter(Type type)
 		{
-			if (Type.IsGenericType && Type.GetGenericTypeDefinition() == typeof(Dictionary<,>))
+			if (type.IsGenericType && type.GetGenericTypeDefinition() == typeof(Dictionary<,>))
 			{
-				Type ConverterType = typeof(CbDictionaryConverter<,>).MakeGenericType(Type.GenericTypeArguments);
-				return (ICbConverter)Activator.CreateInstance(ConverterType)!;
+				Type converterType = typeof(CbDictionaryConverter<,>).MakeGenericType(type.GenericTypeArguments);
+				return (ICbConverter)Activator.CreateInstance(converterType)!;
 			}
 			return null;
 		}
diff --git a/Engine/Source/Programs/Shared/EpicGames.Serialization/Converters/CbEnumConverter.cs b/Engine/Source/Programs/Shared/EpicGames.Serialization/Converters/CbEnumConverter.cs
--- a/Engine/Source/Programs/Shared/EpicGames.Serialization/Converters/CbEnumConverter.cs
+++ b/Engine/Source/Programs/Shared/EpicGames.Serialization/Converters/CbEnumConverter.cs
@@ -12,25 +12,25 @@ namespace EpicGames.Serialization.Converters
 {
 	class CbEnumConverter<T> : CbConverterBase<T>, ICbConverterMethods where T : Enum
 	{
-		Func<CbField, T> ReadFunc;
-		Action<CbWriter, T> WriteFunc;
-		Action<CbWriter, Utf8String, T> WriteNamedFunc;
+		Func<CbField, T> _readFunc;
+		Action<CbWriter, T> _writeFunc;
+		Action<CbWriter, Utf8String, T> _writeNamedFunc;
 
 		public CbEnumConverter()
 		{
-			Type Type = typeof(T);
+			Type type = typeof(T);
 
-			ReadMethod = new DynamicMethod($"Read_{Type.Name}", Type, new Type[] { typeof(CbField) });
-			CreateEnumReader(Type, ReadMethod.GetILGenerator());
-			ReadFunc = (Func<CbField, T>)ReadMethod.CreateDelegate(typeof(Func<CbField, T>));
+			ReadMethod = new DynamicMethod($"Read_{type.Name}", type, new Type[] { typeof(CbField) });
+			CreateEnumReader(type, ReadMethod.GetILGenerator());
+			_readFunc = (Func<CbField, T>)ReadMethod.CreateDelegate(typeof(Func<CbField, T>));
 
-			WriteMethod = new DynamicMethod($"Write_{Type.Name}", null, new Type[] { typeof(CbWriter), Type });
-			CreateEnumWriter(Type, WriteMethod.GetILGenerator());
-			WriteFunc = (Action<CbWriter, T>)WriteMethod.CreateDelegate(typeof(Action<CbWriter, T>));
+			WriteMethod = new DynamicMethod($"Write_{type.Name}", null, new Type[] { typeof(CbWriter), type });
+			CreateEnumWriter(type, WriteMethod.GetILGenerator());
+			_writeFunc = (Action<CbWriter, T>)WriteMethod.CreateDelegate(typeof(Action<CbWriter, T>));
 
-			WriteNamedMethod = new DynamicMethod($"WriteNamed_{Type.Name}", null, new Type[] { typeof(CbWriter), typeof(Utf8String), Type });
-			CreateNamedEnumWriter(Type, WriteNamedMethod.GetILGenerator());
-			WriteNamedFunc = (Action<CbWriter, Utf8String, T>)WriteNamedMethod.CreateDelegate(typeof(Action<CbWriter, Utf8String, T>));
+			WriteNamedMethod = new DynamicMethod($"WriteNamed_{type.Name}", null, new Type[] { typeof(CbWriter), typeof(Utf8String), type });
+			CreateNamedEnumWriter(type, WriteNamedMethod.GetILGenerator());
+			_writeNamedFunc = (Action<CbWriter, Utf8String, T>)WriteNamedMethod.CreateDelegate(typeof(Action<CbWriter, Utf8String, T>));
 		}
 
 		public DynamicMethod ReadMethod { get; }
@@ -42,69 +42,69 @@ namespace EpicGames.Serialization.Converters
 		public DynamicMethod WriteNamedMethod { get; }
 		MethodInfo ICbConverterMethods.WriteNamedMethod => WriteNamedMethod;
 
-		public override T Read(CbField Field) => ReadFunc(Field);
+		public override T Read(CbField field) => _readFunc(field);
 
-		public override void Write(CbWriter Writer, T Value) => WriteFunc(Writer, Value);
+		public override void Write(CbWriter writer, T value) => _writeFunc(writer, value);
 
-		public override void WriteNamed(CbWriter Writer, Utf8String Name, T Value) => WriteNamedFunc(Writer, Name, Value);
+		public override void WriteNamed(CbWriter writer, Utf8String name, T value) => _writeNamedFunc(writer, name, value);
 
-		static void CreateEnumReader(Type Type, ILGenerator Generator)
+		static void CreateEnumReader(Type type, ILGenerator generator)
 		{
-			Generator.Emit(OpCodes.Ldarg_0);
-			Generator.EmitCall(OpCodes.Call, GetMethodInfo<CbField>(x => x.AsInt32()), null);
-			Generator.Emit(OpCodes.Ret);
+			generator.Emit(OpCodes.Ldarg_0);
+			generator.EmitCall(OpCodes.Call, GetMethodInfo<CbField>(x => x.AsInt32()), null);
+			generator.Emit(OpCodes.Ret);
 		}
 
-		static void CreateEnumWriter(Type Type, ILGenerator Generator)
+		static void CreateEnumWriter(Type type, ILGenerator generator)
 		{
-			Generator.Emit(OpCodes.Ldarg_1);
+			generator.Emit(OpCodes.Ldarg_1);
 
-			Label SkipLabel = Generator.DefineLabel();
-			Generator.Emit(OpCodes.Brfalse, SkipLabel);
+			Label skipLabel = generator.DefineLabel();
+			generator.Emit(OpCodes.Brfalse, skipLabel);
 
-			Generator.Emit(OpCodes.Ldarg_0);
-			Generator.Emit(OpCodes.Ldarg_1);
-			Generator.Emit(OpCodes.Conv_I8);
-			Generator.EmitCall(OpCodes.Call, GetMethodInfo<CbWriter>(x => x.WriteIntegerValue(0L)), null);
+			generator.Emit(OpCodes.Ldarg_0);
+			generator.Emit(OpCodes.Ldarg_1);
+			generator.Emit(OpCodes.Conv_I8);
+			generator.EmitCall(OpCodes.Call, GetMethodInfo<CbWriter>(x => x.WriteIntegerValue(0L)), null);
 
-			Generator.MarkLabel(SkipLabel);
-			Generator.Emit(OpCodes.Ret);
+			generator.MarkLabel(skipLabel);
+			generator.Emit(OpCodes.Ret);
 		}
 
-		static void CreateNamedEnumWriter(Type Type, ILGenerator Generator)
+		static void CreateNamedEnumWriter(Type type, ILGenerator generator)
 		{
-			Generator.Emit(OpCodes.Ldarg_2);
+			generator.Emit(OpCodes.Ldarg_2);
 
-			Label SkipLabel = Generator.DefineLabel();
-			Generator.Emit(OpCodes.Brfalse, SkipLabel);
+			Label skipLabel = generator.DefineLabel();
+			generator.Emit(OpCodes.Brfalse, skipLabel);
 
-			Generator.Emit(OpCodes.Ldarg_0);
-			Generator.Emit(OpCodes.Ldarg_1);
-			Generator.Emit(OpCodes.Ldarg_2);
-			Generator.Emit(OpCodes.Conv_I8);
-			Generator.EmitCall(OpCodes.Call, GetMethodInfo<CbWriter>(x => x.WriteInteger(default, 0L)), null);
+			generator.Emit(OpCodes.Ldarg_0);
+			generator.Emit(OpCodes.Ldarg_1);
+			generator.Emit(OpCodes.Ldarg_2);
+			generator.Emit(OpCodes.Conv_I8);
+			generator.EmitCall(OpCodes.Call, GetMethodInfo<CbWriter>(x => x.WriteInteger(default, 0L)), null);
 
-			Generator.MarkLabel(SkipLabel);
-			Generator.Emit(OpCodes.Ret);
+			generator.MarkLabel(skipLabel);
+			generator.Emit(OpCodes.Ret);
 		}
 
-		static MethodInfo GetMethodInfo<TArg>(Expression<Action<TArg>> Expr)
+		static MethodInfo GetMethodInfo<TArg>(Expression<Action<TArg>> expr)
 		{
-			return ((MethodCallExpression)Expr.Body).Method;
+			return ((MethodCallExpression)expr.Body).Method;
 		}
 	}
 
 	class CbEnumConverterFactory : CbConverterFactory
 	{
-		public override ICbConverter? CreateConverter(Type Type)
+		public override ICbConverter? CreateConverter(Type type)
 		{
-			ICbConverter? Converter = null;
-			if (Type.IsEnum)
+			ICbConverter? converter = null;
+			if (type.IsEnum)
 			{
-				Type ConverterType = typeof(CbEnumConverter<>).MakeGenericType(Type);
-				Converter = (ICbConverter?)Activator.CreateInstance(ConverterType);
+				Type converterType = typeof(CbEnumConverter<>).MakeGenericType(type);
+				converter = (ICbConverter?)Activator.CreateInstance(converterType);
 			}
-			return Converter;
+			return converter;
 		}
 	}
 }
diff --git a/Engine/Source/Programs/Shared/EpicGames.Serialization/Converters/CbFieldConverter.cs b/Engine/Source/Programs/Shared/EpicGames.Serialization/Converters/CbFieldConverter.cs
--- a/Engine/Source/Programs/Shared/EpicGames.Serialization/Converters/CbFieldConverter.cs
+++ b/Engine/Source/Programs/Shared/EpicGames.Serialization/Converters/CbFieldConverter.cs
@@ -14,21 +14,21 @@ namespace EpicGames.Serialization.Converters
 	class CbFieldConverter : CbConverterBase<CbField>
 	{
 		/// <inheritdoc/>
-		public override CbField Read(CbField Field)
+		public override CbField Read(CbField field)
 		{
-			return Field;
+			return field;
 		}
 
 		/// <inheritdoc/>
-		public override void Write(CbWriter Writer, CbField Field)
+		public override void Write(CbWriter writer, CbField field)
 		{
-			Writer.WriteFieldValue(Field);
+			writer.WriteFieldValue(field);
 		}
 
 		/// <inheritdoc/>
-		public override void WriteNamed(CbWriter Writer, Utf8String Name, CbField Value)
+		public override void WriteNamed(CbWriter writer, Utf8String name, CbField value)
 		{
-			Writer.WriteField(Name, Value);
+			writer.WriteField(name, value);
 		}
 	}
 
@@ -39,21 +39,21 @@ namespace EpicGames.Serialization.Converters
 	class CbObjectConverter : CbConverterBase<CbObject>
 	{
 		/// <inheritdoc/>
-		public override CbObject Read(CbField Field)
+		public override CbObject Read(CbField field)
 		{
-			return Field.AsObject();
+			return field.AsObject();
 		}
 
 		/// <inheritdoc/>
-		public override void Write(CbWriter Writer, CbObject Object)
+		public override void Write(CbWriter writer, CbObject obj)
 		{
-			Writer.WriteFieldValue(Object.AsField());
+			writer.WriteFieldValue(obj.AsField());
 		}
 
 		/// <inheritdoc/>
-		public override void WriteNamed(CbWriter Writer, Utf8String Name, CbObject Object)
+		public override void WriteNamed(CbWriter writer, Utf8String name, CbObject obj)
 		{
-			Writer.WriteField(Name, Object.AsField());
+			writer.WriteField(name, obj.AsField());
 		}
 	}
 }
diff --git a/Engine/Source/Programs/Shared/EpicGames.Serialization/Converters/CbListConverter.cs b/Engine/Source/Programs/Shared/EpicGames.Serialization/Converters/CbListConverter.cs
--- a/Engine/Source/Programs/Shared/EpicGames.Serialization/Converters/CbListConverter.cs
+++ b/Engine/Source/Programs/Shared/EpicGames.Serialization/Converters/CbListConverter.cs
@@ -20,36 +20,36 @@ namespace EpicGames.Serialization.Converters
 	class CbListConverter<T> : CbConverterBase<List<T>>
 	{
 		/// <inheritdoc/>
-		public override List<T> Read(CbField Field)
+		public override List<T> Read(CbField field)
 		{
-			List<T> List = new List<T>();
-			foreach (CbField ElementField in Field)
+			List<T> list = new List<T>();
+			foreach (CbField elementField in field)
 			{
-				List.Add(CbSerializer.Deserialize<T>(ElementField));
+				list.Add(CbSerializer.Deserialize<T>(elementField));
 			}
-			return List;
+			return list;
 		}
 
 		/// <inheritdoc/>
-		public override void Write(CbWriter Writer, List<T> List)
+		public override void Write(CbWriter writer, List<T> list)
 		{
-			Writer.BeginArray();
-			foreach (T Element in List)
+			writer.BeginArray();
+			foreach (T element in list)
 			{
-				CbSerializer.Serialize<T>(Writer, Element);
+				CbSerializer.Serialize<T>(writer, element);
 			}
-			Writer.EndArray();
+			writer.EndArray();
 		}
 
 		/// <inheritdoc/>
-		public override void WriteNamed(CbWriter Writer, Utf8String Name, List<T> List)
+		public override void WriteNamed(CbWriter writer, Utf8String name, List<T> list)
 		{
-			Writer.BeginArray(Name);
-			foreach (T Element in List)
+			writer.BeginArray(name);
+			foreach (T element in list)
 			{
-				CbSerializer.Serialize<T>(Writer, Element);
+				CbSerializer.Serialize<T>(writer, element);
 			}
-			Writer.EndArray();
+			writer.EndArray();
 		}
 	}
 
@@ -59,38 +59,38 @@ namespace EpicGames.Serialization.Converters
 	sealed class CbStringListConverter : CbConverterBase<List<Utf8String>>
 	{
 		/// <inheritdoc/>
-		public override List<Utf8String> Read(CbField Field)
+		public override List<Utf8String> Read(CbField field)
 		{
-			List<Utf8String> List = new List<Utf8String>();
-			foreach (CbField ElementField in Field)
+			List<Utf8String> list = new List<Utf8String>();
+			foreach (CbField elementField in field)
 			{
-				List.Add(ElementField.AsUtf8String());
+				list.Add(elementField.AsUtf8String());
 			}
-			return List;
+			return list;
 		}
 
 		/// <inheritdoc/>
-		public override void Write(CbWriter Writer, List<Utf8String> List)
+		public override void Write(CbWriter writer, List<Utf8String> list)
 		{
-			Writer.BeginUniformArray(CbFieldType.String);
-			foreach (Utf8String String in List)
+			writer.BeginUniformArray(CbFieldType.String);
+			foreach (Utf8String str in list)
 			{
-				Writer.WriteUtf8StringValue(String);
+				writer.WriteUtf8StringValue(str);
 			}
-			Writer.EndUniformArray();
+			writer.EndUniformArray();
 		}
 
 		/// <inheritdoc/>
-		public override void WriteNamed(CbWriter Writer, Utf8String Name, List<Utf8String> List)
+		public override void WriteNamed(CbWriter writer, Utf8String name, List<Utf8String> list)
 		{
-			if (List.Count > 0)
+			if (list.Count > 0)
 			{
-				Writer.BeginUniformArray(Name, CbFieldType.String);
-				foreach (Utf8String String in List)
+				writer.BeginUniformArray(name, CbFieldType.String);
+				foreach (Utf8String str in list)
 				{
-					Writer.WriteUtf8StringValue(String);
+					writer.WriteUtf8StringValue(str);
 				}
-				Writer.EndUniformArray();
+				writer.EndUniformArray();
 			}
 		}
 	}
@@ -101,18 +101,18 @@ namespace EpicGames.Serialization.Converters
 	class CbListConverterFactory : CbConverterFactory
 	{
 		/// <inheritdoc/>
-		public override ICbConverter? CreateConverter(Type Type)
+		public override ICbConverter? CreateConverter(Type type)
 		{
-			if (Type.IsGenericType && Type.GetGenericTypeDefinition() == typeof(List<>))
+			if (type.IsGenericType && type.GetGenericTypeDefinition() == typeof(List<>))
 			{
-				if (Type.GenericTypeArguments[0] == typeof(Utf8String))
+				if (type.GenericTypeArguments[0] == typeof(Utf8String))
 				{
 					return new CbStringListConverter();
 				}
 				else
 				{
-					Type ConverterType = typeof(CbListConverter<>).MakeGenericType(Type.GenericTypeArguments);
-					return (ICbConverter)Activator.CreateInstance(ConverterType)!;
+					Type converterType = typeof(CbListConverter<>).MakeGenericType(type.GenericTypeArguments);
+					return (ICbConverter)Activator.CreateInstance(converterType)!;
 				}
 			}
 			return null;
diff --git a/Engine/Source/Programs/Shared/EpicGames.Serialization/Converters/CbNullableConverter.cs b/Engine/Source/Programs/Shared/EpicGames.Serialization/Converters/CbNullableConverter.cs
--- a/Engine/Source/Programs/Shared/EpicGames.Serialization/Converters/CbNullableConverter.cs
+++ b/Engine/Source/Programs/Shared/EpicGames.Serialization/Converters/CbNullableConverter.cs
@@ -14,26 +14,26 @@ namespace EpicGames.Serialization.Converters
 	class CbNullableConverter<T> : CbConverterBase<Nullable<T>> where T : struct
 	{
 		/// <inheritdoc/>
-		public override T? Read(CbField Field)
+		public override T? Read(CbField field)
 		{
-			return CbSerializer.Deserialize<T>(Field);
+			return CbSerializer.Deserialize<T>(field);
 		}
 
 		/// <inheritdoc/>
-		public override void Write(CbWriter Writer, T? Nullable)
+		public override void Write(CbWriter writer, T? nullable)
 		{
-			if (Nullable.HasValue)
+			if (nullable.HasValue)
 			{
-				CbSerializer.Serialize<T>(Writer, Nullable.Value);
+				CbSerializer.Serialize<T>(writer, nullable.Value);
 			}
 		}
 
 		/// <inheritdoc/>
-		public override void WriteNamed(CbWriter Writer, Utf8String Name, T? Nullable)
+		public override void WriteNamed(CbWriter writer, Utf8String name, T? nullable)
 		{
-			if (Nullable.HasValue)
+			if (nullable.HasValue)
 			{
-				CbSerializer.Serialize<T>(Writer, Name, Nullable.Value);
+				CbSerializer.Serialize<T>(writer, name, nullable.Value);
 			}
 		}
 	}
@@ -44,12 +44,12 @@ namespace EpicGames.Serialization.Converters
 	class CbNullableConverterFactory : CbConverterFactory
 	{
 		/// <inheritdoc/>
-		public override ICbConverter? CreateConverter(Type Type)
+		public override ICbConverter? CreateConverter(Type type)
 		{
-			if (Type.IsGenericType && Type.GetGenericTypeDefinition() == typeof(Nullable<>))
+			if (type.IsGenericType && type.GetGenericTypeDefinition() == typeof(Nullable<>))
 			{
-				Type ConverterType = typeof(CbNullableConverter<>).MakeGenericType(Type.GenericTypeArguments);
-				return (ICbConverter)Activator.CreateInstance(ConverterType)!;
+				Type converterType = typeof(CbNullableConverter<>).MakeGenericType(type.GenericTypeArguments);
+				return (ICbConverter)Activator.CreateInstance(converterType)!;
 			}
 			return null;
 		}
diff --git a/Engine/Source/Programs/Shared/EpicGames.Serialization/Converters/CbPrimitiveConverter.cs b/Engine/Source/Programs/Shared/EpicGames.Serialization/Converters/CbPrimitiveConverter.cs
--- a/Engine/Source/Programs/Shared/EpicGames.Serialization/Converters/CbPrimitiveConverter.cs
+++ b/Engine/Source/Programs/Shared/EpicGames.Serialization/Converters/CbPrimitiveConverter.cs
@@ -20,22 +20,22 @@ namespace EpicGames.Serialization.Converters
 		public MethodInfo WriteNamedMethod { get; }
 		public Action<CbWriter, Utf8String, T> WriteNamedFunc { get; }
 
-		public CbPrimitiveConverter(Expression<Func<CbField, T>> Read, Expression<Action<CbWriter, T>> Write, Expression<Action<CbWriter, Utf8String, T>> WriteNamed)
+		public CbPrimitiveConverter(Expression<Func<CbField, T>> read, Expression<Action<CbWriter, T>> write, Expression<Action<CbWriter, Utf8String, T>> writeNamed)
 		{
-			this.ReadMethod = ((MethodCallExpression)Read.Body).Method;
-			this.ReadFunc = Read.Compile();
+			this.ReadMethod = ((MethodCallExpression)read.Body).Method;
+			this.ReadFunc = read.Compile();
 
-			this.WriteMethod = ((MethodCallExpression)Write.Body).Method;
-			this.WriteFunc = Write.Compile();
+			this.WriteMethod = ((MethodCallExpression)write.Body).Method;
+			this.WriteFunc = write.Compile();
 
-			this.WriteNamedMethod = ((MethodCallExpression)WriteNamed.Body).Method;
-			this.WriteNamedFunc = WriteNamed.Compile();
+			this.WriteNamedMethod = ((MethodCallExpression)writeNamed.Body).Method;
+			this.WriteNamedFunc = writeNamed.Compile();
 		}
 
-		public override T Read(CbField Field) => ReadFunc(Field);
+		public override T Read(CbField field) => ReadFunc(field);
 
-		public override void Write(CbWriter Writer, T Value) => WriteFunc(Writer, Value);
+		public override void Write(CbWriter writer, T value) => WriteFunc(writer, value);
 
-		public override void WriteNamed(CbWriter Writer, Utf8String Name, T Value) => WriteNamedFunc(Writer, Name, Value);
+		public override void WriteNamed(CbWriter writer, Utf8String name, T value) => WriteNamedFunc(writer, name, value);
 	}
 }
diff --git a/Engine/Source/Programs/Shared/EpicGames.Serialization/VarInt.cs b/Engine/Source/Programs/Shared/EpicGames.Serialization/VarInt.cs
--- a/Engine/Source/Programs/Shared/EpicGames.Serialization/VarInt.cs
+++ b/Engine/Source/Programs/Shared/EpicGames.Serialization/VarInt.cs
@@ -15,40 +15,40 @@ namespace EpicGames.Serialization
 		/// <summary>
 		/// Read a variable-length unsigned integer.
 		/// </summary>
-		/// <param name="Buffer">A variable-length encoding of an unsigned integer</param>
+		/// <param name="buffer">A variable-length encoding of an unsigned integer</param>
 		/// <returns></returns>
-		public static ulong Read(ReadOnlySpan<byte> Buffer)
+		public static ulong Read(ReadOnlySpan<byte> buffer)
 		{
-			return Read(Buffer, out _);
+			return Read(buffer, out _);
 		}
 
 		/// <summary>
 		/// Read a variable-length unsigned integer.
 		/// </summary>
-		/// <param name="Buffer">A variable-length encoding of an unsigned integer</param>
-		/// <param name="BytesRead">The number of bytes consumed from the input</param>
+		/// <param name="buffer">A variable-length encoding of an unsigned integer</param>
+		/// <param name="bytesRead">The number of bytes consumed from the input</param>
 		/// <returns></returns>
-		public static ulong Read(ReadOnlySpan<byte> Buffer, out int BytesRead)
+		public static ulong Read(ReadOnlySpan<byte> buffer, out int bytesRead)
 		{
-			BytesRead = (int)Measure(Buffer);
+			bytesRead = (int)Measure(buffer);
 
-			ulong Value = (ulong)(Buffer[0] & (0xff >> BytesRead));
-			for (int i = 1; i < BytesRead; i++)
+			ulong value = (ulong)(buffer[0] & (0xff >> bytesRead));
+			for (int i = 1; i < bytesRead; i++)
 			{
-				Value <<= 8;
-				Value |= Buffer[i];
+				value <<= 8;
+				value |= buffer[i];
 			}
-			return Value;
+			return value;
 		}
 
 		/// <summary>
 		/// Measure the length in bytes (1-9) of an encoded variable-length integer.
 		/// </summary>
-		/// <param name="Buffer">A variable-length encoding of an(signed or unsigned) integer.</param>
+		/// <param name="buffer">A variable-length encoding of an(signed or unsigned) integer.</param>
 		/// <returns>The number of bytes used to encode the integer, in the range 1-9.</returns>
-		public static int Measure(ReadOnlySpan<byte> Buffer)
+		public static int Measure(ReadOnlySpan<byte> buffer)
 		{
-			byte b = Buffer[0];
+			byte b = buffer[0];
 			b = (byte)~b;
 			return BitOperations.LeadingZeroCount(b) - 23;
 		}
@@ -56,21 +56,21 @@ namespace EpicGames.Serialization
 		/// <summary>
 		/// Measure the number of bytes (1-5) required to encode the 32-bit input.
 		/// </summary>
-		/// <param name="Value"></param>
+		/// <param name="value"></param>
 		/// <returns></returns>
-		public static int Measure(int Value)
+		public static int Measure(int value)
 		{
-			return Measure((ulong)(long)Value);
+			return Measure((ulong)(long)value);
 		}
 
 		/// <summary>
 		/// Measure the number of bytes (1-5) required to encode the 32-bit input.
 		/// </summary>
-		/// <param name="Value"></param>
+		/// <param name="value"></param>
 		/// <returns></returns>
-		public static int Measure(uint Value)
+		public static int Measure(uint value)
 		{
-			return BitOperations.Log2(Value) / 7 + 1;
+			return BitOperations.Log2(value) / 7 + 1;
 		}
 
 		/// <summary>
@@ -86,31 +86,31 @@ namespace EpicGames.Serialization
 		/// <summary>
 		/// Write a variable-length unsigned integer.
 		/// </summary>
-		/// <param name="Value">An unsigned integer to encode</param>
-		/// <param name="Buffer">A buffer of at least 9 bytes to write the output to.</param>
+		/// <param name="value">An unsigned integer to encode</param>
+		/// <param name="buffer">A buffer of at least 9 bytes to write the output to.</param>
 		/// <returns>The number of bytes used in the output</returns>
-		public static int Write(Span<byte> Buffer, long Value)
+		public static int Write(Span<byte> buffer, long value)
 		{
-			return Write(Buffer, (ulong)Value);
+			return Write(buffer, (ulong)value);
 		}
 
 		/// <summary>
 		/// Write a variable-length unsigned integer.
 		/// </summary>
-		/// <param name="Value">An unsigned integer to encode</param>
-		/// <param name="Buffer">A buffer of at least 9 bytes to write the output to.</param>
+		/// <param name="value">An unsigned integer to encode</param>
+		/// <param name="buffer">A buffer of at least 9 bytes to write the output to.</param>
 		/// <returns>The number of bytes used in the output</returns>
-		public static int Write(Span<byte> Buffer, ulong Value)
+		public static int Write(Span<byte> buffer, ulong value)
 		{
-			int ByteCount = Measure(Value);
+			int byteCount = Measure(value);
 
-			for (int Idx = 1; Idx < ByteCount; Idx++)
+			for (int idx = 1; idx < byteCount; idx++)
 			{
-				Buffer[ByteCount - Idx] = (byte)Value;
-				Value >>= 8;
+				buffer[byteCount - idx] = (byte)value;
+				value >>= 8;
 			}
-			Buffer[0] = (byte)((0xff << (9 - (int)ByteCount)) | (byte)Value);
-			return ByteCount;
+			buffer[0] = (byte)((0xff << (9 - (int)byteCount)) | (byte)value);
+			return byteCount;
 		}
 	}
 }
