Reformat EpicGames.Core according to standard coding conventions.

diff --git a/Engine/Source/Programs/AutomationTool/Turnkey/Providers/TurnkeyEditorIOClient.cs b/Engine/Source/Programs/AutomationTool/Turnkey/Providers/TurnkeyEditorIOClient.cs
--- a/Engine/Source/Programs/AutomationTool/Turnkey/Providers/TurnkeyEditorIOClient.cs
+++ b/Engine/Source/Programs/AutomationTool/Turnkey/Providers/TurnkeyEditorIOClient.cs
@@ -275,7 +275,7 @@ namespace Turnkey
 		{
 			// build message
 			StringWriter JsonMessage = new StringWriter();
-			JsonWriter Json = new JsonWriter(JsonMessage, Style:JsonWriterStyle.Readable );
+			JsonWriter Json = new JsonWriter(JsonMessage, style:JsonWriterStyle.Readable );
 			Json.WriteObjectStart();
 				Json.WriteValue("version", Version );
 				MessageFunc(Json);
diff --git a/Engine/Source/Programs/Horde/Horde.Agent/Parser/Matchers/CrashEventMatcher.cs b/Engine/Source/Programs/Horde/Horde.Agent/Parser/Matchers/CrashEventMatcher.cs
--- a/Engine/Source/Programs/Horde/Horde.Agent/Parser/Matchers/CrashEventMatcher.cs
+++ b/Engine/Source/Programs/Horde/Horde.Agent/Parser/Matchers/CrashEventMatcher.cs
@@ -19,7 +19,7 @@ namespace Horde.Agent.Parser.Matchers
 				{
 					if (cursor.IsMatch(maxOffset, "end: stack for UAT"))
 					{
-						LogEventBuilder builder = new LogEventBuilder(cursor, LineCount: maxOffset + 1);
+						LogEventBuilder builder = new LogEventBuilder(cursor, lineCount: maxOffset + 1);
 						return builder.ToMatch(LogEventPriority.BelowNormal, GetLogLevel(cursor), KnownLogEvents.Engine_Crash);
 					}
 				}
diff --git a/Engine/Source/Programs/Shared/EpicGames.Core/AssemblyUtils.cs b/Engine/Source/Programs/Shared/EpicGames.Core/AssemblyUtils.cs
--- a/Engine/Source/Programs/Shared/EpicGames.Core/AssemblyUtils.cs
+++ b/Engine/Source/Programs/Shared/EpicGames.Core/AssemblyUtils.cs
@@ -17,40 +17,34 @@ namespace EpicGames.Core
 		/// assembly path in case shadow copying is enabled.
 		/// </summary>
 		/// <returns>Absolute path and filename to the assembly.</returns>
-		public static string GetOriginalLocation(this Assembly ThisAssembly)
+		public static string GetOriginalLocation(this Assembly thisAssembly)
 		{
-			return new Uri(ThisAssembly.Location).LocalPath;
+			return new Uri(thisAssembly.Location).LocalPath;
 		}
 
 		/// <summary>
 		/// Version info of the executable which runs this code.
 		/// </summary>
-		public static FileVersionInfo ExecutableVersion
-		{
-			get
-			{
-				return FileVersionInfo.GetVersionInfo(Assembly.GetEntryAssembly()!.GetOriginalLocation());
-			}
-		}
+		public static FileVersionInfo ExecutableVersion => FileVersionInfo.GetVersionInfo(Assembly.GetEntryAssembly()!.GetOriginalLocation());
 
 		/// <summary>
 		/// Installs an assembly resolver. Mostly used to get shared assemblies that we don't want copied around to various output locations as happens when "Copy Local" is set to true
 		/// for an assembly reference (which is the default).
 		/// </summary>
-		public static void InstallAssemblyResolver(string PathToBinariesDotNET)
+		public static void InstallAssemblyResolver(string pathToBinariesDotNET)
 		{
 			AppDomain.CurrentDomain.AssemblyResolve += (sender, args) =>
 			{
 				// Name is fully qualified assembly definition - e.g. "p4dn, Version=1.0.0.0, Culture=neutral, PublicKeyToken=ff968dc1933aba6f"
-				string AssemblyName = args.Name!.Split(',')[0];
+				string assemblyName = args.Name!.Split(',')[0];
 
 				return (
-					from KnownAssemblyName in new[] { "SwarmAgent.exe", "../ThirdParty/Ionic/Ionic.Zip.Reduced.dll", "../ThirdParty/Newtonsoft/NewtonSoft.Json.dll" }
-					where AssemblyName.Equals(Path.GetFileNameWithoutExtension(KnownAssemblyName), StringComparison.InvariantCultureIgnoreCase)
-					let ResolvedAssemblyFilename = Path.Combine(PathToBinariesDotNET, KnownAssemblyName)
+					from knownAssemblyName in new[] { "SwarmAgent.exe", "../ThirdParty/Ionic/Ionic.Zip.Reduced.dll", "../ThirdParty/Newtonsoft/NewtonSoft.Json.dll" }
+					where assemblyName.Equals(Path.GetFileNameWithoutExtension(knownAssemblyName), StringComparison.InvariantCultureIgnoreCase)
+					let resolvedAssemblyFilename = Path.Combine(pathToBinariesDotNET, knownAssemblyName)
 					// check if the file exists first. If we just try to load it, we correctly throw an exception, but it's a generic
 					// FileNotFoundException, which is not informative. Better to return null.
-					select File.Exists(ResolvedAssemblyFilename) ? Assembly.LoadFile(ResolvedAssemblyFilename) : null
+					select File.Exists(resolvedAssemblyFilename) ? Assembly.LoadFile(resolvedAssemblyFilename) : null
 					).FirstOrDefault();
 			};
 		}
@@ -58,76 +52,75 @@ namespace EpicGames.Core
 		/// <summary>
 		/// Installs an assembly resolver, which will load *any* assembly which exists recursively within the supplied folder.
 		/// </summary>
-		/// <param name="RootDirectory">The directory to enumerate.</param>
-		public static void InstallRecursiveAssemblyResolver(string RootDirectory)
+		/// <param name="rootDirectory">The directory to enumerate.</param>
+		public static void InstallRecursiveAssemblyResolver(string rootDirectory)
 		{
-			RefreshAssemblyCache(RootDirectory);
+			RefreshAssemblyCache(rootDirectory);
 
 			AppDomain.CurrentDomain.AssemblyResolve += (sender, args) =>
 			{
 				// Name is fully qualified assembly definition - e.g. "p4dn, Version=1.0.0.0, Culture=neutral, PublicKeyToken=ff968dc1933aba6f"
-				string AssemblyName = args.Name!.Split(',')[0];
-				if (AssemblyLocationCache.TryGetValue(AssemblyName, out string? AssemblyLocation))
+				string assemblyName = args.Name!.Split(',')[0];
+				if (s_assemblyLocationCache.TryGetValue(assemblyName, out string? assemblyLocation))
 				{
 					// We have this assembly in our folder.					
-					if (File.Exists(AssemblyLocation))
+					if (File.Exists(assemblyLocation))
 					{
 						// The assembly still exists, so load it.
-						return Assembly.LoadFile(AssemblyLocation);
+						return Assembly.LoadFile(assemblyLocation);
 					}
 					else
 					{
 						// The assembly no longer exists on disk, so remove it from our cache.
-						AssemblyLocationCache.Remove(AssemblyName);
+						s_assemblyLocationCache.Remove(assemblyName);
 					}
 				}
 
 				// The assembly wasn't found, though may have been compiled or copied as a dependency
-				RefreshAssemblyCache(RootDirectory, string.Format("{0}.dll", AssemblyName));
-				if (AssemblyLocationCache.TryGetValue(AssemblyName, out AssemblyLocation))
+				RefreshAssemblyCache(rootDirectory, String.Format("{0}.dll", assemblyName));
+				if (s_assemblyLocationCache.TryGetValue(assemblyName, out assemblyLocation))
 				{
-					return Assembly.LoadFile(AssemblyLocation);
+					return Assembly.LoadFile(assemblyLocation);
 				}
 
 				return null;
 			};
 		}
 
-		private static void RefreshAssemblyCache(string RootDirectory, string Pattern = "*.dll")
+		private static void RefreshAssemblyCache(string rootDirectory, string pattern = "*.dll")
 		{
 			// Initialize our cache of assemblies by enumerating all files in the given folder.
-			foreach (string DiscoveredAssembly in Directory.EnumerateFiles(RootDirectory, Pattern, SearchOption.AllDirectories))
+			foreach (string discoveredAssembly in Directory.EnumerateFiles(rootDirectory, pattern, SearchOption.AllDirectories))
 			{
-				AddFileToAssemblyCache(DiscoveredAssembly);
+				AddFileToAssemblyCache(discoveredAssembly);
 			}
-
 		}
 
-		public static void AddFileToAssemblyCache(string AssemblyPath)
+		public static void AddFileToAssemblyCache(string assemblyPath)
 		{
-			string AssemblyName = Path.GetFileNameWithoutExtension(AssemblyPath);
-			DateTime AssemblyLastWriteTime = File.GetLastWriteTimeUtc(AssemblyPath);
-			if (AssemblyLocationCache.ContainsKey(AssemblyName))
+			string assemblyName = Path.GetFileNameWithoutExtension(assemblyPath);
+			DateTime assemblyLastWriteTime = File.GetLastWriteTimeUtc(assemblyPath);
+			if (s_assemblyLocationCache.ContainsKey(assemblyName))
 			{
 				// We already have this assembly in our cache. Only replace it if the discovered file is newer (to avoid stale assemblies breaking stuff).
-				if (AssemblyLastWriteTime > AssemblyWriteTimes[AssemblyName])
+				if (assemblyLastWriteTime > s_assemblyWriteTimes[assemblyName])
 				{
-					AssemblyLocationCache[AssemblyName] = AssemblyPath;
-					AssemblyWriteTimes[AssemblyName] = AssemblyLastWriteTime;
+					s_assemblyLocationCache[assemblyName] = assemblyPath;
+					s_assemblyWriteTimes[assemblyName] = assemblyLastWriteTime;
 				}
 			}
 			else
 			{
 				// This is the first copy of this assembly ... add it to our cache.
-				AssemblyLocationCache.Add(AssemblyName, AssemblyPath);
-				AssemblyWriteTimes.Add(AssemblyName, AssemblyLastWriteTime);
+				s_assemblyLocationCache.Add(assemblyName, assemblyPath);
+				s_assemblyWriteTimes.Add(assemblyName, assemblyLastWriteTime);
 			}
 		}
 
 		// Map of assembly name to path on disk
-		private static readonly Dictionary<string, string> AssemblyLocationCache = new Dictionary<string, string>(StringComparer.OrdinalIgnoreCase);
+		private static readonly Dictionary<string, string> s_assemblyLocationCache = new Dictionary<string, string>(StringComparer.OrdinalIgnoreCase);
 		// Track last modified date of each assembly, so we can ensure we always reference the latest one in the case of stale assemblies on disk.
-		private static readonly Dictionary<string, DateTime> AssemblyWriteTimes = new Dictionary<string, DateTime>(StringComparer.OrdinalIgnoreCase);
+		private static readonly Dictionary<string, DateTime> s_assemblyWriteTimes = new Dictionary<string, DateTime>(StringComparer.OrdinalIgnoreCase);
 
 	}
 }
diff --git a/Engine/Source/Programs/Shared/EpicGames.Core/AsyncEvent.cs b/Engine/Source/Programs/Shared/EpicGames.Core/AsyncEvent.cs
--- a/Engine/Source/Programs/Shared/EpicGames.Core/AsyncEvent.cs
+++ b/Engine/Source/Programs/Shared/EpicGames.Core/AsyncEvent.cs
@@ -13,15 +13,15 @@ namespace EpicGames.Core
 		/// <summary>
 		/// Completion source to wait on
 		/// </summary>
-		TaskCompletionSource<bool> Source = new TaskCompletionSource<bool>();
+		TaskCompletionSource<bool> _source = new TaskCompletionSource<bool>();
 
 		/// <summary>
 		/// Signal the event
 		/// </summary>
 		public void Set()
 		{
-			TaskCompletionSource<bool> PrevSource = Interlocked.Exchange(ref Source, new TaskCompletionSource<bool>());
-			PrevSource.SetResult(true);
+			TaskCompletionSource<bool> prevSource = Interlocked.Exchange(ref _source, new TaskCompletionSource<bool>());
+			prevSource.SetResult(true);
 		}
 
 		/// <summary>
@@ -30,15 +30,12 @@ namespace EpicGames.Core
 		/// <returns>True if the event is set</returns>
 		public bool IsSet()
 		{
-			return Source.Task.IsCompleted;
+			return _source.Task.IsCompleted;
 		}
 
 		/// <summary>
 		/// Waits for this event to be set
 		/// </summary>
-		public Task Task
-		{
-			get { return Source.Task; }
-		}
+		public Task Task => _source.Task;
 	}
 }
diff --git a/Engine/Source/Programs/Shared/EpicGames.Core/AsyncUtils.cs b/Engine/Source/Programs/Shared/EpicGames.Core/AsyncUtils.cs
--- a/Engine/Source/Programs/Shared/EpicGames.Core/AsyncUtils.cs
+++ b/Engine/Source/Programs/Shared/EpicGames.Core/AsyncUtils.cs
@@ -2,7 +2,6 @@
 
 using System;
 using System.Collections.Generic;
-using System.Text;
 using System.Threading;
 using System.Threading.Tasks;
 
@@ -16,30 +15,30 @@ namespace EpicGames.Core
 		/// <summary>
 		/// Converts a cancellation token to a waitable task
 		/// </summary>
-		/// <param name="Token">Cancellation token</param>
+		/// <param name="token">Cancellation token</param>
 		/// <returns></returns>
-		public static Task AsTask(this CancellationToken Token)
+		public static Task AsTask(this CancellationToken token)
 		{
-			return Task.Delay(-1, Token).ContinueWith(x => { });
+			return Task.Delay(-1, token).ContinueWith(x => { });
 		}
 
 		/// <summary>
 		/// Attempts to get the result of a task, if it has finished
 		/// </summary>
 		/// <typeparam name="T"></typeparam>
-		/// <param name="Task"></param>
-		/// <param name="Result"></param>
+		/// <param name="task"></param>
+		/// <param name="result"></param>
 		/// <returns></returns>
-		public static bool TryGetResult<T>(this Task<T> Task, out T Result)
+		public static bool TryGetResult<T>(this Task<T> task, out T result)
 		{
-			if (Task.IsCompleted)
+			if (task.IsCompleted)
 			{
-				Result = Task.Result;
+				result = task.Result;
 				return true;
 			}
 			else
 			{
-				Result = default!;
+				result = default!;
 				return false;
 			}
 		}
@@ -47,74 +46,74 @@ namespace EpicGames.Core
 		/// <summary>
 		/// Waits for a time period to elapse or the task to be cancelled, without throwing an cancellation exception
 		/// </summary>
-		/// <param name="Time">Time to wait</param>
-		/// <param name="Token">Cancellation token</param>
+		/// <param name="time">Time to wait</param>
+		/// <param name="token">Cancellation token</param>
 		/// <returns></returns>
-		public static Task DelayNoThrow(TimeSpan Time, CancellationToken Token)
+		public static Task DelayNoThrow(TimeSpan time, CancellationToken token)
 		{
-			return Task.Delay(Time, Token).ContinueWith(x => { });
+			return Task.Delay(time, token).ContinueWith(x => { });
 		}
 
 		/// <summary>
 		/// Removes all the complete tasks from a list, allowing each to throw exceptions as necessary
 		/// </summary>
-		/// <param name="Tasks">List of tasks to remove tasks from</param>
-		public static void RemoveCompleteTasks(this List<Task> Tasks)
+		/// <param name="tasks">List of tasks to remove tasks from</param>
+		public static void RemoveCompleteTasks(this List<Task> tasks)
 		{
-			List<Exception> Exceptions = new List<Exception>();
+			List<Exception> exceptions = new List<Exception>();
 
-			int OutIdx = 0;
-			for (int Idx = 0; Idx < Tasks.Count; Idx++)
+			int outIdx = 0;
+			for (int idx = 0; idx < tasks.Count; idx++)
 			{
-				if (Tasks[Idx].IsCompleted)
+				if (tasks[idx].IsCompleted)
 				{
-					AggregateException? Exception = Tasks[Idx].Exception;
-					if (Exception != null)
+					AggregateException? exception = tasks[idx].Exception;
+					if (exception != null)
 					{
-						Exceptions.AddRange(Exception.InnerExceptions);
+						exceptions.AddRange(exception.InnerExceptions);
 					}
 				}
 				else
 				{
-					if (Idx != OutIdx)
+					if (idx != outIdx)
 					{
-						Tasks[OutIdx] = Tasks[Idx];
+						tasks[outIdx] = tasks[idx];
 					}
-					OutIdx++;
+					outIdx++;
 				}
 			}
-			Tasks.RemoveRange(OutIdx, Tasks.Count - OutIdx);
+			tasks.RemoveRange(outIdx, tasks.Count - outIdx);
 
-			if (Exceptions.Count > 0)
+			if (exceptions.Count > 0)
 			{
-				throw new AggregateException(Exceptions);
+				throw new AggregateException(exceptions);
 			}
 		}
 
 		/// <summary>
 		/// Removes all the complete tasks from a list, allowing each to throw exceptions as necessary
 		/// </summary>
-		/// <param name="Tasks">List of tasks to remove tasks from</param>
+		/// <param name="tasks">List of tasks to remove tasks from</param>
 		/// <returns>Return values from the completed tasks</returns>
-		public static List<T> RemoveCompleteTasks<T>(this List<Task<T>> Tasks)
+		public static List<T> RemoveCompleteTasks<T>(this List<Task<T>> tasks)
 		{
-			List<T> Results = new List<T>();
+			List<T> results = new List<T>();
 
-			int OutIdx = 0;
-			for (int Idx = 0; Idx < Tasks.Count; Idx++)
+			int outIdx = 0;
+			for (int idx = 0; idx < tasks.Count; idx++)
 			{
-				if (Tasks[Idx].IsCompleted)
+				if (tasks[idx].IsCompleted)
 				{
-					Results.Add(Tasks[Idx].Result);
+					results.Add(tasks[idx].Result);
 				}
-				else if (Idx != OutIdx)
+				else if (idx != outIdx)
 				{
-					Tasks[OutIdx++] = Tasks[Idx];
+					tasks[outIdx++] = tasks[idx];
 				}
 			}
-			Tasks.RemoveRange(OutIdx, Tasks.Count - OutIdx);
+			tasks.RemoveRange(outIdx, tasks.Count - outIdx);
 
-			return Results;
+			return results;
 		}
 	}
 }
diff --git a/Engine/Source/Programs/Shared/EpicGames.Core/BinaryArchive.cs b/Engine/Source/Programs/Shared/EpicGames.Core/BinaryArchive.cs
--- a/Engine/Source/Programs/Shared/EpicGames.Core/BinaryArchive.cs
+++ b/Engine/Source/Programs/Shared/EpicGames.Core/BinaryArchive.cs
@@ -26,135 +26,135 @@ namespace EpicGames.Core
 			public MethodInfo WriteMethodInfo { get; }
 			public Action<BinaryArchiveWriter, object> WriteMethod { get; }
 
-			public TypeSerializer(MethodInfo ReadMethodInfo, MethodInfo WriteMethodInfo)
+			public TypeSerializer(MethodInfo readMethodInfo, MethodInfo writeMethodInfo)
 			{
-				this.ReadMethodInfo = ReadMethodInfo;
-				this.WriteMethodInfo = WriteMethodInfo;
+				ReadMethodInfo = readMethodInfo;
+				WriteMethodInfo = writeMethodInfo;
 
-				Type Type = ReadMethodInfo.ReturnType;
-				this.ReadMethod = CreateBoxedReadMethod(Type, ReadMethodInfo);
-				this.WriteMethod = CreateBoxedWriteMethod(Type, WriteMethodInfo);
+				Type type = readMethodInfo.ReturnType;
+				ReadMethod = CreateBoxedReadMethod(type, readMethodInfo);
+				WriteMethod = CreateBoxedWriteMethod(type, writeMethodInfo);
 			}
 
-			static Func<BinaryArchiveReader, object> CreateBoxedReadMethod(Type Type, MethodInfo MethodInfo)
+			static Func<BinaryArchiveReader, object> CreateBoxedReadMethod(Type type, MethodInfo methodInfo)
 			{
-				DynamicMethod ReaderMethod = new DynamicMethod($"Boxed_{MethodInfo.Name}", typeof(object), new[] { typeof(BinaryArchiveReader) });
+				DynamicMethod readerMethod = new DynamicMethod($"Boxed_{methodInfo.Name}", typeof(object), new[] { typeof(BinaryArchiveReader) });
 
-				ILGenerator Generator = ReaderMethod.GetILGenerator(64);
-				Generator.Emit(OpCodes.Ldarg_0);
-				Generator.EmitCall(OpCodes.Call, MethodInfo, null);
-				if (!Type.IsClass)
+				ILGenerator generator = readerMethod.GetILGenerator(64);
+				generator.Emit(OpCodes.Ldarg_0);
+				generator.EmitCall(OpCodes.Call, methodInfo, null);
+				if (!type.IsClass)
 				{
-					Generator.Emit(OpCodes.Box);
+					generator.Emit(OpCodes.Box);
 				}
-				Generator.Emit(OpCodes.Ret);
+				generator.Emit(OpCodes.Ret);
 
-				return (Func<BinaryArchiveReader, object>)ReaderMethod.CreateDelegate(typeof(Func<BinaryArchiveReader, object>));
+				return (Func<BinaryArchiveReader, object>)readerMethod.CreateDelegate(typeof(Func<BinaryArchiveReader, object>));
 			}
 
-			static Action<BinaryArchiveWriter, object> CreateBoxedWriteMethod(Type Type, MethodInfo MethodInfo)
+			static Action<BinaryArchiveWriter, object> CreateBoxedWriteMethod(Type type, MethodInfo methodInfo)
 			{
-				DynamicMethod WriterMethod = new DynamicMethod($"Boxed_{MethodInfo.Name}", typeof(void), new[] { typeof(BinaryArchiveWriter), typeof(object) });
+				DynamicMethod writerMethod = new DynamicMethod($"Boxed_{methodInfo.Name}", typeof(void), new[] { typeof(BinaryArchiveWriter), typeof(object) });
 
-				ILGenerator Generator = WriterMethod.GetILGenerator(64);
-				Generator.Emit(OpCodes.Ldarg_0);
-				Generator.Emit(OpCodes.Ldarg_1);
-				Generator.Emit(OpCodes.Unbox_Any, Type);
-				Generator.EmitCall(OpCodes.Call, MethodInfo, null);
-				Generator.Emit(OpCodes.Ret);
+				ILGenerator generator = writerMethod.GetILGenerator(64);
+				generator.Emit(OpCodes.Ldarg_0);
+				generator.Emit(OpCodes.Ldarg_1);
+				generator.Emit(OpCodes.Unbox_Any, type);
+				generator.EmitCall(OpCodes.Call, methodInfo, null);
+				generator.Emit(OpCodes.Ret);
 
-				return (Action<BinaryArchiveWriter, object>)WriterMethod.CreateDelegate(typeof(Action<BinaryArchiveWriter, object>));
+				return (Action<BinaryArchiveWriter, object>)writerMethod.CreateDelegate(typeof(Action<BinaryArchiveWriter, object>));
 			}
 
-			public static TypeSerializer Create<T>(Expression<Func<BinaryArchiveReader, T>> ReaderExpr, Expression<Action<BinaryArchiveWriter, T>> WriterExpr)
+			public static TypeSerializer Create<T>(Expression<Func<BinaryArchiveReader, T>> readerExpr, Expression<Action<BinaryArchiveWriter, T>> writerExpr)
 			{
-				MethodInfo ReadMethod = ((MethodCallExpression)ReaderExpr.Body).Method;
-				MethodInfo WriteMethod = ((MethodCallExpression)WriterExpr.Body).Method;
-				return new TypeSerializer(ReadMethod, WriteMethod);
+				MethodInfo readMethod = ((MethodCallExpression)readerExpr.Body).Method;
+				MethodInfo writeMethod = ((MethodCallExpression)writerExpr.Body).Method;
+				return new TypeSerializer(readMethod, writeMethod);
 			}
 		}
 
-		static readonly ConcurrentDictionary<Type, TypeSerializer> TypeToSerializerInfo = new ConcurrentDictionary<Type, TypeSerializer>()
+		static readonly ConcurrentDictionary<Type, TypeSerializer> s_typeToSerializerInfo = new ConcurrentDictionary<Type, TypeSerializer>()
 		{
-			[typeof(byte)] = TypeSerializer.Create(Reader => Reader.ReadByte(), (Writer, Value) => Writer.WriteByte(Value)),
-			[typeof(sbyte)] = TypeSerializer.Create(Reader => Reader.ReadSignedByte(), (Writer, Value) => Writer.WriteSignedByte(Value)),
-			[typeof(short)] = TypeSerializer.Create(Reader => Reader.ReadShort(), (Writer, Value) => Writer.WriteShort(Value)),
-			[typeof(ushort)] = TypeSerializer.Create(Reader => Reader.ReadUnsignedShort(), (Writer, Value) => Writer.WriteUnsignedShort(Value)),
-			[typeof(int)] = TypeSerializer.Create(Reader => Reader.ReadInt(), (Writer, Value) => Writer.WriteInt(Value)),
-			[typeof(uint)] = TypeSerializer.Create(Reader => Reader.ReadUnsignedInt(), (Writer, Value) => Writer.WriteUnsignedInt(Value)),
-			[typeof(long)] = TypeSerializer.Create(Reader => Reader.ReadLong(), (Writer, Value) => Writer.WriteLong(Value)),
-			[typeof(ulong)] = TypeSerializer.Create(Reader => Reader.ReadUnsignedLong(), (Writer, Value) => Writer.WriteUnsignedLong(Value)),
-			[typeof(string)] = TypeSerializer.Create(Reader => Reader.ReadString(), (Writer, Value) => Writer.WriteString(Value))
+			[typeof(byte)] = TypeSerializer.Create(reader => reader.ReadByte(), (writer, value) => writer.WriteByte(value)),
+			[typeof(sbyte)] = TypeSerializer.Create(reader => reader.ReadSignedByte(), (writer, value) => writer.WriteSignedByte(value)),
+			[typeof(short)] = TypeSerializer.Create(reader => reader.ReadShort(), (writer, value) => writer.WriteShort(value)),
+			[typeof(ushort)] = TypeSerializer.Create(reader => reader.ReadUnsignedShort(), (writer, value) => writer.WriteUnsignedShort(value)),
+			[typeof(int)] = TypeSerializer.Create(reader => reader.ReadInt(), (writer, value) => writer.WriteInt(value)),
+			[typeof(uint)] = TypeSerializer.Create(reader => reader.ReadUnsignedInt(), (writer, value) => writer.WriteUnsignedInt(value)),
+			[typeof(long)] = TypeSerializer.Create(reader => reader.ReadLong(), (writer, value) => writer.WriteLong(value)),
+			[typeof(ulong)] = TypeSerializer.Create(reader => reader.ReadUnsignedLong(), (writer, value) => writer.WriteUnsignedLong(value)),
+			[typeof(string)] = TypeSerializer.Create(reader => reader.ReadString(), (writer, value) => writer.WriteString(value))
 		};
 
-		static readonly ConcurrentDictionary<Type, ContentHash> TypeToDigest = new ConcurrentDictionary<Type, ContentHash>();
-		static readonly ConcurrentDictionary<ContentHash, Type> DigestToType = new ConcurrentDictionary<ContentHash, Type>();
+		static readonly ConcurrentDictionary<Type, ContentHash> s_typeToDigest = new ConcurrentDictionary<Type, ContentHash>();
+		static readonly ConcurrentDictionary<ContentHash, Type> s_digestToType = new ConcurrentDictionary<ContentHash, Type>();
 
-		static MethodInfo GetMethodInfo(Expression<Action> Expr)
+		static MethodInfo GetMethodInfo(Expression<Action> expr)
 		{
-			return ((MethodCallExpression)Expr.Body).Method;
+			return ((MethodCallExpression)expr.Body).Method;
 		}
 
-		static MethodInfo GetMethodInfo<T>(Expression<Action<T>> Expr)
+		static MethodInfo GetMethodInfo<T>(Expression<Action<T>> expr)
 		{
-			return ((MethodCallExpression)Expr.Body).Method;
+			return ((MethodCallExpression)expr.Body).Method;
 		}
 
-		static MethodInfo GetGenericMethodInfo<T, T1>(Expression<Action<T, T1>> Expr)
+		static MethodInfo GetGenericMethodInfo<T, T1>(Expression<Action<T, T1>> expr)
 		{
-			return ((MethodCallExpression)Expr.Body).Method.GetGenericMethodDefinition();
+			return ((MethodCallExpression)expr.Body).Method.GetGenericMethodDefinition();
 		}
 
-		static readonly MethodInfo ReadBoolMethodInfo = GetMethodInfo<BinaryArchiveReader>(x => x.ReadBool());
-		static readonly MethodInfo ReadIntMethodInfo = GetMethodInfo<BinaryArchiveReader>(x => x.ReadInt());
-		static readonly MethodInfo ReadArrayMethodInfo = GetGenericMethodInfo<BinaryArchiveReader, Func<int>>((x, y) => x.ReadArray(y));
-		static readonly MethodInfo ReadPolymorphicObjectMethodInfo = GetMethodInfo(() => ReadObject(null!));
-		static readonly MethodInfo WriteBoolMethodInfo = GetMethodInfo<BinaryArchiveWriter>(x => x.WriteBool(false));
-		static readonly MethodInfo WriteIntMethodInfo = GetMethodInfo<BinaryArchiveWriter>(x => x.WriteInt(0));
-		static readonly MethodInfo WriteArrayMethodInfo = GetGenericMethodInfo<BinaryArchiveWriter, Action<int>>((x, y) => x.WriteArray(null!, y));
-		static readonly MethodInfo WritePolymorphicObjectMethodInfo = GetMethodInfo(() => WritePolymorphicObject(null!, null!));
+		static readonly MethodInfo s_readBoolMethodInfo = GetMethodInfo<BinaryArchiveReader>(x => x.ReadBool());
+		static readonly MethodInfo s_readIntMethodInfo = GetMethodInfo<BinaryArchiveReader>(x => x.ReadInt());
+		static readonly MethodInfo s_readArrayMethodInfo = GetGenericMethodInfo<BinaryArchiveReader, Func<int>>((x, y) => x.ReadArray(y));
+		static readonly MethodInfo s_readPolymorphicObjectMethodInfo = GetMethodInfo(() => ReadObject(null!));
+		static readonly MethodInfo s_writeBoolMethodInfo = GetMethodInfo<BinaryArchiveWriter>(x => x.WriteBool(false));
+		static readonly MethodInfo s_writeIntMethodInfo = GetMethodInfo<BinaryArchiveWriter>(x => x.WriteInt(0));
+		static readonly MethodInfo s_writeArrayMethodInfo = GetGenericMethodInfo<BinaryArchiveWriter, Action<int>>((x, y) => x.WriteArray(null!, y));
+		static readonly MethodInfo s_writePolymorphicObjectMethodInfo = GetMethodInfo(() => WritePolymorphicObject(null!, null!));
 
 		/// <summary>
 		/// Writes an arbitrary type to the given archive. May be an object or value type.
 		/// </summary>
 		/// <typeparam name="T">The type to write</typeparam>
-		/// <param name="Writer">Writer to serialize to</param>
-		/// <param name="Value">Value to write</param>
-		public static void Write<T>(this BinaryArchiveWriter Writer, T Value)
+		/// <param name="writer">Writer to serialize to</param>
+		/// <param name="value">Value to write</param>
+		public static void Write<T>(this BinaryArchiveWriter writer, T value)
 		{
-			Type Type = typeof(T);
-			if (Type.IsClass && !Type.IsSealed)
+			Type type = typeof(T);
+			if (type.IsClass && !type.IsSealed)
 			{
-				Writer.WriteObjectReference(Value!, () => WriteNewPolymorphicObject(Writer, Value!));
+				writer.WriteObjectReference(value!, () => WriteNewPolymorphicObject(writer, value!));
 			}
 			else
 			{
-				FindOrAddSerializerInfo(Type).WriteMethod(Writer, Value!);
+				FindOrAddSerializerInfo(type).WriteMethod(writer, value!);
 			}
 		}
 
 		/// <summary>
 		/// Registers the given type. Allows serializing/deserializing it through calls to ReadType/WriteType.
 		/// </summary>
-		/// <param name="Type">Type to register</param>
-		public static void RegisterType(Type Type)
+		/// <param name="type">Type to register</param>
+		public static void RegisterType(Type type)
 		{
-			ContentHash Digest = ContentHash.MD5($"{Type.Assembly.FullName}\n{Type.FullName}");
-			TypeToDigest.TryAdd(Type, Digest);
-			DigestToType.TryAdd(Digest, Type);
+			ContentHash digest = ContentHash.MD5($"{type.Assembly.FullName}\n{type.FullName}");
+			s_typeToDigest.TryAdd(type, digest);
+			s_digestToType.TryAdd(digest, type);
 		}
 
 		/// <summary>
 		/// Registers all types in the given assembly with the <see cref="BinarySerializableAttribute"/> attribute for serialization
 		/// </summary>
-		/// <param name="Assembly">Assembly to search in</param>
-		public static void RegisterTypes(Assembly Assembly)
+		/// <param name="assembly">Assembly to search in</param>
+		public static void RegisterTypes(Assembly assembly)
 		{
-			foreach (Type Type in Assembly.GetTypes())
+			foreach (Type type in assembly.GetTypes())
 			{
-				if (Type.GetCustomAttribute<BinarySerializableAttribute>() != null)
+				if (type.GetCustomAttribute<BinarySerializableAttribute>() != null)
 				{
-					RegisterType(Type);
+					RegisterType(type);
 				}
 			}
 		}
@@ -162,277 +162,277 @@ namespace EpicGames.Core
 		/// <summary>
 		/// Reads a tyoe from the archive
 		/// </summary>
-		/// <param name="Reader">Reader to deserializer the type from</param>
+		/// <param name="reader">Reader to deserializer the type from</param>
 		/// <returns>The matching type</returns>
-		public static Type? ReadType(this BinaryArchiveReader Reader)
+		public static Type? ReadType(this BinaryArchiveReader reader)
 		{
-			return Reader.ReadObjectReference(() => DigestToType[Reader.ReadContentHash()!]);
+			return reader.ReadObjectReference(() => s_digestToType[reader.ReadContentHash()!]);
 		}
 
 		/// <summary>
 		/// Writes a type to an archive
 		/// </summary>
-		/// <param name="Writer">Writer to serialize to</param>
-		/// <param name="Type">The type to serialize</param>
-		public static void WriteType(this BinaryArchiveWriter Writer, Type Type)
+		/// <param name="writer">Writer to serialize to</param>
+		/// <param name="type">The type to serialize</param>
+		public static void WriteType(this BinaryArchiveWriter writer, Type type)
 		{
-			Writer.WriteObjectReference(Type, () => Writer.WriteContentHash(TypeToDigest[Type]));
+			writer.WriteObjectReference(type, () => writer.WriteContentHash(s_typeToDigest[type]));
 		}
 
-		static object? ReadNewPolymorphicObject(BinaryArchiveReader Reader)
+		static object? ReadNewPolymorphicObject(BinaryArchiveReader reader)
 		{
-			Type? FinalType = Reader.ReadType();
-			if (FinalType == null)
+			Type? finalType = reader.ReadType();
+			if (finalType == null)
 			{
 				return null;
 			}
 
-			TypeSerializer Info = FindOrAddSerializerInfo(FinalType);
-			return Info.ReadMethod(Reader)!;
+			TypeSerializer info = FindOrAddSerializerInfo(finalType);
+			return info.ReadMethod(reader)!;
 		}
 
-		public static object? ReadObject(this BinaryArchiveReader Reader)
+		public static object? ReadObject(this BinaryArchiveReader reader)
 		{
-			return Reader.ReadUntypedObjectReference(() => ReadNewPolymorphicObject(Reader));
+			return reader.ReadUntypedObjectReference(() => ReadNewPolymorphicObject(reader));
 		}
 
-		static void WriteNewPolymorphicObject(BinaryArchiveWriter Writer, object Value)
+		static void WriteNewPolymorphicObject(BinaryArchiveWriter writer, object value)
 		{
-			Type ActualType = Value.GetType();
-			Writer.WriteType(ActualType);
+			Type actualType = value.GetType();
+			writer.WriteType(actualType);
 
-			TypeSerializer Info = FindOrAddSerializerInfo(ActualType);
-			Info.WriteMethod(Writer, Value);
+			TypeSerializer info = FindOrAddSerializerInfo(actualType);
+			info.WriteMethod(writer, value);
 		}
 
-		static void WritePolymorphicObject(this BinaryArchiveWriter Writer, object Value)
+		static void WritePolymorphicObject(this BinaryArchiveWriter writer, object value)
 		{
-			Writer.WriteObjectReference(Value, () => WriteNewPolymorphicObject(Writer, Value));
+			writer.WriteObjectReference(value, () => WriteNewPolymorphicObject(writer, value));
 		}
 
 		/// <summary>
 		/// Finds an existing serializer for the given type, or creates one
 		/// </summary>
-		/// <param name="Type">Type to create a serializer for</param>
-		/// <param name="ConverterType">Type of the converter to use</param>
+		/// <param name="type">Type to create a serializer for</param>
+		/// <param name="converterType">Type of the converter to use</param>
 		/// <returns>Instance of the type serializer</returns>
-		static TypeSerializer FindOrAddSerializerInfo(Type Type, Type? ConverterType = null)
+		static TypeSerializer FindOrAddSerializerInfo(Type type, Type? converterType = null)
 		{
-			Type SerializerKey = ConverterType ?? Type;
+			Type serializerKey = converterType ?? type;
 
 			// Get the serializer info
-			if (!TypeToSerializerInfo.TryGetValue(SerializerKey, out TypeSerializer? SerializerInfo))
+			if (!s_typeToSerializerInfo.TryGetValue(serializerKey, out TypeSerializer? serializerInfo))
 			{
-				lock (TypeToSerializerInfo)
+				lock (s_typeToSerializerInfo)
 				{
-					SerializerInfo = CreateTypeSerializer(Type, ConverterType);
-					TypeToSerializerInfo[SerializerKey] = SerializerInfo;
+					serializerInfo = CreateTypeSerializer(type, converterType);
+					s_typeToSerializerInfo[serializerKey] = serializerInfo;
 				}
 			}
 
-			return SerializerInfo;
+			return serializerInfo;
 		}
 
 		/// <summary>
 		/// Creates a serializer for the given type
 		/// </summary>
-		/// <param name="Type">Type to create a serializer from</param>
-		/// <param name="ConverterType">Converter for the type</param>
+		/// <param name="type">Type to create a serializer from</param>
+		/// <param name="converterType">Converter for the type</param>
 		/// <returns>New instance of a type serializer</returns>
-		static TypeSerializer CreateTypeSerializer(Type Type, Type? ConverterType = null)
+		static TypeSerializer CreateTypeSerializer(Type type, Type? converterType = null)
 		{
 			// If there's a specific converter defined, generate serialization methods using that
-			if (ConverterType != null)
+			if (converterType != null)
 			{
 				// Make sure the converter is derived from IBinaryConverter
-				Type InterfaceType = typeof(IBinaryConverter<>).MakeGenericType(Type);
-				if (!InterfaceType.IsAssignableFrom(ConverterType))
+				Type interfaceType = typeof(IBinaryConverter<>).MakeGenericType(type);
+				if (!interfaceType.IsAssignableFrom(converterType))
 				{
-					throw new NotImplementedException($"Converter does not implement IBinaryArchiveConverter<{Type.Name}>");
+					throw new NotImplementedException($"Converter does not implement IBinaryArchiveConverter<{type.Name}>");
 				}
 
 				// Instantiate the converter, and store it in a static variable
-				Type InstantiatorType = typeof(Instantiator<>).MakeGenericType(ConverterType);
-				FieldInfo ConverterField = InstantiatorType.GetField(nameof(Instantiator<object>.Instance))!;
+				Type instantiatorType = typeof(Instantiator<>).MakeGenericType(converterType);
+				FieldInfo converterField = instantiatorType.GetField(nameof(Instantiator<object>.Instance))!;
 
 				// Get the mapping between interface methods and instance methods
-				InterfaceMapping InterfaceMapping = ConverterType.GetInterfaceMap(InterfaceType);
+				InterfaceMapping interfaceMapping = converterType.GetInterfaceMap(interfaceType);
 
 				// Create a reader method
-				DynamicMethod ReaderMethod = new DynamicMethod($"BinaryArchiveReader_Dynamic_{Type.Name}_With_{ConverterType.Name}", Type, new[] { typeof(BinaryArchiveReader) });
-				int ReaderIdx = Array.FindIndex(InterfaceMapping.InterfaceMethods, x => x.Name.Equals(nameof(IBinaryConverter<object>.Read), StringComparison.Ordinal));
-				GenerateConverterReaderForwardingMethod(ReaderMethod, ConverterField, InterfaceMapping.TargetMethods[ReaderIdx]);
+				DynamicMethod readerMethod = new DynamicMethod($"BinaryArchiveReader_Dynamic_{type.Name}_With_{converterType.Name}", type, new[] { typeof(BinaryArchiveReader) });
+				int readerIdx = Array.FindIndex(interfaceMapping.InterfaceMethods, x => x.Name.Equals(nameof(IBinaryConverter<object>.Read), StringComparison.Ordinal));
+				GenerateConverterReaderForwardingMethod(readerMethod, converterField, interfaceMapping.TargetMethods[readerIdx]);
 
 				// Create a writer method
-				DynamicMethod WriterMethod = new DynamicMethod($"BinaryArchiveWriter_Dynamic_{Type.Name}_With_{ConverterType.Name}", typeof(void), new[] { typeof(BinaryArchiveWriter), typeof(object) });
-				int WriterIdx = Array.FindIndex(InterfaceMapping.InterfaceMethods, x => x.Name.Equals(nameof(IBinaryConverter<object>.Write), StringComparison.Ordinal));
-				GenerateConverterWriterForwardingMethod(WriterMethod, ConverterField, InterfaceMapping.TargetMethods[WriterIdx]);
+				DynamicMethod writerMethod = new DynamicMethod($"BinaryArchiveWriter_Dynamic_{type.Name}_With_{converterType.Name}", typeof(void), new[] { typeof(BinaryArchiveWriter), typeof(object) });
+				int writerIdx = Array.FindIndex(interfaceMapping.InterfaceMethods, x => x.Name.Equals(nameof(IBinaryConverter<object>.Write), StringComparison.Ordinal));
+				GenerateConverterWriterForwardingMethod(writerMethod, converterField, interfaceMapping.TargetMethods[writerIdx]);
 
-				return new TypeSerializer(ReaderMethod, WriterMethod);
+				return new TypeSerializer(readerMethod, writerMethod);
 			}
 
 			// Get the default converter type
-			if (BinaryConverter.TryGetConverterType(Type, out Type? DefaultConverterType))
+			if (BinaryConverter.TryGetConverterType(type, out Type? defaultConverterType))
 			{
-				return FindOrAddSerializerInfo(Type, DefaultConverterType);
+				return FindOrAddSerializerInfo(type, defaultConverterType);
 			}
 
 			// Check if it's an array
-			if (Type.IsArray)
+			if (type.IsArray)
 			{
-				Type ElementType = Type.GetElementType()!;
+				Type elementType = type.GetElementType()!;
 
-				MethodInfo ReaderMethod = ReadArrayMethodInfo.MakeGenericMethod(ElementType);
-				MethodInfo WriterMethod = WriteArrayMethodInfo.MakeGenericMethod(ElementType);
+				MethodInfo readerMethod = s_readArrayMethodInfo.MakeGenericMethod(elementType);
+				MethodInfo writerMethod = s_writeArrayMethodInfo.MakeGenericMethod(elementType);
 
-				return new TypeSerializer(ReaderMethod, WriterMethod);
+				return new TypeSerializer(readerMethod, writerMethod);
 			}
 
 			// Check if it's a class
-			if (Type.IsClass)
+			if (type.IsClass)
 			{
 				// Get all the class properties
-				List<PropertyInfo> Properties = new List<PropertyInfo>();
-				foreach (PropertyInfo Property in Type.GetProperties(BindingFlags.Public | BindingFlags.NonPublic | BindingFlags.Instance))
+				List<PropertyInfo> properties = new List<PropertyInfo>();
+				foreach (PropertyInfo property in type.GetProperties(BindingFlags.Public | BindingFlags.NonPublic | BindingFlags.Instance))
 				{
-					MethodInfo? GetMethod = Property.GetGetMethod();
-					MethodInfo? SetMethod = Property.GetSetMethod();
-					if (GetMethod != null && SetMethod != null && Property.GetCustomAttribute<BinaryIgnoreAttribute>() == null)
+					MethodInfo? getMethod = property.GetGetMethod();
+					MethodInfo? setMethod = property.GetSetMethod();
+					if (getMethod != null && setMethod != null && property.GetCustomAttribute<BinaryIgnoreAttribute>() == null)
 					{
-						Properties.Add(Property);
+						properties.Add(property);
 					}
 				}
 
 				// Find the type constructor
-				ConstructorInfo? ConstructorInfo = Type.GetConstructor(BindingFlags.Instance | BindingFlags.Public | BindingFlags.NonPublic, null, Array.Empty<Type>(), null);
-				if (ConstructorInfo == null)
+				ConstructorInfo? constructorInfo = type.GetConstructor(BindingFlags.Instance | BindingFlags.Public | BindingFlags.NonPublic, null, Array.Empty<Type>(), null);
+				if (constructorInfo == null)
 				{
-					throw new NotImplementedException($"Type '{Type.Name}' does not have a parameterless constructor");
+					throw new NotImplementedException($"Type '{type.Name}' does not have a parameterless constructor");
 				}
 
 				// Create the methods
-				DynamicMethod WriterMethod = new DynamicMethod($"BinaryArchiveWriter_Dynamic_{Type.Name}", typeof(void), new[] { typeof(BinaryArchiveWriter), Type });
-				GenerateClassWriterMethod(WriterMethod, Properties);
+				DynamicMethod writerMethod = new DynamicMethod($"BinaryArchiveWriter_Dynamic_{type.Name}", typeof(void), new[] { typeof(BinaryArchiveWriter), type });
+				GenerateClassWriterMethod(writerMethod, properties);
 
-				DynamicMethod ReaderMethod = new DynamicMethod($"BinaryArchiveReader_Dynamic_{Type.Name}", Type, new[] { typeof(BinaryArchiveReader) });
-				GenerateClassReaderMethod(ReaderMethod, ConstructorInfo, Properties);
+				DynamicMethod readerMethod = new DynamicMethod($"BinaryArchiveReader_Dynamic_{type.Name}", type, new[] { typeof(BinaryArchiveReader) });
+				GenerateClassReaderMethod(readerMethod, constructorInfo, properties);
 
-				return new TypeSerializer(ReaderMethod, WriterMethod);
+				return new TypeSerializer(readerMethod, writerMethod);
 			}
 
-			throw new NotImplementedException($"Unable to create a serializer for {Type.Name}");
+			throw new NotImplementedException($"Unable to create a serializer for {type.Name}");
 		}
 
-		static void GenerateClassWriterMethod(DynamicMethod WriterMethod, List<PropertyInfo> Properties)
+		static void GenerateClassWriterMethod(DynamicMethod writerMethod, List<PropertyInfo> properties)
 		{
 			// Get the IL generator
-			ILGenerator Generator = WriterMethod.GetILGenerator(256 + (Properties.Count * 24));
+			ILGenerator generator = writerMethod.GetILGenerator(256 + (properties.Count * 24));
 
 			// Check if the argument is null, and write 'false' if it is, then early out.
-			Label NonNullInstance = Generator.DefineLabel();
-			Generator.Emit(OpCodes.Ldarg_1);                    // [0] = Arg1 (Instance)
-			Generator.Emit(OpCodes.Brtrue, NonNullInstance);    // [EMPTY] Instance != null -> NonNullInstance
+			Label nonNullInstance = generator.DefineLabel();
+			generator.Emit(OpCodes.Ldarg_1);                    // [0] = Arg1 (Instance)
+			generator.Emit(OpCodes.Brtrue, nonNullInstance);    // [EMPTY] Instance != null -> NonNullInstance
 
-			Generator.Emit(OpCodes.Ldarg_0);                    // [0] = Arg0 (Writer)
-			Generator.Emit(OpCodes.Ldc_I4_0);                   // [1] = false
-			Generator.EmitCall(OpCodes.Call, WriteBoolMethodInfo, null);  // [EMPTY] WriteBool(Write, false)
+			generator.Emit(OpCodes.Ldarg_0);                    // [0] = Arg0 (Writer)
+			generator.Emit(OpCodes.Ldc_I4_0);                   // [1] = false
+			generator.EmitCall(OpCodes.Call, s_writeBoolMethodInfo, null);  // [EMPTY] WriteBool(Write, false)
 
-			Generator.Emit(OpCodes.Ret);                        // [EMPTY] Return
+			generator.Emit(OpCodes.Ret);                        // [EMPTY] Return
 
 			// Otherwise write true.
-			Generator.MarkLabel(NonNullInstance);
+			generator.MarkLabel(nonNullInstance);
 
-			Generator.Emit(OpCodes.Ldarg_0);                    // [0] = Arg0 (Writer)
-			Generator.Emit(OpCodes.Ldc_I4_1);                   // [1] = true
-			Generator.EmitCall(OpCodes.Call, WriteBoolMethodInfo, null);  // [EMPTY] WriteBool(Writer, true)
+			generator.Emit(OpCodes.Ldarg_0);                    // [0] = Arg0 (Writer)
+			generator.Emit(OpCodes.Ldc_I4_1);                   // [1] = true
+			generator.EmitCall(OpCodes.Call, s_writeBoolMethodInfo, null);  // [EMPTY] WriteBool(Writer, true)
 
 			// Write all the properties
-			foreach (PropertyInfo Property in Properties)
+			foreach (PropertyInfo property in properties)
 			{
-				BinaryConverterAttribute? Converter = Property.GetCustomAttribute<BinaryConverterAttribute>();
+				BinaryConverterAttribute? converter = property.GetCustomAttribute<BinaryConverterAttribute>();
 
-				MethodInfo? WritePropertyMethod;
-				if (Property.PropertyType.IsClass && !Property.PropertyType.IsSealed && Converter == null)
+				MethodInfo? writePropertyMethod;
+				if (property.PropertyType.IsClass && !property.PropertyType.IsSealed && converter == null)
 				{
-					WritePropertyMethod = WritePolymorphicObjectMethodInfo;
+					writePropertyMethod = s_writePolymorphicObjectMethodInfo;
 				}
 				else
 				{
-					WritePropertyMethod = FindOrAddSerializerInfo(Property.PropertyType, Converter?.Type).WriteMethodInfo;
+					writePropertyMethod = FindOrAddSerializerInfo(property.PropertyType, converter?.Type).WriteMethodInfo;
 				}
 
-				Generator.Emit(OpCodes.Ldarg_0);                                            // [0] = Arg0 (BinaryArchiveWriter)
-				Generator.Emit(OpCodes.Ldarg_1);                                            // [1] = Arg1 (Instance)
-				Generator.EmitCall(OpCodes.Call, Property.GetGetMethod(true)!, null);       // [1] = GetMethod(Instance)
-				Generator.EmitCall(OpCodes.Call, WritePropertyMethod, null);                // [EMPTY] Write(BinaryArchiveWriter, GetMethod(Instance));
+				generator.Emit(OpCodes.Ldarg_0);                                            // [0] = Arg0 (BinaryArchiveWriter)
+				generator.Emit(OpCodes.Ldarg_1);                                            // [1] = Arg1 (Instance)
+				generator.EmitCall(OpCodes.Call, property.GetGetMethod(true)!, null);       // [1] = GetMethod(Instance)
+				generator.EmitCall(OpCodes.Call, writePropertyMethod, null);                // [EMPTY] Write(BinaryArchiveWriter, GetMethod(Instance));
 			}
 
 			// Return
-			Generator.Emit(OpCodes.Ret);
+			generator.Emit(OpCodes.Ret);
 		}
 
-		static void GenerateClassReaderMethod(DynamicMethod ReaderMethod, ConstructorInfo ConstructorInfo, List<PropertyInfo> Properties)
+		static void GenerateClassReaderMethod(DynamicMethod readerMethod, ConstructorInfo constructorInfo, List<PropertyInfo> properties)
 		{
 			// Get the IL generator
-			ILGenerator Generator = ReaderMethod.GetILGenerator(256 + (Properties.Count * 24));
+			ILGenerator generator = readerMethod.GetILGenerator(256 + (properties.Count * 24));
 
 			// Check if it was a null instance
-			Generator.Emit(OpCodes.Ldarg_0);                    // [0] = Arg1 (Reader)
-			Generator.EmitCall(OpCodes.Call, ReadBoolMethodInfo, null);   // [0] = Reader.ReadBool()
+			generator.Emit(OpCodes.Ldarg_0);                    // [0] = Arg1 (Reader)
+			generator.EmitCall(OpCodes.Call, s_readBoolMethodInfo, null);   // [0] = Reader.ReadBool()
 
-			Label NonNullInstance = Generator.DefineLabel();
-			Generator.Emit(OpCodes.Brtrue, NonNullInstance);    // Bool != null -> NonNullInstance
-			Generator.Emit(OpCodes.Ldnull);                     // [0] = null
-			Generator.Emit(OpCodes.Ret);                        // return
+			Label nonNullInstance = generator.DefineLabel();
+			generator.Emit(OpCodes.Brtrue, nonNullInstance);    // Bool != null -> NonNullInstance
+			generator.Emit(OpCodes.Ldnull);                     // [0] = null
+			generator.Emit(OpCodes.Ret);                        // return
 
 			// Construct an instance
-			Generator.MarkLabel(NonNullInstance);
-			Generator.Emit(OpCodes.Newobj, ConstructorInfo);    // [0] = new Type()
+			generator.MarkLabel(nonNullInstance);
+			generator.Emit(OpCodes.Newobj, constructorInfo);    // [0] = new Type()
 
 			// Read all the properties
-			foreach (PropertyInfo Property in Properties)
+			foreach (PropertyInfo property in properties)
 			{
-				BinaryConverterAttribute? Converter = Property.GetCustomAttribute<BinaryConverterAttribute>();
+				BinaryConverterAttribute? converter = property.GetCustomAttribute<BinaryConverterAttribute>();
 
-				MethodInfo? ReadPropertyMethod;
-				if (Property.PropertyType.IsClass && !Property.PropertyType.IsSealed && Converter == null)
+				MethodInfo? readPropertyMethod;
+				if (property.PropertyType.IsClass && !property.PropertyType.IsSealed && converter == null)
 				{
-					ReadPropertyMethod = ReadPolymorphicObjectMethodInfo;
+					readPropertyMethod = s_readPolymorphicObjectMethodInfo;
 				}
 				else
 				{
-					ReadPropertyMethod = FindOrAddSerializerInfo(Property.PropertyType, Converter?.Type).ReadMethodInfo;
+					readPropertyMethod = FindOrAddSerializerInfo(property.PropertyType, converter?.Type).ReadMethodInfo;
 				}
 
-				Generator.Emit(OpCodes.Dup);                                                // [1] = new Type()
+				generator.Emit(OpCodes.Dup);                                                // [1] = new Type()
 
-				Generator.Emit(OpCodes.Ldarg_0);                                            // [2] = Arg1 (Reader)
-				Generator.Emit(OpCodes.Call, ReadPropertyMethod);                           // [2] = Reader.ReadMethod() or ReadMethod(Reader)
+				generator.Emit(OpCodes.Ldarg_0);                                            // [2] = Arg1 (Reader)
+				generator.Emit(OpCodes.Call, readPropertyMethod);                           // [2] = Reader.ReadMethod() or ReadMethod(Reader)
 
-				Generator.Emit(OpCodes.Call, Property.GetSetMethod(true)!);                 // SetMethod(new Type(), ReadMethod(Reader))
+				generator.Emit(OpCodes.Call, property.GetSetMethod(true)!);                 // SetMethod(new Type(), ReadMethod(Reader))
 			}
 
 			// Return the instance
-			Generator.Emit(OpCodes.Ret);                        // [0] = new Type()
+			generator.Emit(OpCodes.Ret);                        // [0] = new Type()
 		}
 
-		static void GenerateConverterReaderForwardingMethod(DynamicMethod ReaderMethod, FieldInfo ConverterField, MethodInfo TargetMethod)
+		static void GenerateConverterReaderForwardingMethod(DynamicMethod readerMethod, FieldInfo converterField, MethodInfo targetMethod)
 		{
-			ILGenerator Generator = ReaderMethod.GetILGenerator(64);
-			Generator.Emit(OpCodes.Ldsfld, ConverterField);
-			Generator.Emit(OpCodes.Ldarg_0);
-			Generator.Emit(OpCodes.Call, TargetMethod);
-			Generator.Emit(OpCodes.Ret);
+			ILGenerator generator = readerMethod.GetILGenerator(64);
+			generator.Emit(OpCodes.Ldsfld, converterField);
+			generator.Emit(OpCodes.Ldarg_0);
+			generator.Emit(OpCodes.Call, targetMethod);
+			generator.Emit(OpCodes.Ret);
 		}
 
-		static void GenerateConverterWriterForwardingMethod(DynamicMethod WriterMethod, FieldInfo ConverterField, MethodInfo TargetMethod)
+		static void GenerateConverterWriterForwardingMethod(DynamicMethod writerMethod, FieldInfo converterField, MethodInfo targetMethod)
 		{
-			ILGenerator Generator = WriterMethod.GetILGenerator(64);
-			Generator.Emit(OpCodes.Ldsfld, ConverterField);
-			Generator.Emit(OpCodes.Ldarg_0);
-			Generator.Emit(OpCodes.Ldarg_1);
-			Generator.Emit(OpCodes.Call, TargetMethod);
-			Generator.Emit(OpCodes.Ret);
+			ILGenerator generator = writerMethod.GetILGenerator(64);
+			generator.Emit(OpCodes.Ldsfld, converterField);
+			generator.Emit(OpCodes.Ldarg_0);
+			generator.Emit(OpCodes.Ldarg_1);
+			generator.Emit(OpCodes.Call, targetMethod);
+			generator.Emit(OpCodes.Ret);
 		}
 	}
 }
diff --git a/Engine/Source/Programs/Shared/EpicGames.Core/BinaryArchiveAttributes.cs b/Engine/Source/Programs/Shared/EpicGames.Core/BinaryArchiveAttributes.cs
--- a/Engine/Source/Programs/Shared/EpicGames.Core/BinaryArchiveAttributes.cs
+++ b/Engine/Source/Programs/Shared/EpicGames.Core/BinaryArchiveAttributes.cs
@@ -40,10 +40,10 @@ namespace EpicGames.Core
 		/// <summary>
 		/// Constructor
 		/// </summary>
-		/// <param name="Type">The serializer type</param>
-		public BinaryConverterAttribute(Type Type)
+		/// <param name="type">The serializer type</param>
+		public BinaryConverterAttribute(Type type)
 		{
-			this.Type = Type;
+			Type = type;
 		}
 	}
 }
diff --git a/Engine/Source/Programs/Shared/EpicGames.Core/BinaryArchiveReader.cs b/Engine/Source/Programs/Shared/EpicGames.Core/BinaryArchiveReader.cs
--- a/Engine/Source/Programs/Shared/EpicGames.Core/BinaryArchiveReader.cs
+++ b/Engine/Source/Programs/Shared/EpicGames.Core/BinaryArchiveReader.cs
@@ -17,49 +17,49 @@ namespace EpicGames.Core
 		/// <summary>
 		/// The input stream.
 		/// </summary>
-		Stream? Stream;
+		Stream? _stream;
 
 		/// <summary>
 		/// The input buffer
 		/// </summary>
-		byte[]? Buffer;
+		byte[]? _buffer;
 
 		/// <summary>
 		/// Current position within the buffer
 		/// </summary>
-		int BufferPos;
+		int _bufferPos;
 
 		/// <summary>
 		/// List of previously serialized objects
 		/// </summary>
-		readonly List<object?> Objects = new List<object?>();
+		readonly List<object?> _objects = new List<object?>();
 
 		/// <summary>
 		/// Constructor
 		/// </summary>
-		/// <param name="Buffer">The buffer to read from</param>
-		public BinaryArchiveReader(byte[] Buffer)
+		/// <param name="buffer">The buffer to read from</param>
+		public BinaryArchiveReader(byte[] buffer)
 		{
-			this.Buffer = Buffer;
+			_buffer = buffer;
 		}
 
 		/// <summary>
 		/// Constructor
 		/// </summary>
-		/// <param name="FileName">File to read from</param>
-		public BinaryArchiveReader(FileReference FileName)
+		/// <param name="fileName">File to read from</param>
+		public BinaryArchiveReader(FileReference fileName)
 		{
-			this.Buffer = FileReference.ReadAllBytes(FileName);
+			_buffer = FileReference.ReadAllBytes(fileName);
 		}
 
 		/// <summary>
 		/// Constructor
 		/// </summary>
-		/// <param name="Stream">Stream to read from</param>
-		public BinaryArchiveReader(Stream Stream)
+		/// <param name="stream">Stream to read from</param>
+		public BinaryArchiveReader(Stream stream)
 		{
-			Buffer = new byte[Stream.Length];
-			Stream.Read(Buffer, 0, Buffer.Length);
+			_buffer = new byte[stream.Length];
+			stream.Read(_buffer, 0, _buffer.Length);
 		}
 
 		/// <summary>
@@ -67,13 +67,13 @@ namespace EpicGames.Core
 		/// </summary>
 		public void Dispose()
 		{
-			if (Stream != null)
+			if (_stream != null)
 			{
-				Stream.Dispose();
-				Stream = null;
+				_stream.Dispose();
+				_stream = null;
 			}
 
-			Buffer = null;
+			_buffer = null;
 		}
 
 		/// <summary>
@@ -91,9 +91,9 @@ namespace EpicGames.Core
 		/// <returns>The value that was read</returns>
 		public byte ReadByte()
 		{
-			byte Value = Buffer![BufferPos];
-			BufferPos++;
-			return Value;
+			byte value = _buffer![_bufferPos];
+			_bufferPos++;
+			return value;
 		}
 
 		/// <summary>
@@ -120,9 +120,9 @@ namespace EpicGames.Core
 		/// <returns>The value that was read</returns>
 		public ushort ReadUnsignedShort()
 		{
-			ushort Value = (ushort)(Buffer![BufferPos + 0] | (Buffer[BufferPos + 1] << 8));
-			BufferPos += 2;
-			return Value;
+			ushort value = (ushort)(_buffer![_bufferPos + 0] | (_buffer[_bufferPos + 1] << 8));
+			_bufferPos += 2;
+			return value;
 		}
 
 		/// <summary>
@@ -140,9 +140,9 @@ namespace EpicGames.Core
 		/// <returns>The value that was read</returns>
 		public uint ReadUnsignedInt()
 		{
-			uint Value = (uint)(Buffer![BufferPos + 0] | (Buffer[BufferPos + 1] << 8) | (Buffer[BufferPos + 2] << 16) | (Buffer[BufferPos + 3] << 24));
-			BufferPos += 4;
-			return Value;
+			uint value = (uint)(_buffer![_bufferPos + 0] | (_buffer[_bufferPos + 1] << 8) | (_buffer[_bufferPos + 2] << 16) | (_buffer[_bufferPos + 3] << 24));
+			_bufferPos += 4;
+			return value;
 		}
 
 		/// <summary>
@@ -160,9 +160,9 @@ namespace EpicGames.Core
 		/// <returns>The value that was read</returns>
 		public ulong ReadUnsignedLong()
 		{
-			ulong Value = (ulong)ReadUnsignedInt();
-			Value |= (ulong)ReadUnsignedInt() << 32;
-			return Value;
+			ulong value = (ulong)ReadUnsignedInt();
+			value |= (ulong)ReadUnsignedInt() << 32;
+			return value;
 		}
 
 		/// <summary>
@@ -180,14 +180,14 @@ namespace EpicGames.Core
 		/// <returns>The value that was read</returns>
 		public string? ReadString()
 		{
-			byte[]? Bytes = ReadByteArray();
-			if (Bytes == null)
+			byte[]? bytes = ReadByteArray();
+			if (bytes == null)
 			{
 				return null;
 			}
 			else
 			{
-				return Encoding.UTF8.GetString(Bytes);
+				return Encoding.UTF8.GetString(bytes);
 			}
 		}
 
@@ -221,97 +221,96 @@ namespace EpicGames.Core
 		/// <summary>
 		/// Reads an array of primitive types from the stream
 		/// </summary>
-		/// <param name="ElementSize">Size of a single element</param>
+		/// <param name="elementSize">Size of a single element</param>
 		/// <returns>The data that was read</returns>
-		private T[]? ReadPrimitiveArray<T>(int ElementSize) where T : struct
+		private T[]? ReadPrimitiveArray<T>(int elementSize) where T : struct
 		{
-			int Length = ReadInt();
-			if (Length < 0)
+			int length = ReadInt();
+			if (length < 0)
 			{
 				return null;
 			}
 			else
 			{
-				T[] Result = new T[Length];
-				ReadBulkData(Result, Length * ElementSize);
-				return Result;
+				T[] result = new T[length];
+				ReadBulkData(result, length * elementSize);
+				return result;
 			}
 		}
 
 		/// <summary>
 		/// Reads a byte array from the stream
 		/// </summary>
-		/// <param name="Length">Length of the array to read</param>
+		/// <param name="length">Length of the array to read</param>
 		/// <returns>The data that was read</returns>
-		public byte[] ReadFixedSizeByteArray(int Length)
+		public byte[] ReadFixedSizeByteArray(int length)
 		{
-			return ReadFixedSizePrimitiveArray<byte>(sizeof(byte), Length);
+			return ReadFixedSizePrimitiveArray<byte>(sizeof(byte), length);
 		}
 
 		/// <summary>
 		/// Reads a short array from the stream
 		/// </summary>
-		/// <param name="Length">Length of the array to read</param>
+		/// <param name="length">Length of the array to read</param>
 		/// <returns>The data that was read</returns>
-		public short[] ReadFixedSizeShortArray(int Length)
+		public short[] ReadFixedSizeShortArray(int length)
 		{
-			return ReadFixedSizePrimitiveArray<short>(sizeof(short), Length);
+			return ReadFixedSizePrimitiveArray<short>(sizeof(short), length);
 		}
 
 		/// <summary>
 		/// Reads an int array from the stream
 		/// </summary>
-		/// <param name="Length">Length of the array to read</param>
+		/// <param name="length">Length of the array to read</param>
 		/// <returns>The data that was read</returns>
-		public int[] ReadFixedSizeIntArray(int Length)
+		public int[] ReadFixedSizeIntArray(int length)
 		{
-			return ReadFixedSizePrimitiveArray<int>(sizeof(int), Length);
+			return ReadFixedSizePrimitiveArray<int>(sizeof(int), length);
 		}
 
 		/// <summary>
 		/// Reads an array of primitive types from the stream
 		/// </summary>
-		/// <param name="ElementSize">Size of a single element</param>
-		/// <param name="ElementCount">Number of elements to read</param>
+		/// <param name="elementSize">Size of a single element</param>
+		/// <param name="elementCount">Number of elements to read</param>
 		/// <returns>The data that was read</returns>
-		private T[] ReadFixedSizePrimitiveArray<T>(int ElementSize, int ElementCount) where T : struct
+		private T[] ReadFixedSizePrimitiveArray<T>(int elementSize, int elementCount) where T : struct
 		{
-			T[] Result = new T[ElementCount];
-			ReadBulkData(Result, ElementSize * ElementCount);
-			return Result;
+			T[] result = new T[elementCount];
+			ReadBulkData(result, elementSize * elementCount);
+			return result;
 		}
 
 		/// <summary>
 		/// Reads bulk data from the stream into the given buffer
 		/// </summary>
-		/// <param name="Data">Array which receives the data that was read</param>
-		/// <param name="Size">Size of data to read</param>
-		private void ReadBulkData(Array Data, int Size)
+		/// <param name="data">Array which receives the data that was read</param>
+		/// <param name="size">Size of data to read</param>
+		private void ReadBulkData(Array data, int size)
 		{
-			System.Buffer.BlockCopy(Buffer!, BufferPos, Data, 0, Size);
-			BufferPos += Size;
+			System.Buffer.BlockCopy(_buffer!, _bufferPos, data, 0, size);
+			_bufferPos += size;
 		}
 
-
 		/// <summary>
 		/// Reads an array of items
 		/// </summary>
 		/// <returns>New array</returns>
-		public T[]? ReadArray<T>(Func<T> ReadElement)
+		public T[]? ReadArray<T>(Func<T> readElement)
 		{
-			int Count = ReadInt();
-			if (Count < 0)
+			int count = ReadInt();
+			if (count < 0)
 			{
 				return null;
 			}
 			else
 			{
-				T[] Result = new T[Count];
-				for (int Idx = 0; Idx < Count; Idx++)
+				T[] result = new T[count];
+				for (int idx = 0; idx < count; idx++)
 				{
-					Result[Idx] = ReadElement();
+					result[idx] = readElement();
 				}
-				return Result;
+				return result;
 			}
 		}
 
@@ -319,23 +318,23 @@ namespace EpicGames.Core
 		/// Reads a list of items
 		/// </summary>
 		/// <typeparam name="T">The element type for the list</typeparam>
-		/// <param name="ReadElement">Delegate used to read a single element</param>
+		/// <param name="readElement">Delegate used to read a single element</param>
 		/// <returns>List of items</returns>
-		public List<T>? ReadList<T>(Func<T> ReadElement)
+		public List<T>? ReadList<T>(Func<T> readElement)
 		{
-			int Count = ReadInt();
-			if (Count < 0)
+			int count = ReadInt();
+			if (count < 0)
 			{
 				return null;
 			}
 			else
 			{
-				List<T> Result = new List<T>(Count);
-				for (int Idx = 0; Idx < Count; Idx++)
+				List<T> result = new List<T>(count);
+				for (int idx = 0; idx < count; idx++)
 				{
-					Result.Add(ReadElement());
+					result.Add(readElement());
 				}
-				return Result;
+				return result;
 			}
 		}
 
@@ -343,23 +342,23 @@ namespace EpicGames.Core
 		/// Reads a hashset of items
 		/// </summary>
 		/// <typeparam name="T">The element type for the set</typeparam>
-		/// <param name="ReadElement">Delegate used to read a single element</param>
+		/// <param name="readElement">Delegate used to read a single element</param>
 		/// <returns>Set of items</returns>
-		public HashSet<T>? ReadHashSet<T>(Func<T> ReadElement)
+		public HashSet<T>? ReadHashSet<T>(Func<T> readElement)
 		{
-			int Count = ReadInt();
-			if (Count < 0)
+			int count = ReadInt();
+			if (count < 0)
 			{
 				return null;
 			}
 			else
 			{
-				HashSet<T> Result = new HashSet<T>();
-				for (int Idx = 0; Idx < Count; Idx++)
+				HashSet<T> result = new HashSet<T>();
+				for (int idx = 0; idx < count; idx++)
 				{
-					Result.Add(ReadElement());
+					result.Add(readElement());
 				}
-				return Result;
+				return result;
 			}
 		}
 
@@ -367,77 +366,77 @@ namespace EpicGames.Core
 		/// Reads a hashset of items
 		/// </summary>
 		/// <typeparam name="T">The element type for the set</typeparam>
-		/// <param name="ReadElement">Delegate used to read a single element</param>
-		/// <param name="Comparer">Comparison function for the set</param>
+		/// <param name="readElement">Delegate used to read a single element</param>
+		/// <param name="comparer">Comparison function for the set</param>
 		/// <returns>Set of items</returns>
-		public HashSet<T>? ReadHashSet<T>(Func<T> ReadElement, IEqualityComparer<T> Comparer)
+		public HashSet<T>? ReadHashSet<T>(Func<T> readElement, IEqualityComparer<T> comparer)
 		{
-			int Count = ReadInt();
-			if (Count < 0)
+			int count = ReadInt();
+			if (count < 0)
 			{
 				return null;
 			}
 			else
 			{
-				HashSet<T> Result = new HashSet<T>(Comparer);
-				for (int Idx = 0; Idx < Count; Idx++)
+				HashSet<T> result = new HashSet<T>(comparer);
+				for (int idx = 0; idx < count; idx++)
 				{
-					Result.Add(ReadElement());
+					result.Add(readElement());
 				}
-				return Result;
+				return result;
 			}
 		}
 
 		/// <summary>
 		/// Reads a dictionary of items
 		/// </summary>
-		/// <typeparam name="K">Type of the dictionary key</typeparam>
-		/// <typeparam name="V">Type of the dictionary value</typeparam>
-		/// <param name="ReadKey">Delegate used to read a single key</param>
-		/// <param name="ReadValue">Delegate used to read a single value</param>
+		/// <typeparam name="TK">Type of the dictionary key</typeparam>
+		/// <typeparam name="TV">Type of the dictionary value</typeparam>
+		/// <param name="readKey">Delegate used to read a single key</param>
+		/// <param name="readValue">Delegate used to read a single value</param>
 		/// <returns>New dictionary instance</returns>
-		public Dictionary<K, V>? ReadDictionary<K, V>(Func<K> ReadKey, Func<V> ReadValue) where K : notnull
+		public Dictionary<TK, TV>? ReadDictionary<TK, TV>(Func<TK> readKey, Func<TV> readValue) where TK : notnull
 		{
-			int Count = ReadInt();
-			if (Count < 0)
+			int count = ReadInt();
+			if (count < 0)
 			{
 				return null;
 			}
 			else
 			{
-				Dictionary<K, V> Result = new Dictionary<K, V>(Count);
-				for (int Idx = 0; Idx < Count; Idx++)
+				Dictionary<TK, TV> result = new Dictionary<TK, TV>(count);
+				for (int idx = 0; idx < count; idx++)
 				{
-					Result.Add(ReadKey(), ReadValue());
+					result.Add(readKey(), readValue());
 				}
-				return Result;
+				return result;
 			}
 		}
 
 		/// <summary>
 		/// Reads a dictionary of items
 		/// </summary>
-		/// <typeparam name="K">Type of the dictionary key</typeparam>
-		/// <typeparam name="V">Type of the dictionary value</typeparam>
-		/// <param name="ReadKey">Delegate used to read a single key</param>
-		/// <param name="ReadValue">Delegate used to read a single value</param>
-		/// <param name="Comparer">Comparison function for keys in the dictionary</param>
+		/// <typeparam name="TK">Type of the dictionary key</typeparam>
+		/// <typeparam name="TV">Type of the dictionary value</typeparam>
+		/// <param name="readKey">Delegate used to read a single key</param>
+		/// <param name="readValue">Delegate used to read a single value</param>
+		/// <param name="comparer">Comparison function for keys in the dictionary</param>
 		/// <returns>New dictionary instance</returns>
-		public Dictionary<K, V>? ReadDictionary<K, V>(Func<K> ReadKey, Func<V> ReadValue, IEqualityComparer<K> Comparer) where K : notnull
+		public Dictionary<TK, TV>? ReadDictionary<TK, TV>(Func<TK> readKey, Func<TV> readValue, IEqualityComparer<TK> comparer) where TK : notnull
 		{
-			int Count = ReadInt();
-			if (Count < 0)
+			int count = ReadInt();
+			if (count < 0)
 			{
 				return null;
 			}
 			else
 			{
-				Dictionary<K, V> Result = new Dictionary<K, V>(Count, Comparer);
-				for (int Idx = 0; Idx < Count; Idx++)
+				Dictionary<TK, TV> result = new Dictionary<TK, TV>(count, comparer);
+				for (int idx = 0; idx < count; idx++)
 				{
-					Result.Add(ReadKey(), ReadValue());
+					result.Add(readKey(), readValue());
 				}
-				return Result;
+				return result;
 			}
 		}
 
@@ -445,12 +444,12 @@ namespace EpicGames.Core
 		/// Reads a nullable object from the archive
 		/// </summary>
 		/// <typeparam name="T">The nullable type</typeparam>
-		/// <param name="ReadValue">Delegate used to read a value</param>
-		public Nullable<T> ReadNullable<T>(Func<T> ReadValue) where T : struct
+		/// <param name="readValue">Delegate used to read a value</param>
+		public Nullable<T> ReadNullable<T>(Func<T> readValue) where T : struct
 		{
 			if (ReadBool())
 			{
-				return new Nullable<T>(ReadValue());
+				return new Nullable<T>(readValue());
 			}
 			else
 			{
@@ -462,13 +461,13 @@ namespace EpicGames.Core
 		/// Reads an object, which may be null, from the archive. Does not handle de-duplicating object references. 
 		/// </summary>
 		/// <typeparam name="T">Type of the object to read</typeparam>
-		/// <param name="Read">Delegate used to read the object</param>
+		/// <param name="read">Delegate used to read the object</param>
 		/// <returns>The object instance</returns>
-		public T? ReadOptionalObject<T>(Func<T> Read) where T : class
+		public T? ReadOptionalObject<T>(Func<T> read) where T : class
 		{
 			if (ReadBool())
 			{
-				return Read();
+				return read();
 			}
 			else
 			{
@@ -482,25 +481,25 @@ namespace EpicGames.Core
 		/// serialize a reference to itself.
 		/// </summary>
 		/// <typeparam name="T">Type of the object to read.</typeparam>
-		/// <param name="CreateObject">Delegate used to create an object instance</param>
-		/// <param name="ReadObject">Delegate used to read an object instance</param>
+		/// <param name="createObject">Delegate used to create an object instance</param>
+		/// <param name="readObject">Delegate used to read an object instance</param>
 		/// <returns>Object instance</returns>
-		public T? ReadObjectReference<T>(Func<T> CreateObject, Action<T> ReadObject) where T : class
+		public T? ReadObjectReference<T>(Func<T> createObject, Action<T> readObject) where T : class
 		{
-			int Index = ReadInt();
-			if (Index < 0)
+			int index = ReadInt();
+			if (index < 0)
 			{
 				return null;
 			}
 			else
 			{
-				if (Index == Objects.Count)
+				if (index == _objects.Count)
 				{
-					T Object = CreateObject();
-					Objects.Add(Object);
-					ReadObject(Object);
+					T obj = createObject();
+					_objects.Add(obj);
+					readObject(obj);
 				}
-				return (T?)Objects[Index];
+				return (T?)_objects[index];
 			}
 		}
 
@@ -509,28 +508,28 @@ namespace EpicGames.Core
 		/// Since the reader only receives the object reference when the CreateObject delegate returns, it is not possible for the object to serialize a reference to itself.
 		/// </summary>
 		/// <typeparam name="T">Type of the object to read.</typeparam>
-		/// <param name="ReadObject">Delegate used to create an object instance. The object may not reference itself recursively.</param>
+		/// <param name="readObject">Delegate used to create an object instance. The object may not reference itself recursively.</param>
 		/// <returns>Object instance</returns>
-		public object? ReadUntypedObjectReference(Func<object?> ReadObject)
+		public object? ReadUntypedObjectReference(Func<object?> readObject)
 		{
-			int Index = ReadInt();
-			if (Index < 0)
+			int index = ReadInt();
+			if (index < 0)
 			{
 				return null;
 			}
 			else
 			{
 				// Temporarily add the reader to the object list, so we can detect invalid recursive references. 
-				if (Index == Objects.Count)
+				if (index == _objects.Count)
 				{
-					Objects.Add(null);
-					Objects[Index] = ReadObject();
+					_objects.Add(null);
+					_objects[index] = readObject();
 				}
-				if (Objects[Index] == null)
+				if (_objects[index] == null)
 				{
 					throw new InvalidOperationException("Attempt to serialize reference to object recursively.");
 				}
-				return Objects[Index];
+				return _objects[index];
 			}
 		}
 
@@ -539,39 +538,39 @@ namespace EpicGames.Core
 		/// Since the reader only receives the object reference when the CreateObject delegate returns, it is not possible for the object to serialize a reference to itself.
 		/// </summary>
 		/// <typeparam name="T">Type of the object to read.</typeparam>
-		/// <param name="ReadObject">Delegate used to create an object instance. The object may not reference itself recursively.</param>
+		/// <param name="readObject">Delegate used to create an object instance. The object may not reference itself recursively.</param>
 		/// <returns>Object instance</returns>
-		public T? ReadObjectReference<T>(Func<T> ReadObject) where T : class => (T?)ReadUntypedObjectReference(ReadObject);
+		public T? ReadObjectReference<T>(Func<T> readObject) where T : class => (T?)ReadUntypedObjectReference(readObject);
 
 		/// <summary>
 		/// Helper method for validating that deserialized objects are not null
 		/// </summary>
 		/// <typeparam name="T">Type of the deserialized object</typeparam>
-		/// <param name="Param">The object instance</param>
+		/// <param name="param">The object instance</param>
 		/// <returns>The object instance</returns>
 		[return: NotNull]
-		public static T NotNull<T>(T? Param) where T : class
+		public static T NotNull<T>(T? param) where T : class
 		{
-			if (Param == null)
+			if (param == null)
 			{
 				throw new InvalidDataException("Object stored in archive is not allowed to be null.");
 			}
-			return Param;
+			return param;
 		}
 
 		/// <summary>
 		/// Helper method for validating that deserialized objects are not null
 		/// </summary>
 		/// <typeparam name="T">Type of the deserialized object</typeparam>
-		/// <param name="Param">The object instance</param>
+		/// <param name="param">The object instance</param>
 		/// <returns>The object instance</returns>
-		public static T NotNullStruct<T>(T? Param) where T : struct
+		public static T NotNullStruct<T>(T? param) where T : struct
 		{
-			if (Param == null)
+			if (param == null)
 			{
 				throw new InvalidDataException("Object stored in archive is not allowed to be null.");
 			}
-			return Param.Value;
+			return param.Value;
 		}
 	}
 }
diff --git a/Engine/Source/Programs/Shared/EpicGames.Core/BinaryArchiveWriter.cs b/Engine/Source/Programs/Shared/EpicGames.Core/BinaryArchiveWriter.cs
--- a/Engine/Source/Programs/Shared/EpicGames.Core/BinaryArchiveWriter.cs
+++ b/Engine/Source/Programs/Shared/EpicGames.Core/BinaryArchiveWriter.cs
@@ -18,58 +18,58 @@ namespace EpicGames.Core
 		/// </summary>
 		class ReferenceComparer : IEqualityComparer<object>
 		{
-			bool IEqualityComparer<object>.Equals(object? A, object? B)
+			bool IEqualityComparer<object>.Equals(object? a, object? b)
 			{
-				return A == B;
+				return a == b;
 			}
 
-			int IEqualityComparer<object>.GetHashCode(object X)
+			int IEqualityComparer<object>.GetHashCode(object x)
 			{
-				return RuntimeHelpers.GetHashCode(X);
+				return RuntimeHelpers.GetHashCode(x);
 			}
 		}
 
 		/// <summary>
 		/// Instance of the ReferenceComparer class which can be shared by all archive writers
 		/// </summary>
-		static readonly ReferenceComparer ReferenceComparerInstance = new ReferenceComparer();
+		static readonly ReferenceComparer s_referenceComparerInstance = new ReferenceComparer();
 
 		/// <summary>
 		/// The output stream being written to
 		/// </summary>
-		Stream? Stream;
+		Stream? _stream;
 
 		/// <summary>
 		/// Buffer for data to be written to the stream
 		/// </summary>
-		byte[] Buffer;
+		byte[] _buffer;
 
 		/// <summary>
 		/// Current position within the output buffer
 		/// </summary>
-		int BufferPos;
+		int _bufferPos;
 
 		/// <summary>
 		/// Map of object instance to unique id
 		/// </summary>
-		readonly Dictionary<object, int> ObjectToUniqueId = new Dictionary<object, int>(ReferenceComparerInstance);
+		readonly Dictionary<object, int> _objectToUniqueId = new Dictionary<object, int>(s_referenceComparerInstance);
 
 		/// <summary>
 		/// Constructor
 		/// </summary>
-		/// <param name="Stream">The output stream</param>
-		public BinaryArchiveWriter(Stream Stream)
+		/// <param name="stream">The output stream</param>
+		public BinaryArchiveWriter(Stream stream)
 		{
-			this.Stream = Stream;
-			this.Buffer = new byte[4096];
+			_stream = stream;
+			_buffer = new byte[4096];
 		}
 
 		/// <summary>
 		/// Constructor
 		/// </summary>
-		/// <param name="FileName">File to write to</param>
-		public BinaryArchiveWriter(FileReference FileName)
-			: this(File.Open(FileName.FullName, FileMode.Create, FileAccess.Write, FileShare.Read))
+		/// <param name="fileName">File to write to</param>
+		public BinaryArchiveWriter(FileReference fileName)
+			: this(File.Open(fileName.FullName, FileMode.Create, FileAccess.Write, FileShare.Read))
 		{
 		}
 
@@ -80,10 +80,10 @@ namespace EpicGames.Core
 		{
 			Flush();
 
-			if (Stream != null)
+			if (_stream != null)
 			{
-				Stream.Dispose();
-				Stream = null;
+				_stream.Dispose();
+				_stream = null;
 			}
 		}
 
@@ -92,25 +92,25 @@ namespace EpicGames.Core
 		/// </summary>
 		public void Flush()
 		{
-			if (BufferPos > 0)
+			if (_bufferPos > 0)
 			{
-				Stream!.Write(Buffer, 0, BufferPos);
-				BufferPos = 0;
+				_stream!.Write(_buffer, 0, _bufferPos);
+				_bufferPos = 0;
 			}
 		}
 
 		/// <summary>
 		/// Ensures there is a minimum amount of space in the output buffer
 		/// </summary>
-		/// <param name="NumBytes">Minimum amount of space required in the output buffer</param>
-		private void EnsureSpace(int NumBytes)
+		/// <param name="numBytes">Minimum amount of space required in the output buffer</param>
+		private void EnsureSpace(int numBytes)
 		{
-			if (BufferPos + NumBytes > Buffer.Length)
+			if (_bufferPos + numBytes > _buffer.Length)
 			{
 				Flush();
-				if (NumBytes > Buffer.Length)
+				if (numBytes > _buffer.Length)
 				{
-					Buffer = new byte[NumBytes];
+					_buffer = new byte[numBytes];
 				}
 			}
 		}
@@ -118,236 +118,236 @@ namespace EpicGames.Core
 		/// <summary>
 		/// Writes a bool to the output
 		/// </summary>
-		/// <param name="Value">Value to write</param>
-		public void WriteBool(bool Value)
+		/// <param name="value">Value to write</param>
+		public void WriteBool(bool value)
 		{
-			WriteByte(Value ? (byte)1 : (byte)0);
+			WriteByte(value ? (byte)1 : (byte)0);
 		}
 
 		/// <summary>
 		/// Writes a single byte to the output
 		/// </summary>
-		/// <param name="Value">Value to write</param>
-		public void WriteByte(byte Value)
+		/// <param name="value">Value to write</param>
+		public void WriteByte(byte value)
 		{
 			EnsureSpace(1);
 
-			Buffer[BufferPos] = Value;
+			_buffer[_bufferPos] = value;
 
-			BufferPos++;
+			_bufferPos++;
 		}
 
 		/// <summary>
 		/// Writes a single signed byte to the output
 		/// </summary>
-		/// <param name="Value">Value to write</param>
-		public void WriteSignedByte(sbyte Value)
+		/// <param name="value">Value to write</param>
+		public void WriteSignedByte(sbyte value)
 		{
-			WriteByte((byte)Value);
+			WriteByte((byte)value);
 		}
 
 		/// <summary>
 		/// Writes a single short to the output
 		/// </summary>
-		/// <param name="Value">Value to write</param>
-		public void WriteShort(short Value)
+		/// <param name="value">Value to write</param>
+		public void WriteShort(short value)
 		{
-			WriteUnsignedShort((ushort)Value);
+			WriteUnsignedShort((ushort)value);
 		}
 
 		/// <summary>
 		/// Writes a single unsigned short to the output
 		/// </summary>
-		/// <param name="Value">Value to write</param>
-		public void WriteUnsignedShort(ushort Value)
+		/// <param name="value">Value to write</param>
+		public void WriteUnsignedShort(ushort value)
 		{
 			EnsureSpace(2);
 
-			Buffer[BufferPos + 0] = (byte)Value;
-			Buffer[BufferPos + 1] = (byte)(Value >> 8);
+			_buffer[_bufferPos + 0] = (byte)value;
+			_buffer[_bufferPos + 1] = (byte)(value >> 8);
 
-			BufferPos += 2;
+			_bufferPos += 2;
 		}
 
 		/// <summary>
 		/// Writes a single int to the output
 		/// </summary>
-		/// <param name="Value">Value to write</param>
-		public void WriteInt(int Value)
+		/// <param name="value">Value to write</param>
+		public void WriteInt(int value)
 		{
-			WriteUnsignedInt((uint)Value);
+			WriteUnsignedInt((uint)value);
 		}
 
 		/// <summary>
 		/// Writes a single unsigned int to the output
 		/// </summary>
-		/// <param name="Value">Value to write</param>
-		public void WriteUnsignedInt(uint Value)
+		/// <param name="value">Value to write</param>
+		public void WriteUnsignedInt(uint value)
 		{
 			EnsureSpace(4);
 
-			Buffer[BufferPos + 0] = (byte)Value;
-			Buffer[BufferPos + 1] = (byte)(Value >> 8);
-			Buffer[BufferPos + 2] = (byte)(Value >> 16);
-			Buffer[BufferPos + 3] = (byte)(Value >> 24);
+			_buffer[_bufferPos + 0] = (byte)value;
+			_buffer[_bufferPos + 1] = (byte)(value >> 8);
+			_buffer[_bufferPos + 2] = (byte)(value >> 16);
+			_buffer[_bufferPos + 3] = (byte)(value >> 24);
 
-			BufferPos += 4;
+			_bufferPos += 4;
 		}
 
 		/// <summary>
 		/// Writes a single long to the output
 		/// </summary>
-		/// <param name="Value">Value to write</param>
-		public void WriteLong(long Value)
+		/// <param name="value">Value to write</param>
+		public void WriteLong(long value)
 		{
-			WriteUnsignedLong((ulong)Value);
+			WriteUnsignedLong((ulong)value);
 		}
 
 		/// <summary>
 		/// Writes a single unsigned long to the output
 		/// </summary>
-		/// <param name="Value">Value to write</param>
-		public void WriteUnsignedLong(ulong Value)
+		/// <param name="value">Value to write</param>
+		public void WriteUnsignedLong(ulong value)
 		{
 			EnsureSpace(8);
 
-			Buffer[BufferPos + 0] = (byte)Value;
-			Buffer[BufferPos + 1] = (byte)(Value >> 8);
-			Buffer[BufferPos + 2] = (byte)(Value >> 16);
-			Buffer[BufferPos + 3] = (byte)(Value >> 24);
-			Buffer[BufferPos + 4] = (byte)(Value >> 32);
-			Buffer[BufferPos + 5] = (byte)(Value >> 40);
-			Buffer[BufferPos + 6] = (byte)(Value >> 48);
-			Buffer[BufferPos + 7] = (byte)(Value >> 56);
+			_buffer[_bufferPos + 0] = (byte)value;
+			_buffer[_bufferPos + 1] = (byte)(value >> 8);
+			_buffer[_bufferPos + 2] = (byte)(value >> 16);
+			_buffer[_bufferPos + 3] = (byte)(value >> 24);
+			_buffer[_bufferPos + 4] = (byte)(value >> 32);
+			_buffer[_bufferPos + 5] = (byte)(value >> 40);
+			_buffer[_bufferPos + 6] = (byte)(value >> 48);
+			_buffer[_bufferPos + 7] = (byte)(value >> 56);
 
-			BufferPos += 8;
+			_bufferPos += 8;
 		}
 
 		/// <summary>
 		/// Writes a double (64 bit floating point value) to the stream
 		/// </summary>
-		/// <param name="Value">Value to write</param>
-		public void WriteDouble(double Value)
+		/// <param name="value">Value to write</param>
+		public void WriteDouble(double value)
 		{
-			WriteLong(BitConverter.DoubleToInt64Bits(Value));
+			WriteLong(BitConverter.DoubleToInt64Bits(value));
 		}
 
 		/// <summary>
 		/// Writes a string to the output
 		/// </summary>
-		/// <param name="Value">Value to write</param>
-		public void WriteString(string? Value)
+		/// <param name="value">Value to write</param>
+		public void WriteString(string? value)
 		{
-			byte[]? Bytes;
-			if (Value == null)
+			byte[]? bytes;
+			if (value == null)
 			{
-				Bytes = null;
+				bytes = null;
 			}
 			else
 			{
-				Bytes = Encoding.UTF8.GetBytes(Value);
+				bytes = Encoding.UTF8.GetBytes(value);
 			}
-			WriteByteArray(Bytes);
+			WriteByteArray(bytes);
 		}
 
 		/// <summary>
 		/// Writes an array of bytes to the output
 		/// </summary>
-		/// <param name="Data">Data to write. May be null.</param>
-		public void WriteByteArray(byte[]? Data)
+		/// <param name="data">Data to write. May be null.</param>
+		public void WriteByteArray(byte[]? data)
 		{
-			WritePrimitiveArray(Data, sizeof(byte));
+			WritePrimitiveArray(data, sizeof(byte));
 		}
 
 		/// <summary>
 		/// Writes an array of shorts to the output
 		/// </summary>
-		/// <param name="Data">Data to write. May be null.</param>
-		public void WriteShortArray(short[]? Data)
+		/// <param name="data">Data to write. May be null.</param>
+		public void WriteShortArray(short[]? data)
 		{
-			WritePrimitiveArray(Data, sizeof(short));
+			WritePrimitiveArray(data, sizeof(short));
 		}
 
 		/// <summary>
 		/// Writes an array of ints to the output
 		/// </summary>
-		/// <param name="Data">Data to write. May be null.</param>
-		public void WriteIntArray(int[]? Data)
+		/// <param name="data">Data to write. May be null.</param>
+		public void WriteIntArray(int[]? data)
 		{
-			WritePrimitiveArray(Data, sizeof(int));
+			WritePrimitiveArray(data, sizeof(int));
 		}
 
 		/// <summary>
 		/// Writes an array of primitive types to the output.
 		/// </summary>
-		/// <param name="Data">Data to write. May be null.</param>
-		private void WritePrimitiveArray<T>(T[]? Data, int ElementSize) where T : struct
+		/// <param name="data">Data to write. May be null.</param>
+		private void WritePrimitiveArray<T>(T[]? data, int elementSize) where T : struct
 		{
-			if (Data == null)
+			if (data == null)
 			{
 				WriteInt(-1);
 			}
 			else
 			{
-				WriteInt(Data.Length);
-				WriteBulkData(Data, Data.Length * ElementSize);
+				WriteInt(data.Length);
+				WriteBulkData(data, data.Length * elementSize);
 			}
 		}
 
 		/// <summary>
 		/// Writes an array of bytes to the output
 		/// </summary>
-		/// <param name="Data">Data to write. May be null.</param>
-		public void WriteFixedSizeByteArray(byte[] Data)
+		/// <param name="data">Data to write. May be null.</param>
+		public void WriteFixedSizeByteArray(byte[] data)
 		{
-			WriteFixedSizePrimitiveArray(Data, sizeof(byte));
+			WriteFixedSizePrimitiveArray(data, sizeof(byte));
 		}
 
 		/// <summary>
 		/// Writes an array of shorts to the output
 		/// </summary>
-		/// <param name="Data">Data to write. May be null.</param>
-		public void WriteFixedSizeShortArray(short[] Data)
+		/// <param name="data">Data to write. May be null.</param>
+		public void WriteFixedSizeShortArray(short[] data)
 		{
-			WriteFixedSizePrimitiveArray(Data, sizeof(short));
+			WriteFixedSizePrimitiveArray(data, sizeof(short));
 		}
 
 		/// <summary>
 		/// Writes an array of ints to the output
 		/// </summary>
-		/// <param name="Data">Data to write. May be null.</param>
-		public void WriteFixedSizeIntArray(int[] Data)
+		/// <param name="data">Data to write. May be null.</param>
+		public void WriteFixedSizeIntArray(int[] data)
 		{
-			WriteFixedSizePrimitiveArray(Data, sizeof(int));
+			WriteFixedSizePrimitiveArray(data, sizeof(int));
 		}
 
 		/// <summary>
 		/// Writes an array of primitive types to the output.
 		/// </summary>
-		/// <param name="Data">Data to write. May be null.</param>
-		private void WriteFixedSizePrimitiveArray<T>(T[] Data, int ElementSize) where T : struct
+		/// <param name="data">Data to write. May be null.</param>
+		private void WriteFixedSizePrimitiveArray<T>(T[] data, int elementSize) where T : struct
 		{
-			WriteBulkData(Data, Data.Length * ElementSize);
+			WriteBulkData(data, data.Length * elementSize);
 		}
 
 		/// <summary>
 		/// Writes primitive data from the given array to the output buffer.
 		/// </summary>
-		/// <param name="Data">Data to write.</param>
-		/// <param name="Size">Size of the data, in bytes</param>
-		private void WriteBulkData(Array Data, int Size)
+		/// <param name="data">Data to write.</param>
+		/// <param name="size">Size of the data, in bytes</param>
+		private void WriteBulkData(Array data, int size)
 		{
-			if (Size > 0)
+			if (size > 0)
 			{
-				for (int Pos = 0; ;)
+				for (int pos = 0; ;)
 				{
-					int CopySize = Math.Min(Size - Pos, Buffer.Length - BufferPos);
+					int copySize = Math.Min(size - pos, _buffer.Length - _bufferPos);
 
-					System.Buffer.BlockCopy(Data, Pos, Buffer, BufferPos, CopySize);
-					BufferPos += CopySize;
-					Pos += CopySize;
+					System.Buffer.BlockCopy(data, pos, _buffer, _bufferPos, copySize);
+					_bufferPos += copySize;
+					pos += copySize;
 
-					if (Pos == Size)
+					if (pos == size)
 					{
 						break;
 					}
@@ -361,20 +361,20 @@ namespace EpicGames.Core
 		/// Write an array of items to the archive
 		/// </summary>
 		/// <typeparam name="T">Type of the element</typeparam>
-		/// <param name="Items">Array of items</param>
-		/// <param name="WriteElement">Writes an individual element to the archive</param>
-		public void WriteArray<T>(T[]? Items, Action<T> WriteElement)
+		/// <param name="items">Array of items</param>
+		/// <param name="writeElement">Writes an individual element to the archive</param>
+		public void WriteArray<T>(T[]? items, Action<T> writeElement)
 		{
-			if (Items == null)
+			if (items == null)
 			{
 				WriteInt(-1);
 			}
 			else
 			{
-				WriteInt(Items.Length);
-				for (int Idx = 0; Idx < Items.Length; Idx++)
+				WriteInt(items.Length);
+				for (int idx = 0; idx < items.Length; idx++)
 				{
-					WriteElement(Items[Idx]);
+					writeElement(items[idx]);
 				}
 			}
 		}
@@ -383,20 +383,20 @@ namespace EpicGames.Core
 		/// Write a list of items to the archive
 		/// </summary>
 		/// <typeparam name="T">Type of the element</typeparam>
-		/// <param name="Items">List of items</param>
-		/// <param name="WriteElement">Writes an individual element to the archive</param>
-		public void WriteList<T>(IReadOnlyList<T>? Items, Action<T> WriteElement)
+		/// <param name="items">List of items</param>
+		/// <param name="writeElement">Writes an individual element to the archive</param>
+		public void WriteList<T>(IReadOnlyList<T>? items, Action<T> writeElement)
 		{
-			if (Items == null)
+			if (items == null)
 			{
 				WriteInt(-1);
 			}
 			else
 			{
-				WriteInt(Items.Count);
-				for (int Idx = 0; Idx < Items.Count; Idx++)
+				WriteInt(items.Count);
+				for (int idx = 0; idx < items.Count; idx++)
 				{
-					WriteElement(Items[Idx]);
+					writeElement(items[idx]);
 				}
 			}
 		}
@@ -405,20 +405,20 @@ namespace EpicGames.Core
 		/// Writes a hashset of items
 		/// </summary>
 		/// <typeparam name="T">The element type for the set</typeparam>
-		/// <param name="Set">The set to write</param>
-		/// <param name="WriteElement">Delegate used to read a single element</param>
-		public void WriteHashSet<T>(HashSet<T> Set, Action<T> WriteElement)
+		/// <param name="set">The set to write</param>
+		/// <param name="writeElement">Delegate used to read a single element</param>
+		public void WriteHashSet<T>(HashSet<T> set, Action<T> writeElement)
 		{
-			if (Set == null)
+			if (set == null)
 			{
 				WriteInt(-1);
 			}
 			else
 			{
-				WriteInt(Set.Count);
-				foreach (T Element in Set)
+				WriteInt(set.Count);
+				foreach (T element in set)
 				{
-					WriteElement(Element);
+					writeElement(element);
 				}
 			}
 		}
@@ -426,24 +426,24 @@ namespace EpicGames.Core
 		/// <summary>
 		/// Writes a dictionary of items
 		/// </summary>
-		/// <typeparam name="K">Type of the dictionary key</typeparam>
-		/// <typeparam name="V">Type of the dictionary value</typeparam>
-		/// <param name="Dictionary">The dictionary to write</param>
-		/// <param name="WriteKey">Delegate used to read a single key</param>
-		/// <param name="WriteValue">Delegate used to read a single value</param>
-		public void WriteDictionary<K, V>(IReadOnlyDictionary<K, V> Dictionary, Action<K> WriteKey, Action<V> WriteValue) where K : notnull
+		/// <typeparam name="TK">Type of the dictionary key</typeparam>
+		/// <typeparam name="TV">Type of the dictionary value</typeparam>
+		/// <param name="dictionary">The dictionary to write</param>
+		/// <param name="writeKey">Delegate used to read a single key</param>
+		/// <param name="writeValue">Delegate used to read a single value</param>
+		public void WriteDictionary<TK, TV>(IReadOnlyDictionary<TK, TV> dictionary, Action<TK> writeKey, Action<TV> writeValue) where TK : notnull
 		{
-			if (Dictionary == null)
+			if (dictionary == null)
 			{
 				WriteInt(-1);
 			}
 			else
 			{
-				WriteInt(Dictionary.Count);
-				foreach (KeyValuePair<K, V> Pair in Dictionary)
+				WriteInt(dictionary.Count);
+				foreach (KeyValuePair<TK, TV> pair in dictionary)
 				{
-					WriteKey(Pair.Key);
-					WriteValue(Pair.Value);
+					writeKey(pair.Key);
+					writeValue(pair.Value);
 				}
 			}
 		}
@@ -452,14 +452,14 @@ namespace EpicGames.Core
 		/// Writes a nullable object to the archive
 		/// </summary>
 		/// <typeparam name="T">The nullable type</typeparam>
-		/// <param name="Item">Item to write</param>
-		/// <param name="WriteValue">Delegate used to write a value</param>
-		public void WriteNullable<T>(Nullable<T> Item, Action<T> WriteValue) where T : struct
+		/// <param name="item">Item to write</param>
+		/// <param name="writeValue">Delegate used to write a value</param>
+		public void WriteNullable<T>(Nullable<T> item, Action<T> writeValue) where T : struct
 		{
-			if (Item.HasValue)
+			if (item.HasValue)
 			{
 				WriteBool(true);
-				WriteValue(Item.Value);
+				writeValue(item.Value);
 			}
 			else
 			{
@@ -471,18 +471,18 @@ namespace EpicGames.Core
 		/// Writes an object to the output, checking whether it is null or not. Does not preserve object references; each object written is duplicated.
 		/// </summary>
 		/// <typeparam name="T">Type of the object to serialize</typeparam>
-		/// <param name="Object">Reference to check for null before serializing</param>
-		/// <param name="WriteObject">Delegate used to write the object</param>
-		public void WriteOptionalObject<T>(T Object, Action WriteObject) where T : class
+		/// <param name="obj">Reference to check for null before serializing</param>
+		/// <param name="writeObject">Delegate used to write the object</param>
+		public void WriteOptionalObject<T>(T obj, Action writeObject) where T : class
 		{
-			if (Object == null)
+			if (obj == null)
 			{
 				WriteBool(false);
 			}
 			else
 			{
 				WriteBool(true);
-				WriteObject();
+				writeObject();
 			}
 		}
 
@@ -497,25 +497,25 @@ namespace EpicGames.Core
 		/// <summary>
 		/// Writes an object to the output. If the specific instance has already been written, preserves the reference to that.
 		/// </summary>
-		/// <param name="Object">The object to serialize</param>
-		/// <param name="WriteObject">Delegate used to write the object</param>
-		public void WriteObjectReference(object? Object, Action WriteObject)
+		/// <param name="obj">The object to serialize</param>
+		/// <param name="writeObject">Delegate used to write the object</param>
+		public void WriteObjectReference(object? obj, Action writeObject)
 		{
-			if (Object == null)
+			if (obj == null)
 			{
 				WriteInt(-1);
 			}
 			else
 			{
-				if (ObjectToUniqueId.TryGetValue(Object, out int Index))
+				if (_objectToUniqueId.TryGetValue(obj, out int index))
 				{
-					WriteInt(Index);
+					WriteInt(index);
 				}
 				else
 				{
-					WriteInt(ObjectToUniqueId.Count);
-					ObjectToUniqueId.Add(Object, ObjectToUniqueId.Count);
-					WriteObject();
+					WriteInt(_objectToUniqueId.Count);
+					_objectToUniqueId.Add(obj, _objectToUniqueId.Count);
+					writeObject();
 				}
 			}
 		}
@@ -523,11 +523,11 @@ namespace EpicGames.Core
 		/// <summary>
 		/// Writes an object to the output. If the specific instance has already been written, preserves the reference to that.
 		/// </summary>
-		/// <param name="Object">The object to serialize</param>
-		/// <param name="WriteObject">Delegate used to write the object</param>
-		public void WriteObjectReference<T>(T Object, Action WriteObject) where T : class?
+		/// <param name="obj">The object to serialize</param>
+		/// <param name="writeObject">Delegate used to write the object</param>
+		public void WriteObjectReference<T>(T obj, Action writeObject) where T : class?
 		{
-			WriteObjectReference((object?)Object, WriteObject);
+			WriteObjectReference((object?)obj, writeObject);
 		}
 	}
 }
diff --git a/Engine/Source/Programs/Shared/EpicGames.Core/BinaryConverter.cs b/Engine/Source/Programs/Shared/EpicGames.Core/BinaryConverter.cs
--- a/Engine/Source/Programs/Shared/EpicGames.Core/BinaryConverter.cs
+++ b/Engine/Source/Programs/Shared/EpicGames.Core/BinaryConverter.cs
@@ -26,16 +26,16 @@ namespace EpicGames.Core
 		/// <summary>
 		/// Reads a value from the archive
 		/// </summary>
-		/// <param name="Reader">The archive reader</param>
+		/// <param name="reader">The archive reader</param>
 		/// <returns>New instance of the value</returns>
-		TValue Read(BinaryArchiveReader Reader);
+		TValue Read(BinaryArchiveReader reader);
 
 		/// <summary>
 		/// Writes a value to the archive
 		/// </summary>
-		/// <param name="Writer">The archive writer</param>
-		/// <param name="Value">The value to write</param>
-		void Write(BinaryArchiveWriter Writer, TValue Value);
+		/// <param name="writer">The archive writer</param>
+		/// <param name="value">The value to write</param>
+		void Write(BinaryArchiveWriter writer, TValue value);
 	}
 
 	/// <summary>
@@ -46,53 +46,53 @@ namespace EpicGames.Core
 		/// <summary>
 		/// Map from type to the converter type
 		/// </summary>
-		static readonly ConcurrentDictionary<Type, Type> TypeToConverterType = new ConcurrentDictionary<Type, Type>();
+		static readonly ConcurrentDictionary<Type, Type> s_typeToConverterType = new ConcurrentDictionary<Type, Type>();
 
 		/// <summary>
 		/// Explicitly register the converter for a type. If Type is a generic type, the converter should also be a generic type with the same type arguments.
 		/// </summary>
-		/// <param name="Type">Type to register a converter for</param>
-		/// <param name="ConverterType">The converter type</param>
-		public static void RegisterConverter(Type Type, Type ConverterType)
+		/// <param name="type">Type to register a converter for</param>
+		/// <param name="converterType">The converter type</param>
+		public static void RegisterConverter(Type type, Type converterType)
 		{
-			if (!TypeToConverterType.TryAdd(Type, ConverterType))
+			if (!s_typeToConverterType.TryAdd(type, converterType))
 			{
-				throw new Exception($"Type '{Type.Name}' already has a registered converter ({TypeToConverterType[Type].Name})");
+				throw new Exception($"Type '{type.Name}' already has a registered converter ({s_typeToConverterType[type].Name})");
 			}
 		}
 
 		/// <summary>
 		/// Attempts to get the converter for a particular type
 		/// </summary>
-		/// <param name="Type">The type to use</param>
-		/// <param name="ConverterType">The converter type</param>
+		/// <param name="type">The type to use</param>
+		/// <param name="converterType">The converter type</param>
 		/// <returns>True if a converter was found</returns>
-		public static bool TryGetConverterType(Type Type, out Type? ConverterType)
+		public static bool TryGetConverterType(Type type, out Type? converterType)
 		{
-			if (TypeToConverterType.TryGetValue(Type, out Type? CustomConverterType))
+			if (s_typeToConverterType.TryGetValue(type, out Type? customConverterType))
 			{
-				ConverterType = CustomConverterType;
+				converterType = customConverterType;
 				return true;
 			}
 
-			BinaryConverterAttribute? ConverterAttribute = Type.GetCustomAttribute<BinaryConverterAttribute>();
-			if (ConverterAttribute != null)
+			BinaryConverterAttribute? converterAttribute = type.GetCustomAttribute<BinaryConverterAttribute>();
+			if (converterAttribute != null)
 			{
-				ConverterType = ConverterAttribute.Type;
+				converterType = converterAttribute.Type;
 				return true;
 			}
 
-			if (Type.IsGenericType)
+			if (type.IsGenericType)
 			{
-				BinaryConverterAttribute? GenericConverterAttribute = Type.GetGenericTypeDefinition().GetCustomAttribute<BinaryConverterAttribute>();
-				if (GenericConverterAttribute != null)
+				BinaryConverterAttribute? genericConverterAttribute = type.GetGenericTypeDefinition().GetCustomAttribute<BinaryConverterAttribute>();
+				if (genericConverterAttribute != null)
 				{
-					ConverterType = GenericConverterAttribute.Type.MakeGenericType(Type.GetGenericArguments());
+					converterType = genericConverterAttribute.Type.MakeGenericType(type.GetGenericArguments());
 					return true;
 				}
 			}
 
-			ConverterType = null;
+			converterType = null;
 			return false;
 		}
 	}
