Reformat EpicGames.Core according to standard coding conventions.

diff --git a/Engine/Source/Programs/Shared/EpicGames.Core/BinaryFormatterUtils.cs b/Engine/Source/Programs/Shared/EpicGames.Core/BinaryFormatterUtils.cs
--- a/Engine/Source/Programs/Shared/EpicGames.Core/BinaryFormatterUtils.cs
+++ b/Engine/Source/Programs/Shared/EpicGames.Core/BinaryFormatterUtils.cs
@@ -13,45 +13,45 @@ namespace EpicGames.Core
 		/// <summary>
 		/// Load an object from a file on disk, using the binary formatter
 		/// </summary>
-		/// <param name="Location">File to read from</param>
+		/// <param name="location">File to read from</param>
 		/// <returns>Instance of the object that was read from disk</returns>
-		public static T Load<T>(FileReference Location)
+		public static T Load<T>(FileReference location)
 		{
-			using FileStream Stream = new FileStream(Location.FullName, FileMode.Open, FileAccess.Read);
-			BinaryFormatter Formatter = new BinaryFormatter();
-			return (T)Formatter.Deserialize(Stream);
+			using FileStream stream = new FileStream(location.FullName, FileMode.Open, FileAccess.Read);
+			BinaryFormatter formatter = new BinaryFormatter();
+			return (T)formatter.Deserialize(stream);
 		}
 
 		/// <summary>
 		/// Saves a file to disk, using the binary formatter
 		/// </summary>
-		/// <param name="Location">File to write to</param>
-		/// <param name="Object">Object to serialize</param>
-		public static void Save(FileReference Location, object Object)
+		/// <param name="location">File to write to</param>
+		/// <param name="obj">Object to serialize</param>
+		public static void Save(FileReference location, object obj)
 		{
-			DirectoryReference.CreateDirectory(Location.Directory);
-			using FileStream Stream = new FileStream(Location.FullName, FileMode.Create, FileAccess.Write);
-			BinaryFormatter Formatter = new BinaryFormatter();
-			Formatter.Serialize(Stream, Object);
+			DirectoryReference.CreateDirectory(location.Directory);
+			using FileStream stream = new FileStream(location.FullName, FileMode.Create, FileAccess.Write);
+			BinaryFormatter formatter = new BinaryFormatter();
+			formatter.Serialize(stream, obj);
 		}
 
 		/// <summary>
 		/// Saves a file to disk using the binary formatter, without updating the timestamp if it hasn't changed
 		/// </summary>
-		/// <param name="Location">File to write to</param>
-		/// <param name="Object">Object to serialize</param>
-		public static void SaveIfDifferent(FileReference Location, object Object)
+		/// <param name="location">File to write to</param>
+		/// <param name="obj">Object to serialize</param>
+		public static void SaveIfDifferent(FileReference location, object obj)
 		{
-			byte[] Contents;
-			using (MemoryStream Stream = new MemoryStream())
+			byte[] contents;
+			using (MemoryStream stream = new MemoryStream())
 			{
-				BinaryFormatter Formatter = new BinaryFormatter();
-				Formatter.Serialize(Stream, Object);
-				Contents = Stream.ToArray();
+				BinaryFormatter formatter = new BinaryFormatter();
+				formatter.Serialize(stream, obj);
+				contents = stream.ToArray();
 			}
 
-			DirectoryReference.CreateDirectory(Location.Directory);
-			FileReference.WriteAllBytesIfDifferent(Location, Contents);
+			DirectoryReference.CreateDirectory(location.Directory);
+			FileReference.WriteAllBytesIfDifferent(location, contents);
 		}
 	}
 }
diff --git a/Engine/Source/Programs/Shared/EpicGames.Core/BinaryReaderExtensions.cs b/Engine/Source/Programs/Shared/EpicGames.Core/BinaryReaderExtensions.cs
--- a/Engine/Source/Programs/Shared/EpicGames.Core/BinaryReaderExtensions.cs
+++ b/Engine/Source/Programs/Shared/EpicGames.Core/BinaryReaderExtensions.cs
@@ -12,140 +12,140 @@ namespace EpicGames.Core
 		/// Read an objects from a binary reader
 		/// </summary>
 		/// <typeparam name="T">The element type for the array</typeparam>
-		/// <param name="Reader">Reader to read data from</param>
+		/// <param name="reader">Reader to read data from</param>
 		/// <returns>Object instance.</returns>
-		public static T ReadObject<T>(this BinaryReader Reader) where T : class, IBinarySerializable
+		public static T ReadObject<T>(this BinaryReader reader) where T : class, IBinarySerializable
 		{
-			return (T)Activator.CreateInstance(typeof(T), Reader)!;
+			return (T)Activator.CreateInstance(typeof(T), reader)!;
 		}
 
 		/// <summary>
 		/// Read an array of strings from a binary reader
 		/// </summary>
-		/// <param name="Reader">Reader to read data from</param>
+		/// <param name="reader">Reader to read data from</param>
 		/// <returns>Array of strings, as serialized. May be null.</returns>
-		public static string[]? ReadStringArray(this BinaryReader Reader)
+		public static string[]? ReadStringArray(this BinaryReader reader)
 		{
-			return Reader.ReadArray(() => Reader.ReadString());
+			return reader.ReadArray(() => reader.ReadString());
 		}
 
 		/// <summary>
 		/// Read an array of objects from a binary reader
 		/// </summary>
 		/// <typeparam name="T">The element type for the array</typeparam>
-		/// <param name="Reader">Reader to read data from</param>
+		/// <param name="reader">Reader to read data from</param>
 		/// <returns>Array of objects, as serialized. May be null.</returns>
-		public static T[]? ReadArray<T>(this BinaryReader Reader) where T : class, IBinarySerializable
+		public static T[]? ReadArray<T>(this BinaryReader reader) where T : class, IBinarySerializable
 		{
-			return ReadArray<T>(Reader, () => Reader.ReadObject<T>());
+			return ReadArray<T>(reader, () => reader.ReadObject<T>());
 		}
 
 		/// <summary>
 		/// Read an array of objects from a binary reader
 		/// </summary>
 		/// <typeparam name="T">The element type for the array</typeparam>
-		/// <param name="Reader">Reader to read data from</param>
-		/// <param name="ReadElement">Delegate to call to serialize each element</param>
+		/// <param name="reader">Reader to read data from</param>
+		/// <param name="readElement">Delegate to call to serialize each element</param>
 		/// <returns>Array of objects, as serialized. May be null.</returns>
-		public static T[]? ReadArray<T>(this BinaryReader Reader, Func<T> ReadElement)
+		public static T[]? ReadArray<T>(this BinaryReader reader, Func<T> readElement)
 		{
-			int NumItems = Reader.ReadInt32();
-			if (NumItems < 0)
+			int numItems = reader.ReadInt32();
+			if (numItems < 0)
 			{
 				return null;
 			}
 
-			T[] Items = new T[NumItems];
-			for (int Idx = 0; Idx < NumItems; Idx++)
+			T[] items = new T[numItems];
+			for (int idx = 0; idx < numItems; idx++)
 			{
-				Items[Idx] = ReadElement();
+				items[idx] = readElement();
 			}
-			return Items;
+			return items;
 		}
 
 		/// <summary>
 		/// Read a list of strings from a binary reader
 		/// </summary>
-		/// <param name="Reader">Reader to read data from</param>
+		/// <param name="reader">Reader to read data from</param>
 		/// <returns>Array of strings, as serialized. May be null.</returns>
-		public static List<string>? ReadStringList(this BinaryReader Reader)
+		public static List<string>? ReadStringList(this BinaryReader reader)
 		{
-			return Reader.ReadList(() => Reader.ReadString());
+			return reader.ReadList(() => reader.ReadString());
 		}
 
 		/// <summary>
 		/// Read a list of objects from a binary reader
 		/// </summary>
 		/// <typeparam name="T">The element type for the list</typeparam>
-		/// <param name="Reader">Reader to read data from</param>
+		/// <param name="reader">Reader to read data from</param>
 		/// <param name="ReadElement">Delegate to call to serialize each element</param>
 		/// <returns>List of objects, as serialized. May be null.</returns>
-		public static List<T>? ReadList<T>(this BinaryReader Reader) where T : class, IBinarySerializable
+		public static List<T>? ReadList<T>(this BinaryReader reader) where T : class, IBinarySerializable
 		{
-			return ReadList<T>(Reader, () => Reader.ReadObject<T>());
+			return ReadList<T>(reader, () => reader.ReadObject<T>());
 		}
 
 		/// <summary>
 		/// Read a list of objects from a binary reader
 		/// </summary>
 		/// <typeparam name="T">The element type for the list</typeparam>
-		/// <param name="Reader">Reader to read data from</param>
-		/// <param name="ReadElement">Delegate to call to serialize each element</param>
+		/// <param name="reader">Reader to read data from</param>
+		/// <param name="readElement">Delegate to call to serialize each element</param>
 		/// <returns>List of objects, as serialized. May be null.</returns>
-		public static List<T>? ReadList<T>(this BinaryReader Reader, Func<T> ReadElement)
+		public static List<T>? ReadList<T>(this BinaryReader reader, Func<T> readElement)
 		{
-			int NumItems = Reader.ReadInt32();
-			if (NumItems < 0)
+			int numItems = reader.ReadInt32();
+			if (numItems < 0)
 			{
 				return null;
 			}
 
-			List<T> Items = new List<T>(NumItems);
-			for (int Idx = 0; Idx < NumItems; Idx++)
+			List<T> items = new List<T>(numItems);
+			for (int idx = 0; idx < numItems; idx++)
 			{
-				Items.Add(ReadElement());
+				items.Add(readElement());
 			}
-			return Items;
+			return items;
 		}
 
 		/// <summary>
 		/// Read a list of objects from a binary reader
 		/// </summary>
 		/// <typeparam name="T">The element type for the list</typeparam>
-		/// <param name="Reader">Reader to read data from</param>
-		/// <param name="ReadKey">Delegate to call to serialize each key</param>
-		/// <param name="ReadValue">Delegate to call to serialize each value</param>
+		/// <param name="reader">Reader to read data from</param>
+		/// <param name="readKey">Delegate to call to serialize each key</param>
+		/// <param name="readValue">Delegate to call to serialize each value</param>
 		/// <returns>List of objects, as serialized. May be null.</returns>
-		public static Dictionary<K, V>? ReadDictionary<K, V>(this BinaryReader Reader, Func<K> ReadKey, Func<V> ReadValue) where K : notnull
+		public static Dictionary<TK, TV>? ReadDictionary<TK, TV>(this BinaryReader reader, Func<TK> readKey, Func<TV> readValue) where TK : notnull
 		{
-			int NumItems = Reader.ReadInt32();
-			if (NumItems < 0)
+			int numItems = reader.ReadInt32();
+			if (numItems < 0)
 			{
 				return null;
 			}
 
-			Dictionary<K, V> Items = new Dictionary<K, V>(NumItems);
-			for (int Idx = 0; Idx < NumItems; Idx++)
+			Dictionary<TK, TV> items = new Dictionary<TK, TV>(numItems);
+			for (int idx = 0; idx < numItems; idx++)
 			{
-				K Key = ReadKey();
-				V Value = ReadValue();
-				Items.Add(Key, Value);
+				TK key = readKey();
+				TV value = readValue();
+				items.Add(key, value);
 			}
-			return Items;
+			return items;
 		}
 
 		/// <summary>
 		/// Read a nullable object from a binary reader
 		/// </summary>
 		/// <typeparam name="T">Type of the object</typeparam>
-		/// <param name="Reader">Reader to read data from</param>
-		/// <param name="ReadItem">Function to read the payload, if non-null</param>
+		/// <param name="reader">Reader to read data from</param>
+		/// <param name="readItem">Function to read the payload, if non-null</param>
 		/// <returns>Object instance or null</returns>
-		public static T? ReadNullable<T>(this BinaryReader Reader, Func<T> ReadItem) where T : class
+		public static T? ReadNullable<T>(this BinaryReader reader, Func<T> readItem) where T : class
 		{
-			if (Reader.ReadBoolean())
+			if (reader.ReadBoolean())
 			{
-				return ReadItem();
+				return readItem();
 			}
 			else
 			{
@@ -156,51 +156,51 @@ namespace EpicGames.Core
 		/// <summary>
 		/// Reads a value of a specific type from a binary reader
 		/// </summary>
-		/// <param name="Reader">Reader for input data</param>
-		/// <param name="ObjectType">Type of value to read</param>
+		/// <param name="reader">Reader for input data</param>
+		/// <param name="objectType">Type of value to read</param>
 		/// <returns>The value read from the stream</returns>
-		public static object? ReadObject(this BinaryReader Reader, Type ObjectType)
+		public static object? ReadObject(this BinaryReader reader, Type objectType)
 		{
-			if (ObjectType == typeof(string))
+			if (objectType == typeof(string))
 			{
-				return Reader.ReadString();
+				return reader.ReadString();
 			}
-			else if (ObjectType == typeof(bool))
+			else if (objectType == typeof(bool))
 			{
-				return Reader.ReadBoolean();
+				return reader.ReadBoolean();
 			}
-			else if (ObjectType == typeof(int))
+			else if (objectType == typeof(int))
 			{
-				return Reader.ReadInt32();
+				return reader.ReadInt32();
 			}
-			else if (ObjectType == typeof(float))
+			else if (objectType == typeof(float))
 			{
-				return Reader.ReadSingle();
+				return reader.ReadSingle();
 			}
-			else if (ObjectType == typeof(double))
+			else if (objectType == typeof(double))
 			{
-				return Reader.ReadDouble();
+				return reader.ReadDouble();
 			}
-			else if (ObjectType == typeof(string[]))
+			else if (objectType == typeof(string[]))
 			{
-				return Reader.ReadStringArray();
+				return reader.ReadStringArray();
 			}
-			else if (ObjectType == typeof(bool?))
+			else if (objectType == typeof(bool?))
 			{
-				int Value = Reader.ReadInt32();
-				return (Value == -1) ? (bool?)null : (Value == 0) ? (bool?)false : (bool?)true;
+				int value = reader.ReadInt32();
+				return (value == -1) ? (bool?)null : (value == 0) ? (bool?)false : (bool?)true;
 			}
-			else if (ObjectType == typeof(FileReference))
+			else if (objectType == typeof(FileReference))
 			{
-				return Reader.ReadFileReference();
+				return reader.ReadFileReference();
 			}
-			else if (ObjectType.IsEnum)
+			else if (objectType.IsEnum)
 			{
-				return Enum.ToObject(ObjectType, Reader.ReadInt32());
+				return Enum.ToObject(objectType, reader.ReadInt32());
 			}
 			else
 			{
-				throw new Exception(String.Format("Reading binary objects of type '{0}' is not currently supported.", ObjectType.Name));
+				throw new Exception(String.Format("Reading binary objects of type '{0}' is not currently supported.", objectType.Name));
 			}
 		}
 	}
diff --git a/Engine/Source/Programs/Shared/EpicGames.Core/BinaryWriterExtensions.cs b/Engine/Source/Programs/Shared/EpicGames.Core/BinaryWriterExtensions.cs
--- a/Engine/Source/Programs/Shared/EpicGames.Core/BinaryWriterExtensions.cs
+++ b/Engine/Source/Programs/Shared/EpicGames.Core/BinaryWriterExtensions.cs
@@ -13,7 +13,7 @@ namespace EpicGames.Core
 	/// </summary>
 	public interface IBinarySerializable
 	{
-		void Write(BinaryWriter Writer);
+		void Write(BinaryWriter writer);
 	}
 
 	/// <summary>
@@ -24,53 +24,53 @@ namespace EpicGames.Core
 		/// <summary>
 		/// Writes an item implementing the IBinarySerializable interface to a binary writer. Included for symmetry with standard Writer.Write(X) calls.
 		/// </summary>
-		/// <param name="Writer">Writer to serialize to</param>
-		/// <param name="Item">The item to write</param>
-		public static void Write(this BinaryWriter Writer, IBinarySerializable Item)
+		/// <param name="writer">Writer to serialize to</param>
+		/// <param name="item">The item to write</param>
+		public static void Write(this BinaryWriter writer, IBinarySerializable item)
 		{
-			Item.Write(Writer);
+			item.Write(writer);
 		}
 
 		/// <summary>
 		/// Writes an array of strings to a binary writer.
 		/// </summary>
-		/// <param name="Writer">Writer to serialize to</param>
-		/// <param name="Items">Array of items</param>
-		public static void Write(this BinaryWriter Writer, string[] Items)
+		/// <param name="writer">Writer to serialize to</param>
+		/// <param name="items">Array of items</param>
+		public static void Write(this BinaryWriter writer, string[] items)
 		{
-			Write(Writer, Items, Item => Writer.Write(Item));
+			Write(writer, items, item => writer.Write(item));
 		}
 
 		/// <summary>
 		/// Writes an array to a binary writer.
 		/// </summary>
 		/// <typeparam name="T">The array element type</typeparam>
-		/// <param name="Writer">Writer to serialize to</param>
-		/// <param name="Items">Array of items</param>
-		public static void Write<T>(this BinaryWriter Writer, T[] Items) where T : class, IBinarySerializable
+		/// <param name="writer">Writer to serialize to</param>
+		/// <param name="items">Array of items</param>
+		public static void Write<T>(this BinaryWriter writer, T[] items) where T : class, IBinarySerializable
 		{
-			Write(Writer, Items, Item => Writer.Write(Item));
+			Write(writer, items, item => writer.Write(item));
 		}
 
 		/// <summary>
 		/// Writes an array to a binary writer.
 		/// </summary>
 		/// <typeparam name="T">The array element type</typeparam>
-		/// <param name="Writer">Writer to serialize to</param>
-		/// <param name="Items">Array of items</param>
-		/// <param name="WriteElement">Delegate to call to serialize each element</param>
-		public static void Write<T>(this BinaryWriter Writer, T[] Items, Action<T> WriteElement)
+		/// <param name="writer">Writer to serialize to</param>
+		/// <param name="items">Array of items</param>
+		/// <param name="writeElement">Delegate to call to serialize each element</param>
+		public static void Write<T>(this BinaryWriter writer, T[] items, Action<T> writeElement)
 		{
-			if (Items == null)
+			if (items == null)
 			{
-				Writer.Write(-1);
+				writer.Write(-1);
 			}
 			else
 			{
-				Writer.Write(Items.Length);
-				for (int Idx = 0; Idx < Items.Length; Idx++)
+				writer.Write(items.Length);
+				for (int idx = 0; idx < items.Length; idx++)
 				{
-					WriteElement(Items[Idx]);
+					writeElement(items[idx]);
 				}
 			}
 		}
@@ -78,43 +78,43 @@ namespace EpicGames.Core
 		/// <summary>
 		/// Writes a list of strings to a binary writer.
 		/// </summary>
-		/// <param name="Writer">Writer to serialize to</param>
-		/// <param name="Items">Array of items</param>
-		public static void Write(this BinaryWriter Writer, List<string> Items)
+		/// <param name="writer">Writer to serialize to</param>
+		/// <param name="items">Array of items</param>
+		public static void Write(this BinaryWriter writer, List<string> items)
 		{
-			Write(Writer, Items, Item => Writer.Write(Item));
+			Write(writer, items, item => writer.Write(item));
 		}
 
 		/// <summary>
 		/// Writes a list to a binary writer.
 		/// </summary>
 		/// <typeparam name="T">The array element type</typeparam>
-		/// <param name="Writer">Writer to serialize to</param>
-		/// <param name="Items">Array of items</param>
-		public static void Write<T>(this BinaryWriter Writer, List<T> Items) where T : class, IBinarySerializable
+		/// <param name="writer">Writer to serialize to</param>
+		/// <param name="items">Array of items</param>
+		public static void Write<T>(this BinaryWriter writer, List<T> items) where T : class, IBinarySerializable
 		{
-			Write(Writer, Items, Item => Writer.Write(Item));
+			Write(writer, items, item => writer.Write(item));
 		}
 
 		/// <summary>
 		/// Writes a list to a binary writer.
 		/// </summary>
 		/// <typeparam name="T">The list element type</typeparam>
-		/// <param name="Writer">Writer to serialize to</param>
-		/// <param name="Items">List of items</param>
-		/// <param name="WriteElement">Delegate to call to serialize each element</param>
-		public static void Write<T>(this BinaryWriter Writer, List<T> Items, Action<T> WriteElement)
+		/// <param name="writer">Writer to serialize to</param>
+		/// <param name="items">List of items</param>
+		/// <param name="writeElement">Delegate to call to serialize each element</param>
+		public static void Write<T>(this BinaryWriter writer, List<T> items, Action<T> writeElement)
 		{
-			if (Items == null)
+			if (items == null)
 			{
-				Writer.Write(-1);
+				writer.Write(-1);
 			}
 			else
 			{
-				Writer.Write(Items.Count);
-				for (int Idx = 0; Idx < Items.Count; Idx++)
+				writer.Write(items.Count);
+				for (int idx = 0; idx < items.Count; idx++)
 				{
-					WriteElement(Items[Idx]);
+					writeElement(items[idx]);
 				}
 			}
 		}
@@ -122,26 +122,26 @@ namespace EpicGames.Core
 		/// <summary>
 		/// Write a dictionary to a binary writer
 		/// </summary>
-		/// <typeparam name="K">The key type for the dictionary</typeparam>
-		/// <typeparam name="V">The value type for the dictionary</typeparam>
+		/// <typeparam name="TK">The key type for the dictionary</typeparam>
+		/// <typeparam name="TV">The value type for the dictionary</typeparam>
 		/// <param name="Reader">Reader to read data from</param>
-		/// <param name="Items">List of items to be written</param>
-		/// <param name="WriteKey">Delegate to call to serialize each key</param>
-		/// <param name="WriteKey">Delegate to call to serialize each value</param>
+		/// <param name="items">List of items to be written</param>
+		/// <param name="writeKey">Delegate to call to serialize each key</param>
+		/// <param name="writeKey">Delegate to call to serialize each value</param>
 		/// <returns>Dictionary of objects, as serialized. May be null.</returns>
-		public static void Write<K, V>(this BinaryWriter Writer, Dictionary<K, V> Items, Action<K> WriteKey, Action<V> WriteValue) where K : notnull
+		public static void Write<TK, TV>(this BinaryWriter writer, Dictionary<TK, TV> items, Action<TK> writeKey, Action<TV> writeValue) where TK : notnull
 		{
-			if (Items == null)
+			if (items == null)
 			{
-				Writer.Write(-1);
+				writer.Write(-1);
 			}
 			else
 			{
-				Writer.Write(Items.Count);
-				foreach (KeyValuePair<K, V> Item in Items)
+				writer.Write(items.Count);
+				foreach (KeyValuePair<TK, TV> item in items)
 				{
-					WriteKey(Item.Key);
-					WriteValue(Item.Value);
+					writeKey(item.Key);
+					writeValue(item.Value);
 				}
 			}
 		}
@@ -150,70 +150,70 @@ namespace EpicGames.Core
 		/// Read a nullable object from a binary reader
 		/// </summary>
 		/// <typeparam name="T">Type of the object</typeparam>
-		/// <param name="Writer">Reader to read data from</param>
-		/// <param name="WriteItem">Function to read the payload, if non-null</param>
+		/// <param name="writer">Reader to read data from</param>
+		/// <param name="writeItem">Function to read the payload, if non-null</param>
 		/// <returns>Object instance or null</returns>
-		public static void WriteNullable<T>(this BinaryWriter Writer, T Item, Action WriteItem) where T : class
+		public static void WriteNullable<T>(this BinaryWriter writer, T item, Action writeItem) where T : class
 		{
-			if (Item == null)
+			if (item == null)
 			{
-				Writer.Write(false);
+				writer.Write(false);
 			}
 			else
 			{
-				Writer.Write(true);
-				WriteItem();
+				writer.Write(true);
+				writeItem();
 			}
 		}
 
 		/// <summary>
 		/// Writes a value of a specific type to a binary writer
 		/// </summary>
-		/// <param name="Writer">Writer for output data</param>
-		/// <param name="FieldType">Type of value to write</param>
-		/// <param name="Value">The value to output</param>
-		public static void Write(this BinaryWriter Writer, Type FieldType, object Value)
+		/// <param name="writer">Writer for output data</param>
+		/// <param name="fieldType">Type of value to write</param>
+		/// <param name="value">The value to output</param>
+		public static void Write(this BinaryWriter writer, Type fieldType, object value)
 		{
-			if (FieldType == typeof(string))
+			if (fieldType == typeof(string))
 			{
-				Writer.Write((string)Value);
+				writer.Write((string)value);
 			}
-			else if (FieldType == typeof(bool))
+			else if (fieldType == typeof(bool))
 			{
-				Writer.Write((bool)Value);
+				writer.Write((bool)value);
 			}
-			else if (FieldType == typeof(int))
+			else if (fieldType == typeof(int))
 			{
-				Writer.Write((int)Value);
+				writer.Write((int)value);
 			}
-			else if (FieldType == typeof(float))
+			else if (fieldType == typeof(float))
 			{
-				Writer.Write((float)Value);
+				writer.Write((float)value);
 			}
-			else if (FieldType == typeof(double))
+			else if (fieldType == typeof(double))
 			{
-				Writer.Write((double)Value);
+				writer.Write((double)value);
 			}
-			else if (FieldType.IsEnum)
+			else if (fieldType.IsEnum)
 			{
-				Writer.Write((int)Value);
+				writer.Write((int)value);
 			}
-			else if (FieldType == typeof(string[]))
+			else if (fieldType == typeof(string[]))
 			{
-				Writer.Write((string[])Value);
+				writer.Write((string[])value);
 			}
-			else if (FieldType == typeof(bool?))
+			else if (fieldType == typeof(bool?))
 			{
-				bool? NullableValue = (bool?)Value;
-				Writer.Write(NullableValue.HasValue ? NullableValue.Value ? 1 : 0 : -1);
+				bool? nullableValue = (bool?)value;
+				writer.Write(nullableValue.HasValue ? nullableValue.Value ? 1 : 0 : -1);
 			}
-			else if (FieldType == typeof(FileReference))
+			else if (fieldType == typeof(FileReference))
 			{
-				Writer.Write((FileReference)Value);
+				writer.Write((FileReference)value);
 			}
 			else
 			{
-				throw new Exception(String.Format("Unsupported type '{0}' for binary serialization", FieldType.Name));
+				throw new Exception(String.Format("Unsupported type '{0}' for binary serialization", fieldType.Name));
 			}
 		}
 	}
diff --git a/Engine/Source/Programs/Shared/EpicGames.Core/Boxed.cs b/Engine/Source/Programs/Shared/EpicGames.Core/Boxed.cs
--- a/Engine/Source/Programs/Shared/EpicGames.Core/Boxed.cs
+++ b/Engine/Source/Programs/Shared/EpicGames.Core/Boxed.cs
@@ -16,10 +16,10 @@ namespace EpicGames.Core
 		/// <summary>
 		/// Default constructor
 		/// </summary>
-		/// <param name="Value">The value to copy</param>
-		public Boxed(T Value)
+		/// <param name="value">The value to copy</param>
+		public Boxed(T value)
 		{
-			this.Value = Value;
+			Value = value;
 		}
 
 		/// <inheritdoc/>
@@ -31,19 +31,19 @@ namespace EpicGames.Core
 		/// <summary>
 		/// Implicit conversion of a boxed value to value type
 		/// </summary>
-		/// <param name="Boxed"></param>
-		public static implicit operator T(Boxed<T> Boxed)
+		/// <param name="boxed"></param>
+		public static implicit operator T(Boxed<T> boxed)
 		{
-			return Boxed.Value;
+			return boxed.Value;
 		}
 
 		/// <summary>
 		/// Implicit conversion of a boxed value to value type
 		/// </summary>
-		/// <param name="Value"></param>
-		public static implicit operator Boxed<T>(T Value)
+		/// <param name="value"></param>
+		public static implicit operator Boxed<T>(T value)
 		{
-			return new Boxed<T>(Value);
+			return new Boxed<T>(value);
 		}
 	}
 }
diff --git a/Engine/Source/Programs/Shared/EpicGames.Core/BuzHash.cs b/Engine/Source/Programs/Shared/EpicGames.Core/BuzHash.cs
--- a/Engine/Source/Programs/Shared/EpicGames.Core/BuzHash.cs
+++ b/Engine/Source/Programs/Shared/EpicGames.Core/BuzHash.cs
@@ -1,8 +1,6 @@
 // Copyright Epic Games, Inc. All Rights Reserved.
 
 using System;
-using System.Collections.Generic;
-using System.Text;
 
 namespace EpicGames.Core
 {
@@ -11,7 +9,7 @@ namespace EpicGames.Core
 	/// </summary>
 	public class BuzHash
 	{
-		static uint[] Table = new uint[256]
+		static readonly uint[] s_table = new uint[256]
 		{
 			0x458be752, 0xc10748cc, 0xfbbcdbb8, 0x6ded5b68, 0xb10a82b5, 0x20d75648, 0xdfc5665f, 0xa8428801, 0x7ebf5191, 0x841135c7, 0x65cc53b3,
 			0x280a597c, 0x16f60255, 0xc78cbc3e, 0x294415f5, 0xb938d494, 0xec85c4e6, 0xb7d33edc, 0xe549b544, 0xfdeda5aa, 0x882bf287, 0x3116737c,
@@ -39,8 +37,8 @@ namespace EpicGames.Core
 			0xf9c18d66, 0x593ade65, 0xd95ddf11,
 		};
 
-		int Count;
-		uint State;
+		int _count;
+		uint _state;
 
 		/// <summary>
 		/// Constructor
@@ -53,12 +51,12 @@ namespace EpicGames.Core
 		/// <summary>
 		/// Constructor
 		/// </summary>
-		/// <param name="Count"></param>
-		/// <param name="State"></param>
-		public BuzHash(int Count, uint State)
+		/// <param name="count"></param>
+		/// <param name="state"></param>
+		public BuzHash(int count, uint state)
 		{
-			this.Count = Count;
-			this.State = State;
+			_count = count;
+			_state = state;
 		}
 
 		/// <summary>
@@ -66,84 +64,84 @@ namespace EpicGames.Core
 		/// </summary>
 		public void Reset()
 		{ 
-			Count = 0; 
-			State = 0; 
+			_count = 0; 
+			_state = 0; 
 		}
 
 		/// <summary>
 		/// Add a byte to the window and update the hash
 		/// </summary>
-		/// <param name="Value">New value to append to the window</param>
-		public void Add(byte Value)
+		/// <param name="value">New value to append to the window</param>
+		public void Add(byte value)
 		{
-			State = Add(State, Value);
-			Count++;
+			_state = Add(_state, value);
+			_count++;
 		}
 
 		/// <summary>
 		/// Add a byte to the window and update the hash
 		/// </summary>
 		/// <param name="Value">New value to append to the window</param>
-		public void Add(ReadOnlySpan<byte> Span)
+		public void Add(ReadOnlySpan<byte> span)
 		{
-			State = Add(State, Span);
-			Count += Span.Length;
+			_state = Add(_state, span);
+			_count += span.Length;
 		}
 
 		/// <summary>
 		/// Static method for updating a hash given a particular byte value
 		/// </summary>
-		/// <param name="State">The current hash value</param>
-		/// <param name="Value">Byte to add to the hash</param>
+		/// <param name="state">The current hash value</param>
+		/// <param name="value">Byte to add to the hash</param>
 		/// <returns>New hash value</returns>
-		public static uint Add(uint State, byte Value)
+		public static uint Add(uint state, byte value)
 		{
-			return Rol32(State, 1) ^ Table[Value];
+			return Rol32(state, 1) ^ s_table[value];
 		}
 
 		/// <summary>
 		/// Static method for appending a range of data to a hash value
 		/// </summary>
-		/// <param name="State">The current hash value</param>
-		/// <param name="Data">Data to append to the hash</param>
+		/// <param name="state">The current hash value</param>
+		/// <param name="data">Data to append to the hash</param>
 		/// <returns>New hash value</returns>
-		public static uint Add(uint State, ReadOnlySpan<byte> Data)
+		public static uint Add(uint state, ReadOnlySpan<byte> data)
 		{
-			for (int Idx = 0; Idx < Data.Length; Idx++)
+			for (int idx = 0; idx < data.Length; idx++)
 			{
-				State = Add(State, Data[Idx]);
+				state = Add(state, data[idx]);
 			}
-			return State;
+			return state;
 		}
 
 		/// <summary>
 		/// Removes a byte from the start of the hash window
 		/// </summary>
-		/// <param name="Value">Value at the start of the window</param>
-		public void Sub(byte Value)
+		/// <param name="value">Value at the start of the window</param>
+		public void Sub(byte value)
 		{
-			State = State ^ Rol32(Table[Value], (int)(Count - 1));
-			Count--;
+			_state ^= Rol32(s_table[value], (int)(_count - 1));
+			_count--;
 		}
 
 		/// <summary>
 		/// Removes a byte from the start of the hash window
 		/// </summary>
-		/// <param name="Value">Value at the start of the window</param>
-		public static uint Sub(uint State, byte Value, int Count)
+		/// <param name="value">Value at the start of the window</param>
+		public static uint Sub(uint state, byte value, int count)
 		{
-			return State ^ Rol32(Table[Value], (int)(Count - 1));
+			return state ^ Rol32(s_table[value], (int)(count - 1));
 		}
 
 		/// <summary>
 		/// Update the hash with the given data
 		/// </summary>
-		/// <param name="Data"></param>
-		public void Update(ReadOnlySpan<byte> Data)
+		/// <param name="data"></param>
+		public void Update(ReadOnlySpan<byte> data)
 		{
-			for (int Idx = 0; Idx < Data.Length; ++Idx)
+			for (int idx = 0; idx < data.Length; ++idx)
 			{
-				Add(Data[Idx]);
+				Add(data[idx]);
 			}
 		}
 
@@ -153,13 +151,13 @@ namespace EpicGames.Core
 		/// <returns>The current hash value</returns>
 		public uint Get()
 		{
-			return State;
+			return _state;
 		}
 
-		static uint Rol32(uint V, int N)
+		static uint Rol32(uint v, int n)
 		{
-			N &= 31;
-			return ((V) << (N)) | ((V) >> (32 - N));
+			n &= 31;
+			return ((v) << (n)) | ((v) >> (32 - n));
 		}
 	}
 }
diff --git a/Engine/Source/Programs/Shared/EpicGames.Core/CancellationTask.cs b/Engine/Source/Programs/Shared/EpicGames.Core/CancellationTask.cs
--- a/Engine/Source/Programs/Shared/EpicGames.Core/CancellationTask.cs
+++ b/Engine/Source/Programs/Shared/EpicGames.Core/CancellationTask.cs
@@ -14,37 +14,34 @@ namespace EpicGames.Core
 		/// <summary>
 		/// Completion source for the task
 		/// </summary>
-		readonly TaskCompletionSource<bool> CompletionSource;
+		readonly TaskCompletionSource<bool> _completionSource;
 
 		/// <summary>
 		/// Registration handle with the cancellation token
 		/// </summary>
-		readonly CancellationTokenRegistration Registration;
+		readonly CancellationTokenRegistration _registration;
 
 		/// <summary>
 		/// Constructor
 		/// </summary>
-		/// <param name="Token">The cancellation token to register with</param>
-		public CancellationTask(CancellationToken Token)
+		/// <param name="token">The cancellation token to register with</param>
+		public CancellationTask(CancellationToken token)
 		{
-			CompletionSource = new TaskCompletionSource<bool>();
-			Registration = Token.Register(() => CompletionSource.TrySetResult(true));
+			_completionSource = new TaskCompletionSource<bool>();
+			_registration = token.Register(() => _completionSource.TrySetResult(true));
 		}
 
 		/// <summary>
 		/// The task that can be waited on
 		/// </summary>
-		public Task Task
-		{
-			get { return CompletionSource.Task; }
-		}
+		public Task Task => _completionSource.Task;
 
 		/// <summary>
 		/// Dispose of any allocated resources
 		/// </summary>
 		public void Dispose()
 		{
-			Registration.Dispose();
+			_registration.Dispose();
 		}
 	}
 }
diff --git a/Engine/Source/Programs/Shared/EpicGames.Core/CertificateUtils.cs b/Engine/Source/Programs/Shared/EpicGames.Core/CertificateUtils.cs
--- a/Engine/Source/Programs/Shared/EpicGames.Core/CertificateUtils.cs
+++ b/Engine/Source/Programs/Shared/EpicGames.Core/CertificateUtils.cs
@@ -1,10 +1,8 @@
 // Copyright Epic Games, Inc. All Rights Reserved.
 
 using System;
-using System.Collections.Generic;
 using System.Security.Cryptography;
 using System.Security.Cryptography.X509Certificates;
-using System.Text;
 using System.Text.RegularExpressions;
 
 namespace EpicGames.Core
@@ -17,33 +15,33 @@ namespace EpicGames.Core
 		/// <summary>
 		/// Create a self signed cert for the agent
 		/// </summary>
-		/// <param name="DnsName"></param>
-		/// <param name="FriendlyName">Friendly name for the certificate</param>
+		/// <param name="dnsName"></param>
+		/// <param name="friendlyName">Friendly name for the certificate</param>
 		/// <returns></returns>
-		public static byte[] CreateSelfSignedCert(string DnsName, string FriendlyName)
+		public static byte[] CreateSelfSignedCert(string dnsName, string friendlyName)
 		{
-			DnsName = Regex.Replace(DnsName, @"^[a-zA-Z]+://", "");
-			DnsName = Regex.Replace(DnsName, "/.*$", "");
+			dnsName = Regex.Replace(dnsName, @"^[a-zA-Z]+://", "");
+			dnsName = Regex.Replace(dnsName, "/.*$", "");
 
-			using (RSA Algorithm = RSA.Create(2048))
+			using (RSA algorithm = RSA.Create(2048))
 			{
-				X500DistinguishedName distinguishedName = new X500DistinguishedName($"CN={DnsName}");
-				CertificateRequest Request = new CertificateRequest(distinguishedName, Algorithm, HashAlgorithmName.SHA256, RSASignaturePadding.Pkcs1);
+				X500DistinguishedName distinguishedName = new X500DistinguishedName($"CN={dnsName}");
+				CertificateRequest request = new CertificateRequest(distinguishedName, algorithm, HashAlgorithmName.SHA256, RSASignaturePadding.Pkcs1);
 
-				Request.CertificateExtensions.Add(new X509BasicConstraintsExtension(false, false, 0, true));
-				Request.CertificateExtensions.Add(new X509KeyUsageExtension(X509KeyUsageFlags.KeyEncipherment | X509KeyUsageFlags.DigitalSignature, true));
-				Request.CertificateExtensions.Add(new X509EnhancedKeyUsageExtension(new OidCollection { new Oid("1.3.6.1.5.5.7.3.1") }, true));
+				request.CertificateExtensions.Add(new X509BasicConstraintsExtension(false, false, 0, true));
+				request.CertificateExtensions.Add(new X509KeyUsageExtension(X509KeyUsageFlags.KeyEncipherment | X509KeyUsageFlags.DigitalSignature, true));
+				request.CertificateExtensions.Add(new X509EnhancedKeyUsageExtension(new OidCollection { new Oid("1.3.6.1.5.5.7.3.1") }, true));
 
-				SubjectAlternativeNameBuilder AlternativeNameBuilder = new SubjectAlternativeNameBuilder();
-				AlternativeNameBuilder.AddDnsName(DnsName);
-				Request.CertificateExtensions.Add(AlternativeNameBuilder.Build(true));
+				SubjectAlternativeNameBuilder alternativeNameBuilder = new SubjectAlternativeNameBuilder();
+				alternativeNameBuilder.AddDnsName(dnsName);
+				request.CertificateExtensions.Add(alternativeNameBuilder.Build(true));
 
 				// NB: MacOS requires 825 days or fewer (https://support.apple.com/en-us/HT210176)
-				using (X509Certificate2 Certificate = Request.CreateSelfSigned(new DateTimeOffset(DateTime.UtcNow.AddDays(-1)), new DateTimeOffset(DateTime.UtcNow.AddDays(800))))
+				using (X509Certificate2 certificate = request.CreateSelfSigned(new DateTimeOffset(DateTime.UtcNow.AddDays(-1)), new DateTimeOffset(DateTime.UtcNow.AddDays(800))))
 				{
-					Certificate.FriendlyName = FriendlyName;
-					byte[] PrivateCertData = Certificate.Export(X509ContentType.Pkcs12); // Note: Need to reimport this to use immediately, otherwise key is ephemeral
-					return PrivateCertData;
+					certificate.FriendlyName = friendlyName;
+					byte[] privateCertData = certificate.Export(X509ContentType.Pkcs12); // Note: Need to reimport this to use immediately, otherwise key is ephemeral
+					return privateCertData;
 				}
 			}
 		}
diff --git a/Engine/Source/Programs/Shared/EpicGames.Core/CityHash.cs b/Engine/Source/Programs/Shared/EpicGames.Core/CityHash.cs
--- a/Engine/Source/Programs/Shared/EpicGames.Core/CityHash.cs
+++ b/Engine/Source/Programs/Shared/EpicGames.Core/CityHash.cs
@@ -1,4 +1,6 @@
-﻿// Copyright (c) 2011 Google, Inc.
+﻿#pragma warning disable IDE0073
+
+// Copyright (c) 2011 Google, Inc.
 //
 // Permission is hereby granted, free of charge, to any person obtaining a copy
 // of this software and associated documentation files (the "Software"), to deal
@@ -31,6 +33,9 @@ using System;
 
 namespace EpicGames.Core
 {
+#pragma warning disable IDE0054
+#pragma warning disable IDE1006
+
 	public class CityHash
 	{
 		// Some primes between 2^63 and 2^64 for various uses.
@@ -420,9 +425,9 @@ namespace EpicGames.Core
 			return CityHash64WithSeeds(s, len, k2, seed);
 		}
 
-		public ulong CityHash64WithSeed(ReadOnlySpan<byte> Span, ulong seed)
+		public ulong CityHash64WithSeed(ReadOnlySpan<byte> span, ulong seed)
 		{
-			return CityHash64WithSeeds(Span, k2, seed);
+			return CityHash64WithSeeds(span, k2, seed);
 		}
 
 		public static unsafe ulong CityHash64WithSeeds(byte* s, uint len, ulong seed0, ulong seed1)
@@ -456,3 +461,5 @@ namespace EpicGames.Core
 		}
 	}
 }
+#pragma warning restore IDE1006
+#pragma warning restore IDE0054
diff --git a/Engine/Source/Programs/Shared/EpicGames.Core/CommandHost.cs b/Engine/Source/Programs/Shared/EpicGames.Core/CommandHost.cs
--- a/Engine/Source/Programs/Shared/EpicGames.Core/CommandHost.cs
+++ b/Engine/Source/Programs/Shared/EpicGames.Core/CommandHost.cs
@@ -1,13 +1,12 @@
 // Copyright Epic Games, Inc. All Rights Reserved.
 
-using Microsoft.Extensions.DependencyInjection;
-using Microsoft.Extensions.Logging;
 using System;
 using System.Collections.Generic;
 using System.Linq;
 using System.Reflection;
-using System.Text;
 using System.Threading.Tasks;
+using Microsoft.Extensions.DependencyInjection;
+using Microsoft.Extensions.Logging;
 
 namespace EpicGames.Core
 {
@@ -19,22 +18,22 @@ namespace EpicGames.Core
 		/// <summary>
 		/// Configure this object with the given command line arguments
 		/// </summary>
-		/// <param name="Arguments">Command line arguments</param>
-		void Configure(CommandLineArguments Arguments, ILogger Logger);
+		/// <param name="arguments">Command line arguments</param>
+		void Configure(CommandLineArguments arguments, ILogger logger);
 
 		/// <summary>
 		/// Gets all command line parameters to show in help for this command
 		/// </summary>
-		/// <param name="Arguments">The command line arguments</param>
+		/// <param name="arguments">The command line arguments</param>
 		/// <returns>List of name/description pairs</returns>
-		List<KeyValuePair<string, string>> GetParameters(CommandLineArguments Arguments);
+		List<KeyValuePair<string, string>> GetParameters(CommandLineArguments arguments);
 
 		/// <summary>
 		/// Execute this command
 		/// </summary>
-		/// <param name="Logger">The logger to use for this command</param>
+		/// <param name="logger">The logger to use for this command</param>
 		/// <returns>Exit code</returns>
-		Task<int> ExecuteAsync(ILogger Logger);
+		Task<int> ExecuteAsync(ILogger logger);
 	}
 
 	/// <summary>
@@ -55,7 +54,7 @@ namespace EpicGames.Core
 		/// <summary>
 		/// Create a command instance
 		/// </summary>
-		public ICommand CreateInstance(IServiceProvider ServiceProvider);
+		public ICommand CreateInstance(IServiceProvider serviceProvider);
 	}
 
 	/// <summary>
@@ -66,34 +65,34 @@ namespace EpicGames.Core
 		/// <summary>
 		/// Names for this command
 		/// </summary>
-		public string[] Names;
+		public string[] Names { get; }
 
 		/// <summary>
 		/// Short description for the mode. Will be displayed in the help text.
 		/// </summary>
-		public string Description;
+		public string Description { get; }
 
 		/// <summary>
 		/// Constructor
 		/// </summary>
-		/// <param name="Name">Name of the mode</param>
-		/// <param name="Description">Short description for display in the help text</param>
-		public CommandAttribute(string Name, string Description)
+		/// <param name="name">Name of the mode</param>
+		/// <param name="description">Short description for display in the help text</param>
+		public CommandAttribute(string name, string description)
 		{
-			this.Names = new string[] { Name };
-			this.Description = Description;
+			Names = new string[] { name };
+			Description = description;
 		}
 
 		/// <summary>
 		/// Constructor
 		/// </summary>
-		/// <param name="Category">Category for this command</param>
-		/// <param name="Name">Name of the mode</param>
-		/// <param name="Description">Short description for display in the help text</param>
-		public CommandAttribute(string Category, string Name, string Description)
+		/// <param name="category">Category for this command</param>
+		/// <param name="name">Name of the mode</param>
+		/// <param name="description">Short description for display in the help text</param>
+		public CommandAttribute(string category, string name, string description)
 		{
-			this.Names = new string[] { Category, Name };
-			this.Description = Description;
+			Names = new string[] { category, name };
+			Description = description;
 		}
 	}
 
@@ -103,19 +102,19 @@ namespace EpicGames.Core
 	public abstract class Command : ICommand
 	{
 		/// <inheritdoc/>
-		public virtual void Configure(CommandLineArguments Arguments, ILogger Logger)
+		public virtual void Configure(CommandLineArguments arguments, ILogger logger)
 		{
-			Arguments.ApplyTo(this, Logger);
+			arguments.ApplyTo(this, logger);
 		}
 
 		/// <inheritdoc/>
-		public virtual List<KeyValuePair<string, string>> GetParameters(CommandLineArguments Arguments)
+		public virtual List<KeyValuePair<string, string>> GetParameters(CommandLineArguments arguments)
 		{
 			return CommandLineArguments.GetParameters(GetType());
 		}
 
 		/// <inheritdoc/>
-		public abstract Task<int> ExecuteAsync(ILogger Logger);
+		public abstract Task<int> ExecuteAsync(ILogger logger);
 	}
 
 	/// <summary>
@@ -126,16 +125,16 @@ namespace EpicGames.Core
 		public string[] Names { get; }
 		public string Description { get; }
 
-		public Type Type;
+		public Type Type { get; }
 
-		public CommandFactory(string[] Names, string Description, Type Type)
+		public CommandFactory(string[] names, string description, Type type)
 		{
-			this.Names = Names;
-			this.Description = Description;
-			this.Type = Type;
+			Names = names;
+			Description = description;
+			Type = type;
 		}
 
-		public ICommand CreateInstance(IServiceProvider ServiceProvider) => (ICommand)ServiceProvider.GetRequiredService(Type);
+		public ICommand CreateInstance(IServiceProvider serviceProvider) => (ICommand)serviceProvider.GetRequiredService(Type);
 		public override string ToString() => String.Join(" ", Names);
 	}
 
@@ -147,20 +146,20 @@ namespace EpicGames.Core
 		/// <summary>
 		/// Adds services for executing the 
 		/// </summary>
-		/// <param name="Services"></param>
-		/// <param name="Assembly"></param>
-		public static void AddCommandsFromAssembly(this IServiceCollection Services, Assembly Assembly)
+		/// <param name="services"></param>
+		/// <param name="assembly"></param>
+		public static void AddCommandsFromAssembly(this IServiceCollection services, Assembly assembly)
 		{
-			List<(CommandAttribute, Type)> Commands = new List<(CommandAttribute, Type)>();
-			foreach (Type Type in Assembly.GetTypes())
+			List<(CommandAttribute, Type)> commands = new List<(CommandAttribute, Type)>();
+			foreach (Type type in assembly.GetTypes())
 			{
-				if (typeof(ICommand).IsAssignableFrom(Type) && !Type.IsAbstract)
+				if (typeof(ICommand).IsAssignableFrom(type) && !type.IsAbstract)
 				{
-					CommandAttribute? Attribute = Type.GetCustomAttribute<CommandAttribute>();
-					if (Attribute != null)
+					CommandAttribute? attribute = type.GetCustomAttribute<CommandAttribute>();
+					if (attribute != null)
 					{
-						Services.AddTransient(Type);
-						Services.AddTransient(typeof(ICommandFactory), SP => new CommandFactory(Attribute.Names, Attribute.Description, Type));
+						services.AddTransient(type);
+						services.AddTransient(typeof(ICommandFactory), sp => new CommandFactory(attribute.Names, attribute.Description, type));
 					}
 				}
 			}
@@ -169,31 +168,31 @@ namespace EpicGames.Core
 		/// <summary>
 		/// Entry point for executing registered command types in a particular assembly
 		/// </summary>
-		/// <param name="Args">Command line arguments</param>
-		/// <param name="ServiceProvider">The service provider for the application</param>
-		/// <param name="DefaultCommandType">The default command type</param>
+		/// <param name="args">Command line arguments</param>
+		/// <param name="serviceProvider">The service provider for the application</param>
+		/// <param name="defaultCommandType">The default command type</param>
 		/// <returns>Return code from the command</returns>
-		public static async Task<int> RunAsync(CommandLineArguments Args, IServiceProvider ServiceProvider, Type? DefaultCommandType)
+		public static async Task<int> RunAsync(CommandLineArguments args, IServiceProvider serviceProvider, Type? defaultCommandType)
 		{
 			// Find all the command types
-			List<ICommandFactory> CommandFactories = ServiceProvider.GetServices<ICommandFactory>().ToList();
+			List<ICommandFactory> commandFactories = serviceProvider.GetServices<ICommandFactory>().ToList();
 
 			// Check if there's a matching command
-			ICommand? Command = null;
-			ICommandFactory? CommandFactory = null;
+			ICommand? command = null;
+			ICommandFactory? commandFactory = null;
 
 			// Parse the positional arguments for the command name
-			string[] PositionalArgs = Args.GetPositionalArguments();
-			if (PositionalArgs.Length == 0)
+			string[] positionalArgs = args.GetPositionalArguments();
+			if (positionalArgs.Length == 0)
 			{
-				if (DefaultCommandType == null || Args.HasOption("-Help"))
+				if (defaultCommandType == null || args.HasOption("-Help"))
 				{
 					Console.WriteLine("Usage:");
 					Console.WriteLine("    [Command] [-Option1] [-Option2]...");
 					Console.WriteLine("");
 					Console.WriteLine("Commands:");
 
-					PrintCommands(CommandFactories);
+					PrintCommands(commandFactories);
 
 					Console.WriteLine("");
 					Console.WriteLine("Specify \"<CommandName> -Help\" for command-specific help");
@@ -201,75 +200,75 @@ namespace EpicGames.Core
 				}
 				else 
 				{
-					Command = (ICommand)ServiceProvider.GetService(DefaultCommandType);
+					command = (ICommand)serviceProvider.GetService(defaultCommandType);
 				}
 			}
 			else
 			{
-				foreach (ICommandFactory Factory in CommandFactories)
+				foreach (ICommandFactory factory in commandFactories)
 				{
-					if (Factory.Names.SequenceEqual(PositionalArgs, StringComparer.OrdinalIgnoreCase))
+					if (factory.Names.SequenceEqual(positionalArgs, StringComparer.OrdinalIgnoreCase))
 					{
-						Command = Factory.CreateInstance(ServiceProvider);
-						CommandFactory = Factory;
+						command = factory.CreateInstance(serviceProvider);
+						commandFactory = factory;
 						break;
 					}
 				}
-				if (Command == null)
+				if (command == null)
 				{
-					ConsoleUtils.WriteError($"Invalid command '{String.Join(" ", PositionalArgs)}'");
+					ConsoleUtils.WriteError($"Invalid command '{String.Join(" ", positionalArgs)}'");
 					Console.WriteLine("");
 					Console.WriteLine("Available commands:");
 
-					PrintCommands(CommandFactories);
+					PrintCommands(commandFactories);
 					return 1;
 				}
 			}
 
 			// If the help flag is specified, print the help info and exit immediately
-			if (Args.HasOption("-Help"))
+			if (args.HasOption("-Help"))
 			{
-				if (CommandFactory == null)
+				if (commandFactory == null)
 				{
-					HelpUtils.PrintHelp(null, null, Command.GetParameters(Args));
+					HelpUtils.PrintHelp(null, null, command.GetParameters(args));
 				}
 				else
 				{
-					HelpUtils.PrintHelp(String.Join(" ", CommandFactory.Names), CommandFactory.Description, Command.GetParameters(Args));
+					HelpUtils.PrintHelp(String.Join(" ", commandFactory.Names), commandFactory.Description, command.GetParameters(args));
 				}
 				return 1;
 			}
 
 			// Configure the command
-			ILogger Logger = ServiceProvider.GetRequiredService<ILoggerProvider>().CreateLogger("CommandHost");
+			ILogger logger = serviceProvider.GetRequiredService<ILoggerProvider>().CreateLogger("CommandHost");
 			try
 			{
-				Command.Configure(Args, Logger);
-				Args.CheckAllArgumentsUsed(Logger);
+				command.Configure(args, logger);
+				args.CheckAllArgumentsUsed(logger);
 			}
-			catch (CommandLineArgumentException Ex)
+			catch (CommandLineArgumentException ex)
 			{
-				ConsoleUtils.WriteError(Ex.Message);
+				ConsoleUtils.WriteError(ex.Message);
 				Console.WriteLine("");
 				Console.WriteLine("Valid parameters:");
 
-				HelpUtils.PrintTable(Command.GetParameters(Args), 4, 24);
+				HelpUtils.PrintTable(command.GetParameters(args), 4, 24);
 				return 1;
 			}
 
 			// Execute all the commands
 			try
 			{
-				return await Command.ExecuteAsync(Logger);
+				return await command.ExecuteAsync(logger);
 			}
-			catch (FatalErrorException Ex)
+			catch (FatalErrorException ex)
 			{
-				Logger.LogCritical(Ex, "Fatal error.");
-				return Ex.ExitCode;
+				logger.LogCritical(ex, "Fatal error.");
+				return ex.ExitCode;
 			}
-			catch (Exception Ex)
+			catch (Exception ex)
 			{
-				Logger.LogCritical(Ex, "Fatal error.");
+				logger.LogCritical(ex, "Fatal error.");
 				return 1;
 			}
 		}
@@ -277,15 +276,15 @@ namespace EpicGames.Core
 		/// <summary>
 		/// Print a formatted list of all the available commands
 		/// </summary>
-		/// <param name="Attributes">List of command attributes</param>
-		static void PrintCommands(IEnumerable<ICommandFactory> Attributes)
+		/// <param name="attributes">List of command attributes</param>
+		static void PrintCommands(IEnumerable<ICommandFactory> attributes)
 		{
-			List<KeyValuePair<string, string>> Commands = new List<KeyValuePair<string, string>>();
-			foreach (ICommandFactory Attribute in Attributes)
+			List<KeyValuePair<string, string>> commands = new List<KeyValuePair<string, string>>();
+			foreach (ICommandFactory attribute in attributes)
 			{
-				Commands.Add(new KeyValuePair<string, string>(String.Join(" ", Attribute.Names), Attribute.Description));
+				commands.Add(new KeyValuePair<string, string>(String.Join(" ", attribute.Names), attribute.Description));
 			}
-			HelpUtils.PrintTable(Commands.OrderBy(x => x.Key).ToList(), 4, 20);
+			HelpUtils.PrintTable(commands.OrderBy(x => x.Key).ToList(), 4, 20);
 		}
 	}
 }
