Reformat EpicGames.Core according to standard coding conventions.

diff --git a/Engine/Source/Programs/Shared/EpicGames.Core/CommandLineArguments.cs b/Engine/Source/Programs/Shared/EpicGames.Core/CommandLineArguments.cs
--- a/Engine/Source/Programs/Shared/EpicGames.Core/CommandLineArguments.cs
+++ b/Engine/Source/Programs/Shared/EpicGames.Core/CommandLineArguments.cs
@@ -1,6 +1,5 @@
 // Copyright Epic Games, Inc. All Rights Reserved.
 
-using Microsoft.Extensions.Logging;
 using System;
 using System.Collections;
 using System.Collections.Generic;
@@ -10,7 +9,7 @@ using System.Diagnostics.CodeAnalysis;
 using System.Linq;
 using System.Reflection;
 using System.Text;
-using System.Threading.Tasks;
+using Microsoft.Extensions.Logging;
 
 namespace EpicGames.Core
 {
@@ -23,15 +22,15 @@ namespace EpicGames.Core
 		/// The list of arguments
 		/// </summary>
 		[DebuggerBrowsable(DebuggerBrowsableState.RootHidden)]
-		public string[] Arguments;
+		public string[] Arguments { get; }
 
 		/// <summary>
 		/// Constructor
 		/// </summary>
-		/// <param name="ArgumentList">The argument list to proxy</param>
-		public CommandLineArgumentListView(CommandLineArguments ArgumentList)
+		/// <param name="argumentList">The argument list to proxy</param>
+		public CommandLineArgumentListView(CommandLineArguments argumentList)
 		{
-			Arguments = ArgumentList.GetRawArray();
+			Arguments = argumentList.GetRawArray();
 		}
 	}
 
@@ -43,19 +42,19 @@ namespace EpicGames.Core
 		/// <summary>
 		/// Constructor
 		/// </summary>
-		/// <param name="Message">Message to display for this exception</param>
-		public CommandLineArgumentException(string Message)
-			: base(Message)
+		/// <param name="message">Message to display for this exception</param>
+		public CommandLineArgumentException(string message)
+			: base(message)
 		{
 		}
 
 		/// <summary>
 		/// Constructor
 		/// </summary>
-		/// <param name="Message">Message to display for this exception</param>
-		/// <param name="InnerException">The inner exception</param>
-		public CommandLineArgumentException(string Message, Exception InnerException)
-			: base(Message, InnerException)
+		/// <param name="message">Message to display for this exception</param>
+		/// <param name="innerException">The inner exception</param>
+		public CommandLineArgumentException(string message, Exception innerException)
+			: base(message, innerException)
 		{
 		}
 
@@ -90,127 +89,127 @@ namespace EpicGames.Core
 			public Func<object?, object?> GetValue { get; }
 			public CommandLineAttribute[] Attributes { get; }
 
-			public ArgumentTarget(MemberInfo Member, Type ValueType, Action<object?, object?> SetValue, Func<object?, object?> GetValue, CommandLineAttribute[] Attributes)
+			public ArgumentTarget(MemberInfo member, Type valueType, Action<object?, object?> setValue, Func<object?, object?> getValue, CommandLineAttribute[] attributes)
 			{
-				this.Member = Member;
-				this.ValueType = ValueType;
-				this.SetValue = SetValue;
-				this.GetValue = GetValue;
-				this.Attributes = Attributes;
+				Member = member;
+				ValueType = valueType;
+				SetValue = setValue;
+				GetValue = getValue;
+				Attributes = attributes;
 			}
 		}
 
 		/// <summary>
 		/// The raw array of arguments
 		/// </summary>
-		string[] Arguments;
+		readonly string[] _arguments;
 
 		/// <summary>
 		/// Bitmask indicating which arguments are flags rather than values
 		/// </summary>
-		BitArray FlagArguments;
+		readonly BitArray _flagArguments;
 
 		/// <summary>
 		/// Bitmask indicating which arguments have been used, via calls to GetOption(), GetValues() etc...
 		/// </summary>
-		BitArray UsedArguments;
+		readonly BitArray _usedArguments;
 
 		/// <summary>
 		/// Dictionary of argument names (or prefixes, in the case of "-Foo=" style arguments) to their index into the arguments array.
 		/// </summary>
-		Dictionary<string, int> ArgumentToFirstIndex;
+		readonly Dictionary<string, int> _argumentToFirstIndex;
 
 		/// <summary>
 		/// For each argument which is seen more than once, keeps a list of indices for the second and subsequent arguments.
 		/// </summary>
-		int[] NextArgumentIndex;
+		readonly int[] _nextArgumentIndex;
 
 		/// <summary>
 		/// List of positional arguments
 		/// </summary>
-		List<int> PositionalArgumentIndices = new List<int>();
+		readonly List<int> _positionalArgumentIndices = new List<int>();
 
 		/// <summary>
 		/// Array of characters that separate argument names from values
 		/// </summary>
-		static readonly char[] ValueSeparators = { '=', ':' };
+		static readonly char[] s_valueSeparators = { '=', ':' };
 
 		/// <summary>
 		/// Constructor
 		/// </summary>
-		/// <param name="Arguments">The raw list of arguments</param>
-		public CommandLineArguments(string[] Arguments)
+		/// <param name="arguments">The raw list of arguments</param>
+		public CommandLineArguments(string[] arguments)
 		{
-			this.Arguments = Arguments;
-			this.FlagArguments = new BitArray(Arguments.Length);
-			this.UsedArguments = new BitArray(Arguments.Length);
+			_arguments = arguments;
+			_flagArguments = new BitArray(arguments.Length);
+			_usedArguments = new BitArray(arguments.Length);
 
 			// Clear the linked list of identical arguments
-			NextArgumentIndex = new int[Arguments.Length];
-			for(int Idx = 0; Idx < Arguments.Length; Idx++)
+			_nextArgumentIndex = new int[arguments.Length];
+			for(int idx = 0; idx < arguments.Length; idx++)
 			{
-				NextArgumentIndex[Idx] = -1;
+				_nextArgumentIndex[idx] = -1;
 			}
 
 			// Temporarily store the index of the last matching argument
-			int[] LastArgumentIndex = new int[Arguments.Length];
+			int[] lastArgumentIndex = new int[arguments.Length];
 
 			// Parse the argument array and build a lookup
-			ArgumentToFirstIndex = new Dictionary<string, int>(Arguments.Length, StringComparer.OrdinalIgnoreCase);
-			for(int Idx = 0; Idx < Arguments.Length; Idx++)
+			_argumentToFirstIndex = new Dictionary<string, int>(arguments.Length, StringComparer.OrdinalIgnoreCase);
+			for(int idx = 0; idx < arguments.Length; idx++)
 			{
-				if (Arguments[Idx].Equals("--", StringComparison.Ordinal))
+				if (arguments[idx].Equals("--", StringComparison.Ordinal))
 				{
 					// End of option arguments
-					MarkAsUsed(Idx++);
-					for (; Idx < Arguments.Length; Idx++)
+					MarkAsUsed(idx++);
+					for (; idx < arguments.Length; idx++)
 					{
-						PositionalArgumentIndices.Add(Idx);
+						_positionalArgumentIndices.Add(idx);
 					}
 					break;
 				}
-				else if (Arguments[Idx].StartsWith("-", StringComparison.Ordinal))
+				else if (arguments[idx].StartsWith("-", StringComparison.Ordinal))
 				{
 					// Option argument
-					int SeparatorIdx = Arguments[Idx].IndexOfAny(ValueSeparators);
-					if (SeparatorIdx == -1)
+					int separatorIdx = arguments[idx].IndexOfAny(s_valueSeparators);
+					if (separatorIdx == -1)
 					{
 						// Ignore duplicate -Option flags; they are harmless.
-						if (ArgumentToFirstIndex.ContainsKey(Arguments[Idx]))
+						if (_argumentToFirstIndex.ContainsKey(arguments[idx]))
 						{
-							UsedArguments.Set(Idx, true);
+							_usedArguments.Set(idx, true);
 						}
 						else
 						{
-							ArgumentToFirstIndex.Add(Arguments[Idx], Idx);
+							_argumentToFirstIndex.Add(arguments[idx], idx);
 						}
 
 						// Mark this argument as a flag
-						FlagArguments.Set(Idx, true);
+						_flagArguments.Set(idx, true);
 					}
 					else
 					{
 						// Just take the part up to and including the separator character
-						string Prefix = Arguments[Idx].Substring(0, SeparatorIdx + 1);
+						string prefix = arguments[idx].Substring(0, separatorIdx + 1);
 
 						// Add the prefix to the argument lookup, or update the appropriate matching argument list if it's been seen before
-						int ExistingArgumentIndex;
-						if (ArgumentToFirstIndex.TryGetValue(Prefix, out ExistingArgumentIndex))
+						int existingArgumentIndex;
+						if (_argumentToFirstIndex.TryGetValue(prefix, out existingArgumentIndex))
 						{
-							NextArgumentIndex[LastArgumentIndex[ExistingArgumentIndex]] = Idx;
-							LastArgumentIndex[ExistingArgumentIndex] = Idx;
+							_nextArgumentIndex[lastArgumentIndex[existingArgumentIndex]] = idx;
+							lastArgumentIndex[existingArgumentIndex] = idx;
 						}
 						else
 						{
-							ArgumentToFirstIndex.Add(Prefix, Idx);
-							LastArgumentIndex[Idx] = Idx;
+							_argumentToFirstIndex.Add(prefix, idx);
+							lastArgumentIndex[idx] = idx;
 						}
 					}
 				}
 				else
 				{
 					// Positional argument
-					PositionalArgumentIndices.Add(Idx);
+					_positionalArgumentIndices.Add(idx);
 				}
 			}
 		}
@@ -218,60 +217,54 @@ namespace EpicGames.Core
 		/// <summary>
 		/// The number of arguments in this list
 		/// </summary>
-		public int Count
-		{
-			get { return Arguments.Length; }
-		}
+		public int Count => _arguments.Length;
 
 		/// <summary>
 		/// Access an argument by index
 		/// </summary>
-		/// <param name="Index">Index of the argument</param>
+		/// <param name="index">Index of the argument</param>
 		/// <returns>The argument at the given index</returns>
-		public string this[int Index]
-		{
-			get { return Arguments[Index]; }
-		}
+		public string this[int index] => _arguments[index];
 
 		/// <summary>
 		/// Determines if an argument has been used
 		/// </summary>
-		/// <param name="Index">Index of the argument</param>
+		/// <param name="index">Index of the argument</param>
 		/// <returns>True if the argument has been used, false otherwise</returns>
-		public bool HasBeenUsed(int Index)
+		public bool HasBeenUsed(int index)
 		{
-			return UsedArguments.Get(Index);
+			return _usedArguments.Get(index);
 		}
 
 		/// <summary>
 		/// Marks an argument as having been used
 		/// </summary>
-		/// <param name="Index">Index of the argument to mark as used</param>
-		public void MarkAsUsed(int Index)
+		/// <param name="index">Index of the argument to mark as used</param>
+		public void MarkAsUsed(int index)
 		{
-			UsedArguments.Set(Index, true);
+			_usedArguments.Set(index, true);
 		}
 
 		/// <summary>
 		/// Marks an argument as not having been used
 		/// </summary>
-		/// <param name="Index">Index of the argument to mark as being unused</param>
-		public void MarkAsUnused(int Index)
+		/// <param name="index">Index of the argument to mark as being unused</param>
+		public void MarkAsUnused(int index)
 		{
-			UsedArguments.Set(Index, true);
+			_usedArguments.Set(index, true);
 		}
 
 		/// <summary>
 		/// Checks if the given option (eg. "-Foo") was specified on the command line.
 		/// </summary>
-		/// <param name="Option">The option to look for</param>
+		/// <param name="option">The option to look for</param>
 		/// <returns>True if the option was found, false otherwise.</returns>
-		public bool HasOption(string Option)
+		public bool HasOption(string option)
 		{
-			int Index;
-			if(ArgumentToFirstIndex.TryGetValue(Option, out Index))
+			int index;
+			if(_argumentToFirstIndex.TryGetValue(option, out index))
 			{
-				UsedArguments.Set(Index, true);
+				_usedArguments.Set(index, true);
 				return true;
 			}
 			return false;
@@ -280,12 +273,12 @@ namespace EpicGames.Core
 		/// <summary>
 		/// Checks for an argument prefixed with the given string is present.
 		/// </summary>
-		/// <param name="Prefix">The argument prefix (eg. "-Foo="). Must end with an '=' character.</param>
+		/// <param name="prefix">The argument prefix (eg. "-Foo="). Must end with an '=' character.</param>
 		/// <returns>True if an argument with the given prefix was specified</returns>
-		public bool HasValue(string Prefix)
+		public bool HasValue(string prefix)
 		{
-			CheckValidPrefix(Prefix);
-			return ArgumentToFirstIndex.ContainsKey(Prefix);
+			CheckValidPrefix(prefix);
+			return _argumentToFirstIndex.ContainsKey(prefix);
 		}
 
 		/// <summary>
@@ -294,7 +287,7 @@ namespace EpicGames.Core
 		/// <returns>Number of positional arguments</returns>
 		public int GetPositionalArgumentCount()
 		{
-			return PositionalArgumentIndices.Count;
+			return _positionalArgumentIndices.Count;
 		}
 
 		/// <summary>
@@ -302,30 +295,30 @@ namespace EpicGames.Core
 		/// </summary>
 		/// <param name="Index">Number of the positional argument</param>
 		/// <returns>Index of the positional argument</returns>
-		public int GetPositionalArgumentIndex(int Num)
+		public int GetPositionalArgumentIndex(int num)
 		{
-			return PositionalArgumentIndices[Num];
+			return _positionalArgumentIndices[num];
 		}
 
 		/// <summary>
 		/// Attempts to read the next unused positional argument
 		/// </summary>
-		/// <param name="Argument">Receives the argument that was read, on success</param>
+		/// <param name="argument">Receives the argument that was read, on success</param>
 		/// <returns>True if an argument was read</returns>
-		public bool TryGetPositionalArgument([NotNullWhen(true)] out string? Argument)
+		public bool TryGetPositionalArgument([NotNullWhen(true)] out string? argument)
 		{
-			for (int Idx = 0; Idx < PositionalArgumentIndices.Count; Idx++)
+			for (int idx = 0; idx < _positionalArgumentIndices.Count; idx++)
 			{
-				int Index = PositionalArgumentIndices[Idx];
-				if (!HasBeenUsed(Index))
+				int index = _positionalArgumentIndices[idx];
+				if (!HasBeenUsed(index))
 				{
-					MarkAsUsed(Index);
-					Argument = Arguments[Index];
+					MarkAsUsed(index);
+					argument = _arguments[index];
 					return true;
 				}
 			}
 
-			Argument = null;
+			argument = null;
 			return false;
 		}
 
@@ -335,329 +328,329 @@ namespace EpicGames.Core
 		/// <returns>Array of positional arguments</returns>
 		public string[] GetPositionalArguments()
 		{
-			string[] PositionalArguments = new string[PositionalArgumentIndices.Count];
-			for (int Idx = 0; Idx < PositionalArguments.Length; Idx++)
+			string[] positionalArguments = new string[_positionalArgumentIndices.Count];
+			for (int idx = 0; idx < positionalArguments.Length; idx++)
 			{
-				int Index = PositionalArgumentIndices[Idx];
-				MarkAsUsed(Index);
-				PositionalArguments[Idx] = Arguments[Index];
+				int index = _positionalArgumentIndices[idx];
+				MarkAsUsed(index);
+				positionalArguments[idx] = _arguments[index];
 			}
-			return PositionalArguments;
+			return positionalArguments;
 		}
 
 		/// <summary>
 		/// Gets the value specified by an argument with the given prefix. Throws an exception if the argument was not specified.
 		/// </summary>
-		/// <param name="Prefix">The argument prefix (eg. "-Foo="). Must end with an '=' character.</param>
+		/// <param name="prefix">The argument prefix (eg. "-Foo="). Must end with an '=' character.</param>
 		/// <returns>Value of the argument</returns>
-		public string GetString(string Prefix)
+		public string GetString(string prefix)
 		{
-			string? Value;
-			if(!TryGetValue(Prefix, out Value))
+			string? value;
+			if(!TryGetValue(prefix, out value))
 			{
-				throw new CommandLineArgumentException(String.Format("Missing '{0}...' argument", Prefix));
+				throw new CommandLineArgumentException(String.Format("Missing '{0}...' argument", prefix));
 			}
-			return Value;
+			return value;
 		}
 
 		/// <summary>
 		/// Gets the value specified by an argument with the given prefix. Throws an exception if the argument was not specified.
 		/// </summary>
-		/// <param name="Prefix">The argument prefix (eg. "-Foo="). Must end with an '=' character.</param>
+		/// <param name="prefix">The argument prefix (eg. "-Foo="). Must end with an '=' character.</param>
 		/// <returns>Value of the argument</returns>
-		public int GetInteger(string Prefix)
+		public int GetInteger(string prefix)
 		{
-			int Value;
-			if(!TryGetValue(Prefix, out Value))
+			int value;
+			if(!TryGetValue(prefix, out value))
 			{
-				throw new CommandLineArgumentException(String.Format("Missing '{0}...' argument", Prefix));
+				throw new CommandLineArgumentException(String.Format("Missing '{0}...' argument", prefix));
 			}
-			return Value;
+			return value;
 		}
 
 		/// <summary>
 		/// Gets the value specified by an argument with the given prefix. Throws an exception if the argument was not specified.
 		/// </summary>
-		/// <param name="Prefix">The argument prefix (eg. "-Foo="). Must end with an '=' character.</param>
+		/// <param name="prefix">The argument prefix (eg. "-Foo="). Must end with an '=' character.</param>
 		/// <returns>Value of the argument</returns>
-		public FileReference GetFileReference(string Prefix)
+		public FileReference GetFileReference(string prefix)
 		{
-			FileReference? Value;
-			if(!TryGetValue(Prefix, out Value))
+			FileReference? value;
+			if(!TryGetValue(prefix, out value))
 			{
-				throw new CommandLineArgumentException(String.Format("Missing '{0}...' argument", Prefix));
+				throw new CommandLineArgumentException(String.Format("Missing '{0}...' argument", prefix));
 			}
-			return Value;
+			return value;
 		}
 
 		/// <summary>
 		/// Gets the value specified by an argument with the given prefix. Throws an exception if the argument was not specified.
 		/// </summary>
-		/// <param name="Prefix">The argument prefix (eg. "-Foo="). Must end with an '=' character.</param>
+		/// <param name="prefix">The argument prefix (eg. "-Foo="). Must end with an '=' character.</param>
 		/// <returns>Value of the argument</returns>
-		public DirectoryReference GetDirectoryReference(string Prefix)
+		public DirectoryReference GetDirectoryReference(string prefix)
 		{
-			DirectoryReference? Value;
-			if(!TryGetValue(Prefix, out Value))
+			DirectoryReference? value;
+			if(!TryGetValue(prefix, out value))
 			{
-				throw new CommandLineArgumentException(String.Format("Missing '{0}...' argument", Prefix));
+				throw new CommandLineArgumentException(String.Format("Missing '{0}...' argument", prefix));
 			}
-			return Value;
+			return value;
 		}
 
 		/// <summary>
 		/// Gets the value specified by an argument with the given prefix. Throws an exception if the argument was not specified.
 		/// </summary>
-		/// <param name="Prefix">The argument prefix (eg. "-Foo="). Must end with an '=' character.</param>
+		/// <param name="prefix">The argument prefix (eg. "-Foo="). Must end with an '=' character.</param>
 		/// <returns>Value of the argument</returns>
-		public T GetEnum<T>(string Prefix) where T : struct
+		public T GetEnum<T>(string prefix) where T : struct
 		{
-			T Value;
-			if(!TryGetValue(Prefix, out Value))
+			T value;
+			if(!TryGetValue(prefix, out value))
 			{
-				throw new CommandLineArgumentException(String.Format("Missing '{0}...' argument", Prefix));
+				throw new CommandLineArgumentException(String.Format("Missing '{0}...' argument", prefix));
 			}
-			return Value;
+			return value;
 		}
 
 		/// <summary>
 		/// Gets the value specified by an argument with the given prefix, or a default value.
 		/// </summary>
-		/// <param name="Prefix">The argument prefix (eg. "-Foo="). Must end with an '=' character.</param>
-		/// <param name="DefaultValue">Default value for the argument</param>
+		/// <param name="prefix">The argument prefix (eg. "-Foo="). Must end with an '=' character.</param>
+		/// <param name="defaultValue">Default value for the argument</param>
 		/// <returns>Value of the argument</returns>
-		[return: NotNullIfNotNull("DefaultValue")]
-		public string? GetStringOrDefault(string Prefix, string? DefaultValue)
+		[return: NotNullIfNotNull("defaultValue")]
+		public string? GetStringOrDefault(string prefix, string? defaultValue)
 		{
-			string? Value;
-			if(!TryGetValue(Prefix, out Value))
+			string? value;
+			if(!TryGetValue(prefix, out value))
 			{
-				Value = DefaultValue;
+				value = defaultValue;
 			}
-			return Value;
+			return value;
 		}
 
 		/// <summary>
 		/// Gets the value specified by an argument with the given prefix, or a default value.
 		/// </summary>
-		/// <param name="Prefix">The argument prefix (eg. "-Foo="). Must end with an '=' character.</param>
-		/// <param name="DefaultValue">Default value for the argument</param>
+		/// <param name="prefix">The argument prefix (eg. "-Foo="). Must end with an '=' character.</param>
+		/// <param name="defaultValue">Default value for the argument</param>
 		/// <returns>Value of the argument</returns>
-		public int GetIntegerOrDefault(string Prefix, int DefaultValue)
+		public int GetIntegerOrDefault(string prefix, int defaultValue)
 		{
-			int Value;
-			if(!TryGetValue(Prefix, out Value))
+			int value;
+			if(!TryGetValue(prefix, out value))
 			{
-				Value = DefaultValue;
+				value = defaultValue;
 			}
-			return Value;
+			return value;
 		}
 
 		/// <summary>
 		/// Gets the value specified by an argument with the given prefix, or a default value.
 		/// </summary>
-		/// <param name="Prefix">The argument prefix (eg. "-Foo="). Must end with an '=' character.</param>
-		/// <param name="DefaultValue">Default value for the argument</param>
+		/// <param name="prefix">The argument prefix (eg. "-Foo="). Must end with an '=' character.</param>
+		/// <param name="defaultValue">Default value for the argument</param>
 		/// <returns>Value of the argument</returns>
-		[return: NotNullIfNotNull("DefaultValue")]
-		public FileReference? GetFileReferenceOrDefault(string Prefix, FileReference? DefaultValue)
+		[return: NotNullIfNotNull("defaultValue")]
+		public FileReference? GetFileReferenceOrDefault(string prefix, FileReference? defaultValue)
 		{
-			FileReference? Value;
-			if(!TryGetValue(Prefix, out Value))
+			FileReference? value;
+			if(!TryGetValue(prefix, out value))
 			{
-				Value = DefaultValue;
+				value = defaultValue;
 			}
-			return Value;
+			return value;
 		}
 
 		/// <summary>
 		/// Gets the value specified by an argument with the given prefix, or a default value.
 		/// </summary>
-		/// <param name="Prefix">The argument prefix (eg. "-Foo="). Must end with an '=' character.</param>
-		/// <param name="DefaultValue">Default value for the argument</param>
+		/// <param name="prefix">The argument prefix (eg. "-Foo="). Must end with an '=' character.</param>
+		/// <param name="defaultValue">Default value for the argument</param>
 		/// <returns>Value of the argument</returns>
-		[return: NotNullIfNotNull("DefaultValue")]
-		public DirectoryReference? GetDirectoryReferenceOrDefault(string Prefix, DirectoryReference? DefaultValue)
+		[return: NotNullIfNotNull("defaultValue")]
+		public DirectoryReference? GetDirectoryReferenceOrDefault(string prefix, DirectoryReference? defaultValue)
 		{
-			DirectoryReference? Value;
-			if(!TryGetValue(Prefix, out Value))
+			DirectoryReference? value;
+			if(!TryGetValue(prefix, out value))
 			{
-				Value = DefaultValue;
+				value = defaultValue;
 			}
-			return Value;
+			return value;
 		}
 
 		/// <summary>
 		/// Gets the value specified by an argument with the given prefix, or a default value.
 		/// </summary>
-		/// <param name="Prefix">The argument prefix (eg. "-Foo="). Must end with an '=' character.</param>
-		/// <param name="DefaultValue">Default value for the argument</param>
+		/// <param name="prefix">The argument prefix (eg. "-Foo="). Must end with an '=' character.</param>
+		/// <param name="defaultValue">Default value for the argument</param>
 		/// <returns>Value of the argument</returns>
-		public T GetEnumOrDefault<T>(string Prefix, T DefaultValue) where T : struct
+		public T GetEnumOrDefault<T>(string prefix, T defaultValue) where T : struct
 		{
-			T Value;
-			if(!TryGetValue(Prefix, out Value))
+			T value;
+			if(!TryGetValue(prefix, out value))
 			{
-				Value = DefaultValue;
+				value = defaultValue;
 			}
-			return Value;
+			return value;
 		}
 
 		/// <summary>
 		/// Tries to gets the value specified by an argument with the given prefix.
 		/// </summary>
-		/// <param name="Prefix">The argument prefix (eg. "-Foo="). Must end with an '=' character.</param>
-		/// <param name="Value">Value of the argument, if found</param>
+		/// <param name="prefix">The argument prefix (eg. "-Foo="). Must end with an '=' character.</param>
+		/// <param name="value">Value of the argument, if found</param>
 		/// <returns>True if the argument was found (and Value was set), false otherwise.</returns>
-		public bool TryGetValue(string Prefix, [NotNullWhen(true)] out string? Value)
+		public bool TryGetValue(string prefix, [NotNullWhen(true)] out string? value)
 		{
-			CheckValidPrefix(Prefix);
+			CheckValidPrefix(prefix);
 
-			int Index;
-			if(!ArgumentToFirstIndex.TryGetValue(Prefix, out Index))
+			int index;
+			if(!_argumentToFirstIndex.TryGetValue(prefix, out index))
 			{
-				Value = null;
+				value = null;
 				return false;
 			}
 
-			if(NextArgumentIndex[Index] != -1)
+			if(_nextArgumentIndex[index] != -1)
 			{
-				throw new CommandLineArgumentException(String.Format("Multiple {0}... arguments are specified", Prefix));
+				throw new CommandLineArgumentException(String.Format("Multiple {0}... arguments are specified", prefix));
 			}
 
-			UsedArguments.Set(Index, true);
-			Value = Arguments[Index].Substring(Prefix.Length);
+			_usedArguments.Set(index, true);
+			value = _arguments[index].Substring(prefix.Length);
 			return true;
 		}
 
 		/// <summary>
 		/// Tries to gets the value specified by an argument with the given prefix.
 		/// </summary>
-		/// <param name="Prefix">The argument prefix (eg. "-Foo="). Must end with an '=' character.</param>
-		/// <param name="Value">Value of the argument, if found</param>
+		/// <param name="prefix">The argument prefix (eg. "-Foo="). Must end with an '=' character.</param>
+		/// <param name="value">Value of the argument, if found</param>
 		/// <returns>True if the argument was found (and Value was set), false otherwise.</returns>
-		public bool TryGetValue(string Prefix, out int Value)
+		public bool TryGetValue(string prefix, out int value)
 		{
 			// Try to get the string value of this argument
-			string? StringValue;
-			if(!TryGetValue(Prefix, out StringValue))
+			string? stringValue;
+			if(!TryGetValue(prefix, out stringValue))
 			{
-				Value = 0;
+				value = 0;
 				return false;
 			}
 
 			// Try to parse it. If it fails, throw an exception.
 			try
 			{
-				Value = int.Parse(StringValue);
+				value = Int32.Parse(stringValue);
 				return true;
 			}
-			catch(Exception Ex)
+			catch(Exception ex)
 			{
-				throw new CommandLineArgumentException(String.Format("The argument '{0}{1}' does not specify a valid integer", Prefix, StringValue), Ex);
+				throw new CommandLineArgumentException(String.Format("The argument '{0}{1}' does not specify a valid integer", prefix, stringValue), ex);
 			}
 		}
 
 		/// <summary>
 		/// Tries to gets the value specified by an argument with the given prefix.
 		/// </summary>
-		/// <param name="Prefix">The argument prefix (eg. "-Foo="). Must end with an '=' character.</param>
-		/// <param name="Value">Value of the argument, if found</param>
+		/// <param name="prefix">The argument prefix (eg. "-Foo="). Must end with an '=' character.</param>
+		/// <param name="value">Value of the argument, if found</param>
 		/// <returns>True if the argument was found (and Value was set), false otherwise.</returns>
-		public bool TryGetValue(string Prefix, [NotNullWhen(true)] out FileReference? Value)
+		public bool TryGetValue(string prefix, [NotNullWhen(true)] out FileReference? value)
 		{
 			// Try to get the string value of this argument
-			string? StringValue;
-			if(!TryGetValue(Prefix, out StringValue))
+			string? stringValue;
+			if(!TryGetValue(prefix, out stringValue))
 			{
-				Value = null;
+				value = null;
 				return false;
 			}
 
 			// Try to parse it. If it fails, throw an exception.
 			try
 			{
-				Value = new FileReference(StringValue);
+				value = new FileReference(stringValue);
 				return true;
 			}
-			catch(Exception Ex)
+			catch(Exception ex)
 			{
-				throw new CommandLineArgumentException(String.Format("The argument '{0}{1}' does not specify a valid file name", Prefix, StringValue), Ex);
+				throw new CommandLineArgumentException(String.Format("The argument '{0}{1}' does not specify a valid file name", prefix, stringValue), ex);
 			}
 		}
 
 		/// <summary>
 		/// Tries to gets the value specified by an argument with the given prefix.
 		/// </summary>
-		/// <param name="Prefix">The argument prefix (eg. "-Foo="). Must end with an '=' character.</param>
-		/// <param name="Value">Value of the argument, if found</param>
+		/// <param name="prefix">The argument prefix (eg. "-Foo="). Must end with an '=' character.</param>
+		/// <param name="value">Value of the argument, if found</param>
 		/// <returns>True if the argument was found (and Value was set), false otherwise.</returns>
-		public bool TryGetValue(string Prefix, [NotNullWhen(true)] out DirectoryReference? Value)
+		public bool TryGetValue(string prefix, [NotNullWhen(true)] out DirectoryReference? value)
 		{
 			// Try to get the string value of this argument
-			string? StringValue;
-			if(!TryGetValue(Prefix, out StringValue))
+			string? stringValue;
+			if(!TryGetValue(prefix, out stringValue))
 			{
-				Value = null;
+				value = null;
 				return false;
 			}
 
 			// Try to parse it. If it fails, throw an exception.
 			try
 			{
-				Value = new DirectoryReference(StringValue);
+				value = new DirectoryReference(stringValue);
 				return true;
 			}
-			catch(Exception Ex)
+			catch(Exception ex)
 			{
-				throw new CommandLineArgumentException(String.Format("The argument '{0}{1}' does not specify a valid directory name", Prefix, StringValue), Ex);
+				throw new CommandLineArgumentException(String.Format("The argument '{0}{1}' does not specify a valid directory name", prefix, stringValue), ex);
 			}
 		}
 
 		/// <summary>
 		/// Tries to gets the value specified by an argument with the given prefix.
 		/// </summary>
-		/// <param name="Prefix">The argument prefix (eg. "-Foo="). Must end with an '=' character.</param>
-		/// <param name="Value">Value of the argument, if found</param>
+		/// <param name="prefix">The argument prefix (eg. "-Foo="). Must end with an '=' character.</param>
+		/// <param name="value">Value of the argument, if found</param>
 		/// <returns>True if the argument was found (and Value was set), false otherwise.</returns>
-		public bool TryGetValue<T>(string Prefix, out T Value) where T : struct
+		public bool TryGetValue<T>(string prefix, out T value) where T : struct
 		{
 			// Try to get the string value of this argument
-			string? StringValue;
-			if(!TryGetValue(Prefix, out StringValue))
+			string? stringValue;
+			if(!TryGetValue(prefix, out stringValue))
 			{
-				Value = new T();
+				value = new T();
 				return false;
 			}
 
 			// Try to parse it. If it fails, throw an exception.
 			try
 			{
-				Value = (T)Enum.Parse(typeof(T), StringValue, true);
+				value = (T)Enum.Parse(typeof(T), stringValue, true);
 				return true;
 			}
-			catch(Exception Ex)
+			catch(Exception ex)
 			{
-				throw new CommandLineArgumentException(String.Format("The argument '{0}{1}' does not specify a valid {2}", Prefix, StringValue, typeof(T).Name), Ex);
+				throw new CommandLineArgumentException(String.Format("The argument '{0}{1}' does not specify a valid {2}", prefix, stringValue, typeof(T).Name), ex);
 			}
 		}
 
 		/// <summary>
 		/// Returns all arguments with the given prefix.
 		/// </summary>
-		/// <param name="Prefix">The argument prefix (eg. "-Foo="). Must end with an '=' character.</param>
+		/// <param name="prefix">The argument prefix (eg. "-Foo="). Must end with an '=' character.</param>
 		/// <returns>Sequence of values for the given prefix.</returns>
-		public IEnumerable<string> GetValues(string Prefix)
+		public IEnumerable<string> GetValues(string prefix)
 		{
-			CheckValidPrefix(Prefix);
+			CheckValidPrefix(prefix);
 
-			int Index;
-			if(ArgumentToFirstIndex.TryGetValue(Prefix, out Index))
+			int index;
+			if(_argumentToFirstIndex.TryGetValue(prefix, out index))
 			{
-				for(; Index != -1; Index = NextArgumentIndex[Index])
+				for(; index != -1; index = _nextArgumentIndex[index])
 				{
-					UsedArguments.Set(Index, true);
-					yield return Arguments[Index].Substring(Prefix.Length);
+					_usedArguments.Set(index, true);
+					yield return _arguments[index].Substring(prefix.Length);
 				}
 			}
 		}
@@ -665,16 +658,16 @@ namespace EpicGames.Core
 		/// <summary>
 		/// Returns all arguments with the given prefix, allowing multiple arguments to be specified in a single argument with a separator character.
 		/// </summary>
-		/// <param name="Prefix">The argument prefix (eg. "-Foo="). Must end with an '=' character.</param>
-		/// <param name="Separator">The separator character (eg. '+')</param>
+		/// <param name="prefix">The argument prefix (eg. "-Foo="). Must end with an '=' character.</param>
+		/// <param name="separator">The separator character (eg. '+')</param>
 		/// <returns>Sequence of values for the given prefix.</returns>
-		public IEnumerable<string> GetValues(string Prefix, char Separator)
+		public IEnumerable<string> GetValues(string prefix, char separator)
 		{
-			foreach(string Value in GetValues(Prefix))
+			foreach(string value in GetValues(prefix))
 			{
-				foreach(string SplitValue in Value.Split(Separator))
+				foreach(string splitValue in value.Split(separator))
 				{
-					yield return SplitValue;
+					yield return splitValue;
 				}
 			}
 		}
@@ -682,146 +675,146 @@ namespace EpicGames.Core
 		/// <summary>
 		/// Gets the prefix for a particular argument
 		/// </summary>
-		/// <param name="Target">The target hosting the attribute</param>
-		/// <param name="Attribute">The attribute instance</param>
+		/// <param name="target">The target hosting the attribute</param>
+		/// <param name="attribute">The attribute instance</param>
 		/// <returns>Prefix for this argument</returns>
-		private static string GetArgumentPrefix(ArgumentTarget Target, CommandLineAttribute Attribute)
+		private static string GetArgumentPrefix(ArgumentTarget target, CommandLineAttribute attribute)
 		{
 			// Get the inner field type, unwrapping nullable types
-			Type ValueType = Target.ValueType;
-			if (ValueType.IsGenericType && ValueType.GetGenericTypeDefinition() == typeof(Nullable<>))
+			Type valueType = target.ValueType;
+			if (valueType.IsGenericType && valueType.GetGenericTypeDefinition() == typeof(Nullable<>))
 			{
-				ValueType = ValueType.GetGenericArguments()[0];
+				valueType = valueType.GetGenericArguments()[0];
 			}
 
-			string? Prefix = Attribute.Prefix;
-			if (Prefix == null)
+			string? prefix = attribute.Prefix;
+			if (prefix == null)
 			{
-				if (ValueType == typeof(bool))
+				if (valueType == typeof(bool))
 				{
-					Prefix = String.Format("-{0}", Target.Member.Name);
+					prefix = String.Format("-{0}", target.Member.Name);
 				}
 				else
 				{
-					Prefix = String.Format("-{0}=", Target.Member.Name);
+					prefix = String.Format("-{0}=", target.Member.Name);
 				}
 			}
 			else
 			{
-				if (ValueType != typeof(bool) && Attribute.Value == null && !Prefix.EndsWith("=") && !Prefix.EndsWith(":"))
+				if (valueType != typeof(bool) && attribute.Value == null && !prefix.EndsWith("=") && !prefix.EndsWith(":"))
 				{
-					Prefix = Prefix + "=";
+					prefix += "=";
 				}
 			}
-			return Prefix;
+			return prefix;
 		}
 
 		/// <summary>
 		/// Applies these arguments to fields with the [CommandLine] attribute in the given object.
 		/// </summary>
-		/// <param name="TargetObject">The object to configure</param>
-		public void ApplyTo(object TargetObject)
+		/// <param name="targetObject">The object to configure</param>
+		public void ApplyTo(object targetObject)
 		{
-			ApplyTo(TargetObject, Log.Logger);
+			ApplyTo(targetObject, Log.Logger);
 		}
 
 		/// <summary>
 		/// Applies these arguments to fields with the [CommandLine] attribute in the given object.
 		/// </summary>
-		/// <param name="TargetObject">The object to configure</param>
-		/// <param name="Logger">Sink for error/warning messages</param>
-		public void ApplyTo(object TargetObject, ILogger Logger)
+		/// <param name="targetObject">The object to configure</param>
+		/// <param name="logger">Sink for error/warning messages</param>
+		public void ApplyTo(object targetObject, ILogger logger)
 		{
-			List<string> MissingArguments = new List<string>();
+			List<string> missingArguments = new List<string>();
 
 			// Build a mapping from name to field and attribute for this object
-			List<ArgumentTarget> Targets = GetArgumentTargetsForType(TargetObject.GetType());
-			foreach (ArgumentTarget Target in Targets)
+			List<ArgumentTarget> targets = GetArgumentTargetsForType(targetObject.GetType());
+			foreach (ArgumentTarget target in targets)
 			{
 				// If any attribute is required, keep track of it so we can include an error for it
-				string? RequiredPrefix = null;
+				string? requiredPrefix = null;
 
 				// Keep track of whether a value has already been assigned to this field
-				string? AssignedArgument = null;
+				string? assignedArgument = null;
 
 				// Loop through all the attributes for different command line options that can modify it
-				foreach(CommandLineAttribute Attribute in Target.Attributes)
+				foreach(CommandLineAttribute attribute in target.Attributes)
 				{
 					// Get the appropriate prefix for this attribute
-					string Prefix = GetArgumentPrefix(Target, Attribute);
+					string prefix = GetArgumentPrefix(target, attribute);
 
 					// Get the value with the correct prefix
-					int FirstIndex;
-					if(ArgumentToFirstIndex.TryGetValue(Prefix, out FirstIndex))
+					int firstIndex;
+					if(_argumentToFirstIndex.TryGetValue(prefix, out firstIndex))
 					{
-						for(int Index = FirstIndex; Index != -1; Index = NextArgumentIndex[Index])
+						for(int index = firstIndex; index != -1; index = _nextArgumentIndex[index])
 						{
 							// Get the argument text
-							string Argument = Arguments[Index];
+							string argument = _arguments[index];
 
 							// Get the text for this value
-							string ValueText;
-							if(Attribute.Value != null)
+							string valueText;
+							if(attribute.Value != null)
 							{
-								ValueText = Attribute.Value;
+								valueText = attribute.Value;
 							}
-							else if(FlagArguments.Get(Index))
+							else if(_flagArguments.Get(index))
 							{
-								ValueText = "true";
+								valueText = "true";
 							}
 							else
 							{
-								ValueText = Argument.Substring(Prefix.Length);
+								valueText = argument.Substring(prefix.Length);
 							}
 
 							// Apply the value to the field
-							if(Attribute.ListSeparator == 0)
+							if(attribute.ListSeparator == 0)
 							{
-								if(ApplyArgument(TargetObject, Target, Argument, ValueText, AssignedArgument, Logger))
+								if(ApplyArgument(targetObject, target, argument, valueText, assignedArgument, logger))
 								{
-									AssignedArgument = Argument;
+									assignedArgument = argument;
 								}
 							}
 							else
 							{
-								foreach(string ItemValueText in ValueText.Split(Attribute.ListSeparator))
+								foreach(string itemValueText in valueText.Split(attribute.ListSeparator))
 								{
-									if(ApplyArgument(TargetObject, Target, Argument, ItemValueText, AssignedArgument, Logger))
+									if(ApplyArgument(targetObject, target, argument, itemValueText, assignedArgument, logger))
 									{
-										AssignedArgument = Argument;
+										assignedArgument = argument;
 									}
 								}
 							}
 
 							// Mark this argument as used
-							UsedArguments.Set(Index, true);
+							_usedArguments.Set(index, true);
 						}
 					}
 
 					// If this attribute is marked as required, keep track of it so we can warn if the field is not assigned to
-					if(Attribute.Required && RequiredPrefix == null)
+					if(attribute.Required && requiredPrefix == null)
 					{
-						RequiredPrefix = Prefix;
+						requiredPrefix = prefix;
 					}
 				}
 
 				// Make sure that this field has been assigned to
-				if(AssignedArgument == null && RequiredPrefix != null)
+				if(assignedArgument == null && requiredPrefix != null)
 				{
-					MissingArguments.Add(RequiredPrefix);
+					missingArguments.Add(requiredPrefix);
 				}
 			}
 
 			// If any arguments were missing, print an error about them
-			if(MissingArguments.Count > 0)
+			if(missingArguments.Count > 0)
 			{
-				if(MissingArguments.Count == 1)
+				if(missingArguments.Count == 1)
 				{
-					throw new CommandLineArgumentException(String.Format("Missing {0} argument", MissingArguments[0].Replace("=", "=...")));
+					throw new CommandLineArgumentException(String.Format("Missing {0} argument", missingArguments[0].Replace("=", "=...")));
 				}
 				else
 				{
-					throw new CommandLineArgumentException(String.Format("Missing {0} arguments", StringUtils.FormatList(MissingArguments.Select(x => x.Replace("=", "=...")))));
+					throw new CommandLineArgumentException(String.Format("Missing {0} arguments", StringUtils.FormatList(missingArguments.Select(x => x.Replace("=", "=...")))));
 				}
 			}
 		}
@@ -829,175 +822,175 @@ namespace EpicGames.Core
 		/// <summary>
 		/// Applies these arguments to fields with the [CommandLine] attribute in the given object.
 		/// </summary>
-		/// <param name="Logger">Sink for error/warning messages</param>
-		public T ApplyTo<T>(ILogger Logger) where T : new()
+		/// <param name="logger">Sink for error/warning messages</param>
+		public T ApplyTo<T>(ILogger logger) where T : new()
 		{
-			T Object = new T();
-			ApplyTo(Object, Logger);
-			return Object;
+			T obj = new T();
+			ApplyTo(obj, logger);
+			return obj;
 		}
 
-		static List<ArgumentTarget> GetArgumentTargetsForType(Type? TargetType)
+		static List<ArgumentTarget> GetArgumentTargetsForType(Type? targetType)
 		{
-			List<ArgumentTarget> Targets = new List<ArgumentTarget>();
-			while (TargetType != null && TargetType != typeof(object))
+			List<ArgumentTarget> targets = new List<ArgumentTarget>();
+			while (targetType != null && targetType != typeof(object))
 			{
-				foreach (FieldInfo FieldInfo in TargetType.GetFields(BindingFlags.Instance | BindingFlags.GetField | BindingFlags.Public | BindingFlags.NonPublic | BindingFlags.DeclaredOnly))
+				foreach (FieldInfo fieldInfo in targetType.GetFields(BindingFlags.Instance | BindingFlags.GetField | BindingFlags.Public | BindingFlags.NonPublic | BindingFlags.DeclaredOnly))
 				{
-					IEnumerable<CommandLineAttribute> Attributes = FieldInfo.GetCustomAttributes<CommandLineAttribute>();
-					if (Attributes.Any())
+					IEnumerable<CommandLineAttribute> attributes = fieldInfo.GetCustomAttributes<CommandLineAttribute>();
+					if (attributes.Any())
 					{
-						Targets.Add(new ArgumentTarget(FieldInfo, FieldInfo.FieldType, FieldInfo.SetValue, FieldInfo.GetValue, Attributes.ToArray()));
+						targets.Add(new ArgumentTarget(fieldInfo, fieldInfo.FieldType, fieldInfo.SetValue, fieldInfo.GetValue, attributes.ToArray()));
 					}
 				}
-				foreach (PropertyInfo PropertyInfo in TargetType.GetProperties(BindingFlags.Instance | BindingFlags.GetProperty | BindingFlags.Public | BindingFlags.NonPublic | BindingFlags.DeclaredOnly))
+				foreach (PropertyInfo propertyInfo in targetType.GetProperties(BindingFlags.Instance | BindingFlags.GetProperty | BindingFlags.Public | BindingFlags.NonPublic | BindingFlags.DeclaredOnly))
 				{
-					IEnumerable<CommandLineAttribute> Attributes = PropertyInfo.GetCustomAttributes<CommandLineAttribute>();
-					if (Attributes.Any())
+					IEnumerable<CommandLineAttribute> attributes = propertyInfo.GetCustomAttributes<CommandLineAttribute>();
+					if (attributes.Any())
 					{
-						Targets.Add(new ArgumentTarget(PropertyInfo, PropertyInfo.PropertyType, PropertyInfo.SetValue, PropertyInfo.GetValue, Attributes.ToArray()));
+						targets.Add(new ArgumentTarget(propertyInfo, propertyInfo.PropertyType, propertyInfo.SetValue, propertyInfo.GetValue, attributes.ToArray()));
 					}
 				}
-				TargetType = TargetType.BaseType;
+				targetType = targetType.BaseType;
 			}
-			return Targets;
+			return targets;
 		}
 
 		/// <summary>
 		/// Gets help text for the arguments of a given type
 		/// </summary>
-		/// <param name="Type">The type to find parameters for</param>
+		/// <param name="type">The type to find parameters for</param>
 		/// <returns>List of parameters</returns>
-		public static List<KeyValuePair<string, string>> GetParameters(Type Type)
+		public static List<KeyValuePair<string, string>> GetParameters(Type type)
 		{
-			List<KeyValuePair<string, string>> Parameters = new List<KeyValuePair<string, string>>();
+			List<KeyValuePair<string, string>> parameters = new List<KeyValuePair<string, string>>();
 
-			List<ArgumentTarget> Targets = GetArgumentTargetsForType(Type);
-			foreach (ArgumentTarget Target in Targets)
+			List<ArgumentTarget> targets = GetArgumentTargetsForType(type);
+			foreach (ArgumentTarget target in targets)
 			{
-				StringBuilder DescriptionBuilder = new StringBuilder();
-				foreach (DescriptionAttribute Attribute in Target.Member.GetCustomAttributes<DescriptionAttribute>())
+				StringBuilder descriptionBuilder = new StringBuilder();
+				foreach (DescriptionAttribute attribute in target.Member.GetCustomAttributes<DescriptionAttribute>())
 				{
-					if(DescriptionBuilder.Length > 0)
+					if(descriptionBuilder.Length > 0)
 					{
-						DescriptionBuilder.Append("\n");
+						descriptionBuilder.Append("\n");
 					}
-					DescriptionBuilder.Append(Attribute.Description);
+					descriptionBuilder.Append(attribute.Description);
 				}
 
-				string Description = DescriptionBuilder.ToString();
-				if (Description.Length == 0)
+				string description = descriptionBuilder.ToString();
+				if (description.Length == 0)
 				{
-					Description = "No description available.";
+					description = "No description available.";
 				}
 
-				foreach (CommandLineAttribute Attribute in Target.Attributes)
+				foreach (CommandLineAttribute attribute in target.Attributes)
 				{
-					string Prefix = GetArgumentPrefix(Target, Attribute);
-					if(Prefix.EndsWith("=", StringComparison.Ordinal))
+					string prefix = GetArgumentPrefix(target, attribute);
+					if(prefix.EndsWith("=", StringComparison.Ordinal))
 					{
-						Prefix += "...";
+						prefix += "...";
 					}
-					Parameters.Add(new KeyValuePair<string, string>(Prefix, Description));
+					parameters.Add(new KeyValuePair<string, string>(prefix, description));
 				}
 			}
-			return Parameters;
+			return parameters;
 		}
 
 		/// <summary>
 		/// Quotes a command line argument, if necessary
 		/// </summary>
-		/// <param name="Argument">The argument that may need quoting</param>
+		/// <param name="argument">The argument that may need quoting</param>
 		/// <returns>Argument which is safe to pass on the command line</returns>
-		public static string Quote(string Argument)
+		public static string Quote(string argument)
 		{
 			// See if the entire string is quoted correctly
 			bool bInQuotes = false;
-			for (int Idx = 0;;Idx++)
+			for (int idx = 0;;idx++)
 			{
-				if (Idx == Argument.Length)
+				if (idx == argument.Length)
 				{
-					return Argument;
+					return argument;
 				}
-				else if (Argument[Idx] == '\"')
+				else if (argument[idx] == '\"')
 				{
 					bInQuotes ^= true;
 				}
-				else if (Argument[Idx] == ' ')
+				else if (argument[idx] == ' ')
 				{
 					break;
 				}
 			}
 
 			// Try to insert a quote after the argument string
-			if (Argument[0] == '-')
+			if (argument[0] == '-')
 			{
-				for(int Idx = 1; Idx < Argument.Length && Argument[Idx] != ' '; Idx++)
+				for(int idx = 1; idx < argument.Length && argument[idx] != ' '; idx++)
 				{
-					if (Argument[Idx] == '=')
+					if (argument[idx] == '=')
 					{
-						return String.Format("{0}=\"{1}\"", Argument.Substring(0, Idx), Argument.Substring(Idx + 1).Replace("\"", "\\\""));
+						return String.Format("{0}=\"{1}\"", argument.Substring(0, idx), argument.Substring(idx + 1).Replace("\"", "\\\""));
 					}
 				}
 			}
 
 			// Quote the whole thing
-			return "\"" + Argument.Replace("\"", "\\\"") + "\"";
+			return "\"" + argument.Replace("\"", "\\\"") + "\"";
 		}
 
 		/// <summary>
 		/// Joins the given arguments into a command line
 		/// </summary>
-		/// <param name="Arguments">List of command line arguments</param>
+		/// <param name="arguments">List of command line arguments</param>
 		/// <returns>Joined command line</returns>
-		public static string Join(IEnumerable<string> Arguments)
+		public static string Join(IEnumerable<string> arguments)
 		{
-			StringBuilder Result = new StringBuilder();
-			foreach (string Argument in Arguments)
+			StringBuilder result = new StringBuilder();
+			foreach (string argument in arguments)
 			{
-				if(Result.Length > 0)
+				if(result.Length > 0)
 				{
-					Result.Append(' ');
+					result.Append(' ');
 				}
-				Result.Append(Quote(Argument));
+				result.Append(Quote(argument));
 			}
-			return Result.ToString();
+			return result.ToString();
 		}
 
 		/// <summary>
 		/// Splits a command line into individual arguments
 		/// </summary>
-		/// <param name="CommandLine">The command line text</param>
+		/// <param name="commandLine">The command line text</param>
 		/// <returns>Array of arguments</returns>
-		public static string[] Split(string CommandLine)
+		public static string[] Split(string commandLine)
 		{
-			StringBuilder Argument = new StringBuilder();
+			StringBuilder argument = new StringBuilder();
 
-			List<string> Arguments = new List<string>();
-			for(int Idx = 0; Idx < CommandLine.Length; Idx++)
+			List<string> arguments = new List<string>();
+			for(int idx = 0; idx < commandLine.Length; idx++)
 			{
-				if(!Char.IsWhiteSpace(CommandLine[Idx]))
+				if(!Char.IsWhiteSpace(commandLine[idx]))
 				{
-					Argument.Clear();
-					for(bool bInQuotes = false; Idx < CommandLine.Length; Idx++)
+					argument.Clear();
+					for(bool bInQuotes = false; idx < commandLine.Length; idx++)
 					{
-						if(CommandLine[Idx] == '\"')
+						if(commandLine[idx] == '\"')
 						{
 							bInQuotes ^= true;
 						}
-						else if(!bInQuotes && Char.IsWhiteSpace(CommandLine[Idx]))
+						else if(!bInQuotes && Char.IsWhiteSpace(commandLine[idx]))
 						{
 							break;
 						}
 						else
 						{
-							Argument.Append(CommandLine[Idx]);
+							argument.Append(commandLine[idx]);
 						}
 					}
-					Arguments.Add(Argument.ToString());
+					arguments.Add(argument.ToString());
 				}
 			}
-			return Arguments.ToArray();
+			return arguments.ToArray();
 		}
 
 		/// <summary>
@@ -1005,57 +998,57 @@ namespace EpicGames.Core
 		/// </summary>
 		/// <param name="Arguments">The arguments to add</param>
 		/// <returns>New argument list</returns>
-		public CommandLineArguments Append(IEnumerable<string> AppendArguments)
+		public CommandLineArguments Append(IEnumerable<string> appendArguments)
 		{
-			CommandLineArguments NewArguments = new CommandLineArguments(Enumerable.Concat(Arguments, AppendArguments).ToArray());
-			for(int Idx = 0; Idx < Arguments.Length; Idx++)
+			CommandLineArguments newArguments = new CommandLineArguments(Enumerable.Concat(_arguments, appendArguments).ToArray());
+			for(int idx = 0; idx < _arguments.Length; idx++)
 			{
-				if(HasBeenUsed(Idx))
+				if(HasBeenUsed(idx))
 				{
-					NewArguments.MarkAsUsed(Idx);
+					newArguments.MarkAsUsed(idx);
 				}
 			}
-			return NewArguments;
+			return newArguments;
 		}
 
 		/// <summary>
 		/// Retrieves all arguments with the given prefix, and returns the remaining a list of strings
 		/// </summary>
-		/// <param name="Prefix">Prefix for the arguments to remove</param>
-		/// <param name="Values">Receives a list of values with the given prefix</param>
+		/// <param name="prefix">Prefix for the arguments to remove</param>
+		/// <param name="values">Receives a list of values with the given prefix</param>
 		/// <returns>New argument list</returns>
-		public CommandLineArguments Remove(string Prefix, out List<string> Values)
+		public CommandLineArguments Remove(string prefix, out List<string> values)
 		{
-			Values = new List<string>();
+			values = new List<string>();
 
 			// Split the arguments into the values array and an array of new arguments
-			int[] NewArgumentIndex = new int[Arguments.Length];
-			List<string> NewArgumentList = new List<string>(Arguments.Length);
-			for(int Idx = 0; Idx < Arguments.Length; Idx++)
+			int[] newArgumentIndex = new int[_arguments.Length];
+			List<string> newArgumentList = new List<string>(_arguments.Length);
+			for(int idx = 0; idx < _arguments.Length; idx++)
 			{
-				string Argument = Arguments[Idx];
-				if(Argument.StartsWith(Prefix, StringComparison.OrdinalIgnoreCase))
+				string argument = _arguments[idx];
+				if(argument.StartsWith(prefix, StringComparison.OrdinalIgnoreCase))
 				{
-					NewArgumentIndex[Idx] = -1;
-					Values.Add(Argument.Substring(Prefix.Length));
+					newArgumentIndex[idx] = -1;
+					values.Add(argument.Substring(prefix.Length));
 				}
 				else
 				{
-					NewArgumentIndex[Idx] = NewArgumentList.Count;
-					NewArgumentList.Add(Argument);
+					newArgumentIndex[idx] = newArgumentList.Count;
+					newArgumentList.Add(argument);
 				}
 			}
 
 			// Create the new argument list, and mark the same arguments as used
-			CommandLineArguments NewArguments = new CommandLineArguments(NewArgumentList.ToArray());
-			for(int Idx = 0; Idx < Arguments.Length; Idx++)
+			CommandLineArguments newArguments = new CommandLineArguments(newArgumentList.ToArray());
+			for(int idx = 0; idx < _arguments.Length; idx++)
 			{
-				if(HasBeenUsed(Idx) && NewArgumentIndex[Idx] != -1)
+				if(HasBeenUsed(idx) && newArgumentIndex[idx] != -1)
 				{
-					NewArguments.MarkAsUsed(NewArgumentIndex[Idx]);
+					newArguments.MarkAsUsed(newArgumentIndex[idx]);
 				}
 			}
-			return NewArguments;
+			return newArguments;
 		}
 
 		/// <summary>
@@ -1069,28 +1062,28 @@ namespace EpicGames.Core
 		/// <summary>
 		/// Checks that there are no unused arguments (and warns if there are)
 		/// </summary>
-		public void CheckAllArgumentsUsed(ILogger Logger)
+		public void CheckAllArgumentsUsed(ILogger logger)
 		{
 			// Find all the unused arguments
-			List<string> RemainingArguments = new List<string>();
-			for(int Idx = 0; Idx < Arguments.Length; Idx++)
+			List<string> remainingArguments = new List<string>();
+			for(int idx = 0; idx < _arguments.Length; idx++)
 			{
-				if(!UsedArguments[Idx])
+				if(!_usedArguments[idx])
 				{
-					RemainingArguments.Add(Arguments[Idx]);
+					remainingArguments.Add(_arguments[idx]);
 				}
 			}
 
 			// Output a warning
-			if(RemainingArguments.Count > 0)
+			if(remainingArguments.Count > 0)
 			{
-				if(RemainingArguments.Count == 1)
+				if(remainingArguments.Count == 1)
 				{
-					Logger.LogWarning(String.Format("Invalid argument: {0}", RemainingArguments[0]));
+					logger.LogWarning(String.Format("Invalid argument: {0}", remainingArguments[0]));
 				}
 				else
 				{
-					Logger.LogWarning(String.Format("Invalid arguments:\n{0}", String.Join("\n", RemainingArguments)));
+					logger.LogWarning(String.Format("Invalid arguments:\n{0}", String.Join("\n", remainingArguments)));
 				}
 			}
 		}
@@ -1101,7 +1094,7 @@ namespace EpicGames.Core
 		/// <returns></returns>
 		public bool AreAnyArgumentsUsed()
 		{
-			return UsedArguments.Cast<bool>().Count(B => B) != 0;
+			return _usedArguments.Cast<bool>().Count(b => b) != 0;
 		}
 
 		/// <summary>
@@ -1110,11 +1103,11 @@ namespace EpicGames.Core
 		/// <returns></returns>
 		public IEnumerable<string> GetUnusedArguments()
 		{
-			for(int Idx = 0; Idx < Arguments.Length; Idx++)
+			for(int idx = 0; idx < _arguments.Length; idx++)
 			{
-				if (!UsedArguments[Idx])
+				if (!_usedArguments[idx])
 				{
-					yield return Arguments[Idx];
+					yield return _arguments[idx];
 				}
 			}
 		}
@@ -1125,85 +1118,85 @@ namespace EpicGames.Core
 		/// <returns></returns>
 		public int CountValueArguments()
 		{
-			return FlagArguments.Cast<bool>().Count(B => !B);
+			return _flagArguments.Cast<bool>().Count(b => !b);
 		}
 
 		/// <summary>
 		/// Checks that a given string is a valid argument prefix
 		/// </summary>
-		/// <param name="Prefix">The prefix to check</param>
-		private static void CheckValidPrefix(string Prefix)
+		/// <param name="prefix">The prefix to check</param>
+		private static void CheckValidPrefix(string prefix)
 		{
-			if(Prefix.Length == 0)
+			if(prefix.Length == 0)
 			{
 				throw new ArgumentException("Argument prefix cannot be empty.");
 			}
-			else if(Prefix[0] != '-')
+			else if(prefix[0] != '-')
 			{
 				throw new ArgumentException("Argument prefix must begin with a hyphen.");
 			}
-			else if(!ValueSeparators.Contains(Prefix[Prefix.Length - 1]))
+			else if(!s_valueSeparators.Contains(prefix[^1]))
 			{
-				throw new ArgumentException(String.Format("Argument prefix must end with '{0}'", String.Join("' or '", ValueSeparators)));
+				throw new ArgumentException(String.Format("Argument prefix must end with '{0}'", String.Join("' or '", s_valueSeparators)));
 			}
 		}
 
 		/// <summary>
 		/// Parses and assigns a value to a field
 		/// </summary>
-		/// <param name="TargetObject">The target object to assign values to</param>
-		/// <param name="Target">The target to assign the value to</param>
-		/// <param name="ArgumentText">The full argument text</param>
-		/// <param name="ValueText">Argument text</param>
-		/// <param name="PreviousArgumentText">The previous text used to configure this field</param>
-		/// <param name="Logger">Logger for error/warning messages</param>
+		/// <param name="targetObject">The target object to assign values to</param>
+		/// <param name="target">The target to assign the value to</param>
+		/// <param name="argumentText">The full argument text</param>
+		/// <param name="valueText">Argument text</param>
+		/// <param name="previousArgumentText">The previous text used to configure this field</param>
+		/// <param name="logger">Logger for error/warning messages</param>
 		/// <returns>True if the value was assigned to the field, false otherwise</returns>
-		private static bool ApplyArgument(object TargetObject, ArgumentTarget Target, string ArgumentText, string ValueText, string? PreviousArgumentText, ILogger Logger)
+		private static bool ApplyArgument(object targetObject, ArgumentTarget target, string argumentText, string valueText, string? previousArgumentText, ILogger logger)
 		{
-			Type ValueType = Target.ValueType;
+			Type valueType = target.ValueType;
 
 			// Check if the field type implements ICollection<>. If so, we can take multiple values.
-			Type? CollectionType = null;
-			foreach (Type InterfaceType in ValueType.GetInterfaces())
+			Type? collectionType = null;
+			foreach (Type interfaceType in valueType.GetInterfaces())
 			{
-				if (InterfaceType.IsGenericType && InterfaceType.GetGenericTypeDefinition() == typeof(ICollection<>))
+				if (interfaceType.IsGenericType && interfaceType.GetGenericTypeDefinition() == typeof(ICollection<>))
 				{
-					ValueType = InterfaceType.GetGenericArguments()[0];
-					CollectionType = InterfaceType;
+					valueType = interfaceType.GetGenericArguments()[0];
+					collectionType = interfaceType;
 					break;
 				}
 			}
 
 			// Try to parse the value
-			object? Value;
-			if(!TryParseValue(ValueType, ValueText, out Value))
+			object? value;
+			if(!TryParseValue(valueType, valueText, out value))
 			{
-				Logger.LogWarning("Unable to parse value for argument '{0}'.", ArgumentText);
+				logger.LogWarning("Unable to parse value for argument '{0}'.", argumentText);
 				return false;
 			}
 
 			// Try to assign values to the target field
-			if (CollectionType == null)
+			if (collectionType == null)
 			{
 				// Check if this field has already been assigned to. Output a warning if the previous value is in conflict with the new one.
-				if(PreviousArgumentText != null)
+				if(previousArgumentText != null)
 				{
-					object? PreviousValue = Target.GetValue(TargetObject);
-					if(!Object.Equals(PreviousValue, Value))
+					object? previousValue = target.GetValue(targetObject);
+					if(!Object.Equals(previousValue, value))
 					{
-						Logger.LogWarning("Argument '{0}' conflicts with '{1}'; ignoring.", ArgumentText, PreviousArgumentText);
+						logger.LogWarning("Argument '{0}' conflicts with '{1}'; ignoring.", argumentText, previousArgumentText);
 					}
 					return false;
 				}
 
 				// Set the value on the target object
-				Target.SetValue(TargetObject, Value);
+				target.SetValue(targetObject, value);
 				return true;
 			}
 			else
 			{
 				// Call the 'Add' method on the collection
-				CollectionType.InvokeMember("Add", BindingFlags.InvokeMethod, null, Target.GetValue(TargetObject), new object[] { Value });
+				collectionType.InvokeMember("Add", BindingFlags.InvokeMethod, null, target.GetValue(targetObject), new object[] { value });
 				return true;
 			}
 		}
@@ -1211,110 +1204,110 @@ namespace EpicGames.Core
 		/// <summary>
 		/// Attempts to parse the given string to a value
 		/// </summary>
-		/// <param name="FieldType">Type of the field to convert to</param>
-		/// <param name="Text">The value text</param>
-		/// <param name="Value">On success, contains the parsed object</param>
+		/// <param name="fieldType">Type of the field to convert to</param>
+		/// <param name="text">The value text</param>
+		/// <param name="value">On success, contains the parsed object</param>
 		/// <returns>True if the text could be parsed, false otherwise</returns>
-		private static bool TryParseValue(Type FieldType, string Text, [NotNullWhen(true)] out object? Value)
+		private static bool TryParseValue(Type fieldType, string text, [NotNullWhen(true)] out object? value)
 		{
-			if (FieldType.IsGenericType && FieldType.GetGenericTypeDefinition() == typeof(Nullable<>))
+			if (fieldType.IsGenericType && fieldType.GetGenericTypeDefinition() == typeof(Nullable<>))
 			{
 				// Try to parse the inner type instead
-				return TryParseValue(FieldType.GetGenericArguments()[0], Text, out Value);
+				return TryParseValue(fieldType.GetGenericArguments()[0], text, out value);
 			}
-			else if (FieldType.IsEnum)
+			else if (fieldType.IsEnum)
 			{
 				// Special handling for enums; parse the value ignoring case.
 				try
 				{
-					Value = Enum.Parse(FieldType, Text, true);
+					value = Enum.Parse(fieldType, text, true);
 					return true;
 				}
 				catch (ArgumentException)
 				{
-					Value = null;
+					value = null;
 					return false;
 				}
 			}
-			else if (FieldType == typeof(FileReference))
+			else if (fieldType == typeof(FileReference))
 			{
 				// Construct a file reference from the string
 				try
 				{
-					Value = new FileReference(Text);
+					value = new FileReference(text);
 					return true;
 				}
 				catch
 				{
-					Value = null;
+					value = null;
 					return false;
 				}
 			}
-			else if (FieldType == typeof(DirectoryReference))
+			else if (fieldType == typeof(DirectoryReference))
 			{
 				// Construct a file reference from the string
 				try
 				{
-					Value = new DirectoryReference(Text);
+					value = new DirectoryReference(text);
 					return true;
 				}
 				catch
 				{
-					Value = null;
+					value = null;
 					return false;
 				}
 			}
-			else if (FieldType == typeof(TimeSpan))
+			else if (fieldType == typeof(TimeSpan))
 			{
 				// Construct a time span form the string
-				double FloatValue;
-				if (Text.EndsWith("h", StringComparison.OrdinalIgnoreCase) && Double.TryParse(Text.Substring(0, Text.Length - 1), out FloatValue))
+				double floatValue;
+				if (text.EndsWith("h", StringComparison.OrdinalIgnoreCase) && Double.TryParse(text.Substring(0, text.Length - 1), out floatValue))
 				{
-					Value = TimeSpan.FromHours(FloatValue);
+					value = TimeSpan.FromHours(floatValue);
 					return true;
 				}
-				else if (Text.EndsWith("m", StringComparison.OrdinalIgnoreCase) && Double.TryParse(Text.Substring(0, Text.Length - 1), out FloatValue))
+				else if (text.EndsWith("m", StringComparison.OrdinalIgnoreCase) && Double.TryParse(text.Substring(0, text.Length - 1), out floatValue))
 				{
-					Value = TimeSpan.FromMinutes(FloatValue);
+					value = TimeSpan.FromMinutes(floatValue);
 					return true;
 				}
-				else if (Text.EndsWith("s", StringComparison.OrdinalIgnoreCase) && Double.TryParse(Text.Substring(0, Text.Length - 1), out FloatValue))
+				else if (text.EndsWith("s", StringComparison.OrdinalIgnoreCase) && Double.TryParse(text.Substring(0, text.Length - 1), out floatValue))
 				{
-					Value = TimeSpan.FromSeconds(FloatValue);
+					value = TimeSpan.FromSeconds(floatValue);
 					return true;
 				}
 
-				TimeSpan TimeSpanValue;
-				if (TimeSpan.TryParse(Text, out TimeSpanValue))
+				TimeSpan timeSpanValue;
+				if (TimeSpan.TryParse(text, out timeSpanValue))
 				{
-					Value = TimeSpanValue;
+					value = timeSpanValue;
 					return true;
 				}
 				else
 				{
-					Value = null;
+					value = null;
 					return false;
 				}
 			}
 			else
 			{
 				// First check for a TypeConverter
-				TypeConverter TypeConverter = TypeDescriptor.GetConverter(FieldType);
-				if (TypeConverter.CanConvertFrom(typeof(string)))
+				TypeConverter typeConverter = TypeDescriptor.GetConverter(fieldType);
+				if (typeConverter.CanConvertFrom(typeof(string)))
 				{
-					Value = TypeConverter.ConvertFrom(Text);
+					value = typeConverter.ConvertFrom(text);
 					return true;
 				}
 
 				// Otherwise let the framework convert between types
 				try
 				{
-					Value = Convert.ChangeType(Text, FieldType);
+					value = Convert.ChangeType(text, fieldType);
 					return true;
 				}
 				catch (InvalidCastException)
 				{
-					Value = null;
+					value = null;
 					return false;
 				}
 			}
@@ -1326,7 +1319,7 @@ namespace EpicGames.Core
 		/// <returns>IEnumerator interface</returns>
 		IEnumerator IEnumerable.GetEnumerator()
 		{
-			return Arguments.GetEnumerator();
+			return _arguments.GetEnumerator();
 		}
 
 		/// <summary>
@@ -1335,7 +1328,7 @@ namespace EpicGames.Core
 		/// <returns>Generic IEnumerator interface</returns>
 		public IEnumerator<string> GetEnumerator()
 		{
-			return ((IEnumerable<string>)Arguments).GetEnumerator();
+			return ((IEnumerable<string>)_arguments).GetEnumerator();
 		}
 
 		/// <summary>
@@ -1344,36 +1337,36 @@ namespace EpicGames.Core
 		/// <returns>Array of arguments</returns>
 		public string[] GetRawArray()
 		{
-			return Arguments;
+			return _arguments;
 		}
 
 		/// <summary>
 		/// Takes a command line argument and adds quotes if necessary
 		/// </summary>
-		/// <param name="Argument">The command line argument</param>
+		/// <param name="argument">The command line argument</param>
 		/// <returns>The command line argument with quotes inserted to escape it if necessary</returns>
-		public static void Append(StringBuilder CommandLine, string Argument)
+		public static void Append(StringBuilder commandLine, string argument)
 		{
-			if(CommandLine.Length > 0)
+			if(commandLine.Length > 0)
 			{
-				CommandLine.Append(' ');
+				commandLine.Append(' ');
 			}
 
-			int SpaceIdx = Argument.IndexOf(' ');
-			if(SpaceIdx == -1)
+			int spaceIdx = argument.IndexOf(' ');
+			if(spaceIdx == -1)
 			{
-				CommandLine.Append(Argument);
+				commandLine.Append(argument);
 			}
 			else
 			{
-				int EqualsIdx = Argument.IndexOf('=');
-				if(EqualsIdx == -1)
+				int equalsIdx = argument.IndexOf('=');
+				if(equalsIdx == -1)
 				{
-					CommandLine.AppendFormat("\"{0}\"", Argument);
+					commandLine.AppendFormat("\"{0}\"", argument);
 				}
 				else
 				{
-					CommandLine.AppendFormat("{0}\"{1}\"", Argument.Substring(0, EqualsIdx + 1), Argument.Substring(EqualsIdx + 1));
+					commandLine.AppendFormat("{0}\"{1}\"", argument.Substring(0, equalsIdx + 1), argument.Substring(equalsIdx + 1));
 				}
 			}
 		}
@@ -1384,12 +1377,12 @@ namespace EpicGames.Core
 		/// <returns></returns>
 		public override string ToString()
 		{
-			StringBuilder Result = new StringBuilder();
-			foreach(string Argument in Arguments)
+			StringBuilder result = new StringBuilder();
+			foreach(string argument in _arguments)
 			{
-				Append(Result, Argument);
+				Append(result, argument);
 			}
-			return Result.ToString();
+			return result.ToString();
 		}
 	}
 }
diff --git a/Engine/Source/Programs/Shared/EpicGames.Core/CommandLineAttribute.cs b/Engine/Source/Programs/Shared/EpicGames.Core/CommandLineAttribute.cs
--- a/Engine/Source/Programs/Shared/EpicGames.Core/CommandLineAttribute.cs
+++ b/Engine/Source/Programs/Shared/EpicGames.Core/CommandLineAttribute.cs
@@ -1,11 +1,6 @@
 // Copyright Epic Games, Inc. All Rights Reserved.
 
 using System;
-using System.Collections.Generic;
-using System.Linq;
-using System.Reflection;
-using System.Text;
-using System.Threading.Tasks;
 
 namespace EpicGames.Core
 {
@@ -18,39 +13,39 @@ namespace EpicGames.Core
 		/// <summary>
 		/// Prefix for the option, with a leading '-' and trailing '=' character if a value is expected.
 		/// </summary>
-		public string? Prefix;
+		public string? Prefix { get; set; }
 
 		/// <summary>
 		/// Specifies a fixed value for this argument. Specifying an alternate value is not permitted.
 		/// </summary>
-		public string? Value = null;
+		public string? Value { get; set; } = null;
 
 		/// <summary>
 		/// Whether this argument is required
 		/// </summary>
-		public bool Required;
+		public bool Required { get; set; }
 
 		/// <summary>
 		/// For collection types, specifies the separator character between multiple arguments
 		/// </summary>
-		public char ListSeparator = '\0';
+		public char ListSeparator { get; set; } = '\0';
 
 		/// <summary>
 		/// Description of the operation.
 		/// </summary>
-		public string? Description;
+		public string? Description { get; set; }
 
 		/// <summary>
 		/// Constructor
 		/// </summary>
-		/// <param name="Prefix">Prefix for this argument</param>
-		public CommandLineAttribute(string? Prefix = null)
+		/// <param name="prefix">Prefix for this argument</param>
+		public CommandLineAttribute(string? prefix = null)
 		{
-			this.Prefix = Prefix;
+			Prefix = prefix;
 
-			if(Prefix != null)
+			if(prefix != null)
 			{
-				if(!Prefix.StartsWith("-"))
+				if(!prefix.StartsWith("-"))
 				{
 					throw new Exception("Command-line arguments must begin with a '-' character");
 				}
diff --git a/Engine/Source/Programs/Shared/EpicGames.Core/CommandLineExtensions.cs b/Engine/Source/Programs/Shared/EpicGames.Core/CommandLineExtensions.cs
--- a/Engine/Source/Programs/Shared/EpicGames.Core/CommandLineExtensions.cs
+++ b/Engine/Source/Programs/Shared/EpicGames.Core/CommandLineExtensions.cs
@@ -1,7 +1,5 @@
 // Copyright Epic Games, Inc. All Rights Reserved.
 
-using System;
-using System.Collections.Generic;
 using System.Text;
 
 namespace EpicGames.Core
@@ -14,60 +12,60 @@ namespace EpicGames.Core
 		/// <summary>
 		/// Determines if the given argument needs to be escaped
 		/// </summary>
-		/// <param name="Argument">The argument to check</param>
+		/// <param name="argument">The argument to check</param>
 		/// <returns>True if the argument needs to be escaped</returns>
-		static bool NeedsEscaping(string Argument)
+		static bool NeedsEscaping(string argument)
 		{
-			return Argument.Contains(' ') || Argument.Contains('\"');
+			return argument.Contains(' ') || argument.Contains('\"');
 		}
 
 		/// <summary>
 		/// Appends command line argument with a prefixed space. The argument may contain spaces or quotes.
 		/// </summary>
-		/// <param name="Builder">The command line to append to</param>
-		/// <param name="Argument">The argument to append</param>
-		public static void AppendArgument(this StringBuilder Builder, string Argument)
+		/// <param name="builder">The command line to append to</param>
+		/// <param name="argument">The argument to append</param>
+		public static void AppendArgument(this StringBuilder builder, string argument)
 		{
-			if (Builder.Length > 0)
+			if (builder.Length > 0)
 			{
-				Builder.Append(' ');
+				builder.Append(' ');
 			}
 
-			int EqualsIdx = Argument.IndexOf('=');
-			if (EqualsIdx != -1)
+			int equalsIdx = argument.IndexOf('=');
+			if (equalsIdx != -1)
 			{
-				string Name = Argument.Substring(0, EqualsIdx + 1);
-				if (!NeedsEscaping(Name))
+				string name = argument.Substring(0, equalsIdx + 1);
+				if (!NeedsEscaping(name))
 				{
-					Builder.Append(Name);
-					Argument = Argument.Substring(EqualsIdx + 1);
+					builder.Append(name);
+					argument = argument.Substring(equalsIdx + 1);
 				}
 			}
 
-			AppendCommandLineArgumentWithoutSpace(Builder, Argument);
+			AppendCommandLineArgumentWithoutSpace(builder, argument);
 		}
 
 		/// <summary>
 		/// Appends command line argument with a prefixed space. The argument may contain spaces or quotes.
 		/// </summary>
-		/// <param name="Builder">The command line to append to</param>
-		/// <param name="Name">Name of the argument (eg. -Foo=)</param>
-		/// <param name="Value">Value of the argument</param>
-		public static void AppendArgument(this StringBuilder Builder, string Name, string Value)
+		/// <param name="builder">The command line to append to</param>
+		/// <param name="name">Name of the argument (eg. -Foo=)</param>
+		/// <param name="value">Value of the argument</param>
+		public static void AppendArgument(this StringBuilder builder, string name, string value)
 		{
-			if (Builder.Length > 0)
+			if (builder.Length > 0)
 			{
-				Builder.Append(' ');
+				builder.Append(' ');
 			}
 
-			if (NeedsEscaping(Name))
+			if (NeedsEscaping(name))
 			{
-				AppendCommandLineArgumentWithoutSpace(Builder, Name + Value);
+				AppendCommandLineArgumentWithoutSpace(builder, name + value);
 			}
 			else
 			{
-				Builder.Append(Name);
-				AppendCommandLineArgumentWithoutSpace(Builder, Value);
+				builder.Append(name);
+				AppendCommandLineArgumentWithoutSpace(builder, value);
 			}
 		}
 
@@ -75,62 +73,62 @@ namespace EpicGames.Core
 		/// Appends an escaped command line argument. The argument may contain spaces or quotes, and is escaped according to the rules in
 		/// https://docs.microsoft.com/en-us/windows/win32/api/shellapi/nf-shellapi-commandlinetoargvw.
 		/// </summary>
-		/// <param name="Builder">The builder to append to</param>
-		/// <param name="Argument">The argument to escape</param>
-		public static void AppendCommandLineArgumentWithoutSpace(this StringBuilder Builder, string Argument)
+		/// <param name="builder">The builder to append to</param>
+		/// <param name="argument">The argument to escape</param>
+		public static void AppendCommandLineArgumentWithoutSpace(this StringBuilder builder, string argument)
 		{
-			if (!NeedsEscaping(Argument))
+			if (!NeedsEscaping(argument))
 			{
 				// No escaping necessary if the argument doesn't contain any special characters
-				Builder.Append(Argument);
+				builder.Append(argument);
 			}
 			else
 			{
 				// Escape the whole string following the rules on the CommandLineToArgV MSDN page. 
-				Builder.Append('\"');
-				for (int Idx = 0; Idx < Argument.Length; Idx++)
+				builder.Append('\"');
+				for (int idx = 0; idx < argument.Length; idx++)
 				{
-					char Character = Argument[Idx];
-					if (Character == '\"')
+					char character = argument[idx];
+					if (character == '\"')
 					{
 						// Escape a single quotation mark
-						Builder.Append("\\\"");
+						builder.Append("\\\"");
 					}
-					else if (Character == '\\')
+					else if (character == '\\')
 					{
 						// Special handling for slashes which may be followed by a quotation mark, as dictated by CommandLineToArgV
-						int StartIdx = Idx;
+						int startIdx = idx;
 						for (; ; )
 						{
-							int NextIdx = Idx + 1;
-							if (NextIdx == Argument.Length)
+							int nextIdx = idx + 1;
+							if (nextIdx == argument.Length)
 							{
 								// Will have a trailing quotation mark toggling 'in quotes' mode (2n)
-								Builder.Append('\\', (NextIdx - StartIdx) * 2);
+								builder.Append('\\', (nextIdx - startIdx) * 2);
 								break;
 							}
-							else if (Argument[NextIdx] == '\"')
+							else if (argument[nextIdx] == '\"')
 							{
 								// Needs to have a trailing quotation mark, so need to escape each backslash plus the quotation mark (2n+1)
-								Builder.Append('\\', (NextIdx - StartIdx) * 2 + 1);
+								builder.Append('\\', (nextIdx - startIdx) * 2 + 1);
 								break;
 							}
-							else if (Argument[NextIdx] != '\\')
+							else if (argument[nextIdx] != '\\')
 							{
 								// No trailing quote; can just pass through verbatim
-								Builder.Append('\\', (NextIdx - StartIdx));
+								builder.Append('\\', (nextIdx - startIdx));
 								break;
 							}
-							Idx = NextIdx;
+							idx = nextIdx;
 						}
 					}
 					else
 					{
 						// Regular character
-						Builder.Append(Character);
+						builder.Append(character);
 					}
 				}
-				Builder.Append('\"');
+				builder.Append('\"');
 			}
 		}
 	}
diff --git a/Engine/Source/Programs/Shared/EpicGames.Core/ConsoleUtils.cs b/Engine/Source/Programs/Shared/EpicGames.Core/ConsoleUtils.cs
--- a/Engine/Source/Programs/Shared/EpicGames.Core/ConsoleUtils.cs
+++ b/Engine/Source/Programs/Shared/EpicGames.Core/ConsoleUtils.cs
@@ -1,8 +1,6 @@
 // Copyright Epic Games, Inc. All Rights Reserved.
 
 using System;
-using System.Collections.Generic;
-using System.Text;
 
 namespace EpicGames.Core
 {
@@ -19,67 +17,67 @@ namespace EpicGames.Core
 			get
 			{
 				// Get the window width, using a default value if there's no console attached to this process.
-				int NewWindowWidth;
+				int newWindowWidth;
 				try
 				{
-					NewWindowWidth = Console.WindowWidth;
+					newWindowWidth = Console.WindowWidth;
 				}
 				catch
 				{
-					NewWindowWidth = 240;
+					newWindowWidth = 240;
 				}
 
-				if (NewWindowWidth <= 0)
+				if (newWindowWidth <= 0)
 				{
-					NewWindowWidth = 240;
+					newWindowWidth = 240;
 				}
 
-				return NewWindowWidth;
+				return newWindowWidth;
 			}
 		}
 
 		/// <summary>
 		/// Writes the given text to the console as a sequence of word-wrapped lines
 		/// </summary>
-		/// <param name="Text">The text to write to the console</param>
-		public static void WriteLineWithWordWrap(string Text)
+		/// <param name="text">The text to write to the console</param>
+		public static void WriteLineWithWordWrap(string text)
 		{
-			WriteLineWithWordWrap(Text, 0, 0);
+			WriteLineWithWordWrap(text, 0, 0);
 		}
 
 		/// <summary>
 		/// Writes the given text to the console as a sequence of word-wrapped lines
 		/// </summary>
-		/// <param name="Text">The text to write to the console</param>
-		/// <param name="InitialIndent">Indent for the first line</param>
-		/// <param name="HangingIndent">Indent for lines after the first</param>
-		public static void WriteLineWithWordWrap(string Text, int InitialIndent, int HangingIndent)
+		/// <param name="text">The text to write to the console</param>
+		/// <param name="initialIndent">Indent for the first line</param>
+		/// <param name="hangingIndent">Indent for lines after the first</param>
+		public static void WriteLineWithWordWrap(string text, int initialIndent, int hangingIndent)
 		{
-			foreach (string Line in StringUtils.WordWrap(Text, InitialIndent, HangingIndent, WindowWidth))
+			foreach (string line in StringUtils.WordWrap(text, initialIndent, hangingIndent, WindowWidth))
 			{
-				Console.WriteLine(Line);
+				Console.WriteLine(line);
 			}
 		}
 
 		/// <summary>
 		/// Writes an colored warning message to the console
 		/// </summary>
-		/// <param name="Text">The message to output</param>
-		public static void WriteWarning(string Text)
+		/// <param name="text">The message to output</param>
+		public static void WriteWarning(string text)
 		{
 			Console.ForegroundColor = ConsoleColor.Yellow;
-			Console.WriteLine(Text);
+			Console.WriteLine(text);
 			Console.ResetColor();
 		}
 
 		/// <summary>
 		/// Writes an colored error message to the console
 		/// </summary>
-		/// <param name="Text">The message to output</param>
-		public static void WriteError(string Text)
+		/// <param name="text">The message to output</param>
+		public static void WriteError(string text)
 		{
 			Console.ForegroundColor = ConsoleColor.Red;
-			Console.WriteLine(Text);
+			Console.WriteLine(text);
 			Console.ResetColor();
 		}
 	}
diff --git a/Engine/Source/Programs/Shared/EpicGames.Core/ContentHash.cs b/Engine/Source/Programs/Shared/EpicGames.Core/ContentHash.cs
--- a/Engine/Source/Programs/Shared/EpicGames.Core/ContentHash.cs
+++ b/Engine/Source/Programs/Shared/EpicGames.Core/ContentHash.cs
@@ -1,7 +1,6 @@
 // Copyright Epic Games, Inc. All Rights Reserved.
 
 using System;
-using System.Collections.Generic;
 using System.Diagnostics.CodeAnalysis;
 using System.IO;
 using System.Linq;
@@ -9,7 +8,6 @@ using System.Security.Cryptography;
 using System.Text;
 using System.Text.Json;
 using System.Text.Json.Serialization;
-using System.Threading.Tasks;
 
 namespace EpicGames.Core
 {
@@ -39,213 +37,213 @@ namespace EpicGames.Core
 		/// <summary>
 		/// Constructor
 		/// </summary>
-		/// <param name="Bytes">The hash data</param>
-		public ContentHash(byte[] Bytes)
+		/// <param name="bytes">The hash data</param>
+		public ContentHash(byte[] bytes)
 		{
-			this.Bytes = Bytes;
+			Bytes = bytes;
 		}
 
 		/// <summary>
 		/// Compares two content hashes for equality
 		/// </summary>
-		/// <param name="Other">The object to compare against</param>
+		/// <param name="other">The object to compare against</param>
 		/// <returns>True if the hashes are equal, false otherwise</returns>
-		public override bool Equals(object? Other)
+		public override bool Equals(object? other)
 		{
-			return Equals(Other as ContentHash);
+			return Equals(other as ContentHash);
 		}
 
 		/// <summary>
 		/// Compares two content hashes for equality
 		/// </summary>
-		/// <param name="Other">The hash to compare against</param>
+		/// <param name="other">The hash to compare against</param>
 		/// <returns>True if the hashes are equal, false otherwise</returns>
-		public bool Equals(ContentHash? Other)
+		public bool Equals(ContentHash? other)
 		{
-			if (ReferenceEquals(Other, null))
+			if (other is null)
 			{
 				return false;
 			}
 			else
 			{
-				return Bytes.SequenceEqual(Other.Bytes);
+				return Bytes.SequenceEqual(other.Bytes);
 			}
 		}
 
 		/// <summary>
 		/// Compares two content hash objects for equality
 		/// </summary>
-		/// <param name="A">The first hash to compare</param>
-		/// <param name="B">The second has to compare</param>
+		/// <param name="a">The first hash to compare</param>
+		/// <param name="b">The second has to compare</param>
 		/// <returns>True if the objects are equal, false otherwise</returns>
-		public static bool operator ==(ContentHash? A, ContentHash? B)
+		public static bool operator ==(ContentHash? a, ContentHash? b)
 		{
-			if (ReferenceEquals(A, null))
+			if (a is null)
 			{
-				return ReferenceEquals(B, null);
+				return b is null;
 			}
 			else
 			{
-				return A.Equals(B);
+				return a.Equals(b);
 			}
 		}
 
 		/// <summary>
 		/// Compares two content hash objects for inequality
 		/// </summary>
-		/// <param name="A">The first hash to compare</param>
-		/// <param name="B">The second has to compare</param>
+		/// <param name="a">The first hash to compare</param>
+		/// <param name="b">The second has to compare</param>
 		/// <returns>True if the objects are not equal, false otherwise</returns>
-		public static bool operator !=(ContentHash? A, ContentHash? B)
+		public static bool operator !=(ContentHash? a, ContentHash? b)
 		{
-			return !(A == B);
+			return !(a == b);
 		}
 
 		/// <summary>
 		/// Creates a content hash for a block of data, using a given algorithm.
 		/// </summary>
-		/// <param name="Data">Data to compute the hash for</param>
-		/// <param name="Algorithm">Algorithm to use to create the hash</param>
+		/// <param name="data">Data to compute the hash for</param>
+		/// <param name="algorithm">Algorithm to use to create the hash</param>
 		/// <returns>New content hash instance containing the hash of the data</returns>
-		public static ContentHash Compute(byte[] Data, HashAlgorithm Algorithm)
+		public static ContentHash Compute(byte[] data, HashAlgorithm algorithm)
 		{
-			return new ContentHash(Algorithm.ComputeHash(Data));
+			return new ContentHash(algorithm.ComputeHash(data));
 		}
 
 		/// <summary>
 		/// Creates a content hash for a string, using a given algorithm.
 		/// </summary>
-		/// <param name="Text">Text to compute a hash for</param>
-		/// <param name="Algorithm">Algorithm to use to create the hash</param>
+		/// <param name="text">Text to compute a hash for</param>
+		/// <param name="algorithm">Algorithm to use to create the hash</param>
 		/// <returns>New content hash instance containing the hash of the text</returns>
-		public static ContentHash Compute(string Text, HashAlgorithm Algorithm)
+		public static ContentHash Compute(string text, HashAlgorithm algorithm)
 		{
-			return new ContentHash(Algorithm.ComputeHash(Encoding.Unicode.GetBytes(Text)));
+			return new ContentHash(algorithm.ComputeHash(Encoding.Unicode.GetBytes(text)));
 		}
 
 		/// <summary>
 		/// Creates a content hash for a file, using a given algorithm.
 		/// </summary>
-		/// <param name="Location">File to compute a hash for</param>
-		/// <param name="Algorithm">Algorithm to use to create the hash</param>
+		/// <param name="location">File to compute a hash for</param>
+		/// <param name="algorithm">Algorithm to use to create the hash</param>
 		/// <returns>New content hash instance containing the hash of the file</returns>
-		public static ContentHash Compute(FileReference Location, HashAlgorithm Algorithm)
+		public static ContentHash Compute(FileReference location, HashAlgorithm algorithm)
 		{
-			using (FileStream Stream = FileReference.Open(Location, FileMode.Open, FileAccess.Read, FileShare.Read))
+			using (FileStream stream = FileReference.Open(location, FileMode.Open, FileAccess.Read, FileShare.Read))
 			{
-				return new ContentHash(Algorithm.ComputeHash(Stream));
+				return new ContentHash(algorithm.ComputeHash(stream));
 			}
 		}
 
 		/// <summary>
 		/// Creates a content hash for a block of data using MD5
 		/// </summary>
-		/// <param name="Data">Data to compute the hash for</param>
+		/// <param name="data">Data to compute the hash for</param>
 		/// <returns>New content hash instance containing the hash of the data</returns>
-		public static ContentHash MD5(byte[] Data)
+		public static ContentHash MD5(byte[] data)
 		{
-			using (MD5 Algorithm = System.Security.Cryptography.MD5.Create())
+			using (MD5 algorithm = System.Security.Cryptography.MD5.Create())
 			{
-				return Compute(Data, Algorithm);
+				return Compute(data, algorithm);
 			}
 		}
 
 		/// <summary>
 		/// Creates a content hash for a string using MD5.
 		/// </summary>
-		/// <param name="Text">Text to compute a hash for</param>
+		/// <param name="text">Text to compute a hash for</param>
 		/// <returns>New content hash instance containing the hash of the text</returns>
-		public static ContentHash MD5(string Text)
+		public static ContentHash MD5(string text)
 		{
-			using (MD5 Algorithm = System.Security.Cryptography.MD5.Create())
+			using (MD5 algorithm = System.Security.Cryptography.MD5.Create())
 			{
-				return Compute(Text, Algorithm);
+				return Compute(text, algorithm);
 			}
 		}
 
 		/// <summary>
 		/// Creates a content hash for a file, using a given algorithm.
 		/// </summary>
-		/// <param name="Location">File to compute a hash for</param>
+		/// <param name="location">File to compute a hash for</param>
 		/// <returns>New content hash instance containing the hash of the file</returns>
-		public static ContentHash MD5(FileReference Location)
+		public static ContentHash MD5(FileReference location)
 		{
-			using (MD5 Algorithm = System.Security.Cryptography.MD5.Create())
+			using (MD5 algorithm = System.Security.Cryptography.MD5.Create())
 			{
-				return Compute(Location, Algorithm);
+				return Compute(location, algorithm);
 			}
 		}
 
 		/// <summary>
 		/// Creates a content hash for a block of data using SHA1
 		/// </summary>
-		/// <param name="Data">Data to compute the hash for</param>
+		/// <param name="data">Data to compute the hash for</param>
 		/// <returns>New content hash instance containing the hash of the data</returns>
-		public static ContentHash SHA1(byte[] Data)
+		public static ContentHash SHA1(byte[] data)
 		{
-			using (SHA1 Algorithm = System.Security.Cryptography.SHA1.Create())
+			using (SHA1 algorithm = System.Security.Cryptography.SHA1.Create())
 			{
-				return Compute(Data, Algorithm);
+				return Compute(data, algorithm);
 			}
 		}
 
 		/// <summary>
 		/// Creates a content hash for a string using SHA1.
 		/// </summary>
-		/// <param name="Text">Text to compute a hash for</param>
+		/// <param name="text">Text to compute a hash for</param>
 		/// <returns>New content hash instance containing the hash of the text</returns>
-		public static ContentHash SHA1(string Text)
+		public static ContentHash SHA1(string text)
 		{
-			using (SHA1 Algorithm = System.Security.Cryptography.SHA1.Create())
+			using (SHA1 algorithm = System.Security.Cryptography.SHA1.Create())
 			{
-				return Compute(Text, Algorithm);
+				return Compute(text, algorithm);
 			}
 		}
 
 		/// <summary>
 		/// Creates a content hash for a file using SHA1.
 		/// </summary>
-		/// <param name="Location">File to compute a hash for</param>
+		/// <param name="location">File to compute a hash for</param>
 		/// <returns>New content hash instance containing the hash of the file</returns>
-		public static ContentHash SHA1(FileReference Location)
+		public static ContentHash SHA1(FileReference location)
 		{
-			using (SHA1 Algorithm = System.Security.Cryptography.SHA1.Create())
+			using (SHA1 algorithm = System.Security.Cryptography.SHA1.Create())
 			{
-				return Compute(Location, Algorithm);
+				return Compute(location, algorithm);
 			}
 		}
 
 		/// <summary>
 		/// Parse a hash from a string
 		/// </summary>
-		/// <param name="Text">Text to parse</param>
+		/// <param name="text">Text to parse</param>
 		/// <returns>Value of the hash</returns>
-		public static ContentHash Parse(string Text)
+		public static ContentHash Parse(string text)
 		{
-			ContentHash? Hash;
-			if (!TryParse(Text, out Hash))
+			ContentHash? hash;
+			if (!TryParse(text, out hash))
 			{
-				throw new ArgumentException(String.Format("'{0}' is not a valid content hash", Text));
+				throw new ArgumentException(String.Format("'{0}' is not a valid content hash", text));
 			}
-			return Hash;
+			return hash;
 		}
 
 		/// <summary>
 		/// Parse a hash from a string
 		/// </summary>
-		/// <param name="Text">Text to parse</param>
+		/// <param name="text">Text to parse</param>
 		/// <returns>Value of the hash</returns>
-		public static bool TryParse(string Text, [NotNullWhen(true)] out ContentHash? Hash)
+		public static bool TryParse(string text, [NotNullWhen(true)] out ContentHash? hash)
 		{
-			byte[]? Bytes;
-			if (StringUtils.TryParseHexString(Text, out Bytes))
+			byte[]? bytes;
+			if (StringUtils.TryParseHexString(text, out bytes))
 			{
-				Hash = new ContentHash(Bytes);
+				hash = new ContentHash(bytes);
 				return true;
 			}
 			else
 			{
-				Hash = null;
+				hash = null;
 				return false;
 			}
 		}
@@ -256,12 +254,12 @@ namespace EpicGames.Core
 		/// <returns>Integer value to use as a hash code</returns>
 		public override int GetHashCode()
 		{
-			int HashCode = Bytes[0];
-			for (int Idx = 1; Idx < Bytes.Length; Idx++)
+			int hashCode = Bytes[0];
+			for (int idx = 1; idx < Bytes.Length; idx++)
 			{
-				HashCode = (HashCode * 31) + Bytes[Idx];
+				hashCode = (hashCode * 31) + Bytes[idx];
 			}
-			return HashCode;
+			return hashCode;
 		}
 
 		/// <summary>
@@ -280,15 +278,15 @@ namespace EpicGames.Core
 	public class ContentHashJsonConverter : JsonConverter<ContentHash>
 	{
 		/// <inheritdoc/>
-		public override ContentHash Read(ref Utf8JsonReader Reader, Type TypeToConvert, JsonSerializerOptions Options)
+		public override ContentHash Read(ref Utf8JsonReader reader, Type typeToConvert, JsonSerializerOptions options)
 		{
-			return ContentHash.Parse(Reader.GetString());
+			return ContentHash.Parse(reader.GetString());
 		}
 
 		/// <inheritdoc/>
-		public override void Write(Utf8JsonWriter Writer, ContentHash Value, JsonSerializerOptions Options)
+		public override void Write(Utf8JsonWriter writer, ContentHash value, JsonSerializerOptions options)
 		{
-			Writer.WriteStringValue(Value.ToString());
+			writer.WriteStringValue(value.ToString());
 		}
 	}
 
@@ -300,35 +298,35 @@ namespace EpicGames.Core
 		/// <summary>
 		/// Writes a ContentHash to a binary archive
 		/// </summary>
-		/// <param name="Writer">The writer to output data to</param>
-		/// <param name="Hash">The hash to write</param>
-		public static void WriteContentHash(this BinaryArchiveWriter Writer, ContentHash? Hash)
+		/// <param name="writer">The writer to output data to</param>
+		/// <param name="hash">The hash to write</param>
+		public static void WriteContentHash(this BinaryArchiveWriter writer, ContentHash? hash)
 		{
-			if (ReferenceEquals(Hash, null))
+			if (hash is null)
 			{
-				Writer.WriteByteArray(null);
+				writer.WriteByteArray(null);
 			}
 			else
 			{
-				Writer.WriteByteArray(Hash.Bytes);
+				writer.WriteByteArray(hash.Bytes);
 			}
 		}
 
 		/// <summary>
 		/// Reads a ContentHash from a binary archive
 		/// </summary>
-		/// <param name="Reader">Reader to serialize data from</param>
+		/// <param name="reader">Reader to serialize data from</param>
 		/// <returns>New hash instance</returns>
-		public static ContentHash? ReadContentHash(this BinaryArchiveReader Reader)
+		public static ContentHash? ReadContentHash(this BinaryArchiveReader reader)
 		{
-			byte[]? Data = Reader.ReadByteArray();
-			if (Data == null)
+			byte[]? data = reader.ReadByteArray();
+			if (data == null)
 			{
 				return null;
 			}
 			else
 			{
-				return new ContentHash(Data);
+				return new ContentHash(data);
 			}
 		}
 	}
