Reformat EpicGames.Core according to standard coding conventions.

diff --git a/Engine/Source/Programs/Shared/EpicGames.Core/CsProjectInfo.cs b/Engine/Source/Programs/Shared/EpicGames.Core/CsProjectInfo.cs
--- a/Engine/Source/Programs/Shared/EpicGames.Core/CsProjectInfo.cs
+++ b/Engine/Source/Programs/Shared/EpicGames.Core/CsProjectInfo.cs
@@ -6,8 +6,6 @@ using System.Diagnostics.CodeAnalysis;
 using System.IO;
 using System.Linq;
 using System.Runtime.InteropServices;
-using System.Text;
-using System.Threading.Tasks;
 using System.Xml;
 
 namespace EpicGames.Core
@@ -20,65 +18,65 @@ namespace EpicGames.Core
 		/// <summary>
 		/// Evaluated properties from the project file
 		/// </summary>
-		public Dictionary<string, string> Properties;
+		public Dictionary<string, string> Properties { get; }
 
 		/// <summary>
 		/// Mapping of referenced assemblies to their 'CopyLocal' (aka 'Private') setting.
 		/// </summary>
-		public Dictionary<FileReference, bool> References = new Dictionary<FileReference, bool>();
+		public Dictionary<FileReference, bool> References { get; } = new Dictionary<FileReference, bool>();
 
 		/// <summary>
 		/// Mapping of referenced projects to their 'CopyLocal' (aka 'Private') setting.
 		/// </summary>
-		public Dictionary<FileReference, bool> ProjectReferences = new Dictionary<FileReference, bool>();
+		public Dictionary<FileReference, bool> ProjectReferences { get; } = new Dictionary<FileReference, bool>();
 
 		/// <summary>
 		/// List of compile references in the project.
 		/// </summary>
-		public List<FileReference> CompileReferences = new List<FileReference>();
+		public List<FileReference> CompileReferences { get; } = new List<FileReference>();
 
 		/// <summary>
 		/// Mapping of content IF they are flagged Always or Newer
 		/// </summary>
-		public Dictionary<FileReference, bool> ContentReferences = new Dictionary<FileReference, bool>();
+		public Dictionary<FileReference, bool> ContentReferences { get; } = new Dictionary<FileReference, bool>();
 
 		/// <summary>
 		/// Path to the CSProject file
 		/// </summary>
-		public FileReference ProjectPath;
+		public FileReference ProjectPath { get; }
 
 		/// <summary>
 		/// Constructor
 		/// </summary>
-		/// <param name="InProperties">Initial mapping of property names to values</param>
-		CsProjectInfo(Dictionary<string, string> InProperties, FileReference InProjectPath)
+		/// <param name="inProperties">Initial mapping of property names to values</param>
+		CsProjectInfo(Dictionary<string, string> inProperties, FileReference inProjectPath)
 		{
-			ProjectPath = InProjectPath;
-			Properties = new Dictionary<string, string>(InProperties);
+			ProjectPath = inProjectPath;
+			Properties = new Dictionary<string, string>(inProperties);
 		}
 
 		/// <summary>
 		/// Get the ouptut file for this project
 		/// </summary>
-		/// <param name="File">If successful, receives the assembly path</param>
+		/// <param name="file">If successful, receives the assembly path</param>
 		/// <returns>True if the output file was found</returns>
-		public bool TryGetOutputFile([NotNullWhen(true)] out FileReference? File)
+		public bool TryGetOutputFile([NotNullWhen(true)] out FileReference? file)
 		{
-			DirectoryReference? OutputDir;
-			if(!TryGetOutputDir(out OutputDir))
+			DirectoryReference? outputDir;
+			if(!TryGetOutputDir(out outputDir))
 			{
-				File = null;
+				file = null;
 				return false;
 			}
 
-			string? AssemblyName;
-			if(!TryGetAssemblyName(out AssemblyName))
+			string? assemblyName;
+			if(!TryGetAssemblyName(out assemblyName))
 			{
-				File = null;
+				file = null;
 				return false;
 			}
 
-			File = FileReference.Combine(OutputDir, AssemblyName + ".dll");
+			file = FileReference.Combine(outputDir, assemblyName + ".dll");
 			return true;
 		}
 
@@ -88,55 +86,55 @@ namespace EpicGames.Core
 		/// <returns>True if this is a netcore project</returns>
 		private bool IsNetCoreProject()
 		{
-			string? Framework;
-			return Properties.TryGetValue("TargetFramework", out Framework) && Framework.StartsWith("netcoreapp");
+			string? framework;
+			return Properties.TryGetValue("TargetFramework", out framework) && framework.StartsWith("netcoreapp");
 		}
 
 		/// <summary>
 		/// Resolve the project's output directory
 		/// </summary>
-		/// <param name="BaseDirectory">Base directory to resolve relative paths to</param>
+		/// <param name="baseDirectory">Base directory to resolve relative paths to</param>
 		/// <returns>The configured output directory</returns>
-		public DirectoryReference GetOutputDir(DirectoryReference BaseDirectory)
+		public DirectoryReference GetOutputDir(DirectoryReference baseDirectory)
 		{
-			string? OutputPath;
-			if (Properties.TryGetValue("OutputPath", out OutputPath))
+			string? outputPath;
+			if (Properties.TryGetValue("OutputPath", out outputPath))
 			{
-				return DirectoryReference.Combine(BaseDirectory, OutputPath);
+				return DirectoryReference.Combine(baseDirectory, outputPath);
 			}
 			else if (IsNetCoreProject())
 			{
-				string Configuration = Properties.ContainsKey("Configuration") ? Properties["Configuration"] : "Development";
-				return DirectoryReference.Combine(BaseDirectory, "bin", Configuration, Properties["TargetFramework"]);
+				string configuration = Properties.ContainsKey("Configuration") ? Properties["Configuration"] : "Development";
+				return DirectoryReference.Combine(baseDirectory, "bin", configuration, Properties["TargetFramework"]);
 			}
 			else
 			{
-				return BaseDirectory;
+				return baseDirectory;
 			}
 		}
 
 		/// <summary>
 		/// Resolve the project's output directory
 		/// </summary>
-		/// <param name="OutputDir">If successful, receives the output directory</param>
+		/// <param name="outputDir">If successful, receives the output directory</param>
 		/// <returns>True if the output directory was found</returns>
-		public bool TryGetOutputDir([NotNullWhen(true)] out DirectoryReference? OutputDir)
+		public bool TryGetOutputDir([NotNullWhen(true)] out DirectoryReference? outputDir)
 		{
-			string? OutputPath;
-			if (Properties.TryGetValue("OutputPath", out OutputPath))
+			string? outputPath;
+			if (Properties.TryGetValue("OutputPath", out outputPath))
 			{
-				OutputDir = DirectoryReference.Combine(ProjectPath.Directory, OutputPath);
+				outputDir = DirectoryReference.Combine(ProjectPath.Directory, outputPath);
 				return true;
 			}
 			else if (IsNetCoreProject())
 			{
-				string Configuration = Properties.ContainsKey("Configuration") ? Properties["Configuration"] : "Development";
-				OutputDir = DirectoryReference.Combine(ProjectPath.Directory, "bin", Configuration, Properties["TargetFramework"]);
+				string configuration = Properties.ContainsKey("Configuration") ? Properties["Configuration"] : "Development";
+				outputDir = DirectoryReference.Combine(ProjectPath.Directory, "bin", configuration, Properties["TargetFramework"]);
 				return true;
 			}
 			else
 			{
-				OutputDir = null;
+				outputDir = null;
 				return false;
 			}
 		}
@@ -145,15 +143,15 @@ namespace EpicGames.Core
 		/// Returns the assembly name used by this project
 		/// </summary>
 		/// <returns></returns>
-		public bool TryGetAssemblyName([NotNullWhen(true)] out string? AssemblyName)
+		public bool TryGetAssemblyName([NotNullWhen(true)] out string? assemblyName)
 		{
-			if (Properties.TryGetValue("AssemblyName", out AssemblyName))
+			if (Properties.TryGetValue("AssemblyName", out assemblyName))
 			{
 				return true;
 			}
 			else if (IsNetCoreProject())
 			{
-				AssemblyName = ProjectPath.GetFileNameWithoutExtension();
+				assemblyName = ProjectPath.GetFileNameWithoutExtension();
 				return true;
 			}
 			return false;
@@ -162,65 +160,65 @@ namespace EpicGames.Core
 		/// <summary>
 		/// Finds all build products from this project. This includes content and other assemblies marked to be copied local.
 		/// </summary>
-		/// <param name="OutputDir">The output directory</param>
-		/// <param name="BuildProducts">Receives the set of build products</param>
-		/// <param name="ProjectFileToInfo">Map of project file to information, to resolve build products from referenced projects copied locally</param>
-		public void FindBuildProducts(DirectoryReference OutputDir, HashSet<FileReference> BuildProducts, Dictionary<FileReference, CsProjectInfo> ProjectFileToInfo)
+		/// <param name="outputDir">The output directory</param>
+		/// <param name="buildProducts">Receives the set of build products</param>
+		/// <param name="projectFileToInfo">Map of project file to information, to resolve build products from referenced projects copied locally</param>
+		public void FindBuildProducts(DirectoryReference outputDir, HashSet<FileReference> buildProducts, Dictionary<FileReference, CsProjectInfo> projectFileToInfo)
 		{
 			// Add the standard build products
-			FindCompiledBuildProducts(OutputDir, BuildProducts);
+			FindCompiledBuildProducts(outputDir, buildProducts);
 
 			// Add the referenced assemblies which are marked to be copied into the output directory. This only happens for the main project, and does not happen for referenced projects.
-			foreach(KeyValuePair<FileReference, bool> Reference in References)
+			foreach(KeyValuePair<FileReference, bool> reference in References)
 			{
-				if (Reference.Value)
+				if (reference.Value)
 				{
-					FileReference OutputFile = FileReference.Combine(OutputDir, Reference.Key.GetFileName());
-					AddReferencedAssemblyAndSupportFiles(OutputFile, BuildProducts);
+					FileReference outputFile = FileReference.Combine(outputDir, reference.Key.GetFileName());
+					AddReferencedAssemblyAndSupportFiles(outputFile, buildProducts);
 				}
 			}
 
 			// Copy the build products for any referenced projects. Note that this does NOT operate recursively.
-			foreach(KeyValuePair<FileReference, bool> ProjectReference in ProjectReferences)
+			foreach(KeyValuePair<FileReference, bool> projectReference in ProjectReferences)
 			{
-				CsProjectInfo? OtherProjectInfo;
-				if(ProjectFileToInfo.TryGetValue(ProjectReference.Key, out OtherProjectInfo))
+				CsProjectInfo? otherProjectInfo;
+				if(projectFileToInfo.TryGetValue(projectReference.Key, out otherProjectInfo))
 				{
-					OtherProjectInfo.FindCompiledBuildProducts(OutputDir, BuildProducts);
+					otherProjectInfo.FindCompiledBuildProducts(outputDir, buildProducts);
 				}
 			}
 
 			// Add any copied content. This DOES operate recursively.
-			FindCopiedContent(OutputDir, BuildProducts, ProjectFileToInfo);
+			FindCopiedContent(outputDir, buildProducts, projectFileToInfo);
 		}
 
 		/// <summary>
 		/// Determines all the compiled build products (executable, etc...) directly built from this project.
 		/// </summary>
-		/// <param name="OutputDir">The output directory</param>
-		/// <param name="BuildProducts">Receives the set of build products</param>
-		public void FindCompiledBuildProducts(DirectoryReference OutputDir, HashSet<FileReference> BuildProducts)
+		/// <param name="outputDir">The output directory</param>
+		/// <param name="buildProducts">Receives the set of build products</param>
+		public void FindCompiledBuildProducts(DirectoryReference outputDir, HashSet<FileReference> buildProducts)
 		{
-			string? OutputType, AssemblyName;
-			if (Properties.TryGetValue("OutputType", out OutputType) && TryGetAssemblyName(out AssemblyName))
+			string? outputType, assemblyName;
+			if (Properties.TryGetValue("OutputType", out outputType) && TryGetAssemblyName(out assemblyName))
 			{
-				switch (OutputType)
+				switch (outputType)
 				{
 					case "Exe":
 					case "WinExe":
-						string ExecutableExtension = RuntimeInformation.IsOSPlatform(OSPlatform.Windows) ? ".exe" : "";
-						BuildProducts.Add(FileReference.Combine(OutputDir, AssemblyName + ExecutableExtension));
+						string executableExtension = RuntimeInformation.IsOSPlatform(OSPlatform.Windows) ? ".exe" : "";
+						buildProducts.Add(FileReference.Combine(outputDir, assemblyName + executableExtension));
 						// dotnet outputs a apphost executable and a dll with the actual assembly
-						AddOptionalBuildProduct(FileReference.Combine(OutputDir, AssemblyName + ".dll"), BuildProducts);
-						AddOptionalBuildProduct(FileReference.Combine(OutputDir, AssemblyName + ".pdb"), BuildProducts);
-						AddOptionalBuildProduct(FileReference.Combine(OutputDir, AssemblyName + ".exe.config"), BuildProducts);
-						AddOptionalBuildProduct(FileReference.Combine(OutputDir, AssemblyName + ".exe.mdb"), BuildProducts);
+						AddOptionalBuildProduct(FileReference.Combine(outputDir, assemblyName + ".dll"), buildProducts);
+						AddOptionalBuildProduct(FileReference.Combine(outputDir, assemblyName + ".pdb"), buildProducts);
+						AddOptionalBuildProduct(FileReference.Combine(outputDir, assemblyName + ".exe.config"), buildProducts);
+						AddOptionalBuildProduct(FileReference.Combine(outputDir, assemblyName + ".exe.mdb"), buildProducts);
 						break;
 					case "Library":
-						BuildProducts.Add(FileReference.Combine(OutputDir, AssemblyName + ".dll"));
-						AddOptionalBuildProduct(FileReference.Combine(OutputDir, AssemblyName + ".pdb"), BuildProducts);
-						AddOptionalBuildProduct(FileReference.Combine(OutputDir, AssemblyName + ".dll.config"), BuildProducts);
-						AddOptionalBuildProduct(FileReference.Combine(OutputDir, AssemblyName + ".dll.mdb"), BuildProducts);
+						buildProducts.Add(FileReference.Combine(outputDir, assemblyName + ".dll"));
+						AddOptionalBuildProduct(FileReference.Combine(outputDir, assemblyName + ".pdb"), buildProducts);
+						AddOptionalBuildProduct(FileReference.Combine(outputDir, assemblyName + ".dll.config"), buildProducts);
+						AddOptionalBuildProduct(FileReference.Combine(outputDir, assemblyName + ".dll.mdb"), buildProducts);
 						break;
 				}
 			}
@@ -229,28 +227,28 @@ namespace EpicGames.Core
 		/// <summary>
 		/// Finds all content which will be copied into the output directory for this project. This includes content from any project references as "copy local" recursively (though MSBuild only traverses a single reference for actual binaries, in such cases)
 		/// </summary>
-		/// <param name="OutputDir">The output directory</param>
+		/// <param name="outputDir">The output directory</param>
 		/// <param name="BuildProducts">Receives the set of build products</param>
-		/// <param name="ProjectFileToInfo">Map of project file to information, to resolve build products from referenced projects copied locally</param>
-		private void FindCopiedContent(DirectoryReference OutputDir, HashSet<FileReference> OutputFiles, Dictionary<FileReference, CsProjectInfo> ProjectFileToInfo)
+		/// <param name="projectFileToInfo">Map of project file to information, to resolve build products from referenced projects copied locally</param>
+		private void FindCopiedContent(DirectoryReference outputDir, HashSet<FileReference> outputFiles, Dictionary<FileReference, CsProjectInfo> projectFileToInfo)
 		{
 			// Copy any referenced projects too.
-			foreach(KeyValuePair<FileReference, bool> ProjectReference in ProjectReferences)
+			foreach(KeyValuePair<FileReference, bool> projectReference in ProjectReferences)
 			{
-				CsProjectInfo? OtherProjectInfo;
-				if(ProjectFileToInfo.TryGetValue(ProjectReference.Key, out OtherProjectInfo))
+				CsProjectInfo? otherProjectInfo;
+				if(projectFileToInfo.TryGetValue(projectReference.Key, out otherProjectInfo))
 				{
-					OtherProjectInfo.FindCopiedContent(OutputDir, OutputFiles, ProjectFileToInfo);
+					otherProjectInfo.FindCopiedContent(outputDir, outputFiles, projectFileToInfo);
 				}
 			}
 
 			// Add the content which is copied to the output directory
-			foreach (KeyValuePair<FileReference, bool> ContentReference in ContentReferences)
+			foreach (KeyValuePair<FileReference, bool> contentReference in ContentReferences)
 			{
-				FileReference ContentFile = ContentReference.Key;
-				if (ContentReference.Value)
+				FileReference contentFile = contentReference.Key;
+				if (contentReference.Value)
 				{
-					OutputFiles.Add(FileReference.Combine(OutputDir, ContentFile.GetFileName()));
+					outputFiles.Add(FileReference.Combine(outputDir, contentFile.GetFileName()));
 				}
 			}
 		}
@@ -258,22 +256,22 @@ namespace EpicGames.Core
 		/// <summary>
 		/// Adds the given file and any additional build products to the output set
 		/// </summary>
-		/// <param name="OutputFile">The assembly to add</param>
-		/// <param name="OutputFiles">Set to receive the file and support files</param>
-		static public void AddReferencedAssemblyAndSupportFiles(FileReference OutputFile, HashSet<FileReference> OutputFiles)
+		/// <param name="outputFile">The assembly to add</param>
+		/// <param name="outputFiles">Set to receive the file and support files</param>
+		public static void AddReferencedAssemblyAndSupportFiles(FileReference outputFile, HashSet<FileReference> outputFiles)
 		{
-			OutputFiles.Add(OutputFile);
+			outputFiles.Add(outputFile);
 
-			FileReference SymbolFile = OutputFile.ChangeExtension(".pdb");
-			if (FileReference.Exists(SymbolFile))
+			FileReference symbolFile = outputFile.ChangeExtension(".pdb");
+			if (FileReference.Exists(symbolFile))
 			{
-				OutputFiles.Add(SymbolFile);
+				outputFiles.Add(symbolFile);
 			}
 
-			FileReference DocumentationFile = OutputFile.ChangeExtension(".xml");
-			if (FileReference.Exists(DocumentationFile))
+			FileReference documentationFile = outputFile.ChangeExtension(".xml");
+			if (FileReference.Exists(documentationFile))
 			{
-				OutputFiles.Add(DocumentationFile);
+				outputFiles.Add(documentationFile);
 			}
 		}
 
@@ -285,10 +283,10 @@ namespace EpicGames.Core
 		{
 			bool bIsDotNetCoreProject = false;
 
-			string? TargetFramework;
-			if (Properties.TryGetValue("TargetFramework", out TargetFramework))
+			string? targetFramework;
+			if (Properties.TryGetValue("TargetFramework", out targetFramework))
 			{
-				bIsDotNetCoreProject = TargetFramework.ToLower().Contains("netstandard") || TargetFramework.ToLower().Contains("netcoreapp");
+				bIsDotNetCoreProject = targetFramework.ToLower().Contains("netstandard") || targetFramework.ToLower().Contains("netcoreapp");
 			}
 
 			return bIsDotNetCoreProject;
@@ -297,13 +295,13 @@ namespace EpicGames.Core
 		/// <summary>
 		/// Adds a build product to the output list if it exists
 		/// </summary>
-		/// <param name="BuildProduct">The build product to add</param>
-		/// <param name="BuildProducts">List of output build products</param>
-		public static void AddOptionalBuildProduct(FileReference BuildProduct, HashSet<FileReference> BuildProducts)
+		/// <param name="buildProduct">The build product to add</param>
+		/// <param name="buildProducts">List of output build products</param>
+		public static void AddOptionalBuildProduct(FileReference buildProduct, HashSet<FileReference> buildProducts)
 		{
-			if (FileReference.Exists(BuildProduct))
+			if (FileReference.Exists(buildProduct))
 			{
-				BuildProducts.Add(BuildProduct);
+				buildProducts.Add(buildProduct);
 			}
 		}
 
@@ -311,31 +309,31 @@ namespace EpicGames.Core
 		/// Parses supported elements from the children of the provided note. May recurse
 		/// for conditional elements.
 		/// </summary>
-		/// <param name="Node"></param>
-		/// <param name="ProjectInfo"></param>
-		static void ParseNode(XmlNode Node, CsProjectInfo ProjectInfo)
+		/// <param name="node"></param>
+		/// <param name="projectInfo"></param>
+		static void ParseNode(XmlNode node, CsProjectInfo projectInfo)
 		{
-			foreach (XmlElement Element in Node.ChildNodes.OfType<XmlElement>())
+			foreach (XmlElement element in node.ChildNodes.OfType<XmlElement>())
 			{
-				switch (Element.Name)
+				switch (element.Name)
 				{
 					case "PropertyGroup":
-						if (EvaluateCondition(Element, ProjectInfo))
+						if (EvaluateCondition(element, projectInfo))
 						{
-							ParsePropertyGroup(Element, ProjectInfo);
+							ParsePropertyGroup(element, projectInfo);
 						}
 						break;
 					case "ItemGroup":
-						if (EvaluateCondition(Element, ProjectInfo))
+						if (EvaluateCondition(element, projectInfo))
 						{
-							ParseItemGroup(ProjectInfo.ProjectPath.Directory, Element, ProjectInfo);
+							ParseItemGroup(projectInfo.ProjectPath.Directory, element, projectInfo);
 						}
 						break;
 					case "Choose":
 					case "When":
-						if (EvaluateCondition(Element, ProjectInfo))
+						if (EvaluateCondition(element, projectInfo))
 						{
-							ParseNode(Element, ProjectInfo);
+							ParseNode(element, projectInfo);
 						}
 						break;
 				}
@@ -345,64 +343,64 @@ namespace EpicGames.Core
 		/// <summary>
 		/// Reads project information for the given file.
 		/// </summary>
-		/// <param name="File">The project file to read</param>
-		/// <param name="Properties">Initial set of property values</param>
+		/// <param name="file">The project file to read</param>
+		/// <param name="properties">Initial set of property values</param>
 		/// <returns>The parsed project info</returns>
-		public static CsProjectInfo Read(FileReference File, Dictionary<string, string> Properties)
+		public static CsProjectInfo Read(FileReference file, Dictionary<string, string> properties)
 		{
-			CsProjectInfo? Project;
-			if(!TryRead(File, Properties, out Project))
+			CsProjectInfo? project;
+			if(!TryRead(file, properties, out project))
 			{
-				throw new Exception(String.Format("Unable to read '{0}'", File));
+				throw new Exception(String.Format("Unable to read '{0}'", file));
 			}
-			return Project;
+			return project;
 		}
 
 		/// <summary>
 		/// Attempts to read project information for the given file.
 		/// </summary>
-		/// <param name="File">The project file to read</param>
-		/// <param name="Properties">Initial set of property values</param>
-		/// <param name="OutProjectInfo">If successful, the parsed project info</param>
+		/// <param name="file">The project file to read</param>
+		/// <param name="properties">Initial set of property values</param>
+		/// <param name="outProjectInfo">If successful, the parsed project info</param>
 		/// <returns>True if the project was read successfully, false otherwise</returns>
-		public static bool TryRead(FileReference File, Dictionary<string, string> Properties, [NotNullWhen(true)] out CsProjectInfo? OutProjectInfo)
+		public static bool TryRead(FileReference file, Dictionary<string, string> properties, [NotNullWhen(true)] out CsProjectInfo? outProjectInfo)
 		{
 			// Read the project file
-			XmlDocument Document = new XmlDocument();
-			Document.Load(File.FullName);
+			XmlDocument document = new XmlDocument();
+			document.Load(file.FullName);
 
 			// Check the root element is the right type
 			//			HashSet<FileReference> ProjectBuildProducts = new HashSet<FileReference>();
-			if (Document.DocumentElement.Name != "Project")
+			if (document.DocumentElement.Name != "Project")
 			{
-				OutProjectInfo = null;
+				outProjectInfo = null;
 				return false;
 			}
 
 			// Parse the basic structure of the document, updating properties and recursing into other referenced projects as we go
-			CsProjectInfo ProjectInfo = new CsProjectInfo(Properties, File);
+			CsProjectInfo projectInfo = new CsProjectInfo(properties, file);
 
 			// Parse elements in the root node
-			ParseNode(Document.DocumentElement, ProjectInfo);
+			ParseNode(document.DocumentElement, projectInfo);
 
 			// Return the complete project
-			OutProjectInfo = ProjectInfo;
+			outProjectInfo = projectInfo;
 			return true;
 		}
 
 		/// <summary>
 		/// Parses a 'PropertyGroup' element.
 		/// </summary>
-		/// <param name="ParentElement">The parent 'PropertyGroup' element</param>
+		/// <param name="parentElement">The parent 'PropertyGroup' element</param>
 		/// <param name="Properties">Dictionary mapping property names to values</param>
-		static void ParsePropertyGroup(XmlElement ParentElement, CsProjectInfo ProjectInfo)
+		static void ParsePropertyGroup(XmlElement parentElement, CsProjectInfo projectInfo)
 		{
 			// We need to know the overridden output type and output path for the selected configuration.
-			foreach (XmlElement Element in ParentElement.ChildNodes.OfType<XmlElement>())
+			foreach (XmlElement element in parentElement.ChildNodes.OfType<XmlElement>())
 			{
-				if (EvaluateCondition(Element, ProjectInfo))
+				if (EvaluateCondition(element, projectInfo))
 				{
-					ProjectInfo.Properties[Element.Name] = ExpandProperties(Element.InnerText, ProjectInfo.Properties);
+					projectInfo.Properties[element.Name] = ExpandProperties(element.InnerText, projectInfo.Properties);
 				}
 			}
 		}
@@ -410,43 +408,43 @@ namespace EpicGames.Core
 		/// <summary>
 		/// Parses an 'ItemGroup' element.
 		/// </summary>
-		/// <param name="BaseDirectory">Base directory to resolve relative paths against</param>
-		/// <param name="ParentElement">The parent 'ItemGroup' element</param>
-		/// <param name="ProjectInfo">Project info object to be updated</param>
-		static void ParseItemGroup(DirectoryReference BaseDirectory, XmlElement ParentElement, CsProjectInfo ProjectInfo)
+		/// <param name="baseDirectory">Base directory to resolve relative paths against</param>
+		/// <param name="parentElement">The parent 'ItemGroup' element</param>
+		/// <param name="projectInfo">Project info object to be updated</param>
+		static void ParseItemGroup(DirectoryReference baseDirectory, XmlElement parentElement, CsProjectInfo projectInfo)
 		{
 			// Parse any external assembly references
-			foreach (XmlElement ItemElement in ParentElement.ChildNodes.OfType<XmlElement>())
+			foreach (XmlElement itemElement in parentElement.ChildNodes.OfType<XmlElement>())
 			{
-				switch (ItemElement.Name)
+				switch (itemElement.Name)
 				{
 					case "Reference":
 						// Reference to an external assembly
-						if (EvaluateCondition(ItemElement, ProjectInfo))
+						if (EvaluateCondition(itemElement, projectInfo))
 						{
-							ParseReference(BaseDirectory, ItemElement, ProjectInfo.References);
+							ParseReference(baseDirectory, itemElement, projectInfo.References);
 						}
 						break;
 					case "ProjectReference":
 						// Reference to another project
-						if (EvaluateCondition(ItemElement, ProjectInfo))
+						if (EvaluateCondition(itemElement, projectInfo))
 						{
-							ParseProjectReference(BaseDirectory, ItemElement, ProjectInfo.Properties, ProjectInfo.ProjectReferences);
+							ParseProjectReference(baseDirectory, itemElement, projectInfo.Properties, projectInfo.ProjectReferences);
 						}
 						break;
 					case "Compile":
 						// Reference to another project
-						if (EvaluateCondition(ItemElement, ProjectInfo))
+						if (EvaluateCondition(itemElement, projectInfo))
 						{
-							ParseCompileReference(BaseDirectory, ItemElement, ProjectInfo.CompileReferences);
+							ParseCompileReference(baseDirectory, itemElement, projectInfo.CompileReferences);
 						}
 						break;
 					case "Content":
 					case "None":
 						// Reference to another project
-						if (EvaluateCondition(ItemElement, ProjectInfo))
+						if (EvaluateCondition(itemElement, projectInfo))
 						{
-							ParseContent(BaseDirectory, ItemElement, ProjectInfo.ContentReferences);
+							ParseContent(baseDirectory, itemElement, projectInfo.ContentReferences);
 						}
 						break;
 				}
@@ -456,21 +454,21 @@ namespace EpicGames.Core
 		/// <summary>
 		/// Parses an assembly reference from a given 'Reference' element
 		/// </summary>
-		/// <param name="BaseDirectory">Directory to resolve relative paths against</param>
-		/// <param name="ParentElement">The parent 'Reference' element</param>
-		/// <param name="References">Dictionary of project files to a bool indicating whether the assembly should be copied locally to the referencing project.</param>
-		static void ParseReference(DirectoryReference BaseDirectory, XmlElement ParentElement, Dictionary<FileReference, bool> References)
+		/// <param name="baseDirectory">Directory to resolve relative paths against</param>
+		/// <param name="parentElement">The parent 'Reference' element</param>
+		/// <param name="references">Dictionary of project files to a bool indicating whether the assembly should be copied locally to the referencing project.</param>
+		static void ParseReference(DirectoryReference baseDirectory, XmlElement parentElement, Dictionary<FileReference, bool> references)
 		{
-			string? HintPath = UnescapeString(GetChildElementString(ParentElement, "HintPath", null));
-			if (!String.IsNullOrEmpty(HintPath))
+			string? hintPath = UnescapeString(GetChildElementString(parentElement, "HintPath", null));
+			if (!String.IsNullOrEmpty(hintPath))
 			{
 				// Don't include embedded assemblies; they aren't referenced externally by the compiled executable
-				bool bEmbedInteropTypes = GetChildElementBoolean(ParentElement, "EmbedInteropTypes", false);
+				bool bEmbedInteropTypes = GetChildElementBoolean(parentElement, "EmbedInteropTypes", false);
 				if(!bEmbedInteropTypes)
 				{
-					FileReference AssemblyFile = FileReference.Combine(BaseDirectory, HintPath);
-					bool bPrivate = GetChildElementBoolean(ParentElement, "Private", true);
-					References.Add(AssemblyFile, bPrivate);
+					FileReference assemblyFile = FileReference.Combine(baseDirectory, hintPath);
+					bool bPrivate = GetChildElementBoolean(parentElement, "Private", true);
+					references.Add(assemblyFile, bPrivate);
 				}
 			}
 		}
@@ -478,89 +476,88 @@ namespace EpicGames.Core
 		/// <summary>
 		/// Parses a project reference from a given 'ProjectReference' element
 		/// </summary>
-		/// <param name="BaseDirectory">Directory to resolve relative paths against</param>
-		/// <param name="ParentElement">The parent 'ProjectReference' element</param>
-		/// <param name="Properties">Dictionary of properties for parsing the file</param>
-		/// <param name="ProjectReferences">Dictionary of project files to a bool indicating whether the outputs of the project should be copied locally to the referencing project.</param>
-		static void ParseProjectReference(DirectoryReference BaseDirectory, XmlElement ParentElement, Dictionary<string, string> Properties, Dictionary<FileReference, bool> ProjectReferences)
+		/// <param name="baseDirectory">Directory to resolve relative paths against</param>
+		/// <param name="parentElement">The parent 'ProjectReference' element</param>
+		/// <param name="properties">Dictionary of properties for parsing the file</param>
+		/// <param name="projectReferences">Dictionary of project files to a bool indicating whether the outputs of the project should be copied locally to the referencing project.</param>
+		static void ParseProjectReference(DirectoryReference baseDirectory, XmlElement parentElement, Dictionary<string, string> properties, Dictionary<FileReference, bool> projectReferences)
 		{
-			string? IncludePath = UnescapeString(ParentElement.GetAttribute("Include"));
-			if (!String.IsNullOrEmpty(IncludePath))
+			string? includePath = UnescapeString(parentElement.GetAttribute("Include"));
+			if (!String.IsNullOrEmpty(includePath))
 			{
-				FileReference ProjectFile = FileReference.Combine(BaseDirectory, ExpandProperties(IncludePath, Properties));
-				bool bPrivate = GetChildElementBoolean(ParentElement, "Private", true);
-				ProjectReferences[ProjectFile] = bPrivate;
+				FileReference projectFile = FileReference.Combine(baseDirectory, ExpandProperties(includePath, properties));
+				bool bPrivate = GetChildElementBoolean(parentElement, "Private", true);
+				projectReferences[projectFile] = bPrivate;
 			}
 		}
 
-
 		/// recursive helper used by the function below that will append RemainingComponents one by one to ExistingPath, 
 		/// expanding wildcards as necessary. The complete list of files that match the complete path is returned out OutFoundFiles
-		static void ProcessPathComponents(DirectoryReference ExistingPath, IEnumerable<string> RemainingComponents, List<FileReference> OutFoundFiles)
+		static void ProcessPathComponents(DirectoryReference existingPath, IEnumerable<string> remainingComponents, List<FileReference> outFoundFiles)
 		{
-			if (!RemainingComponents.Any())
+			if (!remainingComponents.Any())
 			{
 				return;
 			}
 
 			// take a look at the first component
-			string CurrentComponent = RemainingComponents.First();
-			RemainingComponents = RemainingComponents.Skip(1);
+			string currentComponent = remainingComponents.First();
+			remainingComponents = remainingComponents.Skip(1);
 
 			// If no other components then this is either a file pattern or a greedy pattern
-			if (!RemainingComponents.Any())
+			if (!remainingComponents.Any())
 			{
 				// ** means include all files under this tree, so enumerate them all
-				if (CurrentComponent.Contains("**"))
+				if (currentComponent.Contains("**"))
 				{
-					OutFoundFiles.AddRange(DirectoryReference.EnumerateFiles(ExistingPath, "*", SearchOption.AllDirectories));
+					outFoundFiles.AddRange(DirectoryReference.EnumerateFiles(existingPath, "*", SearchOption.AllDirectories));
 				}
 				else
 				{
 					// easy, a regular path with a file that may or may not be a wildcard
-					OutFoundFiles.AddRange(DirectoryReference.EnumerateFiles(ExistingPath, CurrentComponent));
+					outFoundFiles.AddRange(DirectoryReference.EnumerateFiles(existingPath, currentComponent));
 				}
 			}
 			else
 			{
 				// new component contains a wildcard, and based on the above we know there are more entries so find 
 				// matching directories
-				if (CurrentComponent.Contains("*"))
+				if (currentComponent.Contains("*"))
 				{
 					// ** means all directories, no matter how deep
-					SearchOption Option = CurrentComponent == "**" ? SearchOption.AllDirectories : SearchOption.TopDirectoryOnly;
+					SearchOption option = currentComponent == "**" ? SearchOption.AllDirectories : SearchOption.TopDirectoryOnly;
 
-					IEnumerable<DirectoryReference> Directories = DirectoryReference.EnumerateDirectories(ExistingPath, CurrentComponent, Option);
+					IEnumerable<DirectoryReference> directories = DirectoryReference.EnumerateDirectories(existingPath, currentComponent, option);
 
 					// if we searched all directories regardless of depth, the rest of the components other than the last (file) are irrelevant
-					if (Option == SearchOption.AllDirectories)
+					if (option == SearchOption.AllDirectories)
 					{
-						RemainingComponents = new[] { RemainingComponents.Last() };
+						remainingComponents = new[] { remainingComponents.Last() };
 
 						// ** includes files in the current directory too
-						Directories = Directories.Concat(new[] { ExistingPath });
+						directories = directories.Concat(new[] { existingPath });
 					}
 
-					foreach (DirectoryReference Dir in Directories)
+					foreach (DirectoryReference dir in directories)
 					{
-						ProcessPathComponents(Dir, RemainingComponents, OutFoundFiles);
+						ProcessPathComponents(dir, remainingComponents, outFoundFiles);
 					}
 				}
 				else
 				{
 					// add this component to our path and recurse.
-					ExistingPath = DirectoryReference.Combine(ExistingPath, CurrentComponent);
+					existingPath = DirectoryReference.Combine(existingPath, currentComponent);
 
 					// but... we can just take all the next components that don't have wildcards in them instead of recursing
 					// into each one!
-					IEnumerable<string> NonWildCardComponents = RemainingComponents.TakeWhile(C => !C.Contains("*"));
-					RemainingComponents = RemainingComponents.Skip(NonWildCardComponents.Count());
+					IEnumerable<string> nonWildCardComponents = remainingComponents.TakeWhile(c => !c.Contains("*"));
+					remainingComponents = remainingComponents.Skip(nonWildCardComponents.Count());
 
-					ExistingPath = DirectoryReference.Combine(ExistingPath, NonWildCardComponents.ToArray());
+					existingPath = DirectoryReference.Combine(existingPath, nonWildCardComponents.ToArray());
 
-					if (Directory.Exists(ExistingPath.FullName))
+					if (Directory.Exists(existingPath.FullName))
 					{
-						ProcessPathComponents(ExistingPath, RemainingComponents, OutFoundFiles);
+						ProcessPathComponents(existingPath, remainingComponents, outFoundFiles);
 					}
 				}
 			}
@@ -576,44 +573,44 @@ namespace EpicGames.Core
 		/// Foo/**
 		/// (the last means include all files under the path).
 		/// </summary>
-		/// <param name="InPath">Path specifier to process</param>
+		/// <param name="inPath">Path specifier to process</param>
 		/// <returns></returns>
-		static IEnumerable<FileReference> FindMatchingFiles(FileReference InPath)
+		static IEnumerable<FileReference> FindMatchingFiles(FileReference inPath)
 		{
-			List<FileReference> FoundFiles = new List<FileReference>();		
+			List<FileReference> foundFiles = new List<FileReference>();		
 
 			// split off the drive root
-			string DriveRoot = Path.GetPathRoot(InPath.FullName)!;
+			string driveRoot = Path.GetPathRoot(inPath.FullName)!;
 
 			// break the rest of the path into components
-			string[] PathComponents = InPath.FullName.Substring(DriveRoot.Length).Split(new char[] { Path.DirectorySeparatorChar, Path.AltDirectorySeparatorChar });
+			string[] pathComponents = inPath.FullName.Substring(driveRoot.Length).Split(new char[] { Path.DirectorySeparatorChar, Path.AltDirectorySeparatorChar });
 
 			// Process all the components recursively
-			ProcessPathComponents(new DirectoryReference(DriveRoot), PathComponents, FoundFiles);
+			ProcessPathComponents(new DirectoryReference(driveRoot), pathComponents, foundFiles);
 
-			return FoundFiles;
+			return foundFiles;
 		}
 
 		/// <summary>
 		/// Parses a project reference from a given 'ProjectReference' element
 		/// </summary>
-		/// <param name="BaseDirectory">Directory to resolve relative paths against</param>
-		/// <param name="ParentElement">The parent 'ProjectReference' element</param>
-		/// <param name="CompileReferences">List of source files.</param>
-		static void ParseCompileReference(DirectoryReference BaseDirectory, XmlElement ParentElement, List<FileReference> CompileReferences)
+		/// <param name="baseDirectory">Directory to resolve relative paths against</param>
+		/// <param name="parentElement">The parent 'ProjectReference' element</param>
+		/// <param name="compileReferences">List of source files.</param>
+		static void ParseCompileReference(DirectoryReference baseDirectory, XmlElement parentElement, List<FileReference> compileReferences)
 		{
-			string? IncludePath = UnescapeString(ParentElement.GetAttribute("Include"));
-			if (!String.IsNullOrEmpty(IncludePath))
+			string? includePath = UnescapeString(parentElement.GetAttribute("Include"));
+			if (!String.IsNullOrEmpty(includePath))
 			{
-				FileReference SourceFile = FileReference.Combine(BaseDirectory, IncludePath);
+				FileReference sourceFile = FileReference.Combine(baseDirectory, includePath);
 
-				if (SourceFile.FullName.Contains("*"))
+				if (sourceFile.FullName.Contains("*"))
 				{
-					CompileReferences.AddRange(FindMatchingFiles(SourceFile));
+					compileReferences.AddRange(FindMatchingFiles(sourceFile));
 				}
 				else
 				{
-					CompileReferences.Add(SourceFile);
+					compileReferences.Add(sourceFile);
 				}
 			}
 		}
@@ -621,113 +618,113 @@ namespace EpicGames.Core
 		/// <summary>
 		/// Parses an assembly reference from a given 'Content' element
 		/// </summary>
-		/// <param name="BaseDirectory">Directory to resolve relative paths against</param>
-		/// <param name="ParentElement">The parent 'Content' element</param>
-		/// <param name="Contents">Dictionary of project files to a bool indicating whether the assembly should be copied locally to the referencing project.</param>
-		static void ParseContent(DirectoryReference BaseDirectory, XmlElement ParentElement, Dictionary<FileReference, bool> Contents)
+		/// <param name="baseDirectory">Directory to resolve relative paths against</param>
+		/// <param name="parentElement">The parent 'Content' element</param>
+		/// <param name="contents">Dictionary of project files to a bool indicating whether the assembly should be copied locally to the referencing project.</param>
+		static void ParseContent(DirectoryReference baseDirectory, XmlElement parentElement, Dictionary<FileReference, bool> contents)
 		{
-			string? IncludePath = UnescapeString(ParentElement.GetAttribute("Include"));
-			if (!String.IsNullOrEmpty(IncludePath))
+			string? includePath = UnescapeString(parentElement.GetAttribute("Include"));
+			if (!String.IsNullOrEmpty(includePath))
 			{
-				string? CopyTo = GetChildElementString(ParentElement, "CopyToOutputDirectory", null);
-				bool ShouldCopy = !String.IsNullOrEmpty(CopyTo) && (CopyTo.Equals("Always", StringComparison.InvariantCultureIgnoreCase) || CopyTo.Equals("PreserveNewest", StringComparison.InvariantCultureIgnoreCase));
-				FileReference ContentFile = FileReference.Combine(BaseDirectory, IncludePath);
-				Contents.Add(ContentFile, ShouldCopy);
+				string? copyTo = GetChildElementString(parentElement, "CopyToOutputDirectory", null);
+				bool shouldCopy = !String.IsNullOrEmpty(copyTo) && (copyTo.Equals("Always", StringComparison.InvariantCultureIgnoreCase) || copyTo.Equals("PreserveNewest", StringComparison.InvariantCultureIgnoreCase));
+				FileReference contentFile = FileReference.Combine(baseDirectory, includePath);
+				contents.Add(contentFile, shouldCopy);
 			}
 		}
 
 		/// <summary>
 		/// Reads the inner text of a child XML element
 		/// </summary>
-		/// <param name="ParentElement">The parent element to check</param>
-		/// <param name="Name">Name of the child element</param>
-		/// <param name="DefaultValue">Default value to return if the child element is missing</param>
+		/// <param name="parentElement">The parent element to check</param>
+		/// <param name="name">Name of the child element</param>
+		/// <param name="defaultValue">Default value to return if the child element is missing</param>
 		/// <returns>The contents of the child element, or default value if it's not present</returns>
-		static string? GetChildElementString(XmlElement ParentElement, string Name, string? DefaultValue)
+		static string? GetChildElementString(XmlElement parentElement, string name, string? defaultValue)
 		{
-			XmlElement ChildElement = ParentElement.ChildNodes.OfType<XmlElement>().FirstOrDefault(x => x.Name == Name);
-			if (ChildElement == null)
+			XmlElement childElement = parentElement.ChildNodes.OfType<XmlElement>().FirstOrDefault(x => x.Name == name);
+			if (childElement == null)
 			{
-				return DefaultValue;
+				return defaultValue;
 			}
 			else
 			{
-				return ChildElement.InnerText ?? DefaultValue;
+				return childElement.InnerText ?? defaultValue;
 			}
 		}
 
 		/// <summary>
 		/// Read a child XML element with the given name, and parse it as a boolean.
 		/// </summary>
-		/// <param name="ParentElement">Parent element to check</param>
-		/// <param name="Name">Name of the child element to look for</param>
-		/// <param name="DefaultValue">Default value to return if the element is missing or not a valid bool</param>
+		/// <param name="parentElement">Parent element to check</param>
+		/// <param name="name">Name of the child element to look for</param>
+		/// <param name="defaultValue">Default value to return if the element is missing or not a valid bool</param>
 		/// <returns>The parsed boolean, or the default value</returns>
-		static bool GetChildElementBoolean(XmlElement ParentElement, string Name, bool DefaultValue)
+		static bool GetChildElementBoolean(XmlElement parentElement, string name, bool defaultValue)
 		{
-			string? Value = GetChildElementString(ParentElement, Name, null);
-			if (Value == null)
+			string? value = GetChildElementString(parentElement, name, null);
+			if (value == null)
 			{
-				return DefaultValue;
+				return defaultValue;
 			}
-			else if (Value.Equals("True", StringComparison.InvariantCultureIgnoreCase))
+			else if (value.Equals("True", StringComparison.InvariantCultureIgnoreCase))
 			{
 				return true;
 			}
-			else if (Value.Equals("False", StringComparison.InvariantCultureIgnoreCase))
+			else if (value.Equals("False", StringComparison.InvariantCultureIgnoreCase))
 			{
 				return false;
 			}
 			else
 			{
-				return DefaultValue;
+				return defaultValue;
 			}
 		}
 
 		/// <summary>
 		/// Evaluate whether the optional MSBuild condition on an XML element evaluates to true. Currently only supports 'ABC' == 'DEF' style expressions, but can be expanded as needed.
 		/// </summary>
-		/// <param name="Element">The XML element to check</param>
+		/// <param name="element">The XML element to check</param>
 		/// <param name="Properties">Dictionary mapping from property names to values.</param>
 		/// <returns></returns>
-		static bool EvaluateCondition(XmlElement Element, CsProjectInfo ProjectInfo)
+		static bool EvaluateCondition(XmlElement element, CsProjectInfo projectInfo)
 		{
 			// Read the condition attribute. If it's not present, assume it evaluates to true.
-			string Condition = Element.GetAttribute("Condition");
-			if (String.IsNullOrEmpty(Condition))
+			string condition = element.GetAttribute("Condition");
+			if (String.IsNullOrEmpty(condition))
 			{
 				return true;
 			}
 
 			// Expand all the properties
-			Condition = ExpandProperties(Condition, ProjectInfo.Properties);
+			condition = ExpandProperties(condition, projectInfo.Properties);
 
 			// Parse literal true/false values
-			bool OutResult;
-			if (bool.TryParse(Condition, out OutResult))
+			bool outResult;
+			if (Boolean.TryParse(condition, out outResult))
 			{
-				return OutResult;
+				return outResult;
 			}
 
 			// Tokenize the condition
-			string[] Tokens = Tokenize(Condition);
+			string[] tokens = Tokenize(condition);
 
-			char[] TokenQuotes = new[] { '\'', '(', ')', '{', '}', '[', ']' };
+			char[] tokenQuotes = new[] { '\'', '(', ')', '{', '}', '[', ']' };
 			
 			// Try to evaluate it. We only support a very limited class of condition expressions at the moment, but it's enough to parse standard projects
 			bool bResult;
 
 			// Handle Exists('Platform\Windows\Gauntlet.TargetDeviceWindows.cs')
-			if (Tokens[0] == "Exists")
+			if (tokens[0] == "Exists")
 			{
 				// remove all quotes, apostrophes etc that are either tokens or wrap tokens (The Tokenize() function is a bit suspect).
-				string[] Arguments = Tokens.Select(S => S.Trim(TokenQuotes)).Where(S => S.Length > 0).ToArray();
+				string[] arguments = tokens.Select(s => s.Trim(tokenQuotes)).Where(s => s.Length > 0).ToArray();
 
-				if (Tokens.Length > 1)
+				if (tokens.Length > 1)
 				{
-					FileSystemReference Dependency = DirectoryReference.Combine(ProjectInfo.ProjectPath.Directory, Arguments[1]);
+					FileSystemReference dependency = DirectoryReference.Combine(projectInfo.ProjectPath.Directory, arguments[1]);
 
-					if (File.Exists(Dependency.FullName) || Directory.Exists(Dependency.FullName))
+					if (File.Exists(dependency.FullName) || Directory.Exists(dependency.FullName))
 					{
 						return true;
 					}
@@ -736,18 +733,18 @@ namespace EpicGames.Core
 				}
 			}
 
-			if (Tokens.Length == 3 && Tokens[0].StartsWith("'") && Tokens[1] == "==" && Tokens[2].StartsWith("'"))
+			if (tokens.Length == 3 && tokens[0].StartsWith("'") && tokens[1] == "==" && tokens[2].StartsWith("'"))
 			{
-				bResult = String.Compare(Tokens[0], Tokens[2], StringComparison.InvariantCultureIgnoreCase) == 0;
+				bResult = String.Compare(tokens[0], tokens[2], StringComparison.InvariantCultureIgnoreCase) == 0;
 			}
-			else if (Tokens.Length == 3 && Tokens[0].StartsWith("'") && Tokens[1] == "!=" && Tokens[2].StartsWith("'"))
+			else if (tokens.Length == 3 && tokens[0].StartsWith("'") && tokens[1] == "!=" && tokens[2].StartsWith("'"))
 			{
-				bResult = String.Compare(Tokens[0], Tokens[2], StringComparison.InvariantCultureIgnoreCase) != 0;
+				bResult = String.Compare(tokens[0], tokens[2], StringComparison.InvariantCultureIgnoreCase) != 0;
 			}
 			else
 			{
-				string Msg = string.Format("Couldn't parse condition {0} in project file {1}", Element.ToString(), ProjectInfo.ProjectPath);
-				throw new Exception(Msg);
+				string msg = String.Format("Couldn't parse condition {0} in project file {1}", element.ToString(), projectInfo.ProjectPath);
+				throw new Exception(msg);
 			}
 			return bResult;
 		}
@@ -755,304 +752,305 @@ namespace EpicGames.Core
 		/// <summary>
 		/// Expand MSBuild properties within a string. If referenced properties are not in this dictionary, the process' environment variables are expanded. Unknown properties are expanded to an empty string.
 		/// </summary>
-		/// <param name="Text">The input string to expand</param>
-		/// <param name="Properties">Dictionary mapping from property names to values.</param>
+		/// <param name="text">The input string to expand</param>
+		/// <param name="properties">Dictionary mapping from property names to values.</param>
 		/// <returns>String with all properties expanded.</returns>
-		static string ExpandProperties(string Text, Dictionary<string, string> Properties)
+		static string ExpandProperties(string text, Dictionary<string, string> properties)
 		{
-			string NewText = Text;
-			for (int Idx = NewText.IndexOf("$("); Idx != -1; Idx = NewText.IndexOf("$(", Idx))
+			string newText = text;
+			for (int idx = newText.IndexOf("$("); idx != -1; idx = newText.IndexOf("$(", idx))
 			{
 				// Find the end of the variable name, accounting for changes in scope
-				int EndIdx = Idx + 2;
-				for(int Depth = 1; Depth > 0; EndIdx++)
+				int endIdx = idx + 2;
+				for(int depth = 1; depth > 0; endIdx++)
 				{
-					if(EndIdx == NewText.Length)
+					if(endIdx == newText.Length)
 					{
 						throw new Exception("Encountered end of string while expanding properties");
 					}
-					else if(NewText[EndIdx] == '(')
+					else if(newText[endIdx] == '(')
 					{
-						Depth++;
+						depth++;
 					}
-					else if(NewText[EndIdx] == ')')
+					else if(newText[endIdx] == ')')
 					{
-						Depth--;
+						depth--;
 					}
 				}
 
 				// Convert the property name to tokens
-				string[] Tokens = Tokenize(NewText.Substring(Idx + 2, (EndIdx - 1) - (Idx + 2)));
+				string[] tokens = Tokenize(newText.Substring(idx + 2, (endIdx - 1) - (idx + 2)));
 
 				// Make sure the first token is a valid property name
-				if(Tokens.Length == 0 || !(Char.IsLetter(Tokens[0][0]) || Tokens[0][0] == '_' || Tokens[0][0] == '[' ))
+				if(tokens.Length == 0 || !(Char.IsLetter(tokens[0][0]) || tokens[0][0] == '_' || tokens[0][0] == '[' ))
 				{
-					throw new Exception(String.Format("Invalid property name '{0}' in .csproj file", Tokens[0]));
+					throw new Exception(String.Format("Invalid property name '{0}' in .csproj file", tokens[0]));
 				}
 
 				// Find the value for it, either from the dictionary or the environment block
-				string Value;
-				if (Properties.TryGetValue(Tokens[0], out string? RetrievedValue))
+				string value;
+				if (properties.TryGetValue(tokens[0], out string? retrievedValue))
 				{
-					Value = RetrievedValue;
+					value = retrievedValue;
 				}
 				else
 				{
-					Value = Environment.GetEnvironmentVariable(Tokens[0]) ?? "";
+					value = Environment.GetEnvironmentVariable(tokens[0]) ?? "";
 				}
 
 				// Evaluate any functions within it
-				int TokenIdx = 1;
-				while(TokenIdx + 3 < Tokens.Length && Tokens[TokenIdx] == "." && Tokens[TokenIdx + 2] == "(")
+				int tokenIdx = 1;
+				while(tokenIdx + 3 < tokens.Length && tokens[tokenIdx] == "." && tokens[tokenIdx + 2] == "(")
 				{
 					// Read the method name
-					string MethodName = Tokens[TokenIdx + 1];
+					string methodName = tokens[tokenIdx + 1];
 
 					// Skip to the first argument
-					TokenIdx += 3;
+					tokenIdx += 3;
 
 					// Parse any arguments
-					List<object> Arguments = new List<object>();
-					if(Tokens[TokenIdx] != ")")
+					List<object> arguments = new List<object>();
+					if(tokens[tokenIdx] != ")")
 					{
-						Arguments.Add(ParseArgument(Tokens[TokenIdx]));
-						TokenIdx++;
+						arguments.Add(ParseArgument(tokens[tokenIdx]));
+						tokenIdx++;
 
-						while(TokenIdx + 1 < Tokens.Length && Tokens[TokenIdx] == ",")
+						while(tokenIdx + 1 < tokens.Length && tokens[tokenIdx] == ",")
 						{
-							Arguments.Add(ParseArgument(Tokens[TokenIdx + 2]));
-							TokenIdx += 2;
+							arguments.Add(ParseArgument(tokens[tokenIdx + 2]));
+							tokenIdx += 2;
 						}
 
-						if(Tokens[TokenIdx] != ")")
+						if(tokens[tokenIdx] != ")")
 						{
 							throw new Exception("Missing closing parenthesis in condition");
 						}
 					}
 
 					// Skip over the closing parenthesis
-					TokenIdx++;
+					tokenIdx++;
 
 					// Execute the method
 					try
 					{
-						Value = typeof(string).InvokeMember(MethodName, System.Reflection.BindingFlags.Instance | System.Reflection.BindingFlags.Public | System.Reflection.BindingFlags.InvokeMethod, Type.DefaultBinder, Value, Arguments.ToArray())!.ToString()!;
+						value = typeof(string).InvokeMember(methodName, System.Reflection.BindingFlags.Instance | System.Reflection.BindingFlags.Public | System.Reflection.BindingFlags.InvokeMethod, Type.DefaultBinder, value, arguments.ToArray())!.ToString()!;
 					}
-					catch(Exception Ex)
+					catch(Exception ex)
 					{
-						throw new Exception(String.Format("Unable to evaluate condition '{0}'", Text), Ex);
+						throw new Exception(String.Format("Unable to evaluate condition '{0}'", text), ex);
 					}
 				}
 
-				if (TokenIdx < Tokens.Length && Tokens[TokenIdx] == ":")
+				if (tokenIdx < tokens.Length && tokens[tokenIdx] == ":")
 				{
-					TokenIdx = Tokens.Length;
+					tokenIdx = tokens.Length;
 				}
 
 				// Make sure there's nothing left over
-				if(TokenIdx != Tokens.Length)
+				if(tokenIdx != tokens.Length)
 				{
-					throw new Exception(String.Format("Unable to parse token '{0}'", NewText));
+					throw new Exception(String.Format("Unable to parse token '{0}'", newText));
 				}
 
 				// Replace the variable with its value
-				NewText = NewText.Substring(0, Idx) + Value + NewText.Substring(EndIdx);
+				newText = newText.Substring(0, idx) + value + newText.Substring(endIdx);
 
 				// Make sure we skip over the expanded variable; we don't want to recurse on it.
-				Idx += Value.Length;
+				idx += value.Length;
 			}
-			return NewText;
+			return newText;
 		}
 
 		/// <summary>
 		/// Parse an argument into a framework type
 		/// </summary>
-		/// <param name="Token">The token to parse</param>
+		/// <param name="token">The token to parse</param>
 		/// <returns>The argument object</returns>
-		static object ParseArgument(string Token)
+		static object ParseArgument(string token)
 		{
 			// Try to parse a string
-			if(Token.Length > 2 && Token[0] == '\'' && Token[Token.Length - 1] == '\'')
+			if(token.Length > 2 && token[0] == '\'' && token[^1] == '\'')
 			{
-				return Token.Substring(1, Token.Length - 2);
+				return token.Substring(1, token.Length - 2);
 			}
 
 			// Try to parse an integer
-			int Value;
-			if(int.TryParse(Token, out Value))
+			int value;
+			if(Int32.TryParse(token, out value))
 			{
-				return Value;
+				return value;
 			}
 
 			// Otherwise throw an exception
-			throw new Exception(String.Format("Unable to parse token '{0}' into a .NET framework type", Token));
+			throw new Exception(String.Format("Unable to parse token '{0}' into a .NET framework type", token));
 		}
 
 		/// <summary>
 		/// Split an MSBuild condition into tokens
 		/// </summary>
-		/// <param name="Condition">The condition expression</param>
+		/// <param name="condition">The condition expression</param>
 		/// <returns>Array of the parsed tokens</returns>
-		static string[] Tokenize(string Condition)
+		static string[] Tokenize(string condition)
 		{
-			List<string> Tokens = new List<string>();
-			for (int Idx = 0; Idx < Condition.Length; )
+			List<string> tokens = new List<string>();
+			for (int idx = 0; idx < condition.Length; )
 			{
-				if(Char.IsWhiteSpace(Condition[Idx]))
+				if(Char.IsWhiteSpace(condition[idx]))
 				{
 					// Whitespace
-					Idx++;
+					idx++;
 				}
-				else if (Idx + 1 < Condition.Length && Condition[Idx] == '=' && Condition[Idx + 1] == '=')
+				else if (idx + 1 < condition.Length && condition[idx] == '=' && condition[idx + 1] == '=')
 				{
 					// "==" operator
-					Idx += 2;
-					Tokens.Add("==");
+					idx += 2;
+					tokens.Add("==");
 				}
-				else if (Idx + 1 < Condition.Length && Condition[Idx] == '!' && Condition[Idx + 1] == '=')
+				else if (idx + 1 < condition.Length && condition[idx] == '!' && condition[idx + 1] == '=')
 				{
 					// "!=" operator
-					Idx += 2;
-					Tokens.Add("!=");
+					idx += 2;
+					tokens.Add("!=");
 				}
-				else if (Condition[Idx] == '\'')
+				else if (condition[idx] == '\'')
 				{
 					// Quoted string
-					int StartIdx = Idx++;
-					for(;;Idx++)
+					int startIdx = idx++;
+					for(;;idx++)
 					{
-						if(Idx == Condition.Length)
+						if(idx == condition.Length)
 						{
-							throw new Exception(String.Format("Missing end quote in condition string ('{0}')", Condition));
+							throw new Exception(String.Format("Missing end quote in condition string ('{0}')", condition));
 						}
-						if(Condition[Idx] == '\'')
+						if(condition[idx] == '\'')
 						{
 							break;
 						}
 					}
-					Idx++;
-					Tokens.Add(Condition.Substring(StartIdx, Idx - StartIdx));
+					idx++;
+					tokens.Add(condition.Substring(startIdx, idx - startIdx));
 				}
-				else if (Condition[Idx] == '[')
+				else if (condition[idx] == '[')
 				{
+#pragma warning disable IDE0059 // Unnecessary assignment of a value
 					// static property function invoke
 					// format: [Class]::Property
 					// alternatively: [Class]::Method()
 					// we consider the entire invocation to be a single token
-					int StartIdx = Idx++;
-					int ClassEndIdx = 0;
-					int MethodEndIdx = 0;
-					int MethodArgsEndIdx = 0;
-					for (; ; Idx++)
+					int startIdx = idx++;
+					int classEndIdx = 0;
+					int methodEndIdx = 0;
+					int methodArgsEndIdx = 0;
+					for (; ; idx++)
 					{
-						while (Idx < Condition.Length && (Char.IsLetterOrDigit(Condition[Idx]) || Condition[Idx] == '_'))
+						while (idx < condition.Length && (Char.IsLetterOrDigit(condition[idx]) || condition[idx] == '_'))
 						{
-							Idx++;
+							idx++;
 						}
 
-						if (Idx == Condition.Length)
+						if (idx == condition.Length)
 						{
-							throw new Exception(String.Format("Found end of condition when searching for end of static property function for condition string ('{0}')", Condition));
+							throw new Exception(String.Format("Found end of condition when searching for end of static property function for condition string ('{0}')", condition));
 						}
-						if (Condition[Idx] == ']')
+						if (condition[idx] == ']')
 						{
-							ClassEndIdx = Idx;
-							Idx++;
+							classEndIdx = idx;
+							idx++;
 							break;
 						}
 					}
 
 					// skip ::
-					if (Condition[Idx] != ':')
+					if (condition[idx] != ':')
 					{
-						throw new Exception(String.Format("Unexpected format of static property function, expected :: after class declaration in condition string ('{0}')", Condition));
+						throw new Exception(String.Format("Unexpected format of static property function, expected :: after class declaration in condition string ('{0}')", condition));
 					}
-					Idx += 2;
+					idx += 2;
 
-					while (Idx < Condition.Length && (Char.IsLetterOrDigit(Condition[Idx]) || Condition[Idx] == '_'))
+					while (idx < condition.Length && (Char.IsLetterOrDigit(condition[idx]) || condition[idx] == '_'))
 					{
-						Idx++;
+						idx++;
 					}
 
-					MethodEndIdx = Idx;
+					methodEndIdx = idx;
 
-					if (Condition[Idx] == '(')
+					if (condition[idx] == '(')
 					{
 						// a method invoke
-						for (; ; Idx++)
+						for (; ; idx++)
 						{
-							while (Idx < Condition.Length && (Char.IsLetterOrDigit(Condition[Idx]) || Condition[Idx] == '_'))
+							while (idx < condition.Length && (Char.IsLetterOrDigit(condition[idx]) || condition[idx] == '_'))
 							{
-								Idx++;
+								idx++;
 							}
 
-							if (Idx == Condition.Length)
+							if (idx == condition.Length)
 							{
-								throw new Exception(String.Format("Found end of condition when searching for ) to indicate end of arguments to static property function for condition string ('{0}')", Condition));
+								throw new Exception(String.Format("Found end of condition when searching for ) to indicate end of arguments to static property function for condition string ('{0}')", condition));
 							}
-							if (Condition[Idx] == ')')
+							if (condition[idx] == ')')
 							{
-								MethodArgsEndIdx = Idx;
-								Idx++;
+								methodArgsEndIdx = idx;
+								idx++;
 								break;
 							}
 						}
-
 					}
+#pragma warning restore IDE0059 // Unnecessary assignment of a value
 
-					Idx++;
-					Tokens.Add(Condition.Substring(StartIdx, Idx - StartIdx));
+					idx++;
+					tokens.Add(condition.Substring(startIdx, idx - startIdx));
 				}
-				else if(Char.IsLetterOrDigit(Condition[Idx]) || Condition[Idx] == '_')
+				else if(Char.IsLetterOrDigit(condition[idx]) || condition[idx] == '_')
 				{
 					// Identifier or number
-					int StartIdx = Idx++;
-					while(Idx < Condition.Length && (Char.IsLetterOrDigit(Condition[Idx]) || Condition[Idx] == '_'))
+					int startIdx = idx++;
+					while(idx < condition.Length && (Char.IsLetterOrDigit(condition[idx]) || condition[idx] == '_'))
 					{
-						Idx++;
+						idx++;
 					}
-					Tokens.Add(Condition.Substring(StartIdx, Idx - StartIdx));
+					tokens.Add(condition.Substring(startIdx, idx - startIdx));
 				}
 				else
 				{
 					// Other token; assume a single character.
-					string Token = Condition.Substring(Idx++, 1);
-					Tokens.Add(Token);
+					string token = condition.Substring(idx++, 1);
+					tokens.Add(token);
 				}
 			}
-			return Tokens.ToArray();
+			return tokens.ToArray();
 		}
 
 		/// <summary>
 		/// Un-escape an MSBuild string (see https://msdn.microsoft.com/en-us/library/bb383819.aspx)
 		/// </summary>
-		/// <param name="Text">String to remove escape characters from</param>
+		/// <param name="text">String to remove escape characters from</param>
 		/// <returns>Unescaped string</returns>
-		static string? UnescapeString(string? Text)
+		static string? UnescapeString(string? text)
 		{
 			const string HexChars = "0123456789abcdef";
 
-			string? NewText = Text;
-			if(NewText != null)
+			string? newText = text;
+			if(newText != null)
 			{
-				for(int Idx = 0; Idx + 2 < NewText.Length; Idx++)
+				for(int idx = 0; idx + 2 < newText.Length; idx++)
 				{
-					if(NewText[Idx] == '%')
+					if(newText[idx] == '%')
 					{
-						int UpperDigitIdx = HexChars.IndexOf(Char.ToLowerInvariant(NewText[Idx + 1]));
-						if(UpperDigitIdx != -1)
+						int upperDigitIdx = HexChars.IndexOf(Char.ToLowerInvariant(newText[idx + 1]));
+						if(upperDigitIdx != -1)
 						{
-							int LowerDigitIdx = HexChars.IndexOf(Char.ToLowerInvariant(NewText[Idx + 2]));
-							if(LowerDigitIdx != -1)
+							int lowerDigitIdx = HexChars.IndexOf(Char.ToLowerInvariant(newText[idx + 2]));
+							if(lowerDigitIdx != -1)
 							{
-								char NewChar = (char)((UpperDigitIdx << 4) | LowerDigitIdx);
-								NewText = NewText.Substring(0, Idx) + NewChar + NewText.Substring(Idx + 3);
+								char newChar = (char)((upperDigitIdx << 4) | lowerDigitIdx);
+								newText = newText.Substring(0, idx) + newChar + newText.Substring(idx + 3);
 							}
 						}
 					}
 				}
 			}
-			return NewText;
+			return newText;
 		}
 	}
 }
diff --git a/Engine/Source/Programs/Shared/EpicGames.Core/Digest.cs b/Engine/Source/Programs/Shared/EpicGames.Core/Digest.cs
--- a/Engine/Source/Programs/Shared/EpicGames.Core/Digest.cs
+++ b/Engine/Source/Programs/Shared/EpicGames.Core/Digest.cs
@@ -1,11 +1,9 @@
 // Copyright Epic Games, Inc. All Rights Reserved.
 
 using System;
-using System.Collections.Generic;
-using System.Buffers;
+using System.Buffers.Binary;
 using System.Security.Cryptography;
 using System.Text;
-using System.Buffers.Binary;
 
 namespace EpicGames.Core
 {
@@ -17,7 +15,7 @@ namespace EpicGames.Core
 		/// <summary>
 		/// Memory storing the digest data
 		/// </summary>
-		public ReadOnlyMemory<byte> Memory;
+		public ReadOnlyMemory<byte> Memory { get; }
 
 		/// <summary>
 		/// Accessor for the span of memory storing the data
@@ -27,73 +25,73 @@ namespace EpicGames.Core
 		/// <summary>
 		/// Constructor
 		/// </summary>
-		/// <param name="Memory">Memory to construct from</param>
-		public Digest(ReadOnlyMemory<byte> Memory)
+		/// <param name="memory">Memory to construct from</param>
+		public Digest(ReadOnlyMemory<byte> memory)
 		{
-			this.Memory = Memory;
+			Memory = memory;
 		}
 
 		/// <summary>
 		/// Creates a content hash for a block of data, using a given algorithm.
 		/// </summary>
-		/// <param name="Data">Data to compute the hash for</param>
+		/// <param name="data">Data to compute the hash for</param>
 		/// <returns>New content hash instance containing the hash of the data</returns>
-		public static Digest<T> Compute<T>(byte[] Data) where T : DigestTraits, new()
+		public static Digest<T> Compute<T>(byte[] data) where T : DigestTraits, new()
 		{
-			using HashAlgorithm Algorithm = Digest<T>.Traits.CreateAlgorithm();
-			return Algorithm.ComputeHash(Data);
+			using HashAlgorithm algorithm = Digest<T>.Traits.CreateAlgorithm();
+			return algorithm.ComputeHash(data);
 		}
 
 		/// <summary>
 		/// Creates a content hash for a block of data, using a given algorithm.
 		/// </summary>
-		/// <param name="Text">Text to compute the hash for</param>
+		/// <param name="text">Text to compute the hash for</param>
 		/// <returns>New content hash instance containing the hash of the data</returns>
-		public static Digest<T> Compute<T>(string Text) where T : DigestTraits, new()
+		public static Digest<T> Compute<T>(string text) where T : DigestTraits, new()
 		{
-			return Compute<T>(Encoding.UTF8.GetBytes(Text));
+			return Compute<T>(Encoding.UTF8.GetBytes(text));
 		}
 
 		/// <summary>
 		/// Creates a content hash for a block of data, using a given algorithm.
 		/// </summary>
-		/// <param name="Data">Data to compute the hash for</param>
+		/// <param name="data">Data to compute the hash for</param>
 		/// <returns>New content hash instance containing the hash of the data</returns>
-		public static Digest<T> Compute<T>(ReadOnlySpan<byte> Data) where T : DigestTraits, new()
+		public static Digest<T> Compute<T>(ReadOnlySpan<byte> data) where T : DigestTraits, new()
 		{
-			byte[] Value = new byte[Digest<T>.Traits.Length];
-			using (HashAlgorithm Algorithm = Digest<T>.Traits.CreateAlgorithm())
+			byte[] value = new byte[Digest<T>.Traits.Length];
+			using (HashAlgorithm algorithm = Digest<T>.Traits.CreateAlgorithm())
 			{
-				if (!Algorithm.TryComputeHash(Data, Value, out int Written) || Written != Value.Length)
+				if (!algorithm.TryComputeHash(data, value, out int written) || written != value.Length)
 				{
 					throw new InvalidOperationException("Unable to compute hash for buffer");
 				}
 			}
-			return new Digest<T>(Value);
+			return new Digest<T>(value);
 		}
 
 		/// <summary>
 		/// Parses a digest from the given hex string
 		/// </summary>
-		/// <param name="Text"></param>
+		/// <param name="text"></param>
 		/// <returns></returns>
-		public static Digest Parse(string Text)
+		public static Digest Parse(string text)
 		{
-			return new Digest(StringUtils.ParseHexString(Text));
+			return new Digest(StringUtils.ParseHexString(text));
 		}
 
 		/// <summary>
 		/// Parses a digest from the given hex string
 		/// </summary>
-		/// <param name="Text"></param>
+		/// <param name="text"></param>
 		/// <returns></returns>
-		public static Digest<T> Parse<T>(string Text) where T : DigestTraits, new()
+		public static Digest<T> Parse<T>(string text) where T : DigestTraits, new()
 		{
-			return new Digest<T>(StringUtils.ParseHexString(Text));
+			return new Digest<T>(StringUtils.ParseHexString(text));
 		}
 
 		/// <inheritdoc/>
-		public override bool Equals(object? Obj) => (Obj is Digest Digest) && Digest.Span.SequenceEqual(Span);
+		public override bool Equals(object? obj) => (obj is Digest digest) && digest.Span.SequenceEqual(Span);
 
 		/// <inheritdoc/>
 		public override int GetHashCode() => BinaryPrimitives.ReadInt32LittleEndian(Span);
@@ -104,41 +102,41 @@ namespace EpicGames.Core
 		/// <summary>
 		/// Test two hash values for equality
 		/// </summary>
-		/// <param name="A"></param>
-		/// <param name="B"></param>
+		/// <param name="a"></param>
+		/// <param name="b"></param>
 		/// <returns></returns>
-		public static bool operator ==(Digest A, Digest B)
+		public static bool operator ==(Digest a, Digest b)
 		{
-			return A.Memory.Span.SequenceEqual(B.Memory.Span);
+			return a.Memory.Span.SequenceEqual(b.Memory.Span);
 		}
 
 		/// <summary>
 		/// Test two hash values for equality
 		/// </summary>
-		/// <param name="A"></param>
-		/// <param name="B"></param>
+		/// <param name="a"></param>
+		/// <param name="b"></param>
 		/// <returns></returns>
-		public static bool operator !=(Digest A, Digest B)
+		public static bool operator !=(Digest a, Digest b)
 		{
-			return !(A == B);
+			return !(a == b);
 		}
 
 		/// <summary>
 		/// Implicit conversion operator from memory objects
 		/// </summary>
-		/// <param name="Memory"></param>
-		public static implicit operator Digest(ReadOnlyMemory<byte> Memory)
+		/// <param name="memory"></param>
+		public static implicit operator Digest(ReadOnlyMemory<byte> memory)
 		{
-			return new Digest(Memory);
+			return new Digest(memory);
 		}
 
 		/// <summary>
 		/// Implicit conversion operator from byte arrays
 		/// </summary>
-		/// <param name="Memory"></param>
-		public static implicit operator Digest(byte[] Memory)
+		/// <param name="memory"></param>
+		public static implicit operator Digest(byte[] memory)
 		{
-			return new Digest(Memory);
+			return new Digest(memory);
 		}
 	}
 
@@ -155,10 +153,10 @@ namespace EpicGames.Core
 		/// <summary>
 		/// Constructor
 		/// </summary>
-		/// <param name="Length"></param>
-		public DigestTraits(int Length)
+		/// <param name="length"></param>
+		public DigestTraits(int length)
 		{
-			this.Length = Length;
+			Length = length;
 		}
 
 		/// <summary>
@@ -257,7 +255,7 @@ namespace EpicGames.Core
 		/// <summary>
 		/// Memory storing the digest data
 		/// </summary>
-		public ReadOnlyMemory<byte> Memory;
+		public ReadOnlyMemory<byte> Memory { get; }
 
 		/// <summary>
 		/// Accessor for the span of memory storing the data
@@ -267,14 +265,14 @@ namespace EpicGames.Core
 		/// <summary>
 		/// Constructor
 		/// </summary>
-		/// <param name="Memory">Memory to construct from</param>
-		public Digest(ReadOnlyMemory<byte> Memory)
+		/// <param name="memory">Memory to construct from</param>
+		public Digest(ReadOnlyMemory<byte> memory)
 		{
-			this.Memory = Memory;
+			Memory = memory;
 		}
 
 		/// <inheritdoc/>
-		public override bool Equals(object? Obj) => (Obj is Digest<T> Hash) && Hash.Span.SequenceEqual(Span);
+		public override bool Equals(object? obj) => (obj is Digest<T> hash) && hash.Span.SequenceEqual(Span);
 
 		/// <inheritdoc/>
 		public override int GetHashCode() => BinaryPrimitives.ReadInt32LittleEndian(Span);
@@ -285,35 +283,35 @@ namespace EpicGames.Core
 		/// <summary>
 		/// Test two hash values for equality
 		/// </summary>
-		/// <param name="A"></param>
-		/// <param name="B"></param>
+		/// <param name="a"></param>
+		/// <param name="b"></param>
 		/// <returns></returns>
-		public static bool operator ==(Digest<T> A, Digest<T> B) => A.Span.SequenceEqual(B.Span);
+		public static bool operator ==(Digest<T> a, Digest<T> b) => a.Span.SequenceEqual(b.Span);
 
 		/// <summary>
 		/// Test two hash values for equality
 		/// </summary>
-		/// <param name="A"></param>
-		/// <param name="B"></param>
+		/// <param name="a"></param>
+		/// <param name="b"></param>
 		/// <returns></returns>
-		public static bool operator !=(Digest<T> A, Digest<T> B) => !A.Span.SequenceEqual(B.Span);
+		public static bool operator !=(Digest<T> a, Digest<T> b) => !a.Span.SequenceEqual(b.Span);
 
 		/// <summary>
 		/// Implicit conversion operator from memory objects
 		/// </summary>
-		/// <param name="Memory"></param>
-		public static implicit operator Digest<T>(ReadOnlyMemory<byte> Memory)
+		/// <param name="memory"></param>
+		public static implicit operator Digest<T>(ReadOnlyMemory<byte> memory)
 		{
-			return new Digest<T>(Memory);
+			return new Digest<T>(memory);
 		}
 
 		/// <summary>
 		/// Implicit conversion operator from byte arrays
 		/// </summary>
-		/// <param name="Memory"></param>
-		public static implicit operator Digest<T>(byte[] Memory)
+		/// <param name="memory"></param>
+		public static implicit operator Digest<T>(byte[] memory)
 		{
-			return new Digest<T>(Memory);
+			return new Digest<T>(memory);
 		}
 	}
 
@@ -325,43 +323,43 @@ namespace EpicGames.Core
 		/// <summary>
 		/// Read a digest from a memory reader
 		/// </summary>
-		/// <param name="Reader"></param>
+		/// <param name="reader"></param>
 		/// <returns></returns>
-		public static Digest ReadDigest(this MemoryReader Reader)
+		public static Digest ReadDigest(this MemoryReader reader)
 		{
-			return new Digest(Reader.ReadVariableLengthBytes());
+			return new Digest(reader.ReadVariableLengthBytes());
 		}
 
 		/// <summary>
 		/// Read a strongly-typed digest from a memory reader
 		/// </summary>
 		/// <typeparam name="T"></typeparam>
-		/// <param name="Reader"></param>
+		/// <param name="reader"></param>
 		/// <returns></returns>
-		public static Digest<T> ReadDigest<T>(this MemoryReader Reader) where T : DigestTraits, new()
+		public static Digest<T> ReadDigest<T>(this MemoryReader reader) where T : DigestTraits, new()
 		{
-			return new Digest<T>(Reader.ReadFixedLengthBytes(Digest<T>.Traits.Length));
+			return new Digest<T>(reader.ReadFixedLengthBytes(Digest<T>.Traits.Length));
 		}
 
 		/// <summary>
 		/// Write a digest to a memory writer
 		/// </summary>
-		/// <param name="Writer"></param>
-		/// <param name="Digest"></param>
-		public static void WriteDigest(this MemoryWriter Writer, Digest Digest)
+		/// <param name="writer"></param>
+		/// <param name="digest"></param>
+		public static void WriteDigest(this MemoryWriter writer, Digest digest)
 		{
-			Writer.WriteVariableLengthBytes(Digest.Span);
+			writer.WriteVariableLengthBytes(digest.Span);
 		}
 
 		/// <summary>
 		/// Write a strongly typed digest to a memory writer
 		/// </summary>
 		/// <typeparam name="T"></typeparam>
-		/// <param name="Writer"></param>
-		/// <param name="Digest"></param>
-		public static void WriteDigest<T>(this MemoryWriter Writer, Digest<T> Digest) where T : DigestTraits, new()
+		/// <param name="writer"></param>
+		/// <param name="digest"></param>
+		public static void WriteDigest<T>(this MemoryWriter writer, Digest<T> digest) where T : DigestTraits, new()
 		{
-			Writer.WriteFixedLengthBytes(Digest.Span);
+			writer.WriteFixedLengthBytes(digest.Span);
 		}
 	}
 }
