Reformat EpicGames.Core according to standard coding conventions.

diff --git a/Engine/Source/Programs/Shared/EpicGames.Core/DirectoryReference.cs b/Engine/Source/Programs/Shared/EpicGames.Core/DirectoryReference.cs
--- a/Engine/Source/Programs/Shared/EpicGames.Core/DirectoryReference.cs
+++ b/Engine/Source/Programs/Shared/EpicGames.Core/DirectoryReference.cs
@@ -4,9 +4,6 @@ using System;
 using System.Collections.Generic;
 using System.Diagnostics.CodeAnalysis;
 using System.IO;
-using System.Linq;
-using System.Text;
-using System.Threading.Tasks;
 
 namespace EpicGames.Core
 {
@@ -27,28 +24,27 @@ namespace EpicGames.Core
 		/// <summary>
 		/// Default constructor.
 		/// </summary>
-		/// <param name="InPath">Path to this directory.</param>
-		public DirectoryReference(string InPath)
-			: base(FixTrailingPathSeparator(Path.GetFullPath(InPath)))
+		/// <param name="inPath">Path to this directory.</param>
+		public DirectoryReference(string inPath)
+			: base(FixTrailingPathSeparator(Path.GetFullPath(inPath)))
 		{
 		}
 
 		/// <summary>
 		/// Construct a DirectoryReference from a DirectoryInfo object.
 		/// </summary>
-		/// <param name="InInfo">Path to this file</param>
-		public DirectoryReference(DirectoryInfo InInfo)
-			: base(FixTrailingPathSeparator(InInfo.FullName))
+		/// <param name="inInfo">Path to this file</param>
+		public DirectoryReference(DirectoryInfo inInfo)
+			: base(FixTrailingPathSeparator(inInfo.FullName))
 		{
 		}
 
 		/// <summary>
 		/// Constructor for creating a directory object directly from two strings.
 		/// </summary>
-		/// <param name="InFullName">The full, sanitized path name</param>
-		/// <param name="InSanitize">Dummy argument used to resolve this overload</param>
-		public DirectoryReference(string InFullName, Sanitize InSanitize)
-			: base(InFullName)
+		/// <param name="fullName">The full, sanitized path name</param>
+		public DirectoryReference(string fullName, Sanitize _)
+			: base(fullName)
 		{
 		}
 
@@ -56,25 +52,25 @@ namespace EpicGames.Core
 		/// Ensures that the correct trailing path separator is appended. On Windows, the root directory (eg. C:\) always has a trailing path separator, but no other
 		/// path does.
 		/// </summary>
-		/// <param name="DirName">Absolute path to the directory</param>
+		/// <param name="dirName">Absolute path to the directory</param>
 		/// <returns>Path to the directory, with the correct trailing path separator</returns>
-		private static string FixTrailingPathSeparator(string DirName)
+		private static string FixTrailingPathSeparator(string dirName)
 		{
-			if(DirName.Length == 2 && DirName[1] == ':')
+			if(dirName.Length == 2 && dirName[1] == ':')
 			{
-				return DirName + Path.DirectorySeparatorChar;
+				return dirName + Path.DirectorySeparatorChar;
 			}
-			else if(DirName.Length == 3 && DirName[1] == ':' && DirName[2] == Path.DirectorySeparatorChar)
+			else if(dirName.Length == 3 && dirName[1] == ':' && dirName[2] == Path.DirectorySeparatorChar)
 			{
-				return DirName;
+				return dirName;
 			}
-			else if(DirName.Length > 1 && DirName[DirName.Length - 1] == Path.DirectorySeparatorChar)
+			else if(dirName.Length > 1 && dirName[^1] == Path.DirectorySeparatorChar)
 			{
-				return DirName.TrimEnd(Path.DirectorySeparatorChar);
+				return dirName.TrimEnd(Path.DirectorySeparatorChar);
 			}
 			else
 			{
-				return DirName;
+				return dirName;
 			}
 		}
 
@@ -100,31 +96,31 @@ namespace EpicGames.Core
 					return null;
 				}
 
-				int ParentLength = FullName.LastIndexOf(Path.DirectorySeparatorChar);
-				if (ParentLength == 2 && FullName[1] == ':')
+				int parentLength = FullName.LastIndexOf(Path.DirectorySeparatorChar);
+				if (parentLength == 2 && FullName[1] == ':')
 				{
-					ParentLength++;
+					parentLength++;
 				}
 
-				if (ParentLength == 0 && FullName[0] == Path.DirectorySeparatorChar)
+				if (parentLength == 0 && FullName[0] == Path.DirectorySeparatorChar)
 				{
 					// we have reached the root
-					ParentLength = 1;
+					parentLength = 1;
 				}
 
-				return new DirectoryReference(FullName.Substring(0, ParentLength), Sanitize.None);
+				return new DirectoryReference(FullName.Substring(0, parentLength), Sanitize.None);
 			}
 		}
 
 		/// <summary>
 		/// Gets the path for a special folder
 		/// </summary>
-		/// <param name="Folder">The folder to receive the path for</param>
+		/// <param name="folder">The folder to receive the path for</param>
 		/// <returns>Directory reference for the given folder, or null if it is not available</returns>
-		public static DirectoryReference? GetSpecialFolder(Environment.SpecialFolder Folder)
+		public static DirectoryReference? GetSpecialFolder(Environment.SpecialFolder folder)
 		{
-			string FolderPath = Environment.GetFolderPath(Folder);
-			return String.IsNullOrEmpty(FolderPath)? null : new DirectoryReference(FolderPath);
+			string folderPath = Environment.GetFolderPath(folder);
+			return String.IsNullOrEmpty(folderPath)? null : new DirectoryReference(folderPath);
 		}
 
 		/// <summary>
@@ -133,68 +129,72 @@ namespace EpicGames.Core
 		/// <returns>True if this path is a root directory, false otherwise</returns>
 		public bool IsRootDirectory()
 		{
-			return FullName[FullName.Length - 1] == Path.DirectorySeparatorChar;
+			return FullName[^1] == Path.DirectorySeparatorChar;
 		}
 
 		/// <summary>
 		/// Combine several fragments with a base directory, to form a new directory name
 		/// </summary>
-		/// <param name="BaseDirectory">The base directory</param>
-		/// <param name="Fragments">Fragments to combine with the base directory</param>
+		/// <param name="baseDirectory">The base directory</param>
+		/// <param name="fragments">Fragments to combine with the base directory</param>
 		/// <returns>The new directory name</returns>
-		public static DirectoryReference Combine(DirectoryReference BaseDirectory, params string[] Fragments)
+		public static DirectoryReference Combine(DirectoryReference baseDirectory, params string[] fragments)
 		{
-			string FullName = FileSystemReference.CombineStrings(BaseDirectory, Fragments);
-			return new DirectoryReference(FullName, Sanitize.None);
+			string fullName = FileSystemReference.CombineStrings(baseDirectory, fragments);
+			return new DirectoryReference(fullName, Sanitize.None);
 		}
 
 		/// <summary>
 		/// Compares two filesystem object names for equality. Uses the canonical name representation, not the display name representation.
 		/// </summary>
-		/// <param name="A">First object to compare.</param>
-		/// <param name="B">Second object to compare.</param>
+		/// <param name="a">First object to compare.</param>
+		/// <param name="b">Second object to compare.</param>
 		/// <returns>True if the names represent the same object, false otherwise</returns>
-		public static bool operator ==(DirectoryReference? A, DirectoryReference? B)
+		public static bool operator ==(DirectoryReference? a, DirectoryReference? b)
 		{
-			if ((object?)A == null)
+			if (a is null)
 			{
-				return (object?)B == null;
+				return b is null;
+			}
+			else if (b is null)
+			{
+				return false;
 			}
 			else
 			{
-				return (object?)B != null && A.FullName.Equals(B.FullName, Comparison);
+				return a.FullName.Equals(b.FullName, Comparison);
 			}
 		}
 
 		/// <summary>
 		/// Compares two filesystem object names for inequality. Uses the canonical name representation, not the display name representation.
 		/// </summary>
-		/// <param name="A">First object to compare.</param>
-		/// <param name="B">Second object to compare.</param>
+		/// <param name="a">First object to compare.</param>
+		/// <param name="b">Second object to compare.</param>
 		/// <returns>False if the names represent the same object, true otherwise</returns>
-		public static bool operator !=(DirectoryReference? A, DirectoryReference? B)
+		public static bool operator !=(DirectoryReference? a, DirectoryReference? b)
 		{
-			return !(A == B);
+			return !(a == b);
 		}
 
 		/// <summary>
 		/// Compares against another object for equality.
 		/// </summary>
-		/// <param name="Obj">other instance to compare.</param>
+		/// <param name="obj">other instance to compare.</param>
 		/// <returns>True if the names represent the same object, false otherwise</returns>
-		public override bool Equals(object? Obj)
+		public override bool Equals(object? obj)
 		{
-			return (Obj is DirectoryReference) && ((DirectoryReference)Obj) == this;
+			return (obj is DirectoryReference dir) && dir == this;
 		}
 
 		/// <summary>
 		/// Compares against another object for equality.
 		/// </summary>
-		/// <param name="Obj">other instance to compare.</param>
+		/// <param name="obj">other instance to compare.</param>
 		/// <returns>True if the names represent the same object, false otherwise</returns>
-		public bool Equals(DirectoryReference? Obj)
+		public bool Equals(DirectoryReference? obj)
 		{
-			return Obj == this;
+			return obj == this;
 		}
 
 		/// <summary>
@@ -207,61 +207,61 @@ namespace EpicGames.Core
 		}
 
 		/// <inheritdoc/>
-		public int CompareTo(DirectoryReference? Other) => Comparer.Compare(FullName, Other?.FullName);
+		public int CompareTo(DirectoryReference? other) => Comparer.Compare(FullName, other?.FullName);
 
 		/// <summary>
 		/// Helper function to create a remote directory reference. Unlike normal DirectoryReference objects, these aren't converted to a full path in the local filesystem.
 		/// </summary>
-		/// <param name="AbsolutePath">The absolute path in the remote file system</param>
+		/// <param name="absolutePath">The absolute path in the remote file system</param>
 		/// <returns>New directory reference</returns>
-		public static DirectoryReference MakeRemote(string AbsolutePath)
+		public static DirectoryReference MakeRemote(string absolutePath)
 		{
-			return new DirectoryReference(AbsolutePath, Sanitize.None);
+			return new DirectoryReference(absolutePath, Sanitize.None);
 		}
 
 		/// <summary>
 		/// Gets the parent directory for a file, or returns null if it's null.
 		/// </summary>
-		/// <param name="File">The file to create a directory reference for</param>
+		/// <param name="file">The file to create a directory reference for</param>
 		/// <returns>The directory containing the file  </returns>
-		[return: NotNullIfNotNull("File")]
-		public static DirectoryReference? FromFile(FileReference? File)
+		[return: NotNullIfNotNull("file")]
+		public static DirectoryReference? FromFile(FileReference? file)
 		{
-			if(File == null)
+			if(file == null)
 			{
 				return null;
 			}
 			else
 			{
-				return File.Directory;
+				return file.Directory;
 			}
 		}
 
 		/// <summary>
 		/// Create a DirectoryReference from a string. If the string is null, returns a null DirectoryReference.
 		/// </summary>
-		/// <param name="DirectoryName">Path for the new object</param>
+		/// <param name="directoryName">Path for the new object</param>
 		/// <returns>Returns a FileReference representing the given string, or null.</returns>
-		public static DirectoryReference? FromString(string? DirectoryName)
+		public static DirectoryReference? FromString(string? directoryName)
 		{
-			if(String.IsNullOrEmpty(DirectoryName))
+			if(String.IsNullOrEmpty(directoryName))
 			{
 				return null;
 			}
 			else
 			{
-				return new DirectoryReference(DirectoryName);
+				return new DirectoryReference(directoryName);
 			}
 		}
 
 		/// <summary>
 		/// Finds the correct case to match the location of this file on disk. Uses the given case for parts of the path that do not exist.
 		/// </summary>
-		/// <param name="Location">The path to find the correct case for</param>
+		/// <param name="location">The path to find the correct case for</param>
 		/// <returns>Location of the file with the correct case</returns>
-		public static DirectoryReference FindCorrectCase(DirectoryReference Location)
+		public static DirectoryReference FindCorrectCase(DirectoryReference location)
 		{
-			return new DirectoryReference(DirectoryUtils.FindCorrectCase(Location.ToDirectoryInfo()));
+			return new DirectoryReference(DirectoryUtils.FindCorrectCase(location.ToDirectoryInfo()));
 		}
 
 		/// <summary>
@@ -287,132 +287,132 @@ namespace EpicGames.Core
 		/// <summary>
 		/// Creates a directory
 		/// </summary>
-		/// <param name="Location">Location of the directory</param>
-		public static void CreateDirectory(DirectoryReference Location)
+		/// <param name="location">Location of the directory</param>
+		public static void CreateDirectory(DirectoryReference location)
 		{
-			Directory.CreateDirectory(Location.FullName);
+			Directory.CreateDirectory(location.FullName);
 		}
 
         /// <summary>
         /// Deletes a directory
         /// </summary>
-		/// <param name="Location">Location of the directory</param>
-        public static void Delete(DirectoryReference Location)
+		/// <param name="location">Location of the directory</param>
+        public static void Delete(DirectoryReference location)
         {
-            Directory.Delete(Location.FullName);
+            Directory.Delete(location.FullName);
         }
 
         /// <summary>
         /// Deletes a directory
         /// </summary>
-		/// <param name="Location">Location of the directory</param>
+		/// <param name="location">Location of the directory</param>
 		/// <param name="bRecursive">Whether to remove directories recursively</param>
-        public static void Delete(DirectoryReference Location, bool bRecursive)
+        public static void Delete(DirectoryReference location, bool bRecursive)
         {
-            Directory.Delete(Location.FullName, bRecursive);
+            Directory.Delete(location.FullName, bRecursive);
         }
 
         /// <summary>
         /// Checks whether the directory exists
         /// </summary>
-		/// <param name="Location">Location of the directory</param>
+		/// <param name="location">Location of the directory</param>
         /// <returns>True if this directory exists</returns>
-        public static bool Exists(DirectoryReference Location)
+        public static bool Exists(DirectoryReference location)
 		{
-			return Directory.Exists(Location.FullName);
+			return Directory.Exists(location.FullName);
 		}
 
 		/// <summary>
 		/// Enumerate files from a given directory
 		/// </summary>
-		/// <param name="BaseDir">Base directory to search in</param>
+		/// <param name="baseDir">Base directory to search in</param>
 		/// <returns>Sequence of file references</returns>
-		public static IEnumerable<FileReference> EnumerateFiles(DirectoryReference BaseDir)
+		public static IEnumerable<FileReference> EnumerateFiles(DirectoryReference baseDir)
 		{
-			foreach (string FileName in Directory.EnumerateFiles(BaseDir.FullName))
+			foreach (string fileName in Directory.EnumerateFiles(baseDir.FullName))
 			{
-				yield return new FileReference(FileName, FileReference.Sanitize.None);
+				yield return new FileReference(fileName, FileReference.Sanitize.None);
 			}
 		}
 
 		/// <summary>
 		/// Enumerate files from a given directory
 		/// </summary>
-		/// <param name="BaseDir">Base directory to search in</param>
-		/// <param name="Pattern">Pattern for matching files</param>
+		/// <param name="baseDir">Base directory to search in</param>
+		/// <param name="pattern">Pattern for matching files</param>
 		/// <returns>Sequence of file references</returns>
-		public static IEnumerable<FileReference> EnumerateFiles(DirectoryReference BaseDir, string Pattern)
+		public static IEnumerable<FileReference> EnumerateFiles(DirectoryReference baseDir, string pattern)
 		{
-			foreach (string FileName in Directory.EnumerateFiles(BaseDir.FullName, Pattern))
+			foreach (string fileName in Directory.EnumerateFiles(baseDir.FullName, pattern))
 			{
-				yield return new FileReference(FileName, FileReference.Sanitize.None);
+				yield return new FileReference(fileName, FileReference.Sanitize.None);
 			}
 		}
 
 		/// <summary>
 		/// Enumerate files from a given directory
 		/// </summary>
-		/// <param name="BaseDir">Base directory to search in</param>
-		/// <param name="Pattern">Pattern for matching files</param>
-		/// <param name="Option">Options for the search</param>
+		/// <param name="baseDir">Base directory to search in</param>
+		/// <param name="pattern">Pattern for matching files</param>
+		/// <param name="option">Options for the search</param>
 		/// <returns>Sequence of file references</returns>
-		public static IEnumerable<FileReference> EnumerateFiles(DirectoryReference BaseDir, string Pattern, SearchOption Option)
+		public static IEnumerable<FileReference> EnumerateFiles(DirectoryReference baseDir, string pattern, SearchOption option)
 		{
-			foreach (string FileName in Directory.EnumerateFiles(BaseDir.FullName, Pattern, Option))
+			foreach (string fileName in Directory.EnumerateFiles(baseDir.FullName, pattern, option))
 			{
-				yield return new FileReference(FileName, FileReference.Sanitize.None);
+				yield return new FileReference(fileName, FileReference.Sanitize.None);
 			}
 		}
 
 		/// <summary>
 		/// Enumerate subdirectories in a given directory
 		/// </summary>
-		/// <param name="BaseDir">Base directory to search in</param>
+		/// <param name="baseDir">Base directory to search in</param>
 		/// <returns>Sequence of directory references</returns>
-		public static IEnumerable<DirectoryReference> EnumerateDirectories(DirectoryReference BaseDir)
+		public static IEnumerable<DirectoryReference> EnumerateDirectories(DirectoryReference baseDir)
 		{
-			foreach (string DirectoryName in Directory.EnumerateDirectories(BaseDir.FullName))
+			foreach (string directoryName in Directory.EnumerateDirectories(baseDir.FullName))
 			{
-				yield return new DirectoryReference(DirectoryName, Sanitize.None);
+				yield return new DirectoryReference(directoryName, Sanitize.None);
 			}
 		}
 
 		/// <summary>
 		/// Enumerate subdirectories in a given directory
 		/// </summary>
-		/// <param name="BaseDir">Base directory to search in</param>
-		/// <param name="Pattern">Pattern for matching directories</param>
+		/// <param name="baseDir">Base directory to search in</param>
+		/// <param name="pattern">Pattern for matching directories</param>
 		/// <returns>Sequence of directory references</returns>
-		public static IEnumerable<DirectoryReference> EnumerateDirectories(DirectoryReference BaseDir, string Pattern)
+		public static IEnumerable<DirectoryReference> EnumerateDirectories(DirectoryReference baseDir, string pattern)
 		{
-			foreach (string DirectoryName in Directory.EnumerateDirectories(BaseDir.FullName, Pattern))
+			foreach (string directoryName in Directory.EnumerateDirectories(baseDir.FullName, pattern))
 			{
-				yield return new DirectoryReference(DirectoryName, Sanitize.None);
+				yield return new DirectoryReference(directoryName, Sanitize.None);
 			}
 		}
 
 		/// <summary>
 		/// Enumerate subdirectories in a given directory
 		/// </summary>
-		/// <param name="BaseDir">Base directory to search in</param>
-		/// <param name="Pattern">Pattern for matching files</param>
-		/// <param name="Option">Options for the search</param>
+		/// <param name="baseDir">Base directory to search in</param>
+		/// <param name="pattern">Pattern for matching files</param>
+		/// <param name="option">Options for the search</param>
 		/// <returns>Sequence of directory references</returns>
-		public static IEnumerable<DirectoryReference> EnumerateDirectories(DirectoryReference BaseDir, string Pattern, SearchOption Option)
+		public static IEnumerable<DirectoryReference> EnumerateDirectories(DirectoryReference baseDir, string pattern, SearchOption option)
 		{
-			foreach (string DirectoryName in Directory.EnumerateDirectories(BaseDir.FullName, Pattern, Option))
+			foreach (string directoryName in Directory.EnumerateDirectories(baseDir.FullName, pattern, option))
 			{
-				yield return new DirectoryReference(DirectoryName, Sanitize.None);
+				yield return new DirectoryReference(directoryName, Sanitize.None);
 			}
 		}
 
 		/// <summary>
 		/// Sets the current directory
 		/// </summary>
-		/// <param name="Location">Location of the new current directory</param>
-		public static void SetCurrentDirectory(DirectoryReference Location)
+		/// <param name="location">Location of the new current directory</param>
+		public static void SetCurrentDirectory(DirectoryReference location)
 		{
-			Directory.SetCurrentDirectory(Location.FullName);
+			Directory.SetCurrentDirectory(location.FullName);
 		}
 
 		#endregion
@@ -426,70 +426,70 @@ namespace EpicGames.Core
 		/// <summary>
 		/// Manually serialize a file reference to a binary stream.
 		/// </summary>
-		/// <param name="Writer">Binary writer to write to</param>
-		/// <param name="Directory">The directory reference to write</param>
-		public static void Write(this BinaryWriter Writer, DirectoryReference Directory)
+		/// <param name="writer">Binary writer to write to</param>
+		/// <param name="directory">The directory reference to write</param>
+		public static void Write(this BinaryWriter writer, DirectoryReference directory)
 		{
-			Writer.Write((Directory == null) ? String.Empty : Directory.FullName);
+			writer.Write((directory == null) ? String.Empty : directory.FullName);
 		}
 
 		/// <summary>
 		/// Manually deserialize a directory reference from a binary stream.
 		/// </summary>
-		/// <param name="Reader">Binary reader to read from</param>
+		/// <param name="reader">Binary reader to read from</param>
 		/// <returns>New DirectoryReference object</returns>
-		public static DirectoryReference? ReadDirectoryReference(this BinaryReader Reader)
+		public static DirectoryReference? ReadDirectoryReference(this BinaryReader reader)
 		{
-			string FullName = Reader.ReadString();
-			return (FullName.Length == 0) ? null : new DirectoryReference(FullName, DirectoryReference.Sanitize.None);
+			string fullName = reader.ReadString();
+			return (fullName.Length == 0) ? null : new DirectoryReference(fullName, DirectoryReference.Sanitize.None);
 		}
 
 		/// <summary>
 		/// Manually deserialize a directory reference from a binary stream.
 		/// </summary>
-		/// <param name="Reader">Binary reader to read from</param>
+		/// <param name="reader">Binary reader to read from</param>
 		/// <returns>New DirectoryReference object</returns>
-		public static DirectoryReference ReadDirectoryReferenceNotNull(this BinaryReader Reader)
+		public static DirectoryReference ReadDirectoryReferenceNotNull(this BinaryReader reader)
 		{
-			return BinaryArchiveReader.NotNull<DirectoryReference>(ReadDirectoryReference(Reader));
+			return BinaryArchiveReader.NotNull<DirectoryReference>(ReadDirectoryReference(reader));
 		}
 
 		/// <summary>
 		/// Writes a directory reference  to a binary archive
 		/// </summary>
-		/// <param name="Writer">The writer to output data to</param>
-		/// <param name="Directory">The item to write</param>
-		public static void WriteDirectoryReference(this BinaryArchiveWriter Writer, DirectoryReference? Directory)
+		/// <param name="writer">The writer to output data to</param>
+		/// <param name="directory">The item to write</param>
+		public static void WriteDirectoryReference(this BinaryArchiveWriter writer, DirectoryReference? directory)
 		{
-			Writer.WriteString(Directory?.FullName);
+			writer.WriteString(directory?.FullName);
 		}
 
 		/// <summary>
 		/// Reads a directory reference from a binary archive
 		/// </summary>
-		/// <param name="Reader">Reader to serialize data from</param>
+		/// <param name="reader">Reader to serialize data from</param>
 		/// <returns>New directory reference instance</returns>
-		public static DirectoryReference? ReadDirectoryReference(this BinaryArchiveReader Reader)
+		public static DirectoryReference? ReadDirectoryReference(this BinaryArchiveReader reader)
 		{
-			string? FullName = Reader.ReadString();
-			if (FullName == null)
+			string? fullName = reader.ReadString();
+			if (fullName == null)
 			{
 				return null;
 			}
 			else
 			{
-				return new DirectoryReference(FullName, DirectoryReference.Sanitize.None);
+				return new DirectoryReference(fullName, DirectoryReference.Sanitize.None);
 			}
 		}
 
 		/// <summary>
 		/// Reads a directory reference from a binary archive
 		/// </summary>
-		/// <param name="Reader">Reader to serialize data from</param>
+		/// <param name="reader">Reader to serialize data from</param>
 		/// <returns>New directory reference instance</returns>
-		public static DirectoryReference ReadDirectoryReferenceNotNull(this BinaryArchiveReader Reader)
+		public static DirectoryReference ReadDirectoryReferenceNotNull(this BinaryArchiveReader reader)
 		{
-			return BinaryArchiveReader.NotNull(ReadDirectoryReference(Reader));
+			return BinaryArchiveReader.NotNull(ReadDirectoryReference(reader));
 		}
 	}
 }
diff --git a/Engine/Source/Programs/Shared/EpicGames.Core/DirectoryUtils.cs b/Engine/Source/Programs/Shared/EpicGames.Core/DirectoryUtils.cs
--- a/Engine/Source/Programs/Shared/EpicGames.Core/DirectoryUtils.cs
+++ b/Engine/Source/Programs/Shared/EpicGames.Core/DirectoryUtils.cs
@@ -1,11 +1,7 @@
 // Copyright Epic Games, Inc. All Rights Reserved.
 
 using System;
-using System.Collections.Generic;
 using System.IO;
-using System.Linq;
-using System.Text;
-using System.Threading.Tasks;
 
 namespace EpicGames.Core
 {
@@ -17,35 +13,35 @@ namespace EpicGames.Core
 		/// <summary>
 		/// Finds the on-disk case of a a directory
 		/// </summary>
-		/// <param name="Info">DirectoryInfo instance describing the directory</param>
+		/// <param name="info">DirectoryInfo instance describing the directory</param>
 		/// <returns>New DirectoryInfo instance that represents the directory with the correct case</returns>
-		public static DirectoryInfo FindCorrectCase(DirectoryInfo Info)
+		public static DirectoryInfo FindCorrectCase(DirectoryInfo info)
 		{
-			DirectoryInfo ParentInfo = Info.Parent;
-			if (ParentInfo == null)
+			DirectoryInfo parentInfo = info.Parent;
+			if (parentInfo == null)
 			{
-				string FullName = Info.FullName;
-				if (FullName.Length >= 2 && (FullName[0] >= 'a' && FullName[0] <= 'z') && FullName[1] == ':')
+				string fullName = info.FullName;
+				if (fullName.Length >= 2 && (fullName[0] >= 'a' && fullName[0] <= 'z') && fullName[1] == ':')
 				{
-					return new DirectoryInfo(Char.ToUpper(FullName[0]) + FullName.Substring(1));
+					return new DirectoryInfo(Char.ToUpper(fullName[0]) + fullName.Substring(1));
 				}
 				else
 				{
-					return Info;
+					return info;
 				}
 			}
 			else
 			{
-				ParentInfo = FindCorrectCase(ParentInfo);
-				if (Info.Exists)
+				parentInfo = FindCorrectCase(parentInfo);
+				if (info.Exists)
 				{
 					try
 					{
-						foreach (DirectoryInfo ChildInfo in ParentInfo.EnumerateDirectories())
+						foreach (DirectoryInfo childInfo in parentInfo.EnumerateDirectories())
 						{
-							if (String.Equals(ChildInfo.Name, Info.Name, DirectoryReference.Comparison))
+							if (String.Equals(childInfo.Name, info.Name, DirectoryReference.Comparison))
 							{
-								return ChildInfo;
+								return childInfo;
 							}
 						}
 					}
@@ -56,7 +52,7 @@ namespace EpicGames.Core
 					}
 				}
 
-				return new DirectoryInfo(Path.Combine(ParentInfo.FullName, Info.Name));
+				return new DirectoryInfo(Path.Combine(parentInfo.FullName, info.Name));
 			}
 		}
 	}
diff --git a/Engine/Source/Programs/Shared/EpicGames.Core/EnumerableExtensions.cs b/Engine/Source/Programs/Shared/EpicGames.Core/EnumerableExtensions.cs
--- a/Engine/Source/Programs/Shared/EpicGames.Core/EnumerableExtensions.cs
+++ b/Engine/Source/Programs/Shared/EpicGames.Core/EnumerableExtensions.cs
@@ -2,7 +2,6 @@
 
 using System;
 using System.Collections.Generic;
-using System.Text;
 
 namespace EpicGames.Core
 {
@@ -15,24 +14,24 @@ namespace EpicGames.Core
 		/// Split the sequence into batches of at most the given size
 		/// </summary>
 		/// <typeparam name="TElement">The element type</typeparam>
-		/// <param name="Sequence">Sequence to split into batches</param>
-		/// <param name="BatchSize">Maximum size of each batch</param>
+		/// <param name="sequence">Sequence to split into batches</param>
+		/// <param name="batchSize">Maximum size of each batch</param>
 		/// <returns>Sequence of batches</returns>
-		public static IEnumerable<IReadOnlyList<TElement>> Batch<TElement>(this IEnumerable<TElement> Sequence, int BatchSize)
+		public static IEnumerable<IReadOnlyList<TElement>> Batch<TElement>(this IEnumerable<TElement> sequence, int batchSize)
 		{
-			List<TElement> Elements = new List<TElement>(BatchSize);
-			foreach (TElement Element in Sequence)
+			List<TElement> elements = new List<TElement>(batchSize);
+			foreach (TElement element in sequence)
 			{
-				Elements.Add(Element);
-				if (Elements.Count == BatchSize)
+				elements.Add(element);
+				if (elements.Count == batchSize)
 				{
-					yield return Elements;
-					Elements.Clear();
+					yield return elements;
+					elements.Clear();
 				}
 			}
-			if (Elements.Count > 0)
+			if (elements.Count > 0)
 			{
-				yield return Elements;
+				yield return elements;
 			}
 		}
 
@@ -40,60 +39,60 @@ namespace EpicGames.Core
 		/// Finds the minimum element by a given field
 		/// </summary>
 		/// <typeparam name="TElement"></typeparam>
-		/// <param name="Sequence"></param>
-		/// <param name="Selector"></param>
+		/// <param name="sequence"></param>
+		/// <param name="selector"></param>
 		/// <returns></returns>
-		public static TElement MinBy<TElement>(this IEnumerable<TElement> Sequence, Func<TElement, int> Selector)
+		public static TElement MinBy<TElement>(this IEnumerable<TElement> sequence, Func<TElement, int> selector)
 		{
-			IEnumerator<TElement> Enumerator = Sequence.GetEnumerator();
-			if (!Enumerator.MoveNext())
+			IEnumerator<TElement> enumerator = sequence.GetEnumerator();
+			if (!enumerator.MoveNext())
 			{
 				throw new Exception("Collection is empty");
 			}
 
-			TElement MinElement = Enumerator.Current;
+			TElement minElement = enumerator.Current;
 
-			int MinValue = Selector(MinElement);
-			while (Enumerator.MoveNext())
+			int minValue = selector(minElement);
+			while (enumerator.MoveNext())
 			{
-				int Value = Selector(Enumerator.Current);
-				if (Value < MinValue)
+				int value = selector(enumerator.Current);
+				if (value < minValue)
 				{
-					MinElement = Enumerator.Current;
+					minElement = enumerator.Current;
 				}
 			}
 
-			return MinElement;
+			return minElement;
 		}
 
 		/// <summary>
 		/// Finds the maximum element by a given field
 		/// </summary>
 		/// <typeparam name="TElement"></typeparam>
-		/// <param name="Sequence"></param>
-		/// <param name="Selector"></param>
+		/// <param name="sequence"></param>
+		/// <param name="selector"></param>
 		/// <returns></returns>
-		public static TElement MaxBy<TElement>(this IEnumerable<TElement> Sequence, Func<TElement, int> Selector)
+		public static TElement MaxBy<TElement>(this IEnumerable<TElement> sequence, Func<TElement, int> selector)
 		{
-			IEnumerator<TElement> Enumerator = Sequence.GetEnumerator();
-			if (!Enumerator.MoveNext())
+			IEnumerator<TElement> enumerator = sequence.GetEnumerator();
+			if (!enumerator.MoveNext())
 			{
 				throw new Exception("Collection is empty");
 			}
 
-			TElement MaxElement = Enumerator.Current;
+			TElement maxElement = enumerator.Current;
 
-			int MaxValue = Selector(MaxElement);
-			while (Enumerator.MoveNext())
+			int maxValue = selector(maxElement);
+			while (enumerator.MoveNext())
 			{
-				int Value = Selector(Enumerator.Current);
-				if (Value > MaxValue)
+				int value = selector(enumerator.Current);
+				if (value > maxValue)
 				{
-					MaxElement = Enumerator.Current;
+					maxElement = enumerator.Current;
 				}
 			}
 
-			return MaxElement;
+			return maxElement;
 		}
 	}
 }
diff --git a/Engine/Source/Programs/Shared/EpicGames.Core/ExceptionUtils.cs b/Engine/Source/Programs/Shared/EpicGames.Core/ExceptionUtils.cs
--- a/Engine/Source/Programs/Shared/EpicGames.Core/ExceptionUtils.cs
+++ b/Engine/Source/Programs/Shared/EpicGames.Core/ExceptionUtils.cs
@@ -4,7 +4,6 @@ using System;
 using System.Collections.Generic;
 using System.Linq;
 using System.Text;
-using System.Threading.Tasks;
 
 namespace EpicGames.Core
 {
@@ -22,44 +21,44 @@ namespace EpicGames.Core
 		/// Adds a context message to a list stored on the given exception. Intended to be used in situations where supplying additional context
 		/// for an exception is valuable, but wrapping it would remove information.
 		/// </summary>
-		/// <param name="Ex">The exception that was thrown</param>
-		/// <param name="Message">Message to append to the context list</param>
-		public static void AddContext(Exception Ex, string Message)
+		/// <param name="ex">The exception that was thrown</param>
+		/// <param name="message">Message to append to the context list</param>
+		public static void AddContext(Exception ex, string message)
 		{
-			List<string>? Messages = Ex.Data[ContextEntryName] as List<string>;
-			if (Messages == null)
+			List<string>? messages = ex.Data[ContextEntryName] as List<string>;
+			if (messages == null)
 			{
-				Messages = new List<string>();
-				Ex.Data[ContextEntryName] = Messages;
+				messages = new List<string>();
+				ex.Data[ContextEntryName] = messages;
 			}
-			Messages.Add(Message);
+			messages.Add(message);
 		}
 
 		/// <summary>
 		/// Adds a context message to a list stored on the given exception. Intended to be used in situations where supplying additional context
 		/// for an exception is valuable, but wrapping it would remove information.
 		/// </summary>
-		/// <param name="Ex">The exception that was thrown</param>
-		/// <param name="Format">Formatting string for </param>
-		/// <param name="Args">Message to append to the context list</param>
-		public static void AddContext(Exception Ex, string Format, params object[] Args)
+		/// <param name="ex">The exception that was thrown</param>
+		/// <param name="format">Formatting string for </param>
+		/// <param name="args">Message to append to the context list</param>
+		public static void AddContext(Exception ex, string format, params object[] args)
 		{
-			AddContext(Ex, String.Format(Format, Args));
+			AddContext(ex, String.Format(format, args));
 		}
 
 		/// <summary>
 		/// Enumerates the context lines from the given exception
 		/// </summary>
-		/// <param name="Ex">The exception to retrieve context from</param>
+		/// <param name="ex">The exception to retrieve context from</param>
 		/// <returns>Sequence of context lines</returns>
-		public static IEnumerable<string> GetContext(Exception Ex)
+		public static IEnumerable<string> GetContext(Exception ex)
 		{
-			List<string>? Messages = Ex.Data[ContextEntryName] as List<string>;
-			if (Messages != null)
+			List<string>? messages = ex.Data[ContextEntryName] as List<string>;
+			if (messages != null)
 			{
-				foreach (string Message in Messages)
+				foreach (string message in messages)
 				{
-					yield return Message;
+					yield return message;
 				}
 			}
 		}
@@ -67,83 +66,83 @@ namespace EpicGames.Core
 		/// <summary>
 		/// Formats an exception for display in the log, including additional lines of context that were attached to it.
 		/// </summary>
-		/// <param name="Ex">The exception to format</param>
+		/// <param name="ex">The exception to format</param>
 		/// <returns>String containing the exception information. May be multiple lines.</returns>
-		public static string FormatException(Exception Ex)
+		public static string FormatException(Exception ex)
 		{
-			StringBuilder ErrorMessage = new StringBuilder();
-			if (Ex is AggregateException)
+			StringBuilder errorMessage = new StringBuilder();
+			if (ex is AggregateException exAgg)
 			{
-				Exception? InnerException = ((AggregateException)Ex).InnerException;
-				if (InnerException != null)
+				Exception? innerException = exAgg.InnerException;
+				if (innerException != null)
 				{
-					ErrorMessage.Append(InnerException.ToString());
-					foreach (string Line in GetContext(InnerException))
+					errorMessage.Append(innerException.ToString());
+					foreach (string line in GetContext(innerException))
 					{
-						ErrorMessage.AppendFormat("\n  {0}", Line);
+						errorMessage.AppendFormat("\n  {0}", line);
 					}
 				}
 			}
 			else
 			{
-				ErrorMessage.Append(Ex.ToString());
+				errorMessage.Append(ex.ToString());
 			}
-			foreach (string Line in GetContext(Ex))
+			foreach (string line in GetContext(ex))
 			{
-				ErrorMessage.AppendFormat("\n{0}", Line);
+				errorMessage.AppendFormat("\n{0}", line);
 			}
-			return ErrorMessage.ToString();
+			return errorMessage.ToString();
 		}
 
 		/// <summary>
 		/// Formats a detailed information about where an exception occurs, including any inner exceptions
 		/// </summary>
-		/// <param name="Ex">The exception to format</param>
+		/// <param name="ex">The exception to format</param>
 		/// <returns>String containing the exception information. May be multiple lines.</returns>
-		public static string FormatExceptionDetails(Exception Ex)
+		public static string FormatExceptionDetails(Exception ex)
 		{
-			List<Exception> ExceptionStack = new List<Exception>();
-			for (Exception? CurrentEx = Ex; CurrentEx != null; CurrentEx = CurrentEx.InnerException)
+			List<Exception> exceptionStack = new List<Exception>();
+			for (Exception? currentEx = ex; currentEx != null; currentEx = currentEx.InnerException)
 			{
-				ExceptionStack.Add(CurrentEx);
+				exceptionStack.Add(currentEx);
 			}
 
-			StringBuilder Message = new StringBuilder();
-			for (int Idx = ExceptionStack.Count - 1; Idx >= 0; Idx--)
+			StringBuilder message = new StringBuilder();
+			for (int idx = exceptionStack.Count - 1; idx >= 0; idx--)
 			{
-				Exception CurrentEx = ExceptionStack[Idx];
-				Message.AppendFormat("{0}{1}: {2}\n{3}", (Idx == ExceptionStack.Count - 1) ? "" : "Wrapped by ", CurrentEx.GetType().Name, CurrentEx.Message, CurrentEx.StackTrace);
+				Exception currentEx = exceptionStack[idx];
+				message.AppendFormat("{0}{1}: {2}\n{3}", (idx == exceptionStack.Count - 1) ? "" : "Wrapped by ", currentEx.GetType().Name, currentEx.Message, currentEx.StackTrace);
 
-				if (CurrentEx.Data.Count > 0)
+				if (currentEx.Data.Count > 0)
 				{
-					foreach (object? Key in CurrentEx.Data.Keys)
+					foreach (object? key in currentEx.Data.Keys)
 					{
-						if (Key == null)
+						if (key == null)
 						{
 							continue;
 						}
 
-						object? Value = CurrentEx.Data[Key];
-						if (Value == null)
+						object? value = currentEx.Data[key];
+						if (value == null)
 						{
 							continue;
 						}
 
-						string ValueString;
-						if(Value is List<string>)
+						string valueString;
+						if(value is List<string> valueList)
 						{
-							ValueString = String.Format("({0})", String.Join(", ", ((List<string>)Value).Select(x => String.Format("\"{0}\"", x))));
+							valueString = String.Format("({0})", String.Join(", ", valueList.Select(x => String.Format("\"{0}\"", x))));
 						}
 						else
 						{
-							ValueString = Value.ToString() ?? String.Empty;
+							valueString = value.ToString() ?? String.Empty;
 						}
 
-						Message.AppendFormat("   data: {0} = {1}", Key, ValueString);
+						message.AppendFormat("   data: {0} = {1}", key, valueString);
 					}
 				}
 			}
-			return Message.ToString();
+			return message.ToString();
 		}
 	}
 }
diff --git a/Engine/Source/Programs/Shared/EpicGames.Core/FatalErrorException.cs b/Engine/Source/Programs/Shared/EpicGames.Core/FatalErrorException.cs
--- a/Engine/Source/Programs/Shared/EpicGames.Core/FatalErrorException.cs
+++ b/Engine/Source/Programs/Shared/EpicGames.Core/FatalErrorException.cs
@@ -1,10 +1,6 @@
 // Copyright Epic Games, Inc. All Rights Reserved.
 
 using System;
-using System.Collections.Generic;
-using System.Linq;
-using System.Text;
-using System.Threading.Tasks;
 
 namespace EpicGames.Core
 {
@@ -16,45 +12,45 @@ namespace EpicGames.Core
 		/// <summary>
 		/// Exit code for the process
 		/// </summary>
-		public int ExitCode = 1;
+		public int ExitCode { get; set; } = 1;
 
 		/// <summary>
 		/// Constructor
 		/// </summary>
-		/// <param name="Message">The error message to display.</param>
-		public FatalErrorException(string Message)
-			: base(Message)
+		/// <param name="message">The error message to display.</param>
+		public FatalErrorException(string message)
+			: base(message)
 		{
 		}
 
 		/// <summary>
 		/// Constructor
 		/// </summary>
-		/// <param name="InnerException">An inner exception to wrap</param>
-		/// <param name="Message">The error message to display.</param>
-		public FatalErrorException(Exception InnerException, string Message)
-			: base(Message, InnerException)
+		/// <param name="innerException">An inner exception to wrap</param>
+		/// <param name="message">The error message to display.</param>
+		public FatalErrorException(Exception innerException, string message)
+			: base(message, innerException)
 		{
 		}
 
 		/// <summary>
 		/// Constructor
 		/// </summary>
-		/// <param name="Format">Formatting string for the error message</param>
-		/// <param name="Arguments">Arguments for the formatting string</param>
-		public FatalErrorException(string Format, params object[] Arguments)
-			: base(String.Format(Format, Arguments))
+		/// <param name="format">Formatting string for the error message</param>
+		/// <param name="arguments">Arguments for the formatting string</param>
+		public FatalErrorException(string format, params object[] arguments)
+			: base(String.Format(format, arguments))
 		{
 		}
 
 		/// <summary>
 		/// Constructor which wraps another exception
 		/// </summary>
-		/// <param name="InnerException">The inner exception being wrapped</param>
-		/// <param name="Format">Format for the message string</param>
-		/// <param name="Arguments">Format arguments</param>
-		public FatalErrorException(Exception InnerException, string Format, params object[] Arguments)
-			: base(String.Format(Format, Arguments), InnerException)
+		/// <param name="innerException">The inner exception being wrapped</param>
+		/// <param name="format">Format for the message string</param>
+		/// <param name="arguments">Format arguments</param>
+		public FatalErrorException(Exception innerException, string format, params object[] arguments)
+			: base(String.Format(format, arguments), innerException)
 		{
 		}
 
diff --git a/Engine/Source/Programs/Shared/EpicGames.Core/FileFilter.cs b/Engine/Source/Programs/Shared/EpicGames.Core/FileFilter.cs
--- a/Engine/Source/Programs/Shared/EpicGames.Core/FileFilter.cs
+++ b/Engine/Source/Programs/Shared/EpicGames.Core/FileFilter.cs
@@ -2,11 +2,8 @@
 
 using System;
 using System.Collections.Generic;
-using System.Diagnostics;
 using System.IO;
 using System.Linq;
-using System.Text;
-using System.Text.RegularExpressions;
 
 namespace EpicGames.Core
 {
@@ -36,123 +33,123 @@ namespace EpicGames.Core
 		/// <summary>
 		/// Root node for the tree.
 		/// </summary>
-		FileFilterNode RootNode;
+		readonly FileFilterNode _rootNode;
 
 		/// <summary>
 		/// The default node, which will match any path
 		/// </summary>
-		FileFilterNode DefaultNode;
+		readonly FileFilterNode _defaultNode;
 
 		/// <summary>
 		/// Terminating nodes for each rule added to the filter
 		/// </summary>
-		List<FileFilterNode> Rules = new List<FileFilterNode>();
+		readonly List<FileFilterNode> _rules = new List<FileFilterNode>();
 
 		/// <summary>
 		/// Default constructor
 		/// </summary>
-		public FileFilter(FileFilterType DefaultType = FileFilterType.Exclude)
+		public FileFilter(FileFilterType defaultType = FileFilterType.Exclude)
 		{
-			RootNode = new FileFilterNode(null, "");
+			_rootNode = new FileFilterNode(null, "");
 
-			DefaultNode = new FileFilterNode(RootNode, "...");
-			DefaultNode.Type = DefaultType;
+			_defaultNode = new FileFilterNode(_rootNode, "...");
+			_defaultNode.Type = defaultType;
 		}
 
 		/// <summary>
 		/// Copy constructor
 		/// </summary>
-		/// <param name="Other">Filter to copy from</param>
-		public FileFilter(FileFilter Other)
-			: this(Other.DefaultNode.Type)
+		/// <param name="other">Filter to copy from</param>
+		public FileFilter(FileFilter other)
+			: this(other._defaultNode.Type)
 		{
-			foreach (FileFilterNode OtherRule in Other.Rules)
+			foreach (FileFilterNode otherRule in other._rules)
 			{
-				AddRule(OtherRule.ToString(), OtherRule.Type);
+				AddRule(otherRule.ToString(), otherRule.Type);
 			}
 		}
 
 		/// <summary>
 		/// Constructs a file filter from a p4-style filespec. Exclude lines are prefixed with a - character.
 		/// </summary>
-		/// <param name="Lines">Lines to parse rules from</param>
-		public FileFilter(IEnumerable<string> Lines)
+		/// <param name="lines">Lines to parse rules from</param>
+		public FileFilter(IEnumerable<string> lines)
 			: this()
 		{
-			foreach (string Line in Lines)
+			foreach (string line in lines)
 			{
-				AddRule(Line);
+				AddRule(line);
 			}
 		}
 
 		/// <summary>
 		/// Adds an include or exclude rule to the filter
 		/// </summary>
-		/// <param name="Rule">Pattern to match. See CreateRegex() for details.</param>
-		public void AddRule(string Rule)
+		/// <param name="rule">Pattern to match. See CreateRegex() for details.</param>
+		public void AddRule(string rule)
 		{
-			if (Rule.StartsWith("-"))
+			if (rule.StartsWith("-"))
 			{
-				Exclude(Rule.Substring(1).TrimStart());
+				Exclude(rule.Substring(1).TrimStart());
 			}
 			else
 			{
-				Include(Rule);
+				Include(rule);
 			}
 		}
 
 		/// <summary>
 		/// Adds an include or exclude rule to the filter. The rule may be 
 		/// </summary>
-		/// <param name="Rule">Pattern to match. See CreateRegex() for details.</param>
-		/// <param name="AllowTags">If the rule starts with a list of tags {a,b,c}, specifies a list of tags to match against</param>
-		public void AddRule(string Rule, params string[] AllowTags)
+		/// <param name="rule">Pattern to match. See CreateRegex() for details.</param>
+		/// <param name="allowTags">If the rule starts with a list of tags {a,b,c}, specifies a list of tags to match against</param>
+		public void AddRule(string rule, params string[] allowTags)
 		{
-			string CleanRule = Rule.Trim();
-			if(CleanRule.StartsWith("{"))
+			string cleanRule = rule.Trim();
+			if(cleanRule.StartsWith("{"))
 			{
 				// Find the end of the condition
-				int ConditionEnd = CleanRule.IndexOf('}');
-				if(ConditionEnd == -1)
+				int conditionEnd = cleanRule.IndexOf('}');
+				if(conditionEnd == -1)
 				{
-					throw new Exception(String.Format("Missing closing parenthesis in rule: {0}", CleanRule));
+					throw new Exception(String.Format("Missing closing parenthesis in rule: {0}", cleanRule));
 				}
 
 				// Check there's a matching tag
-				string[] RuleTags = CleanRule.Substring(1, ConditionEnd - 1).Split(',').Select(x => x.Trim()).ToArray();
-				if(!RuleTags.Any(x => AllowTags.Contains(x)))
+				string[] ruleTags = cleanRule.Substring(1, conditionEnd - 1).Split(',').Select(x => x.Trim()).ToArray();
+				if(!ruleTags.Any(x => allowTags.Contains(x)))
 				{
 					return;
 				}
 
 				// Strip the condition from the rule
-				CleanRule = CleanRule.Substring(ConditionEnd + 1).TrimStart(); 
+				cleanRule = cleanRule.Substring(conditionEnd + 1).TrimStart(); 
 			}
-			AddRule(CleanRule);
+			AddRule(cleanRule);
 		}
 
 		/// <summary>
 		/// Adds several rules to the filter
 		/// </summary>
-		/// <param name="Rules">List of patterns to match.</param>
-		public void AddRules(IEnumerable<string> Rules)
+		/// <param name="rules">List of patterns to match.</param>
+		public void AddRules(IEnumerable<string> rules)
 		{
-			foreach(string Rule in Rules)
+			foreach(string rule in rules)
 			{
-				AddRule(Rule);
+				AddRule(rule);
 			}
 		}
 
 		/// <summary>
 		/// Adds several rules to the filter
 		/// </summary>
-		/// <param name="Rules">List of patterns to match.</param>
-		/// <param name="Type">Whether the rules are include or exclude rules</param>
-		public void AddRules(IEnumerable<string> Rules, FileFilterType Type)
+		/// <param name="rules">List of patterns to match.</param>
+		/// <param name="type">Whether the rules are include or exclude rules</param>
+		public void AddRules(IEnumerable<string> rules, FileFilterType type)
 		{
-			foreach(string Rule in Rules)
+			foreach(string rule in rules)
 			{
-				AddRule(Rule, Type);
+				AddRule(rule, type);
 			}
 		}
 
@@ -160,62 +157,62 @@ namespace EpicGames.Core
 		/// Adds several rules in the given lines. Rules may be prefixed with conditions of the syntax {Key=Value, Key2=Value2}, which
 		/// will be evaluated using variables in the given dictionary before being added.
 		/// </summary>
-		/// <param name="Rules"></param>
-		/// <param name="Tags">Lookup for variables to test against</param>
-		public void AddRules(IEnumerable<string> Rules, params string[] Tags)
+		/// <param name="rules"></param>
+		/// <param name="tags">Lookup for variables to test against</param>
+		public void AddRules(IEnumerable<string> rules, params string[] tags)
 		{
-			foreach(string Rule in Rules)
+			foreach(string rule in rules)
 			{
-				AddRule(Rule, Tags);
+				AddRule(rule, tags);
 			}
 		}
 
 		/// <summary>
 		/// Adds a rule which matches a filename relative to a given base directory.
 		/// </summary>
-		/// <param name="File">The filename to add a rule for</param>
-		/// <param name="BaseDirectory">Base directory for the rule</param>
-		/// <param name="Type">Whether to add an include or exclude rule</param>
-		public void AddRuleForFile(FileReference File, DirectoryReference BaseDirectory, FileFilterType Type)
+		/// <param name="file">The filename to add a rule for</param>
+		/// <param name="baseDirectory">Base directory for the rule</param>
+		/// <param name="type">Whether to add an include or exclude rule</param>
+		public void AddRuleForFile(FileReference file, DirectoryReference baseDirectory, FileFilterType type)
 		{
-			AddRule("/" + File.MakeRelativeTo(BaseDirectory), Type);
+			AddRule("/" + file.MakeRelativeTo(baseDirectory), type);
 		}
 
 		/// <summary>
 		/// Adds rules for files which match the given names
 		/// </summary>
-		/// <param name="Files">The filenames to rules for</param>
-		/// <param name="BaseDirectory">Base directory for the rules</param>
-		/// <param name="Type">Whether to add an include or exclude rule</param>
-		public void AddRuleForFiles(IEnumerable<FileReference> Files, DirectoryReference BaseDirectory, FileFilterType Type)
+		/// <param name="files">The filenames to rules for</param>
+		/// <param name="baseDirectory">Base directory for the rules</param>
+		/// <param name="type">Whether to add an include or exclude rule</param>
+		public void AddRuleForFiles(IEnumerable<FileReference> files, DirectoryReference baseDirectory, FileFilterType type)
 		{
-			foreach (FileReference File in Files)
+			foreach (FileReference file in files)
 			{
-				AddRuleForFile(File, BaseDirectory, Type);
+				AddRuleForFile(file, baseDirectory, type);
 			}
 		}
 
 		/// <summary>
 		/// Reads a configuration file split into sections
 		/// </summary>
-		/// <param name="FileName"></param>
-		/// <param name="SectionName"></param>
-		/// <param name="AllowTags"></param>
-		public void ReadRulesFromFile(FileReference FileName, string SectionName, params string[] AllowTags)
+		/// <param name="fileName"></param>
+		/// <param name="sectionName"></param>
+		/// <param name="allowTags"></param>
+		public void ReadRulesFromFile(FileReference fileName, string sectionName, params string[] allowTags)
 		{
 			bool bInSection = false;
-			foreach(string Line in File.ReadAllLines(FileName.FullName))
+			foreach(string line in File.ReadAllLines(fileName.FullName))
 			{
-				string TrimLine = Line.Trim();
-				if(!TrimLine.StartsWith(";") && TrimLine.Length > 0)
+				string trimLine = line.Trim();
+				if(!trimLine.StartsWith(";") && trimLine.Length > 0)
 				{
-					if(TrimLine.StartsWith("["))
+					if(trimLine.StartsWith("["))
 					{
-						bInSection = (TrimLine == "[" + SectionName + "]");
+						bInSection = (trimLine == "[" + sectionName + "]");
 					}
 					else if(bInSection)
 					{
-						AddRule(Line, AllowTags);
+						AddRule(line, allowTags);
 					}
 				}
 			}
@@ -224,120 +221,120 @@ namespace EpicGames.Core
 		/// <summary>
 		/// Adds an include rule to the filter
 		/// </summary>
-		/// <param name="Pattern">Pattern to match. See CreateRegex() for details.</param>
-		public void Include(string Pattern)
+		/// <param name="pattern">Pattern to match. See CreateRegex() for details.</param>
+		public void Include(string pattern)
 		{
-			AddRule(Pattern, FileFilterType.Include);
+			AddRule(pattern, FileFilterType.Include);
 		}
 
 		/// <summary>
 		/// Adds several exclude rules to the filter
 		/// </summary>
-		/// <param name="Patterns">Patterns to match. See CreateRegex() for details.</param>
-		public void Include(IEnumerable<string> Patterns)
+		/// <param name="patterns">Patterns to match. See CreateRegex() for details.</param>
+		public void Include(IEnumerable<string> patterns)
 		{
-			foreach (string Pattern in Patterns)
+			foreach (string pattern in patterns)
 			{
-				Include(Pattern);
+				Include(pattern);
 			}
 		}
 
 		/// <summary>
 		/// Adds an exclude rule to the filter
 		/// </summary>
-		/// <param name="Pattern">Mask to match. See CreateRegex() for details.</param>
-		public void Exclude(string Pattern)
+		/// <param name="pattern">Mask to match. See CreateRegex() for details.</param>
+		public void Exclude(string pattern)
 		{
-			AddRule(Pattern, FileFilterType.Exclude);
+			AddRule(pattern, FileFilterType.Exclude);
 		}
 
 		/// <summary>
 		/// Adds several exclude rules to the filter
 		/// </summary>
-		/// <param name="Patterns">Patterns to match. See CreateRegex() for details.</param>
-		public void Exclude(IEnumerable<string> Patterns)
+		/// <param name="patterns">Patterns to match. See CreateRegex() for details.</param>
+		public void Exclude(IEnumerable<string> patterns)
 		{
-			foreach (string Pattern in Patterns)
+			foreach (string pattern in patterns)
 			{
-				Exclude(Pattern);
+				Exclude(pattern);
 			}
 		}
 
 		/// <summary>
 		/// Adds an include or exclude rule to the filter
 		/// </summary>
-		/// <param name="Pattern">The pattern which the rule should match</param>
-		/// <param name="Type">Whether to include or exclude files matching this rule</param>
-		public void AddRule(string Pattern, FileFilterType Type)
+		/// <param name="pattern">The pattern which the rule should match</param>
+		/// <param name="type">Whether to include or exclude files matching this rule</param>
+		public void AddRule(string pattern, FileFilterType type)
 		{
-			string NormalizedPattern = Pattern.Replace('\\', '/');
+			string normalizedPattern = pattern.Replace('\\', '/');
 
 			// Remove the slash from the start of the pattern. Any exclude pattern that doesn't contain a directory separator is assumed to apply to any directory (eg. *.cpp), otherwise it's 
 			// taken relative to the root.
-			if (NormalizedPattern.StartsWith("/"))
+			if (normalizedPattern.StartsWith("/"))
 			{
-				NormalizedPattern = NormalizedPattern.Substring(1);
+				normalizedPattern = normalizedPattern.Substring(1);
 			}
-			else if(!NormalizedPattern.Contains("/") && !NormalizedPattern.StartsWith("..."))
+			else if(!normalizedPattern.Contains("/") && !normalizedPattern.StartsWith("..."))
 			{
-				NormalizedPattern = ".../" + NormalizedPattern;
+				normalizedPattern = ".../" + normalizedPattern;
 			}
 
 			// All directories indicate a wildcard match
-			if (NormalizedPattern.EndsWith("/"))
+			if (normalizedPattern.EndsWith("/"))
 			{
-				NormalizedPattern += "...";
+				normalizedPattern += "...";
 			}
 
 			// Replace any directory wildcards mid-string
-			for (int Idx = NormalizedPattern.IndexOf("..."); Idx != -1; Idx = NormalizedPattern.IndexOf("...", Idx))
+			for (int idx = normalizedPattern.IndexOf("..."); idx != -1; idx = normalizedPattern.IndexOf("...", idx))
 			{
-				if (Idx > 0 && NormalizedPattern[Idx - 1] != '/')
+				if (idx > 0 && normalizedPattern[idx - 1] != '/')
 				{
-					NormalizedPattern = NormalizedPattern.Insert(Idx, "*/");
-					Idx++;
+					normalizedPattern = normalizedPattern.Insert(idx, "*/");
+					idx++;
 				}
 
-				Idx += 3;
+				idx += 3;
 
-				if (Idx < NormalizedPattern.Length && NormalizedPattern[Idx] != '/')
+				if (idx < normalizedPattern.Length && normalizedPattern[idx] != '/')
 				{
-					NormalizedPattern = NormalizedPattern.Insert(Idx, "/*");
-					Idx += 2;
+					normalizedPattern = normalizedPattern.Insert(idx, "/*");
+					idx += 2;
 				}
 			}
 
 			// Split the pattern into fragments
-			string[] BranchPatterns = NormalizedPattern.Split('/');
+			string[] branchPatterns = normalizedPattern.Split('/');
 
 			// Add it into the tree
-			FileFilterNode LastNode = RootNode;
-			foreach (string BranchPattern in BranchPatterns)
+			FileFilterNode lastNode = _rootNode;
+			foreach (string branchPattern in branchPatterns)
 			{
-				FileFilterNode NextNode = LastNode.Branches.FirstOrDefault(x => x.Pattern == BranchPattern);
-				if (NextNode == null)
+				FileFilterNode nextNode = lastNode.Branches.FirstOrDefault(x => x.Pattern == branchPattern);
+				if (nextNode == null)
 				{
-					NextNode = new FileFilterNode(LastNode, BranchPattern);
-					LastNode.Branches.Add(NextNode);
+					nextNode = new FileFilterNode(lastNode, branchPattern);
+					lastNode.Branches.Add(nextNode);
 				}
-				LastNode = NextNode;
+				lastNode = nextNode;
 			}
 
 			// We've reached the end of the pattern, so mark it as a leaf node
-			Rules.Add(LastNode);
-			LastNode.RuleNumber = Rules.Count - 1;
-			LastNode.Type = Type;
+			_rules.Add(lastNode);
+			lastNode.RuleNumber = _rules.Count - 1;
+			lastNode.Type = type;
 
 			// Update the maximums along that path
-			for (FileFilterNode? UpdateNode = LastNode; UpdateNode != null; UpdateNode = UpdateNode.Parent)
+			for (FileFilterNode? updateNode = lastNode; updateNode != null; updateNode = updateNode.Parent)
 			{
-				if (Type == FileFilterType.Include)
+				if (type == FileFilterType.Include)
 				{
-					UpdateNode.MaxIncludeRuleNumber = LastNode.RuleNumber;
+					updateNode.MaxIncludeRuleNumber = lastNode.RuleNumber;
 				}
 				else
 				{
-					UpdateNode.MaxExcludeRuleNumber = LastNode.RuleNumber;
+					updateNode.MaxExcludeRuleNumber = lastNode.RuleNumber;
 				}
 			}
 		}
@@ -345,158 +342,158 @@ namespace EpicGames.Core
 		/// <summary>
 		/// Determines whether the given file matches the filter
 		/// </summary>
-		/// <param name="FileName">File to match</param>
+		/// <param name="fileName">File to match</param>
 		/// <returns>True if the file passes the filter</returns>
-		public bool Matches(string FileName)
+		public bool Matches(string fileName)
 		{
-			string[] Tokens = FileName.TrimStart('/', '\\').Split('/', '\\');
+			string[] tokens = fileName.TrimStart('/', '\\').Split('/', '\\');
 
-			FileFilterNode MatchingNode = FindMatchingNode(RootNode, Tokens, 0, DefaultNode);
+			FileFilterNode matchingNode = FindMatchingNode(_rootNode, tokens, 0, _defaultNode);
 
-			return MatchingNode.Type == FileFilterType.Include;
+			return matchingNode.Type == FileFilterType.Include;
 		}
 
 		/// <summary>
 		/// Determines whether it's possible for anything within the given folder name to match the filter. Useful to early out of recursive file searches.
 		/// </summary>
-		/// <param name="FolderName">File to match</param>
+		/// <param name="folderName">File to match</param>
 		/// <returns>True if the file passes the filter</returns>
-		public bool PossiblyMatches(string FolderName)
+		public bool PossiblyMatches(string folderName)
 		{
-			string[] Tokens = FolderName.Trim('/', '\\').Split('/', '\\');
+			string[] tokens = folderName.Trim('/', '\\').Split('/', '\\');
 
-			FileFilterNode MatchingNode = FindMatchingNode(RootNode, Tokens.Union(new string[]{ "" }).ToArray(), 0, DefaultNode);
+			FileFilterNode matchingNode = FindMatchingNode(_rootNode, tokens.Union(new string[]{ "" }).ToArray(), 0, _defaultNode);
 
-			return MatchingNode.Type == FileFilterType.Include || HighestPossibleIncludeMatch(RootNode, Tokens, 0, MatchingNode.RuleNumber) > MatchingNode.RuleNumber;
+			return matchingNode.Type == FileFilterType.Include || HighestPossibleIncludeMatch(_rootNode, tokens, 0, matchingNode.RuleNumber) > matchingNode.RuleNumber;
 		}
 
 		/// <summary>
 		/// Applies the filter to each element in a sequence, and returns the list of files that match
 		/// </summary>
-		/// <param name="FileNames">List of filenames</param>
+		/// <param name="fileNames">List of filenames</param>
 		/// <returns>List of filenames which match the filter</returns>
-		public IEnumerable<string> ApplyTo(IEnumerable<string> FileNames)
+		public IEnumerable<string> ApplyTo(IEnumerable<string> fileNames)
 		{
-			return FileNames.Where(x => Matches(x));
+			return fileNames.Where(x => Matches(x));
 		}
 
 		/// <summary>
 		/// Applies the filter to each element in a sequence, and returns the list of files that match
 		/// </summary>
-		/// <param name="BaseDirectory">The base directory to match files against</param>
-		/// <param name="FileNames">List of filenames</param>
+		/// <param name="baseDirectory">The base directory to match files against</param>
+		/// <param name="fileNames">List of filenames</param>
 		/// <returns>List of filenames which match the filter</returns>
-		public IEnumerable<FileReference> ApplyTo(DirectoryReference BaseDirectory, IEnumerable<FileReference> FileNames)
+		public IEnumerable<FileReference> ApplyTo(DirectoryReference baseDirectory, IEnumerable<FileReference> fileNames)
 		{
-			return FileNames.Where(x => Matches(x.MakeRelativeTo(BaseDirectory)));
+			return fileNames.Where(x => Matches(x.MakeRelativeTo(baseDirectory)));
 		}
 
 		/// <summary>
 		/// Finds a list of files within a given directory which match the filter.
 		/// </summary>
-		/// <param name="DirectoryName">File to match</param>
+		/// <param name="directoryName">File to match</param>
 		/// <param name="bIgnoreSymlinks">Whether to ignore symlinks in the output</param>
 		/// <returns>List of files that pass the filter</returns>
-		public List<FileReference> ApplyToDirectory(DirectoryReference DirectoryName, bool bIgnoreSymlinks)
+		public List<FileReference> ApplyToDirectory(DirectoryReference directoryName, bool bIgnoreSymlinks)
 		{
-			List<FileReference> MatchingFileNames = new List<FileReference>();
-			if (DirectoryReference.Exists(DirectoryName))
+			List<FileReference> matchingFileNames = new List<FileReference>();
+			if (DirectoryReference.Exists(directoryName))
 			{
-				FindMatchesFromDirectory(new DirectoryInfo(DirectoryName.FullName), "", bIgnoreSymlinks, MatchingFileNames);
+				FindMatchesFromDirectory(new DirectoryInfo(directoryName.FullName), "", bIgnoreSymlinks, matchingFileNames);
 			}
-			return MatchingFileNames;
+			return matchingFileNames;
 		}
 
 		/// <summary>
 		/// Finds a list of files within a given directory which match the filter.
 		/// </summary>
-		/// <param name="DirectoryName">File to match</param>
-		/// <param name="PrefixPath"></param>
+		/// <param name="directoryName">File to match</param>
+		/// <param name="prefixPath"></param>
 		/// <param name="bIgnoreSymlinks">Whether to ignore symlinks in the output</param>
 		/// <returns>List of files that pass the filter</returns>
-		public List<FileReference> ApplyToDirectory(DirectoryReference DirectoryName, string PrefixPath, bool bIgnoreSymlinks)
+		public List<FileReference> ApplyToDirectory(DirectoryReference directoryName, string prefixPath, bool bIgnoreSymlinks)
 		{
-			List<FileReference> MatchingFileNames = new List<FileReference>();
-			if (DirectoryReference.Exists(DirectoryName))
+			List<FileReference> matchingFileNames = new List<FileReference>();
+			if (DirectoryReference.Exists(directoryName))
 			{
-				FindMatchesFromDirectory(new DirectoryInfo(DirectoryName.FullName), PrefixPath.Replace('\\', '/').TrimEnd('/') + "/", bIgnoreSymlinks, MatchingFileNames);
+				FindMatchesFromDirectory(new DirectoryInfo(directoryName.FullName), prefixPath.Replace('\\', '/').TrimEnd('/') + "/", bIgnoreSymlinks, matchingFileNames);
 			}
-			return MatchingFileNames;
+			return matchingFileNames;
 		}
 
 		/// <summary>
 		/// Checks whether the given pattern contains a supported wildcard. Useful for distinguishing explicit file references from opportunistic file references.
 		/// </summary>
-		/// <param name="Pattern">The pattern to check</param>
+		/// <param name="pattern">The pattern to check</param>
 		/// <returns>True if the pattern contains a wildcard (?, *, ...), false otherwise.</returns>
-		public static int FindWildcardIndex(string Pattern)
+		public static int FindWildcardIndex(string pattern)
 		{
-			int Result = -1;
-			for(int Idx = 0; Idx < Pattern.Length; Idx++)
+			int result = -1;
+			for(int idx = 0; idx < pattern.Length; idx++)
 			{
-				if(Pattern[Idx] == '?' || Pattern[Idx] == '*' || (Pattern[Idx] == '.' && Idx + 2 < Pattern.Length && Pattern[Idx + 1] == '.' && Pattern[Idx + 2] == '.'))
+				if(pattern[idx] == '?' || pattern[idx] == '*' || (pattern[idx] == '.' && idx + 2 < pattern.Length && pattern[idx + 1] == '.' && pattern[idx + 2] == '.'))
 				{
-					Result = Idx;
+					result = idx;
 					break;
 				}
 			}
-			return Result;
+			return result;
 		}
 
 		/// <summary>
 		/// Resolve an individual wildcard to a set of files
 		/// </summary>
-		/// <param name="Pattern">The pattern to resolve</param>
+		/// <param name="pattern">The pattern to resolve</param>
 		/// <returns>List of files matching the wildcard</returns>
-		public static List<FileReference> ResolveWildcard(string Pattern)
+		public static List<FileReference> ResolveWildcard(string pattern)
 		{
-			for(int Idx = FindWildcardIndex(Pattern); Idx > 0; Idx--)
+			for(int idx = FindWildcardIndex(pattern); idx > 0; idx--)
 			{
-				if(Pattern[Idx] == '/' || Pattern[Idx] == '\\')
+				if(pattern[idx] == '/' || pattern[idx] == '\\')
 				{
-					return ResolveWildcard(new DirectoryReference(Pattern.Substring(0, Idx)), Pattern.Substring(Idx + 1));
+					return ResolveWildcard(new DirectoryReference(pattern.Substring(0, idx)), pattern.Substring(idx + 1));
 				}
 			}
-			return ResolveWildcard(DirectoryReference.GetCurrentDirectory(), Pattern);
+			return ResolveWildcard(DirectoryReference.GetCurrentDirectory(), pattern);
 		}
 
 		/// <summary>
 		/// Resolve an individual wildcard to a set of files
 		/// </summary>
-		/// <param name="BaseDir">Base directory for wildcards</param>
-		/// <param name="Pattern">The pattern to resolve</param>
+		/// <param name="baseDir">Base directory for wildcards</param>
+		/// <param name="pattern">The pattern to resolve</param>
 		/// <returns>List of files matching the wildcard</returns>
-		public static List<FileReference> ResolveWildcard(DirectoryReference BaseDir, string Pattern)
+		public static List<FileReference> ResolveWildcard(DirectoryReference baseDir, string pattern)
 		{
-			FileFilter Filter = new FileFilter(FileFilterType.Exclude);
-			Filter.AddRule(Pattern);
-			return Filter.ApplyToDirectory(BaseDir, true);
+			FileFilter filter = new FileFilter(FileFilterType.Exclude);
+			filter.AddRule(pattern);
+			return filter.ApplyToDirectory(baseDir, true);
 		}
 
 		/// <summary>
 		/// Finds a list of files within a given directory which match the filter.
 		/// </summary>
-		/// <param name="CurrentDirectory">The current directory</param>
-		/// <param name="NamePrefix">Current relative path prefix in the traversed directory tree</param>
+		/// <param name="currentDirectory">The current directory</param>
+		/// <param name="namePrefix">Current relative path prefix in the traversed directory tree</param>
 		/// <param name="bIgnoreSymlinks">Whether to ignore symlinks in the output</param>
-		/// <param name="MatchingFileNames">Receives a list of matching files</param>
+		/// <param name="matchingFileNames">Receives a list of matching files</param>
 		/// <returns>True if the file passes the filter</returns>
-		void FindMatchesFromDirectory(DirectoryInfo CurrentDirectory, string NamePrefix, bool bIgnoreSymlinks, List<FileReference> MatchingFileNames)
+		void FindMatchesFromDirectory(DirectoryInfo currentDirectory, string namePrefix, bool bIgnoreSymlinks, List<FileReference> matchingFileNames)
 		{
-			foreach (FileInfo NextFile in CurrentDirectory.EnumerateFiles())
+			foreach (FileInfo nextFile in currentDirectory.EnumerateFiles())
 			{
-				string FileName = NamePrefix + NextFile.Name;
-				if (Matches(FileName) && (!bIgnoreSymlinks || !NextFile.Attributes.HasFlag(FileAttributes.ReparsePoint)))
+				string fileName = namePrefix + nextFile.Name;
+				if (Matches(fileName) && (!bIgnoreSymlinks || !nextFile.Attributes.HasFlag(FileAttributes.ReparsePoint)))
 				{
-					MatchingFileNames.Add(new FileReference(NextFile));
+					matchingFileNames.Add(new FileReference(nextFile));
 				}
 			}
-			foreach (DirectoryInfo NextDirectory in CurrentDirectory.EnumerateDirectories())
+			foreach (DirectoryInfo nextDirectory in currentDirectory.EnumerateDirectories())
 			{
-				string NextNamePrefix = NamePrefix + NextDirectory.Name;
-				if (PossiblyMatches(NextNamePrefix))
+				string nextNamePrefix = namePrefix + nextDirectory.Name;
+				if (PossiblyMatches(nextNamePrefix))
 				{
-					FindMatchesFromDirectory(NextDirectory, NextNamePrefix + "/", bIgnoreSymlinks, MatchingFileNames);
+					FindMatchesFromDirectory(nextDirectory, nextNamePrefix + "/", bIgnoreSymlinks, matchingFileNames);
 				}
 			}
 		}
@@ -504,86 +501,86 @@ namespace EpicGames.Core
 		/// <summary>
 		/// Finds the node which matches a given list of tokens.
 		/// </summary>
-		/// <param name="CurrentNode"></param>
-		/// <param name="Tokens"></param>
-		/// <param name="TokenIdx"></param>
-		/// <param name="CurrentBestNode"></param>
+		/// <param name="currentNode"></param>
+		/// <param name="tokens"></param>
+		/// <param name="tokenIdx"></param>
+		/// <param name="currentBestNode"></param>
 		/// <returns></returns>
-		FileFilterNode FindMatchingNode(FileFilterNode CurrentNode, string[] Tokens, int TokenIdx, FileFilterNode CurrentBestNode)
+		FileFilterNode FindMatchingNode(FileFilterNode currentNode, string[] tokens, int tokenIdx, FileFilterNode currentBestNode)
 		{
 			// If we've matched all the input tokens, check if this rule is better than any other we've seen
-			if (TokenIdx == Tokens.Length)
+			if (tokenIdx == tokens.Length)
 			{
-				return (CurrentNode.RuleNumber > CurrentBestNode.RuleNumber) ? CurrentNode : CurrentBestNode;
+				return (currentNode.RuleNumber > currentBestNode.RuleNumber) ? currentNode : currentBestNode;
 			}
 
 			// If there is no rule under the current node which is better than the current best node, early out
-			if (CurrentNode.MaxIncludeRuleNumber <= CurrentBestNode.RuleNumber && CurrentNode.MaxExcludeRuleNumber <= CurrentBestNode.RuleNumber)
+			if (currentNode.MaxIncludeRuleNumber <= currentBestNode.RuleNumber && currentNode.MaxExcludeRuleNumber <= currentBestNode.RuleNumber)
 			{
-				return CurrentBestNode;
+				return currentBestNode;
 			}
 
 			// Test all the branches for one that matches
-			FileFilterNode BestNode = CurrentBestNode;
-			foreach (FileFilterNode Branch in CurrentNode.Branches)
+			FileFilterNode bestNode = currentBestNode;
+			foreach (FileFilterNode branch in currentNode.Branches)
 			{
-				if (Branch.Pattern == "...")
+				if (branch.Pattern == "...")
 				{
-					for (int NextTokenIdx = Tokens.Length; NextTokenIdx >= TokenIdx; NextTokenIdx--)
+					for (int nextTokenIdx = tokens.Length; nextTokenIdx >= tokenIdx; nextTokenIdx--)
 					{
-						BestNode = FindMatchingNode(Branch, Tokens, NextTokenIdx, BestNode);
+						bestNode = FindMatchingNode(branch, tokens, nextTokenIdx, bestNode);
 					}
 				}
 				else
 				{
-					if (Branch.IsMatch(Tokens[TokenIdx]))
+					if (branch.IsMatch(tokens[tokenIdx]))
 					{
-						BestNode = FindMatchingNode(Branch, Tokens, TokenIdx + 1, BestNode);
+						bestNode = FindMatchingNode(branch, tokens, tokenIdx + 1, bestNode);
 					}
 				}
 			}
-			return BestNode;
+			return bestNode;
 		}
 
 		/// <summary>
 		/// Returns the highest possible rule number which can match the given list of input tokens, assuming that the list of input tokens is incomplete.
 		/// </summary>
-		/// <param name="CurrentNode">The current node being checked</param>
-		/// <param name="Tokens">The tokens to match</param>
-		/// <param name="TokenIdx">Current token index</param>
-		/// <param name="CurrentBestRuleNumber">The highest rule number seen so far. Used to optimize tree traversals.</param>
+		/// <param name="currentNode">The current node being checked</param>
+		/// <param name="tokens">The tokens to match</param>
+		/// <param name="tokenIdx">Current token index</param>
+		/// <param name="currentBestRuleNumber">The highest rule number seen so far. Used to optimize tree traversals.</param>
 		/// <returns>New highest rule number</returns>
-		int HighestPossibleIncludeMatch(FileFilterNode CurrentNode, string[] Tokens, int TokenIdx, int CurrentBestRuleNumber)
+		int HighestPossibleIncludeMatch(FileFilterNode currentNode, string[] tokens, int tokenIdx, int currentBestRuleNumber)
 		{
 			// If we've matched all the input tokens, check if this rule is better than any other we've seen
-			if (TokenIdx == Tokens.Length)
+			if (tokenIdx == tokens.Length)
 			{
-				return Math.Max(CurrentBestRuleNumber, CurrentNode.MaxIncludeRuleNumber);
+				return Math.Max(currentBestRuleNumber, currentNode.MaxIncludeRuleNumber);
 			}
 
 			// Test all the branches for one that matches
-			int BestRuleNumber = CurrentBestRuleNumber;
-			if (CurrentNode.MaxIncludeRuleNumber > BestRuleNumber)
+			int bestRuleNumber = currentBestRuleNumber;
+			if (currentNode.MaxIncludeRuleNumber > bestRuleNumber)
 			{
-				foreach (FileFilterNode Branch in CurrentNode.Branches)
+				foreach (FileFilterNode branch in currentNode.Branches)
 				{
-					if (Branch.Pattern == "...")
+					if (branch.Pattern == "...")
 					{
-						if (Branch.MaxIncludeRuleNumber > BestRuleNumber)
+						if (branch.MaxIncludeRuleNumber > bestRuleNumber)
 						{
-							BestRuleNumber = Branch.MaxIncludeRuleNumber;
+							bestRuleNumber = branch.MaxIncludeRuleNumber;
 						}
 					}
 					else
 					{
-						if (Branch.IsMatch(Tokens[TokenIdx]))
+						if (branch.IsMatch(tokens[tokenIdx]))
 						{
-							BestRuleNumber = HighestPossibleIncludeMatch(Branch, Tokens, TokenIdx + 1, BestRuleNumber);
+							bestRuleNumber = HighestPossibleIncludeMatch(branch, tokens, tokenIdx + 1, bestRuleNumber);
 						}
 					}
 				}
 			}
-			return BestRuleNumber;
+			return bestRuleNumber;
 		}
 	}
 }
diff --git a/Engine/Source/Programs/Shared/EpicGames.Core/FileFilterNode.cs b/Engine/Source/Programs/Shared/EpicGames.Core/FileFilterNode.cs
--- a/Engine/Source/Programs/Shared/EpicGames.Core/FileFilterNode.cs
+++ b/Engine/Source/Programs/Shared/EpicGames.Core/FileFilterNode.cs
@@ -2,9 +2,6 @@
 
 using System;
 using System.Collections.Generic;
-using System.Linq;
-using System.Text;
-using System.Threading.Tasks;
 
 namespace EpicGames.Core
 {
@@ -26,35 +23,35 @@ namespace EpicGames.Core
 		/// <summary>
 		/// Highest include rule number matched by this node or any child nodes.
 		/// </summary>
-		public int MaxIncludeRuleNumber = -1;
+		public int MaxIncludeRuleNumber { get; set; } = -1;
 
 		/// <summary>
 		/// Highest exclude rule number matched by this node or any child nodes.
 		/// </summary>
-		public int MaxExcludeRuleNumber = -1;
+		public int MaxExcludeRuleNumber { get; set; } = -1;
 
 		/// <summary>
 		/// Child branches of this node, distinct by the the pattern for each.
 		/// </summary>
-		public List<FileFilterNode> Branches = new List<FileFilterNode>();
+		public List<FileFilterNode> Branches { get; set; } = new List<FileFilterNode>();
 
 		/// <summary>
 		/// If a path matches the sequence of nodes ending in this node, the number of the rule which added it. -1 if this was not the last node in a rule.
 		/// </summary>
-		public int RuleNumber;
+		public int RuleNumber { get; set; }
 
 		/// <summary>
 		/// Whether a rule terminating in this node should include files or exclude files.
 		/// </summary>
-		public FileFilterType Type;
+		public FileFilterType Type { get; set; }
 
 		/// <summary>
 		/// Default constructor.
 		/// </summary>
-		public FileFilterNode(FileFilterNode? InParent, string InPattern)
+		public FileFilterNode(FileFilterNode? inParent, string inPattern)
 		{
-			Parent = InParent;
-			Pattern = InPattern;
+			Parent = inParent;
+			Pattern = inPattern;
 			RuleNumber = -1;
 			Type = FileFilterType.Exclude;
 		}
@@ -62,44 +59,44 @@ namespace EpicGames.Core
 		/// <summary>
 		/// Determine if the given token matches the pattern in this node
 		/// </summary>
-		public bool IsMatch(string Token)
+		public bool IsMatch(string token)
 		{
 			if(Pattern.EndsWith("."))
 			{
-				return !Token.Contains('.') && IsMatch(Token, 0, Pattern.Substring(0, Pattern.Length - 1), 0);
+				return !token.Contains('.') && IsMatch(token, 0, Pattern.Substring(0, Pattern.Length - 1), 0);
 			}
 			else
 			{
-				return IsMatch(Token, 0, Pattern, 0);
+				return IsMatch(token, 0, Pattern, 0);
 			}
 		}
 
 		/// <summary>
 		/// Determine if a given token matches a pattern, starting from the given positions within each string.
 		/// </summary>
-		static bool IsMatch(string Token, int TokenIdx, string Pattern, int PatternIdx)
+		static bool IsMatch(string token, int tokenIdx, string pattern, int patternIdx)
 		{
 			for (; ; )
 			{
-				if (PatternIdx == Pattern.Length)
+				if (patternIdx == pattern.Length)
 				{
-					return (TokenIdx == Token.Length);
+					return (tokenIdx == token.Length);
 				}
-				else if (Pattern[PatternIdx] == '*')
+				else if (pattern[patternIdx] == '*')
 				{
-					for (int NextTokenIdx = Token.Length; NextTokenIdx >= TokenIdx; NextTokenIdx--)
+					for (int nextTokenIdx = token.Length; nextTokenIdx >= tokenIdx; nextTokenIdx--)
 					{
-						if (IsMatch(Token, NextTokenIdx, Pattern, PatternIdx + 1))
+						if (IsMatch(token, nextTokenIdx, pattern, patternIdx + 1))
 						{
 							return true;
 						}
 					}
 					return false;
 				}
-				else if (TokenIdx < Token.Length && (Char.ToLower(Token[TokenIdx]) == Char.ToLower(Pattern[PatternIdx]) || Pattern[PatternIdx] == '?'))
+				else if (tokenIdx < token.Length && (Char.ToLower(token[tokenIdx]) == Char.ToLower(pattern[patternIdx]) || pattern[patternIdx] == '?'))
 				{
-					TokenIdx++;
-					PatternIdx++;
+					tokenIdx++;
+					patternIdx++;
 				}
 				else
 				{
