Reformat EpicGames.Core according to standard coding conventions.

diff --git a/Engine/Source/Programs/Shared/EpicGames.Core/FilePattern.cs b/Engine/Source/Programs/Shared/EpicGames.Core/FilePattern.cs
--- a/Engine/Source/Programs/Shared/EpicGames.Core/FilePattern.cs
+++ b/Engine/Source/Programs/Shared/EpicGames.Core/FilePattern.cs
@@ -3,22 +3,20 @@
 using System;
 using System.Collections.Generic;
 using System.IO;
-using System.Linq;
 using System.Text;
 using System.Text.RegularExpressions;
-using System.Threading.Tasks;
 
 namespace EpicGames.Core
 {
 	public class FilePatternException : Exception
 	{
-		public FilePatternException(string Message)
-			: base(Message)
+		public FilePatternException(string message)
+			: base(message)
 		{
 		}
 
-		public FilePatternException(string Format, params object[] Args)
-			: base(String.Format(Format, Args))
+		public FilePatternException(string format, params object[] args)
+			: base(String.Format(format, args))
 		{
 		}
 
@@ -46,92 +44,92 @@ namespace EpicGames.Core
 		/// <summary>
 		/// Constructs a file pattern which matches a single file
 		/// </summary>
-		/// <param name="File">Location of the file</param>
-		public FilePattern(FileReference File)
+		/// <param name="file">Location of the file</param>
+		public FilePattern(FileReference file)
 		{
-			BaseDirectory = File.Directory;
-			Tokens.Add(File.GetFileName());
+			BaseDirectory = file.Directory;
+			Tokens.Add(file.GetFileName());
 		}
 
 		/// <summary>
 		/// Constructs a file pattern from the given string, resolving relative paths to the given directory. 
 		/// </summary>
-		/// <param name="RootDirectory">If a relative path is specified by the pattern, the root directory used to turn it into an absolute path</param>
-		/// <param name="Pattern">The pattern to match. If the pattern ends with a directory separator, an implicit '...' is appended.</param>
-		public FilePattern(DirectoryReference RootDirectory, string Pattern)
+		/// <param name="rootDirectory">If a relative path is specified by the pattern, the root directory used to turn it into an absolute path</param>
+		/// <param name="pattern">The pattern to match. If the pattern ends with a directory separator, an implicit '...' is appended.</param>
+		public FilePattern(DirectoryReference rootDirectory, string pattern)
 		{
 			// Normalize the path separators
-			StringBuilder Text = new StringBuilder(Pattern);
+			StringBuilder text = new StringBuilder(pattern);
 			if(Path.DirectorySeparatorChar != '\\')
 			{
-				Text.Replace('\\', Path.DirectorySeparatorChar);
+				text.Replace('\\', Path.DirectorySeparatorChar);
 			}
 			if(Path.DirectorySeparatorChar != '/')
 			{
-				Text.Replace('/', Path.DirectorySeparatorChar);
+				text.Replace('/', Path.DirectorySeparatorChar);
 			}
 
 			// Find the base directory, stopping when we hit a wildcard. The source directory must end with a path specification.
-			int BaseDirectoryLen = 0;
-			for(int Idx = 0; Idx < Text.Length; Idx++)
+			int baseDirectoryLen = 0;
+			for(int idx = 0; idx < text.Length; idx++)
 			{
-				if(Text[Idx] == Path.DirectorySeparatorChar)
+				if(text[idx] == Path.DirectorySeparatorChar)
 				{
-					BaseDirectoryLen = Idx + 1;
+					baseDirectoryLen = idx + 1;
 				}
-				else if(Text[Idx] == '?' || Text[Idx] == '*' || (Idx + 2 < Text.Length && Text[Idx] == '.' && Text[Idx + 1] == '.' && Text[Idx + 2] == '.'))
+				else if(text[idx] == '?' || text[idx] == '*' || (idx + 2 < text.Length && text[idx] == '.' && text[idx + 1] == '.' && text[idx + 2] == '.'))
 				{
 					break;
 				}
 			}
 
 			// Extract the base directory
-			BaseDirectory = DirectoryReference.Combine(RootDirectory, Text.ToString(0, BaseDirectoryLen));
+			BaseDirectory = DirectoryReference.Combine(rootDirectory, text.ToString(0, baseDirectoryLen));
 
 			// Convert any directory wildcards ("...") into complete directory wildcards ("\\...\\"). We internally treat use "...\\" as the wildcard 
 			// token so we can correctly match zero directories. Patterns such as "foo...bar" should require at least one directory separator, so 
 			// should be converted to "foo*\\...\\*bar".
-			for(int Idx = BaseDirectoryLen; Idx < Text.Length; Idx++)
+			for(int idx = baseDirectoryLen; idx < text.Length; idx++)
 			{
-				if(Text[Idx] == '.' && Text[Idx + 1] == '.' && Text[Idx + 2] == '.')
+				if(text[idx] == '.' && text[idx + 1] == '.' && text[idx + 2] == '.')
 				{
 					// Insert a directory separator before
-					if(Idx > BaseDirectoryLen && Text[Idx - 1] != Path.DirectorySeparatorChar)
+					if(idx > baseDirectoryLen && text[idx - 1] != Path.DirectorySeparatorChar)
 					{
-						Text.Insert(Idx++, '*');
-						Text.Insert(Idx++, Path.DirectorySeparatorChar);
+						text.Insert(idx++, '*');
+						text.Insert(idx++, Path.DirectorySeparatorChar);
 					}
 
 					// Skip past the ellipsis
-					Idx += 3;
+					idx += 3;
 
 					// Insert a directory separator after
-					if(Idx == Text.Length || Text[Idx] != Path.DirectorySeparatorChar)
+					if(idx == text.Length || text[idx] != Path.DirectorySeparatorChar)
 					{
-						Text.Insert(Idx++, Path.DirectorySeparatorChar);
-						Text.Insert(Idx++, '*');
+						text.Insert(idx++, Path.DirectorySeparatorChar);
+						text.Insert(idx++, '*');
 					}
 				}
 			}
 
 			// Parse the tokens
-			int LastIdx = BaseDirectoryLen;
-			for(int Idx = BaseDirectoryLen; Idx < Text.Length; Idx++)
+			int lastIdx = baseDirectoryLen;
+			for(int idx = baseDirectoryLen; idx < text.Length; idx++)
 			{
-				if(Text[Idx] == '?' || Text[Idx] == '*')
+				if(text[idx] == '?' || text[idx] == '*')
 				{
-					Tokens.Add(Text.ToString(LastIdx, Idx - LastIdx));
-					Tokens.Add(Text.ToString(Idx, 1));
-					LastIdx = Idx + 1;
+					Tokens.Add(text.ToString(lastIdx, idx - lastIdx));
+					Tokens.Add(text.ToString(idx, 1));
+					lastIdx = idx + 1;
 				}
-				else if(Idx - 3 >= BaseDirectoryLen && Text[Idx] == Path.DirectorySeparatorChar && Text[Idx - 1] == '.' && Text[Idx - 2] == '.' && Text[Idx - 3] == '.')
+				else if(idx - 3 >= baseDirectoryLen && text[idx] == Path.DirectorySeparatorChar && text[idx - 1] == '.' && text[idx - 2] == '.' && text[idx - 3] == '.')
 				{
-					Tokens.Add(Text.ToString(LastIdx, Idx - 3 - LastIdx));
-					Tokens.Add(Text.ToString(Idx - 3, 4));
-					LastIdx = Idx + 1;
+					Tokens.Add(text.ToString(lastIdx, idx - 3 - lastIdx));
+					Tokens.Add(text.ToString(idx - 3, 4));
+					lastIdx = idx + 1;
 				}
 			}
-			Tokens.Add(Text.ToString(LastIdx, Text.Length - LastIdx));
+			Tokens.Add(text.ToString(lastIdx, text.Length - lastIdx));
 		}
 
 		/// <summary>
@@ -146,18 +144,18 @@ namespace EpicGames.Core
 			}
 			else
 			{
-				StringBuilder Pattern = new StringBuilder();
-				foreach(string Token in Tokens)
+				StringBuilder pattern = new StringBuilder();
+				foreach(string token in Tokens)
 				{
-					Pattern.Append(Token);
+					pattern.Append(token);
 				}
-				if(Pattern.Length > 0)
+				if(pattern.Length > 0)
 				{
-					Pattern.Append(Path.DirectorySeparatorChar);
+					pattern.Append(Path.DirectorySeparatorChar);
 				}
-				Pattern.Append("...");
+				pattern.Append("...");
 
-				return new FilePattern(BaseDirectory, Pattern.ToString());
+				return new FilePattern(BaseDirectory, pattern.ToString());
 			}
 		}
 
@@ -183,8 +181,8 @@ namespace EpicGames.Core
 		/// <returns>True if the pattern is a directory</returns>
 		public bool EndsWithDirectorySeparator()
 		{
-			string LastToken = Tokens[Tokens.Count - 1];
-			return LastToken.Length > 0 && LastToken[LastToken.Length - 1] == Path.DirectorySeparatorChar;
+			string lastToken = Tokens[^1];
+            return lastToken.Length > 0 && lastToken[^1] == Path.DirectorySeparatorChar;
 		}
 
 		/// <summary>
@@ -199,20 +197,20 @@ namespace EpicGames.Core
 		/// <summary>
 		/// Tests whether a pattern is compatible with another pattern (that is, that the number and type of wildcards match)
 		/// </summary>
-		/// <param name="Other">Pattern to compare against</param>
+		/// <param name="other">Pattern to compare against</param>
 		/// <returns>Whether the patterns are compatible.</returns>
-		public bool IsCompatibleWith(FilePattern Other)
+		public bool IsCompatibleWith(FilePattern other)
 		{
 			// Check there are the same number of tokens in each pattern
-			if(Tokens.Count != Other.Tokens.Count)
+			if(Tokens.Count != other.Tokens.Count)
 			{
 				return false;
 			}
 
 			// Check all the wildcard tokens match
-			for(int Idx = 1; Idx < Tokens.Count; Idx += 2)
+			for(int idx = 1; idx < Tokens.Count; idx += 2)
 			{
-				if(Tokens[Idx] != Other.Tokens[Idx])
+				if(Tokens[idx] != other.Tokens[idx])
 				{
 					return false;
 				}
@@ -226,29 +224,29 @@ namespace EpicGames.Core
 		/// <returns>The regex pattern</returns>
 		public string GetRegexPattern()
 		{
-			StringBuilder Pattern = new StringBuilder("^");
-			Pattern.Append(Regex.Escape(Tokens[0]));
-			for(int Idx = 1; Idx < Tokens.Count; Idx += 2)
+			StringBuilder pattern = new StringBuilder("^");
+			pattern.Append(Regex.Escape(Tokens[0]));
+			for(int idx = 1; idx < Tokens.Count; idx += 2)
 			{
 				// Append the wildcard expression
-				if(Tokens[Idx] == "?")
+				if(Tokens[idx] == "?")
 				{
-					Pattern.Append("([^\\/])");
+					pattern.Append("([^\\/])");
 				}
-				else if(Tokens[Idx] == "*")
+				else if(Tokens[idx] == "*")
 				{
-					Pattern.Append("([^\\/]*)");
+					pattern.Append("([^\\/]*)");
 				}
 				else
 				{
-					Pattern.AppendFormat("((?:.+{0})?)", Regex.Escape(Path.DirectorySeparatorChar.ToString()));
+					pattern.AppendFormat("((?:.+{0})?)", Regex.Escape(Path.DirectorySeparatorChar.ToString()));
 				}
 
 				// Append the next sequence of characters to match
-				Pattern.Append(Regex.Escape(Tokens[Idx + 1]));
+				pattern.Append(Regex.Escape(Tokens[idx + 1]));
 			}
-			Pattern.Append("$");
-			return Pattern.ToString();
+			pattern.Append("$");
+			return pattern.ToString();
 		}
 
 		/// <summary>
@@ -257,171 +255,171 @@ namespace EpicGames.Core
 		/// <returns>String representing the regex replacement pattern</returns>
 		public string GetRegexReplacementPattern()
 		{
-			StringBuilder Pattern = new StringBuilder();
-			for(int Idx = 0;;Idx += 2)
+			StringBuilder pattern = new StringBuilder();
+			for(int idx = 0;;idx += 2)
 			{
 				// Append the escaped replacement character
-				Pattern.Append(Tokens[Idx].Replace("$", "$$"));
+				pattern.Append(Tokens[idx].Replace("$", "$$"));
 
 				// Check if we've reached the end of the string
-				if(Idx == Tokens.Count - 1)
+				if(idx == Tokens.Count - 1)
 				{
 					break;
 				}
 
 				// Insert the capture
-				Pattern.AppendFormat("${0}", (Idx / 2) + 1);
+				pattern.AppendFormat("${0}", (idx / 2) + 1);
 			}
-			return Pattern.ToString();
+			return pattern.ToString();
 		}
 
 		/// <summary>
 		/// Creates a file mapping between a set of source patterns and a target pattern. All patterns should have a matching order and number of wildcards.
 		/// </summary>
-		/// <param name="Files">Files to use for the mapping</param>
+		/// <param name="files">Files to use for the mapping</param>
 		/// <param name="SourcePatterns">List of source patterns</param>
-		/// <param name="TargetPattern">Matching output pattern</param>
+		/// <param name="targetPattern">Matching output pattern</param>
 		/// <param name="Filter">Filter to apply to source files</param>
 		/// <param name="TargetFileToSourceFile">Dictionary to receive a mapping from target file to source file. An exception is thrown if multiple source files map to one target file, or a source file is also used as a target file.</param>
-		public static Dictionary<FileReference, FileReference> CreateMapping(HashSet<FileReference>? Files, ref FilePattern SourcePattern, ref FilePattern TargetPattern)
+		public static Dictionary<FileReference, FileReference> CreateMapping(HashSet<FileReference>? files, ref FilePattern sourcePattern, ref FilePattern targetPattern)
 		{
 			// If the source pattern ends in a directory separator, or a set of input files are specified and it doesn't contain wildcards, treat it as a full directory match
-			if(SourcePattern.EndsWithDirectorySeparator())
+			if(sourcePattern.EndsWithDirectorySeparator())
 			{
-				SourcePattern = new FilePattern(SourcePattern.BaseDirectory, String.Join("", SourcePattern.Tokens) + "...");
+				sourcePattern = new FilePattern(sourcePattern.BaseDirectory, String.Join("", sourcePattern.Tokens) + "...");
 			}
-			else if(Files != null)
+			else if(files != null)
 			{
-				SourcePattern = SourcePattern.AsDirectoryPattern();
+				sourcePattern = sourcePattern.AsDirectoryPattern();
 			}
 
 			// If we have multiple potential source files, but no wildcards in the output pattern, assume it's a directory and append the pattern from the source.
-			if(SourcePattern.ContainsWildcards() && !TargetPattern.ContainsWildcards())
+			if(sourcePattern.ContainsWildcards() && !targetPattern.ContainsWildcards())
 			{
-				StringBuilder NewPattern = new StringBuilder();
-				foreach(string Token in TargetPattern.Tokens)
+				StringBuilder newPattern = new StringBuilder();
+				foreach(string token in targetPattern.Tokens)
 				{
-					NewPattern.Append(Token);
+					newPattern.Append(token);
 				}
-				if(NewPattern.Length > 0 && NewPattern[NewPattern.Length - 1] != Path.DirectorySeparatorChar)
+				if(newPattern.Length > 0 && newPattern[^1] != Path.DirectorySeparatorChar)
 				{
-					NewPattern.Append(Path.DirectorySeparatorChar);
+					newPattern.Append(Path.DirectorySeparatorChar);
 				}
-				foreach(string Token in SourcePattern.Tokens)
+				foreach(string token in sourcePattern.Tokens)
 				{
-					NewPattern.Append(Token);
+					newPattern.Append(token);
 				}
-				TargetPattern = new FilePattern(TargetPattern.BaseDirectory, NewPattern.ToString());
+				targetPattern = new FilePattern(targetPattern.BaseDirectory, newPattern.ToString());
 			}
 
 			// If the target pattern ends with a directory separator, treat it as a full directory match if it has wildcards, or a copy of the source pattern if not
-			if(TargetPattern.EndsWithDirectorySeparator())
+			if(targetPattern.EndsWithDirectorySeparator())
 			{
-				TargetPattern = new FilePattern(TargetPattern.BaseDirectory, String.Join("", TargetPattern.Tokens) + "...");
+				targetPattern = new FilePattern(targetPattern.BaseDirectory, String.Join("", targetPattern.Tokens) + "...");
 			}
 
 			// Handle the case where source and target pattern are both individual files
-			Dictionary<FileReference, FileReference> TargetFileToSourceFile = new Dictionary<FileReference, FileReference>();
-			if(SourcePattern.ContainsWildcards() || TargetPattern.ContainsWildcards())
+			Dictionary<FileReference, FileReference> targetFileToSourceFile = new Dictionary<FileReference, FileReference>();
+			if(sourcePattern.ContainsWildcards() || targetPattern.ContainsWildcards())
 			{
 				// Check the two patterns are compatible
-				if(!SourcePattern.IsCompatibleWith(TargetPattern))
+				if(!sourcePattern.IsCompatibleWith(targetPattern))
 				{
-					throw new FilePatternException("File patterns '{0}' and '{1}' do not have matching wildcards", SourcePattern, TargetPattern);
+					throw new FilePatternException("File patterns '{0}' and '{1}' do not have matching wildcards", sourcePattern, targetPattern);
 				}
 
 				// Create a filter to match the source files
-				FileFilter Filter = new FileFilter(FileFilterType.Exclude);
-				Filter.Include(String.Join("", SourcePattern.Tokens));
+				FileFilter filter = new FileFilter(FileFilterType.Exclude);
+				filter.Include(String.Join("", sourcePattern.Tokens));
 
 				// Apply it to the source directory
-				List<FileReference> SourceFiles;
-				if(Files == null)
+				List<FileReference> sourceFiles;
+				if(files == null)
 				{
-					SourceFiles = Filter.ApplyToDirectory(SourcePattern.BaseDirectory, true);
+					sourceFiles = filter.ApplyToDirectory(sourcePattern.BaseDirectory, true);
 				}
 				else
 				{
-					SourceFiles = CheckInputFiles(Files, SourcePattern.BaseDirectory);
+					sourceFiles = CheckInputFiles(files, sourcePattern.BaseDirectory);
 				}
 
 				// Map them onto output files
-				FileReference[] TargetFiles = new FileReference[SourceFiles.Count];
+				FileReference[] targetFiles = new FileReference[sourceFiles.Count];
 
 				// Get the source and target regexes
-				string SourceRegex = SourcePattern.GetRegexPattern();
-				string TargetRegex = TargetPattern.GetRegexReplacementPattern();
-				for(int Idx = 0; Idx < SourceFiles.Count; Idx++)
+				string sourceRegex = sourcePattern.GetRegexPattern();
+				string targetRegex = targetPattern.GetRegexReplacementPattern();
+				for(int idx = 0; idx < sourceFiles.Count; idx++)
 				{
-					string SourceRelativePath = SourceFiles[Idx].MakeRelativeTo(SourcePattern.BaseDirectory);
-					string TargetRelativePath = Regex.Replace(SourceRelativePath, SourceRegex, TargetRegex);
-					TargetFiles[Idx] = FileReference.Combine(TargetPattern.BaseDirectory, TargetRelativePath);
+					string sourceRelativePath = sourceFiles[idx].MakeRelativeTo(sourcePattern.BaseDirectory);
+					string targetRelativePath = Regex.Replace(sourceRelativePath, sourceRegex, targetRegex);
+					targetFiles[idx] = FileReference.Combine(targetPattern.BaseDirectory, targetRelativePath);
 				}
 
 				// Add them to the output map
-				for(int Idx = 0; Idx < TargetFiles.Length; Idx++)
+				for(int idx = 0; idx < targetFiles.Length; idx++)
 				{
-					FileReference? ExistingSourceFile;
-					if(TargetFileToSourceFile.TryGetValue(TargetFiles[Idx], out ExistingSourceFile) && ExistingSourceFile != SourceFiles[Idx])
+					FileReference? existingSourceFile;
+					if(targetFileToSourceFile.TryGetValue(targetFiles[idx], out existingSourceFile) && existingSourceFile != sourceFiles[idx])
 					{
-						throw new FilePatternException("Output file '{0}' is mapped from '{1}' and '{2}'", TargetFiles[Idx], ExistingSourceFile, SourceFiles[Idx]);
+						throw new FilePatternException("Output file '{0}' is mapped from '{1}' and '{2}'", targetFiles[idx], existingSourceFile, sourceFiles[idx]);
 					}
-					TargetFileToSourceFile[TargetFiles[Idx]] = SourceFiles[Idx];
+					targetFileToSourceFile[targetFiles[idx]] = sourceFiles[idx];
 				}
 			}
 			else
 			{
 				// Just copy a single file
-				FileReference SourceFile = SourcePattern.GetSingleFile();
-				if(FileReference.Exists(SourceFile))
+				FileReference sourceFile = sourcePattern.GetSingleFile();
+				if(FileReference.Exists(sourceFile))
 				{
-					FileReference TargetFile = TargetPattern.GetSingleFile();
-					TargetFileToSourceFile[TargetFile] = SourceFile;
+					FileReference targetFile = targetPattern.GetSingleFile();
+					targetFileToSourceFile[targetFile] = sourceFile;
 				}
 				else
 				{
-					throw new FilePatternException("Source file '{0}' does not exist", SourceFile);
+					throw new FilePatternException("Source file '{0}' does not exist", sourceFile);
 				}
 			}
 
 			// Check that no source file is also destination file
-			foreach(FileReference SourceFile in TargetFileToSourceFile.Values)
+			foreach(FileReference sourceFile in targetFileToSourceFile.Values)
 			{
-				if(TargetFileToSourceFile.ContainsKey(SourceFile))
+				if(targetFileToSourceFile.ContainsKey(sourceFile))
 				{
-					throw new FilePatternException("'{0}' is listed as a source and target file", SourceFile);
+					throw new FilePatternException("'{0}' is listed as a source and target file", sourceFile);
 				}
 			}
 
 			// Return the map
-			return TargetFileToSourceFile;
+			return targetFileToSourceFile;
 		}
 
 		/// <summary>
 		/// Checks that the given input files all exist and are under the given base directory
 		/// </summary>
-		/// <param name="InputFiles">Input files to check</param>
-		/// <param name="BaseDirectory">Base directory for files</param>
+		/// <param name="inputFiles">Input files to check</param>
+		/// <param name="baseDirectory">Base directory for files</param>
 		/// <returns>List of valid files</returns>
-		public static List<FileReference> CheckInputFiles(IEnumerable<FileReference> InputFiles, DirectoryReference BaseDirectory)
+		public static List<FileReference> CheckInputFiles(IEnumerable<FileReference> inputFiles, DirectoryReference baseDirectory)
 		{
-			List<FileReference> Files = new List<FileReference>();
-			foreach(FileReference InputFile in InputFiles)
+			List<FileReference> files = new List<FileReference>();
+			foreach(FileReference inputFile in inputFiles)
 			{
-				if(!InputFile.IsUnderDirectory(BaseDirectory))
+				if(!inputFile.IsUnderDirectory(baseDirectory))
 				{
-					throw new FilePatternException("Source file '{0}' is not under '{1}'", InputFile, BaseDirectory);
+					throw new FilePatternException("Source file '{0}' is not under '{1}'", inputFile, baseDirectory);
 				}
-				else if(!FileReference.Exists(InputFile))
+				else if(!FileReference.Exists(inputFile))
 				{
-					throw new FilePatternException("Source file '{0}' does not exist", InputFile);
+					throw new FilePatternException("Source file '{0}' does not exist", inputFile);
 				}
 				else
 				{
-					Files.Add(InputFile);
+					files.Add(inputFile);
 				}
 			}
-			return Files;
+			return files;
 		}
 
 		/// <summary>
diff --git a/Engine/Source/Programs/Shared/EpicGames.Core/FileReference.cs b/Engine/Source/Programs/Shared/EpicGames.Core/FileReference.cs
--- a/Engine/Source/Programs/Shared/EpicGames.Core/FileReference.cs
+++ b/Engine/Source/Programs/Shared/EpicGames.Core/FileReference.cs
@@ -28,11 +28,11 @@ namespace EpicGames.Core
 		/// <summary>
 		/// Default constructor.
 		/// </summary>
-		/// <param name="InPath">Path to this file</param>
-		public FileReference(string InPath)
-			: base(Path.GetFullPath(InPath))
+		/// <param name="inPath">Path to this file</param>
+		public FileReference(string inPath)
+			: base(Path.GetFullPath(inPath))
 		{
-			if(FullName[FullName.Length - 1] == '\\' || FullName[FullName.Length - 1] == '/')
+			if(FullName[^1] == '\\' || FullName[^1] == '/')
 			{
 				throw new ArgumentException("File names may not be terminated by a path separator character");
 			}
@@ -41,37 +41,36 @@ namespace EpicGames.Core
 		/// <summary>
 		/// Construct a FileReference from a FileInfo object.
 		/// </summary>
-		/// <param name="InInfo">Path to this file</param>
-		public FileReference(FileInfo InInfo)
-			: base(InInfo.FullName)
+		/// <param name="inInfo">Path to this file</param>
+		public FileReference(FileInfo inInfo)
+			: base(inInfo.FullName)
 		{
 		}
 
 		/// <summary>
 		/// Default constructor.
 		/// </summary>
-		/// <param name="InFullName">The full sanitized path</param>
-		/// <param name="InSanitize">Dummary argument to use the sanitized overload</param>
-		public FileReference(string InFullName, Sanitize InSanitize)
-			: base(InFullName)
+		/// <param name="fullName">The full sanitized path</param>
+		public FileReference(string fullName, Sanitize _)
+			: base(fullName)
 		{
 		}
 
 		/// <summary>
 		/// Create a FileReference from a string. If the string is null, returns a null FileReference.
 		/// </summary>
-		/// <param name="FileName">FileName for the string</param>
+		/// <param name="fileName">FileName for the string</param>
 		/// <returns>Returns a FileReference representing the given string, or null.</returns>
-		[return: NotNullIfNotNull("FileName")]
-		public static FileReference? FromString(string? FileName)
+		[return: NotNullIfNotNull("fileName")]
+		public static FileReference? FromString(string? fileName)
 		{
-			if(String.IsNullOrEmpty(FileName))
+			if(String.IsNullOrEmpty(fileName))
 			{
 				return null;
 			}
 			else
 			{
-				return new FileReference(FileName);
+				return new FileReference(fileName);
 			}
 		}
 
@@ -99,16 +98,16 @@ namespace EpicGames.Core
 		/// <returns>A string containing the file name without an extension</returns>
 		public string GetFileNameWithoutAnyExtensions()
 		{
-			int StartIdx = FullName.LastIndexOf(Path.DirectorySeparatorChar) + 1;
+			int startIdx = FullName.LastIndexOf(Path.DirectorySeparatorChar) + 1;
 
-			int EndIdx = FullName.IndexOf('.', StartIdx);
-			if (EndIdx < StartIdx)
+			int endIdx = FullName.IndexOf('.', startIdx);
+			if (endIdx < startIdx)
 			{
-				return FullName.Substring(StartIdx);
+				return FullName.Substring(startIdx);
 			}
 			else
 			{
-				return FullName.Substring(StartIdx, EndIdx - StartIdx);
+				return FullName.Substring(startIdx, endIdx - startIdx);
 			}
 		}
 
@@ -124,12 +123,12 @@ namespace EpicGames.Core
 		/// <summary>
 		/// Change the file's extension to something else
 		/// </summary>
-		/// <param name="Extension">The new extension</param>
+		/// <param name="extension">The new extension</param>
 		/// <returns>A FileReference with the same path and name, but with the new extension</returns>
-		public FileReference ChangeExtension(string? Extension)
+		public FileReference ChangeExtension(string? extension)
 		{
-			string NewFullName = Path.ChangeExtension(FullName, Extension);
-			return new FileReference(NewFullName, Sanitize.None);
+			string newFullName = Path.ChangeExtension(FullName, extension);
+			return new FileReference(newFullName, Sanitize.None);
 		}
 
 		/// <summary>
@@ -140,94 +139,95 @@ namespace EpicGames.Core
 		{
 			get
 			{
-				int ParentLength = FullName.LastIndexOf(Path.DirectorySeparatorChar);
+				int parentLength = FullName.LastIndexOf(Path.DirectorySeparatorChar);
 
-				if (ParentLength == 2 && FullName[1] == ':')
+				if (parentLength == 2 && FullName[1] == ':')
 				{
 					// windows root detected (C:)
-					ParentLength++;
+					parentLength++;
 				}
 
-				if (ParentLength == 0 && FullName[0] == Path.DirectorySeparatorChar)
+				if (parentLength == 0 && FullName[0] == Path.DirectorySeparatorChar)
 				{
 					// nix style root (/) detected
-					ParentLength = 1;
+					parentLength = 1;
 				}
 
-				return new DirectoryReference(FullName.Substring(0, ParentLength), DirectoryReference.Sanitize.None);
+				return new DirectoryReference(FullName.Substring(0, parentLength), DirectoryReference.Sanitize.None);
 			}
 		}
 
 		/// <summary>
 		/// Combine several fragments with a base directory, to form a new filename
 		/// </summary>
-		/// <param name="BaseDirectory">The base directory</param>
-		/// <param name="Fragments">Fragments to combine with the base directory</param>
+		/// <param name="baseDirectory">The base directory</param>
+		/// <param name="fragments">Fragments to combine with the base directory</param>
 		/// <returns>The new file name</returns>
-		public static FileReference Combine(DirectoryReference BaseDirectory, params string[] Fragments)
+		public static FileReference Combine(DirectoryReference baseDirectory, params string[] fragments)
 		{
-			string FullName = FileSystemReference.CombineStrings(BaseDirectory, Fragments);
-			return new FileReference(FullName, Sanitize.None);
+			string fullName = FileSystemReference.CombineStrings(baseDirectory, fragments);
+			return new FileReference(fullName, Sanitize.None);
 		}
 
 		/// <summary>
 		/// Append a string to the end of a filename
 		/// </summary>
-		/// <param name="A">The base file reference</param>
-		/// <param name="B">Suffix to be appended</param>
+		/// <param name="a">The base file reference</param>
+		/// <param name="b">Suffix to be appended</param>
 		/// <returns>The new file reference</returns>
-		public static FileReference operator +(FileReference A, string B)
+		public static FileReference operator +(FileReference a, string b)
 		{
-			return new FileReference(A.FullName + B, Sanitize.None);
+			return new FileReference(a.FullName + b, Sanitize.None);
 		}
 
 		/// <summary>
 		/// Compares two filesystem object names for equality. Uses the canonical name representation, not the display name representation.
 		/// </summary>
-		/// <param name="A">First object to compare.</param>
-		/// <param name="B">Second object to compare.</param>
+		/// <param name="a">First object to compare.</param>
+		/// <param name="b">Second object to compare.</param>
 		/// <returns>True if the names represent the same object, false otherwise</returns>
-		public static bool operator ==(FileReference? A, FileReference? B)
+		public static bool operator ==(FileReference? a, FileReference? b)
 		{
-			if ((object?)A == null)
+			if (a is null)
 			{
-				return (object?)B == null;
+				return b is null;
+			}
+			else if (b is null)
+			{
+				return false;
 			}
 			else
 			{
-				return (object?)B != null && A.FullName.Equals(B.FullName, Comparison);
+				return a.FullName.Equals(b.FullName, Comparison);
 			}
 		}
 
 		/// <summary>
 		/// Compares two filesystem object names for inequality. Uses the canonical name representation, not the display name representation.
 		/// </summary>
-		/// <param name="A">First object to compare.</param>
-		/// <param name="B">Second object to compare.</param>
+		/// <param name="a">First object to compare.</param>
+		/// <param name="b">Second object to compare.</param>
 		/// <returns>False if the names represent the same object, true otherwise</returns>
-		public static bool operator !=(FileReference? A, FileReference? B)
+		public static bool operator !=(FileReference? a, FileReference? b)
 		{
-			return !(A == B);
+			return !(a == b);
 		}
 
 		/// <summary>
 		/// Compares against another object for equality.
 		/// </summary>
-		/// <param name="Obj">other instance to compare.</param>
+		/// <param name="obj">other instance to compare.</param>
 		/// <returns>True if the names represent the same object, false otherwise</returns>
-		public override bool Equals(object? Obj)
-		{
-			return (Obj is FileReference) && ((FileReference)Obj) == this;
-		}
+		public override bool Equals(object? obj) => obj is FileReference file && file == this;
 
 		/// <summary>
 		/// Compares against another object for equality.
 		/// </summary>
-		/// <param name="Obj">other instance to compare.</param>
+		/// <param name="obj">other instance to compare.</param>
 		/// <returns>True if the names represent the same object, false otherwise</returns>
-		public bool Equals(FileReference? Obj)
+		public bool Equals(FileReference? obj)
 		{
-			return Obj == this;
+			return obj == this;
 		}
 
 		/// <summary>
@@ -240,31 +240,31 @@ namespace EpicGames.Core
 		}
 
 		/// <inheritdoc/>
-		public int CompareTo(FileReference? Other) => Comparer.Compare(FullName, Other?.FullName);
+		public int CompareTo(FileReference? other) => Comparer.Compare(FullName, other?.FullName);
 
 		/// <summary>
 		/// Helper function to create a remote file reference. Unlike normal FileReference objects, these aren't converted to a full path in the local filesystem, but are
 		/// left as they are passed in.
 		/// </summary>
-		/// <param name="AbsolutePath">The absolute path in the remote file system</param>
+		/// <param name="absolutePath">The absolute path in the remote file system</param>
 		/// <returns>New file reference</returns>
-		public static FileReference MakeRemote(string AbsolutePath)
+		public static FileReference MakeRemote(string absolutePath)
 		{
-			return new FileReference(AbsolutePath, Sanitize.None);
+			return new FileReference(absolutePath, Sanitize.None);
 		}
 
 		/// <summary>
 		/// Makes a file location writeable; 
 		/// </summary>
-		/// <param name="Location">Location of the file</param>
-		public static void MakeWriteable(FileReference Location)
+		/// <param name="location">Location of the file</param>
+		public static void MakeWriteable(FileReference location)
 		{
-			if(Exists(Location))
+			if(Exists(location))
 			{
-				FileAttributes Attributes = GetAttributes(Location);
-				if((Attributes & FileAttributes.ReadOnly) != 0)
+				FileAttributes attributes = GetAttributes(location);
+				if((attributes & FileAttributes.ReadOnly) != 0)
 				{
-					SetAttributes(Location, Attributes & ~FileAttributes.ReadOnly);
+					SetAttributes(location, attributes & ~FileAttributes.ReadOnly);
 				}
 			}
 		}
@@ -272,11 +272,11 @@ namespace EpicGames.Core
 		/// <summary>
 		/// Finds the correct case to match the location of this file on disk. Uses the given case for parts of the path that do not exist.
 		/// </summary>
-		/// <param name="Location">The path to find the correct case for</param>
+		/// <param name="location">The path to find the correct case for</param>
 		/// <returns>Location of the file with the correct case</returns>
-		public static FileReference FindCorrectCase(FileReference Location)
+		public static FileReference FindCorrectCase(FileReference location)
 		{
-			return new FileReference(FileUtils.FindCorrectCase(Location.ToFileInfo()));
+			return new FileReference(FileUtils.FindCorrectCase(location.ToFileInfo()));
 		}
 
 		/// <summary>
@@ -293,319 +293,319 @@ namespace EpicGames.Core
 		/// <summary>
 		/// Copies a file from one location to another
 		/// </summary>
-		/// <param name="SourceLocation">Location of the source file</param>
-		/// <param name="TargetLocation">Location of the target file</param>
-		public static void Copy(FileReference SourceLocation, FileReference TargetLocation)
+		/// <param name="sourceLocation">Location of the source file</param>
+		/// <param name="targetLocation">Location of the target file</param>
+		public static void Copy(FileReference sourceLocation, FileReference targetLocation)
 		{
-			File.Copy(SourceLocation.FullName, TargetLocation.FullName);
+			File.Copy(sourceLocation.FullName, targetLocation.FullName);
 		}
 
 		/// <summary>
 		/// Copies a file from one location to another
 		/// </summary>
-		/// <param name="SourceLocation">Location of the source file</param>
-		/// <param name="TargetLocation">Location of the target file</param>
+		/// <param name="sourceLocation">Location of the source file</param>
+		/// <param name="targetLocation">Location of the target file</param>
 		/// <param name="bOverwrite">Whether to overwrite the file in the target location</param>
-		public static void Copy(FileReference SourceLocation, FileReference TargetLocation, bool bOverwrite)
+		public static void Copy(FileReference sourceLocation, FileReference targetLocation, bool bOverwrite)
 		{
-			File.Copy(SourceLocation.FullName, TargetLocation.FullName, bOverwrite);
+			File.Copy(sourceLocation.FullName, targetLocation.FullName, bOverwrite);
 		}
 
 		/// <summary>
 		/// Deletes this file
 		/// </summary>
-		public static void Delete(FileReference Location)
+		public static void Delete(FileReference location)
 		{
-			File.Delete(Location.FullName);
+			File.Delete(location.FullName);
 		}
 
 		/// <summary>
 		/// Determines whether the given filename exists
 		/// </summary>
 		/// <returns>True if it exists, false otherwise</returns>
-		public static bool Exists(FileReference Location)
+		public static bool Exists(FileReference location)
 		{
-			return File.Exists(Location.FullName);
+			return File.Exists(location.FullName);
 		}
 
 		/// <summary>
 		/// Gets the attributes for a file
 		/// </summary>
-		/// <param name="Location">Location of the file</param>
+		/// <param name="location">Location of the file</param>
 		/// <returns>Attributes for the file</returns>
-		public static FileAttributes GetAttributes(FileReference Location)
+		public static FileAttributes GetAttributes(FileReference location)
 		{
-			return File.GetAttributes(Location.FullName);
+			return File.GetAttributes(location.FullName);
 		}
 
 		/// <summary>
 		/// Gets the time that the file was last written to
 		/// </summary>
-		/// <param name="Location">Location of the file</param>
+		/// <param name="location">Location of the file</param>
 		/// <returns>Last write time, in local time</returns>
-		public static DateTime GetLastWriteTime(FileReference Location)
+		public static DateTime GetLastWriteTime(FileReference location)
 		{
-			return File.GetLastWriteTime(Location.FullName);
+			return File.GetLastWriteTime(location.FullName);
 		}
 
 		/// <summary>
 		/// Gets the time that the file was last written to
 		/// </summary>
-		/// <param name="Location">Location of the file</param>
+		/// <param name="location">Location of the file</param>
 		/// <returns>Last write time, in UTC time</returns>
-		public static DateTime GetLastWriteTimeUtc(FileReference Location)
+		public static DateTime GetLastWriteTimeUtc(FileReference location)
 		{
-			return File.GetLastWriteTimeUtc(Location.FullName);
+			return File.GetLastWriteTimeUtc(location.FullName);
 		}
 
 		/// <summary>
 		/// Moves a file from one location to another
 		/// </summary>
-		/// <param name="SourceLocation">Location of the source file</param>
-		/// <param name="TargetLocation">Location of the target file</param>
-		public static void Move(FileReference SourceLocation, FileReference TargetLocation)
+		/// <param name="sourceLocation">Location of the source file</param>
+		/// <param name="targetLocation">Location of the target file</param>
+		public static void Move(FileReference sourceLocation, FileReference targetLocation)
 		{
-			File.Move(SourceLocation.FullName, TargetLocation.FullName);
+			File.Move(sourceLocation.FullName, targetLocation.FullName);
 		}
 
 		/// <summary>
 		/// Moves a file from one location to another
 		/// </summary>
-		/// <param name="SourceLocation">Location of the source file</param>
-		/// <param name="TargetLocation">Location of the target file</param>
-		/// <param name="Overwrite">Whether to overwrite the file in the target location</param>
-		public static void Move(FileReference SourceLocation, FileReference TargetLocation, bool Overwrite)
+		/// <param name="sourceLocation">Location of the source file</param>
+		/// <param name="targetLocation">Location of the target file</param>
+		/// <param name="overwrite">Whether to overwrite the file in the target location</param>
+		public static void Move(FileReference sourceLocation, FileReference targetLocation, bool overwrite)
 		{
-			File.Move(SourceLocation.FullName, TargetLocation.FullName, Overwrite);
+			File.Move(sourceLocation.FullName, targetLocation.FullName, overwrite);
 		}
 
 		/// <summary>
 		/// Opens a FileStream on the specified path with read/write access
 		/// </summary>
-		/// <param name="Location">Location of the file</param>
-		/// <param name="Mode">Mode to use when opening the file</param>
+		/// <param name="location">Location of the file</param>
+		/// <param name="mode">Mode to use when opening the file</param>
 		/// <returns>New filestream for the given file</returns>
-		public static FileStream Open(FileReference Location, FileMode Mode)
+		public static FileStream Open(FileReference location, FileMode mode)
 		{
-			return File.Open(Location.FullName, Mode);
+			return File.Open(location.FullName, mode);
 		}
 
 		/// <summary>
 		/// Opens a FileStream on the specified path
 		/// </summary>
-		/// <param name="Location">Location of the file</param>
-		/// <param name="Mode">Mode to use when opening the file</param>
-		/// <param name="Access">Sharing mode for the new file</param>
+		/// <param name="location">Location of the file</param>
+		/// <param name="mode">Mode to use when opening the file</param>
+		/// <param name="access">Sharing mode for the new file</param>
 		/// <returns>New filestream for the given file</returns>
-		public static FileStream Open(FileReference Location, FileMode Mode, FileAccess Access)
+		public static FileStream Open(FileReference location, FileMode mode, FileAccess access)
 		{
-			return File.Open(Location.FullName, Mode, Access);
+			return File.Open(location.FullName, mode, access);
 		}
 
 		/// <summary>
 		/// Opens a FileStream on the specified path
 		/// </summary>
-		/// <param name="Location">Location of the file</param>
-		/// <param name="Mode">Mode to use when opening the file</param>
-		/// <param name="Access">Access mode for the new file</param>
-		/// <param name="Share">Sharing mode for the open file</param>
+		/// <param name="location">Location of the file</param>
+		/// <param name="mode">Mode to use when opening the file</param>
+		/// <param name="access">Access mode for the new file</param>
+		/// <param name="share">Sharing mode for the open file</param>
 		/// <returns>New filestream for the given file</returns>
-		public static FileStream Open(FileReference Location, FileMode Mode, FileAccess Access, FileShare Share)
+		public static FileStream Open(FileReference location, FileMode mode, FileAccess access, FileShare share)
 		{
-			return File.Open(Location.FullName, Mode, Access, Share);
+			return File.Open(location.FullName, mode, access, share);
 		}
 
 		/// <summary>
 		/// Reads the contents of a file
 		/// </summary>
-		/// <param name="Location">Location of the file</param>
+		/// <param name="location">Location of the file</param>
 		/// <returns>Byte array containing the contents of the file</returns>
-		public static byte[] ReadAllBytes(FileReference Location)
+		public static byte[] ReadAllBytes(FileReference location)
 		{
-			return File.ReadAllBytes(Location.FullName);
+			return File.ReadAllBytes(location.FullName);
 		}
 
 		/// <summary>
 		/// Reads the contents of a file
 		/// </summary>
-		/// <param name="Location">Location of the file</param>
+		/// <param name="location">Location of the file</param>
 		/// <returns>Byte array containing the contents of the file</returns>
-		public static Task<byte[]> ReadAllBytesAsync(FileReference Location, CancellationToken CancellationToken = default(CancellationToken))
+		public static Task<byte[]> ReadAllBytesAsync(FileReference location, CancellationToken cancellationToken = default)
 		{
-			return File.ReadAllBytesAsync(Location.FullName, CancellationToken);
+			return File.ReadAllBytesAsync(location.FullName, cancellationToken);
 		}
 
 		/// <summary>
 		/// Reads the contents of a file
 		/// </summary>
-		/// <param name="Location">Location of the file</param>
+		/// <param name="location">Location of the file</param>
 		/// <returns>Contents of the file as a single string</returns>
-		public static string ReadAllText(FileReference Location)
+		public static string ReadAllText(FileReference location)
 		{
-			return File.ReadAllText(Location.FullName);
+			return File.ReadAllText(location.FullName);
 		}
 
 		/// <summary>
 		/// Reads the contents of a file
 		/// </summary>
-		/// <param name="Location">Location of the file</param>
-		/// <param name="Encoding">Encoding of the file</param>
+		/// <param name="location">Location of the file</param>
+		/// <param name="encoding">Encoding of the file</param>
 		/// <returns>Contents of the file as a single string</returns>
-		public static string ReadAllText(FileReference Location, Encoding Encoding)
+		public static string ReadAllText(FileReference location, Encoding encoding)
 		{
-			return File.ReadAllText(Location.FullName, Encoding);
+			return File.ReadAllText(location.FullName, encoding);
 		}
 
 		/// <summary>
 		/// Reads the contents of a file
 		/// </summary>
-		/// <param name="Location">Location of the file</param>
+		/// <param name="location">Location of the file</param>
 		/// <returns>Contents of the file as a single string</returns>
-		public static Task<string> ReadAllTextAsync(FileReference Location, CancellationToken CancellationToken = default(CancellationToken))
+		public static Task<string> ReadAllTextAsync(FileReference location, CancellationToken cancellationToken = default)
 		{
-			return File.ReadAllTextAsync(Location.FullName, CancellationToken);
+			return File.ReadAllTextAsync(location.FullName, cancellationToken);
 		}
 
 		/// <summary>
 		/// Reads the contents of a file
 		/// </summary>
-		/// <param name="Location">Location of the file</param>
-		/// <param name="Encoding">Encoding of the file</param>
+		/// <param name="location">Location of the file</param>
+		/// <param name="encoding">Encoding of the file</param>
 		/// <returns>Contents of the file as a single string</returns>
-		public static Task<string> ReadAllTextAsync(FileReference Location, Encoding Encoding, CancellationToken CancellationToken = default(CancellationToken))
+		public static Task<string> ReadAllTextAsync(FileReference location, Encoding encoding, CancellationToken cancellationToken = default)
 		{
-			return File.ReadAllTextAsync(Location.FullName, Encoding, CancellationToken);
+			return File.ReadAllTextAsync(location.FullName, encoding, cancellationToken);
 		}
 
 		/// <summary>
 		/// Reads the contents of a file
 		/// </summary>
-		/// <param name="Location">Location of the file</param>
+		/// <param name="location">Location of the file</param>
 		/// <returns>String array containing the contents of the file</returns>
-		public static string[] ReadAllLines(FileReference Location)
+		public static string[] ReadAllLines(FileReference location)
 		{
-			return File.ReadAllLines(Location.FullName);
+			return File.ReadAllLines(location.FullName);
 		}
 
 		/// <summary>
 		/// Reads the contents of a file
 		/// </summary>
-		/// <param name="Location">Location of the file</param>
-		/// <param name="Encoding">The encoding to use when parsing the file</param>
+		/// <param name="location">Location of the file</param>
+		/// <param name="encoding">The encoding to use when parsing the file</param>
 		/// <returns>String array containing the contents of the file</returns>
-		public static string[] ReadAllLines(FileReference Location, Encoding Encoding)
+		public static string[] ReadAllLines(FileReference location, Encoding encoding)
 		{
-			return File.ReadAllLines(Location.FullName, Encoding);
+			return File.ReadAllLines(location.FullName, encoding);
 		}
 
 		/// <summary>
 		/// Reads the contents of a file
 		/// </summary>
-		/// <param name="Location">Location of the file</param>
+		/// <param name="location">Location of the file</param>
 		/// <returns>String array containing the contents of the file</returns>
-		public static Task<string[]> ReadAllLinesAsync(FileReference Location, CancellationToken CancellationToken = default(CancellationToken))
+		public static Task<string[]> ReadAllLinesAsync(FileReference location, CancellationToken cancellationToken = default)
 		{
-			return File.ReadAllLinesAsync(Location.FullName, CancellationToken);
+			return File.ReadAllLinesAsync(location.FullName, cancellationToken);
 		}
 
 		/// <summary>
 		/// Reads the contents of a file
 		/// </summary>
-		/// <param name="Location">Location of the file</param>
-		/// <param name="Encoding">The encoding to use when parsing the file</param>
+		/// <param name="location">Location of the file</param>
+		/// <param name="encoding">The encoding to use when parsing the file</param>
 		/// <returns>String array containing the contents of the file</returns>
-		public static Task<string[]> ReadAllLinesAsync(FileReference Location, Encoding Encoding, CancellationToken CancellationToken = default(CancellationToken))
+		public static Task<string[]> ReadAllLinesAsync(FileReference location, Encoding encoding, CancellationToken cancellationToken = default)
 		{
-			return File.ReadAllLinesAsync(Location.FullName, Encoding, CancellationToken);
+			return File.ReadAllLinesAsync(location.FullName, encoding, cancellationToken);
 		}
 
 		/// <summary>
 		/// Sets the attributes for a file
 		/// </summary>
-		/// <param name="Location">Location of the file</param>
-		/// <param name="Attributes">New attributes for the file</param>
-		public static void SetAttributes(FileReference Location, FileAttributes Attributes)
+		/// <param name="location">Location of the file</param>
+		/// <param name="attributes">New attributes for the file</param>
+		public static void SetAttributes(FileReference location, FileAttributes attributes)
 		{
-			File.SetAttributes(Location.FullName, Attributes);
+			File.SetAttributes(location.FullName, attributes);
 		}
 
 		/// <summary>
 		/// Sets the time that the file was last written to
 		/// </summary>
-		/// <param name="Location">Location of the file</param>
-		/// <param name="LastWriteTime">Last write time, in local time</param>
-		public static void SetLastWriteTime(FileReference Location, DateTime LastWriteTime)
+		/// <param name="location">Location of the file</param>
+		/// <param name="lastWriteTime">Last write time, in local time</param>
+		public static void SetLastWriteTime(FileReference location, DateTime lastWriteTime)
 		{
-			File.SetLastWriteTime(Location.FullName, LastWriteTime);
+			File.SetLastWriteTime(location.FullName, lastWriteTime);
 		}
 
 		/// <summary>
 		/// Sets the time that the file was last written to
 		/// </summary>
-		/// <param name="Location">Location of the file</param>
-		/// <param name="LastWriteTimeUtc">Last write time, in UTC time</param>
-		public static void SetLastWriteTimeUtc(FileReference Location, DateTime LastWriteTimeUtc)
+		/// <param name="location">Location of the file</param>
+		/// <param name="lastWriteTimeUtc">Last write time, in UTC time</param>
+		public static void SetLastWriteTimeUtc(FileReference location, DateTime lastWriteTimeUtc)
 		{
-			File.SetLastWriteTimeUtc(Location.FullName, LastWriteTimeUtc);
+			File.SetLastWriteTimeUtc(location.FullName, lastWriteTimeUtc);
 		}
 
 		/// <summary>
 		/// Sets the time that the file was last accessed.
 		/// </summary>
-		/// <param name="Location">Location of the file.</param>
-		/// <param name="LastWriteTime">Last access time, in local time.</param>
-		public static void SetLastAccessTime(FileReference Location, DateTime LastWriteTime)
+		/// <param name="location">Location of the file.</param>
+		/// <param name="lastWriteTime">Last access time, in local time.</param>
+		public static void SetLastAccessTime(FileReference location, DateTime lastWriteTime)
 		{
-			File.SetLastWriteTime(Location.FullName, LastWriteTime);
+			File.SetLastWriteTime(location.FullName, lastWriteTime);
 		}
 
 		/// <summary>
 		/// Sets the time that the file was last accessed.
 		/// </summary>
-		/// <param name="Location">Location of the file.</param>
+		/// <param name="location">Location of the file.</param>
 		/// <param name="LastWriteTime">Last access time, in UTC time.</param>
-		public static void SetLastAccessTimeUtc(FileReference Location, DateTime LastWriteTimeUtc)
+		public static void SetLastAccessTimeUtc(FileReference location, DateTime lastWriteTimeUtc)
 		{
-			File.SetLastWriteTimeUtc(Location.FullName, LastWriteTimeUtc);
+			File.SetLastWriteTimeUtc(location.FullName, lastWriteTimeUtc);
 		}
 
 		/// <summary>
 		/// Writes the contents of a file
 		/// </summary>
-		/// <param name="Location">Location of the file</param>
-		/// <param name="Contents">Contents of the file</param>
-		public static void WriteAllBytes(FileReference Location, byte[] Contents)
+		/// <param name="location">Location of the file</param>
+		/// <param name="contents">Contents of the file</param>
+		public static void WriteAllBytes(FileReference location, byte[] contents)
 		{
-			File.WriteAllBytes(Location.FullName, Contents);
+			File.WriteAllBytes(location.FullName, contents);
 		}
 
 		/// <summary>
 		/// Writes the contents of a file
 		/// </summary>
-		/// <param name="Location">Location of the file</param>
-		/// <param name="Contents">Contents of the file</param>
-		public static Task WriteAllBytesAsync(FileReference Location, byte[] Contents, CancellationToken CancellationToken = default(CancellationToken))
+		/// <param name="location">Location of the file</param>
+		/// <param name="contents">Contents of the file</param>
+		public static Task WriteAllBytesAsync(FileReference location, byte[] contents, CancellationToken cancellationToken = default)
 		{
-			return File.WriteAllBytesAsync(Location.FullName, Contents, CancellationToken);
+			return File.WriteAllBytesAsync(location.FullName, contents, cancellationToken);
 		}
 
 		/// <summary>
 		/// Writes the data to the given file, if it's different from what's there already.
 		/// Returns true if contents were written.
 		/// </summary>
-		/// <param name="Location">Location of the file</param>
-		/// <param name="Contents">Contents of the file</param>
-		public static bool WriteAllBytesIfDifferent(FileReference Location, byte[] Contents)
+		/// <param name="location">Location of the file</param>
+		/// <param name="contents">Contents of the file</param>
+		public static bool WriteAllBytesIfDifferent(FileReference location, byte[] contents)
 		{
-			if(FileReference.Exists(Location))
+			if(FileReference.Exists(location))
 			{
-				byte[] CurrentContents = FileReference.ReadAllBytes(Location);
-				if(Contents.AsSpan().SequenceEqual(CurrentContents))
+				byte[] currentContents = FileReference.ReadAllBytes(location);
+				if(contents.AsSpan().SequenceEqual(currentContents))
 				{
 					return false;
 				}
 			}
-			WriteAllBytes(Location, Contents);
+			WriteAllBytes(location, contents);
 			return true;
 		}
 
@@ -613,273 +613,272 @@ namespace EpicGames.Core
 		/// Writes the string to the given file, if it's different from what's there already.
 		/// Returns true if contents were written.
 		/// </summary>
-		/// <param name="Location">Location of the file</param>
-		/// <param name="Contents">Contents of the file</param>
-		public static bool WriteAllTextIfDifferent(FileReference Location, string Contents)
+		/// <param name="location">Location of the file</param>
+		/// <param name="contents">Contents of the file</param>
+		public static bool WriteAllTextIfDifferent(FileReference location, string contents)
 		{
-			if(FileReference.Exists(Location))
+			if(FileReference.Exists(location))
 			{
-				string CurrentContents = FileReference.ReadAllText(Location);
-				if (String.Equals(Contents, CurrentContents))
+				string currentContents = FileReference.ReadAllText(location);
+				if (String.Equals(contents, currentContents))
 				{
 					return false;
 				}
 			}
-			WriteAllText(Location, Contents);
+			WriteAllText(location, contents);
 			return true;
 		}
 
 		/// <summary>
 		/// Writes the contents of a file
 		/// </summary>
-		/// <param name="Location">Location of the file</param>
-		/// <param name="Contents">Contents of the file</param>
-		public static void WriteAllLines(FileReference Location, IEnumerable<string> Contents)
+		/// <param name="location">Location of the file</param>
+		/// <param name="contents">Contents of the file</param>
+		public static void WriteAllLines(FileReference location, IEnumerable<string> contents)
 		{
-			File.WriteAllLines(Location.FullName, Contents);
+			File.WriteAllLines(location.FullName, contents);
 		}
 
 		/// <summary>
 		/// Writes the contents of a file
 		/// </summary>
-		/// <param name="Location">Location of the file</param>
-		/// <param name="Contents">Contents of the file</param>
-		/// <param name="Encoding">The encoding to use when parsing the file</param>
-		public static void WriteAllLines(FileReference Location, IEnumerable<string> Contents, Encoding Encoding)
+		/// <param name="location">Location of the file</param>
+		/// <param name="contents">Contents of the file</param>
+		/// <param name="encoding">The encoding to use when parsing the file</param>
+		public static void WriteAllLines(FileReference location, IEnumerable<string> contents, Encoding encoding)
 		{
-			File.WriteAllLines(Location.FullName, Contents, Encoding);
+			File.WriteAllLines(location.FullName, contents, encoding);
 		}
 
 		/// <summary>
 		/// Writes the contents of a file
 		/// </summary>
-		/// <param name="Location">Location of the file</param>
-		/// <param name="Contents">Contents of the file</param>
-		public static void WriteAllLines(FileReference Location, string[] Contents)
+		/// <param name="location">Location of the file</param>
+		/// <param name="contents">Contents of the file</param>
+		public static void WriteAllLines(FileReference location, string[] contents)
 		{
-			File.WriteAllLines(Location.FullName, Contents);
+			File.WriteAllLines(location.FullName, contents);
 		}
 
 		/// <summary>
 		/// Writes the contents of a file
 		/// </summary>
-		/// <param name="Location">Location of the file</param>
-		/// <param name="Contents">Contents of the file</param>
-		/// <param name="Encoding">The encoding to use when parsing the file</param>
-		public static void WriteAllLines(FileReference Location, string[] Contents, Encoding Encoding)
+		/// <param name="location">Location of the file</param>
+		/// <param name="contents">Contents of the file</param>
+		/// <param name="encoding">The encoding to use when parsing the file</param>
+		public static void WriteAllLines(FileReference location, string[] contents, Encoding encoding)
 		{
-			File.WriteAllLines(Location.FullName, Contents, Encoding);
+			File.WriteAllLines(location.FullName, contents, encoding);
 		}
 
 		/// <summary>
 		/// Writes the contents of a file
 		/// </summary>
-		/// <param name="Location">Location of the file</param>
-		/// <param name="Contents">Contents of the file</param>
-		public static Task WriteAllLinesAsync(FileReference Location, IEnumerable<string> Contents, CancellationToken CancellationToken = default(CancellationToken))
+		/// <param name="location">Location of the file</param>
+		/// <param name="contents">Contents of the file</param>
+		public static Task WriteAllLinesAsync(FileReference location, IEnumerable<string> contents, CancellationToken cancellationToken = default)
 		{
-			return File.WriteAllLinesAsync(Location.FullName, Contents, CancellationToken);
+			return File.WriteAllLinesAsync(location.FullName, contents, cancellationToken);
 		}
 
 		/// <summary>
 		/// Writes the contents of a file
 		/// </summary>
-		/// <param name="Location">Location of the file</param>
-		/// <param name="Contents">Contents of the file</param>
-		/// <param name="Encoding">The encoding to use when parsing the file</param>
-		public static Task WriteAllLinesAsync(FileReference Location, IEnumerable<string> Contents, Encoding Encoding, CancellationToken CancellationToken = default(CancellationToken))
+		/// <param name="location">Location of the file</param>
+		/// <param name="contents">Contents of the file</param>
+		/// <param name="encoding">The encoding to use when parsing the file</param>
+		public static Task WriteAllLinesAsync(FileReference location, IEnumerable<string> contents, Encoding encoding, CancellationToken cancellationToken = default)
 		{
-			return File.WriteAllLinesAsync(Location.FullName, Contents, Encoding, CancellationToken);
+			return File.WriteAllLinesAsync(location.FullName, contents, encoding, cancellationToken);
 		}
 
 		/// <summary>
 		/// Writes the contents of a file
 		/// </summary>
-		/// <param name="Location">Location of the file</param>
-		/// <param name="Contents">Contents of the file</param>
-		public static Task WriteAllLinesAsync(FileReference Location, string[] Contents, CancellationToken CancellationToken = default(CancellationToken))
+		/// <param name="location">Location of the file</param>
+		/// <param name="contents">Contents of the file</param>
+		public static Task WriteAllLinesAsync(FileReference location, string[] contents, CancellationToken cancellationToken = default)
 		{
-			return File.WriteAllLinesAsync(Location.FullName, Contents, CancellationToken);
+			return File.WriteAllLinesAsync(location.FullName, contents, cancellationToken);
 		}
 
 		/// <summary>
 		/// Writes the contents of a file
 		/// </summary>
-		/// <param name="Location">Location of the file</param>
-		/// <param name="Contents">Contents of the file</param>
-		/// <param name="Encoding">The encoding to use when parsing the file</param>
-		public static Task WriteAllLinesAsync(FileReference Location, string[] Contents, Encoding Encoding, CancellationToken CancellationToken = default(CancellationToken))
+		/// <param name="location">Location of the file</param>
+		/// <param name="contents">Contents of the file</param>
+		/// <param name="encoding">The encoding to use when parsing the file</param>
+		public static Task WriteAllLinesAsync(FileReference location, string[] contents, Encoding encoding, CancellationToken cancellationToken = default)
 		{
-			return File.WriteAllLinesAsync(Location.FullName, Contents, Encoding, CancellationToken);
+			return File.WriteAllLinesAsync(location.FullName, contents, encoding, cancellationToken);
 		}
 
 		/// <summary>
 		/// Writes the contents of a file
 		/// </summary>
-		/// <param name="Location">Location of the file</param>
-		/// <param name="Contents">Contents of the file</param>
-		public static void WriteAllText(FileReference Location, string Contents)
+		/// <param name="location">Location of the file</param>
+		/// <param name="contents">Contents of the file</param>
+		public static void WriteAllText(FileReference location, string contents)
 		{
-			File.WriteAllText(Location.FullName, Contents);
+			File.WriteAllText(location.FullName, contents);
 		}
 
 		/// <summary>
 		/// Writes the contents of a file
 		/// </summary>
-		/// <param name="Location">Location of the file</param>
-		/// <param name="Contents">Contents of the file</param>
-		/// <param name="Encoding">The encoding to use when parsing the file</param>
-		public static void WriteAllText(FileReference Location, string Contents, Encoding Encoding)
+		/// <param name="location">Location of the file</param>
+		/// <param name="contents">Contents of the file</param>
+		/// <param name="encoding">The encoding to use when parsing the file</param>
+		public static void WriteAllText(FileReference location, string contents, Encoding encoding)
 		{
-			File.WriteAllText(Location.FullName, Contents, Encoding);
+			File.WriteAllText(location.FullName, contents, encoding);
 		}
 
-
 		/// <summary>
 		/// Writes the contents of a file
 		/// </summary>
-		/// <param name="Location">Location of the file</param>
-		/// <param name="Contents">Contents of the file</param>
-		public static Task WriteAllTextAsync(FileReference Location, string Contents)
+		/// <param name="location">Location of the file</param>
+		/// <param name="contents">Contents of the file</param>
+		public static Task WriteAllTextAsync(FileReference location, string contents)
 		{
-			return File.WriteAllTextAsync(Location.FullName, Contents);
+			return File.WriteAllTextAsync(location.FullName, contents);
 		}
 
 		/// <summary>
 		/// Writes the contents of a file
 		/// </summary>
-		/// <param name="Location">Location of the file</param>
-		/// <param name="Contents">Contents of the file</param>
-		/// <param name="Encoding">The encoding to use when parsing the file</param>
-		public static Task WriteAllTextAsync(FileReference Location, string Contents, Encoding Encoding)
+		/// <param name="location">Location of the file</param>
+		/// <param name="contents">Contents of the file</param>
+		/// <param name="encoding">The encoding to use when parsing the file</param>
+		public static Task WriteAllTextAsync(FileReference location, string contents, Encoding encoding)
 		{
-			return File.WriteAllTextAsync(Location.FullName, Contents, Encoding);
+			return File.WriteAllTextAsync(location.FullName, contents, encoding);
 		}
 
 		/// <summary>
 		/// Appends the contents to a file
 		/// </summary>
-		/// <param name="Location">Location of the file</param>
-		/// <param name="Contents">Contents to append to the file</param>
-		public static void AppendAllLines(FileReference Location, IEnumerable<string> Contents)
+		/// <param name="location">Location of the file</param>
+		/// <param name="contents">Contents to append to the file</param>
+		public static void AppendAllLines(FileReference location, IEnumerable<string> contents)
 		{
-			File.AppendAllLines(Location.FullName, Contents);
+			File.AppendAllLines(location.FullName, contents);
 		}
 
 		/// <summary>
 		/// Appends the contents to a file
 		/// </summary>
-		/// <param name="Location">Location of the file</param>
-		/// <param name="Contents">Contents to append to the file</param>
-		public static Task AppendAllLinesAsync(FileReference Location, IEnumerable<string> Contents)
+		/// <param name="location">Location of the file</param>
+		/// <param name="contents">Contents to append to the file</param>
+		public static Task AppendAllLinesAsync(FileReference location, IEnumerable<string> contents)
 		{
-			return File.AppendAllLinesAsync(Location.FullName, Contents);
+			return File.AppendAllLinesAsync(location.FullName, contents);
 		}
 
 		/// <summary>
 		/// Appends the contents to a file
 		/// </summary>
-		/// <param name="Location">Location of the file</param>
-		/// <param name="Contents">Contents to append to the file</param>
-		/// <param name="Encoding">The encoding to use when parsing the file</param>
-		public static void AppendAllLines(FileReference Location, IEnumerable<string> Contents, Encoding Encoding)
+		/// <param name="location">Location of the file</param>
+		/// <param name="contents">Contents to append to the file</param>
+		/// <param name="encoding">The encoding to use when parsing the file</param>
+		public static void AppendAllLines(FileReference location, IEnumerable<string> contents, Encoding encoding)
 		{
-			File.AppendAllLines(Location.FullName, Contents, Encoding);
+			File.AppendAllLines(location.FullName, contents, encoding);
 		}
 
 		/// <summary>
 		/// Appends the contents to a file
 		/// </summary>
-		/// <param name="Location">Location of the file</param>
-		/// <param name="Contents">Contents to append to the file</param>
-		/// <param name="Encoding">The encoding to use when parsing the file</param>
-		public static Task AppendAllLinesAsync(FileReference Location, IEnumerable<string> Contents, Encoding Encoding)
+		/// <param name="location">Location of the file</param>
+		/// <param name="contents">Contents to append to the file</param>
+		/// <param name="encoding">The encoding to use when parsing the file</param>
+		public static Task AppendAllLinesAsync(FileReference location, IEnumerable<string> contents, Encoding encoding)
 		{
-			return File.AppendAllLinesAsync(Location.FullName, Contents, Encoding);
+			return File.AppendAllLinesAsync(location.FullName, contents, encoding);
 		}
 
 		/// <summary>
 		/// Appends the contents to a file
 		/// </summary>
-		/// <param name="Location">Location of the file</param>
-		/// <param name="Contents">Contents to append to the file</param>
-		public static void AppendAllLines(FileReference Location, string[] Contents)
+		/// <param name="location">Location of the file</param>
+		/// <param name="contents">Contents to append to the file</param>
+		public static void AppendAllLines(FileReference location, string[] contents)
 		{
-			File.AppendAllLines(Location.FullName, Contents);
+			File.AppendAllLines(location.FullName, contents);
 		}
 
 		/// <summary>
 		/// Appends the contents to a file
 		/// </summary>
-		/// <param name="Location">Location of the file</param>
-		/// <param name="Contents">Contents to append to the file</param>
-		public static Task AppendAllLinesAsync(FileReference Location, string[] Contents)
+		/// <param name="location">Location of the file</param>
+		/// <param name="contents">Contents to append to the file</param>
+		public static Task AppendAllLinesAsync(FileReference location, string[] contents)
 		{
-			return File.AppendAllLinesAsync(Location.FullName, Contents);
+			return File.AppendAllLinesAsync(location.FullName, contents);
 		}
 
 		/// <summary>
 		/// Appends the contents to a file
 		/// </summary>
-		/// <param name="Location">Location of the file</param>
-		/// <param name="Contents">Contents to append to the file</param>
-		/// <param name="Encoding">The encoding to use when parsing the file</param>
-		public static void AppendAllLines(FileReference Location, string[] Contents, Encoding Encoding)
+		/// <param name="location">Location of the file</param>
+		/// <param name="contents">Contents to append to the file</param>
+		/// <param name="encoding">The encoding to use when parsing the file</param>
+		public static void AppendAllLines(FileReference location, string[] contents, Encoding encoding)
 		{
-			File.AppendAllLines(Location.FullName, Contents, Encoding);
+			File.AppendAllLines(location.FullName, contents, encoding);
 		}
 
 		/// <summary>
 		/// Appends the contents to a file
 		/// </summary>
-		/// <param name="Location">Location of the file</param>
-		/// <param name="Contents">Contents to append to the file</param>
-		/// <param name="Encoding">The encoding to use when parsing the file</param>
-		public static Task AppendAllLinesAsync(FileReference Location, string[] Contents, Encoding Encoding)
+		/// <param name="location">Location of the file</param>
+		/// <param name="contents">Contents to append to the file</param>
+		/// <param name="encoding">The encoding to use when parsing the file</param>
+		public static Task AppendAllLinesAsync(FileReference location, string[] contents, Encoding encoding)
 		{
-			return File.AppendAllLinesAsync(Location.FullName, Contents, Encoding);
+			return File.AppendAllLinesAsync(location.FullName, contents, encoding);
 		}
 
 		/// <summary>
 		/// Appends the contents to a file
 		/// </summary>
-		/// <param name="Location">Location of the file</param>
-		/// <param name="Contents">Contents to append to the file</param>
-		public static void AppendAllText(FileReference Location, string Contents)
+		/// <param name="location">Location of the file</param>
+		/// <param name="contents">Contents to append to the file</param>
+		public static void AppendAllText(FileReference location, string contents)
 		{
-			File.AppendAllText(Location.FullName, Contents);
+			File.AppendAllText(location.FullName, contents);
 		}
 
 		/// <summary>
 		/// Appends the contents to a file
 		/// </summary>
-		/// <param name="Location">Location of the file</param>
-		/// <param name="Contents">Contents to append to the file</param>
-		public static Task AppendAllTextAsync(FileReference Location, string Contents)
+		/// <param name="location">Location of the file</param>
+		/// <param name="contents">Contents to append to the file</param>
+		public static Task AppendAllTextAsync(FileReference location, string contents)
 		{
-			return File.AppendAllTextAsync(Location.FullName, Contents);
+			return File.AppendAllTextAsync(location.FullName, contents);
 		}
 
 		/// <summary>
 		/// Appends the contents to a file
 		/// </summary>
-		/// <param name="Location">Location of the file</param>
-		/// <param name="Contents">Contents to append to the file</param>
-		/// <param name="Encoding">The encoding to use when parsing the file</param>
-		public static void AppendAllText(FileReference Location, string Contents, Encoding Encoding)
+		/// <param name="location">Location of the file</param>
+		/// <param name="contents">Contents to append to the file</param>
+		/// <param name="encoding">The encoding to use when parsing the file</param>
+		public static void AppendAllText(FileReference location, string contents, Encoding encoding)
 		{
-			File.AppendAllText(Location.FullName, Contents, Encoding);
+			File.AppendAllText(location.FullName, contents, encoding);
 		}
 
 		/// <summary>
 		/// Appends the contents to a file
 		/// </summary>
-		/// <param name="Location">Location of the file</param>
-		/// <param name="Contents">Contents to append to the file</param>
-		/// <param name="Encoding">The encoding to use when parsing the file</param>
-		public static Task AppendAllTextAsync(FileReference Location, string Contents, Encoding Encoding)
+		/// <param name="location">Location of the file</param>
+		/// <param name="contents">Contents to append to the file</param>
+		/// <param name="encoding">The encoding to use when parsing the file</param>
+		public static Task AppendAllTextAsync(FileReference location, string contents, Encoding encoding)
 		{
-			return File.AppendAllTextAsync(Location.FullName, Contents, Encoding);
+			return File.AppendAllTextAsync(location.FullName, contents, encoding);
 		}
 
 		#endregion
@@ -893,137 +892,137 @@ namespace EpicGames.Core
 		/// <summary>
 		/// Manually serialize a file reference to a binary stream.
 		/// </summary>
-		/// <param name="Writer">Binary writer to write to</param>
-		/// <param name="File">The file reference to write</param>
-		public static void Write(this BinaryWriter Writer, FileReference File)
+		/// <param name="writer">Binary writer to write to</param>
+		/// <param name="file">The file reference to write</param>
+		public static void Write(this BinaryWriter writer, FileReference file)
 		{
-			Writer.Write((File == null) ? String.Empty : File.FullName);
+			writer.Write((file == null) ? String.Empty : file.FullName);
 		}
 
 		/// <summary>
 		/// Serializes a file reference, using a lookup table to avoid serializing the same name more than once.
 		/// </summary>
-		/// <param name="Writer">The writer to save this reference to</param>
-		/// <param name="File">A file reference to output; may be null</param>
-		/// <param name="FileToUniqueId">A lookup table that caches previous files that have been output, and maps them to unique id's.</param>
-		public static void Write(this BinaryWriter Writer, FileReference File, Dictionary<FileReference, int> FileToUniqueId)
+		/// <param name="writer">The writer to save this reference to</param>
+		/// <param name="file">A file reference to output; may be null</param>
+		/// <param name="fileToUniqueId">A lookup table that caches previous files that have been output, and maps them to unique id's.</param>
+		public static void Write(this BinaryWriter writer, FileReference file, Dictionary<FileReference, int> fileToUniqueId)
 		{
-			int UniqueId;
-			if (File == null)
+			int uniqueId;
+			if (file == null)
 			{
-				Writer.Write(-1);
+				writer.Write(-1);
 			}
-			else if (FileToUniqueId.TryGetValue(File, out UniqueId))
+			else if (fileToUniqueId.TryGetValue(file, out uniqueId))
 			{
-				Writer.Write(UniqueId);
+				writer.Write(uniqueId);
 			}
 			else
 			{
-				Writer.Write(FileToUniqueId.Count);
-				Writer.Write(File);
-				FileToUniqueId.Add(File, FileToUniqueId.Count);
+				writer.Write(fileToUniqueId.Count);
+				writer.Write(file);
+				fileToUniqueId.Add(file, fileToUniqueId.Count);
 			}
 		}
 
 		/// <summary>
 		/// Manually deserialize a file reference from a binary stream.
 		/// </summary>
-		/// <param name="Reader">Binary reader to read from</param>
+		/// <param name="reader">Binary reader to read from</param>
 		/// <returns>New FileReference object</returns>
-		public static FileReference ReadFileReference(this BinaryReader Reader)
+		public static FileReference ReadFileReference(this BinaryReader reader)
 		{
-			return BinaryArchiveReader.NotNull(ReadFileReferenceOrNull(Reader));
+			return BinaryArchiveReader.NotNull(ReadFileReferenceOrNull(reader));
 		}
 
 		/// <summary>
 		/// Manually deserialize a file reference from a binary stream.
 		/// </summary>
-		/// <param name="Reader">Binary reader to read from</param>
+		/// <param name="reader">Binary reader to read from</param>
 		/// <returns>New FileReference object</returns>
-		public static FileReference? ReadFileReferenceOrNull(this BinaryReader Reader)
+		public static FileReference? ReadFileReferenceOrNull(this BinaryReader reader)
 		{
-			string FullName = Reader.ReadString();
-			return (FullName.Length == 0) ? null : new FileReference(FullName, FileReference.Sanitize.None);
+			string fullName = reader.ReadString();
+			return (fullName.Length == 0) ? null : new FileReference(fullName, FileReference.Sanitize.None);
 		}
 
 		/// <summary>
 		/// Deserializes a file reference, using a lookup table to avoid writing the same name more than once.
 		/// </summary>
-		/// <param name="Reader">The source to read from</param>
-		/// <param name="UniqueFiles">List of previously read file references. The index into this array is used in place of subsequent ocurrences of the file.</param>
+		/// <param name="reader">The source to read from</param>
+		/// <param name="uniqueFiles">List of previously read file references. The index into this array is used in place of subsequent ocurrences of the file.</param>
 		/// <returns>The file reference that was read</returns>
-		public static FileReference ReadFileReference(this BinaryReader Reader, List<FileReference> UniqueFiles)
+		public static FileReference ReadFileReference(this BinaryReader reader, List<FileReference> uniqueFiles)
 		{
-			return BinaryArchiveReader.NotNull(ReadFileReferenceOrNull(Reader, UniqueFiles));
+			return BinaryArchiveReader.NotNull(ReadFileReferenceOrNull(reader, uniqueFiles));
 		}
 
 		/// <summary>
 		/// Deserializes a file reference, using a lookup table to avoid writing the same name more than once.
 		/// </summary>
-		/// <param name="Reader">The source to read from</param>
-		/// <param name="UniqueFiles">List of previously read file references. The index into this array is used in place of subsequent ocurrences of the file.</param>
+		/// <param name="reader">The source to read from</param>
+		/// <param name="uniqueFiles">List of previously read file references. The index into this array is used in place of subsequent ocurrences of the file.</param>
 		/// <returns>The file reference that was read</returns>
-		public static FileReference? ReadFileReferenceOrNull(this BinaryReader Reader, List<FileReference> UniqueFiles)
+		public static FileReference? ReadFileReferenceOrNull(this BinaryReader reader, List<FileReference> uniqueFiles)
 		{
-			int UniqueId = Reader.ReadInt32();
-			if (UniqueId == -1)
+			int uniqueId = reader.ReadInt32();
+			if (uniqueId == -1)
 			{
 				return null;
 			}
-			else if (UniqueId < UniqueFiles.Count)
+			else if (uniqueId < uniqueFiles.Count)
 			{
-				return UniqueFiles[UniqueId];
+				return uniqueFiles[uniqueId];
 			}
 			else
 			{
-				FileReference Result = Reader.ReadFileReference();
-				UniqueFiles.Add(Result);
-				return Result;
+				FileReference result = reader.ReadFileReference();
+				uniqueFiles.Add(result);
+				return result;
 			}
 		}
 
 		/// <summary>
 		/// Writes a FileReference to a binary archive
 		/// </summary>
-		/// <param name="Writer">The writer to output data to</param>
-		/// <param name="File">The file reference to write</param>
-		public static void WriteFileReference(this BinaryArchiveWriter Writer, FileReference? File)
+		/// <param name="writer">The writer to output data to</param>
+		/// <param name="file">The file reference to write</param>
+		public static void WriteFileReference(this BinaryArchiveWriter writer, FileReference? file)
 		{
-			if(File == null)
+			if(file == null)
 			{
-				Writer.WriteString(null);
+				writer.WriteString(null);
 			}
 			else
 			{
-				Writer.WriteString(File.FullName);
+				writer.WriteString(file.FullName);
 			}
 		}
 
 		/// <summary>
 		/// Reads a FileReference from a binary archive
 		/// </summary>
-		/// <param name="Reader">Reader to serialize data from</param>
+		/// <param name="reader">Reader to serialize data from</param>
 		/// <returns>New file reference instance</returns>
-		public static FileReference ReadFileReference(this BinaryArchiveReader Reader)
+		public static FileReference ReadFileReference(this BinaryArchiveReader reader)
 		{
-			return BinaryArchiveReader.NotNull(ReadFileReferenceOrNull(Reader));
+			return BinaryArchiveReader.NotNull(ReadFileReferenceOrNull(reader));
 		}
 
 		/// <summary>
 		/// Reads a FileReference from a binary archive
 		/// </summary>
-		/// <param name="Reader">Reader to serialize data from</param>
+		/// <param name="reader">Reader to serialize data from</param>
 		/// <returns>New file reference instance</returns>
-		public static FileReference? ReadFileReferenceOrNull(this BinaryArchiveReader Reader)
+		public static FileReference? ReadFileReferenceOrNull(this BinaryArchiveReader reader)
 		{
-			string? FullName = Reader.ReadString();
-			if(FullName == null)
+			string? fullName = reader.ReadString();
+			if(fullName == null)
 			{
 				return null;
 			}
 			else
 			{
-				return new FileReference(FullName, FileReference.Sanitize.None);
+				return new FileReference(fullName, FileReference.Sanitize.None);
 			}
 		}
 	}
