Reformat EpicGames.Core according to standard coding conventions.

diff --git a/Engine/Source/Programs/Shared/EpicGames.Core/FileSet.cs b/Engine/Source/Programs/Shared/EpicGames.Core/FileSet.cs
--- a/Engine/Source/Programs/Shared/EpicGames.Core/FileSet.cs
+++ b/Engine/Source/Programs/Shared/EpicGames.Core/FileSet.cs
@@ -5,7 +5,6 @@ using System.Collections;
 using System.Collections.Generic;
 using System.IO;
 using System.Linq;
-using System.Text;
 
 namespace EpicGames.Core
 {
@@ -28,10 +27,10 @@ namespace EpicGames.Core
 		/// <summary>
 		/// Constructor
 		/// </summary>
-		/// <param name="Path">Relative path within the tree</param>
-		public FileSet(string Path)
+		/// <param name="path">Relative path within the tree</param>
+		public FileSet(string path)
 		{
-			this.Path = Path;
+			Path = path;
 		}
 
 		/// <summary>
@@ -51,19 +50,19 @@ namespace EpicGames.Core
 		/// </summary>
 		/// <param name="Files"></param>
 		/// <returns>Tree containing the given files</returns>
-		public static FileSet FromFile(DirectoryReference Directory, string File)
+		public static FileSet FromFile(DirectoryReference directory, string file)
 		{
-			return FromFiles(new[] { (File, FileReference.Combine(Directory, File)) });
+			return FromFiles(new[] { (File: file, FileReference.Combine(directory, file)) });
 		}
 
 		/// <summary>
 		/// Creates a file tree from a given set of files
 		/// </summary>
-		/// <param name="Files"></param>
+		/// <param name="files"></param>
 		/// <returns>Tree containing the given files</returns>
-		public static FileSet FromFiles(IEnumerable<(string, FileReference)> Files)
+		public static FileSet FromFiles(IEnumerable<(string, FileReference)> files)
 		{
-			return new FileSetFromFiles(Files);
+			return new FileSetFromFiles(files);
 		}
 
 		/// <summary>
@@ -71,79 +70,79 @@ namespace EpicGames.Core
 		/// </summary>
 		/// <param name="Files"></param>
 		/// <returns>Tree containing the given files</returns>
-		public static FileSet FromFile(DirectoryReference Directory, FileReference File)
+		public static FileSet FromFile(DirectoryReference directory, FileReference file)
 		{
-			return FromFiles(Directory, new[] { File });
+			return FromFiles(directory, new[] { file });
 		}
 
 		/// <summary>
 		/// Creates a file tree from a given set of files
 		/// </summary>
-		/// <param name="Files"></param>
+		/// <param name="files"></param>
 		/// <returns>Tree containing the given files</returns>
-		public static FileSet FromFiles(DirectoryReference Directory, IEnumerable<FileReference> Files)
+		public static FileSet FromFiles(DirectoryReference directory, IEnumerable<FileReference> files)
 		{
-			return new FileSetFromFiles(Files.Select(x => (x.MakeRelativeTo(Directory), x)));
+			return new FileSetFromFiles(files.Select(x => (x.MakeRelativeTo(directory), x)));
 		}
 
 		/// <summary>
 		/// Creates a file tree from a folder on disk
 		/// </summary>
-		/// <param name="Directory"></param>
+		/// <param name="directory"></param>
 		/// <returns></returns>
-		public static FileSet FromDirectory(DirectoryReference Directory)
+		public static FileSet FromDirectory(DirectoryReference directory)
 		{
-			return new FileSetFromDirectory(new DirectoryInfo(Directory.FullName));
+			return new FileSetFromDirectory(new DirectoryInfo(directory.FullName));
 		}
 
 		/// <summary>
 		/// Creates a file tree from a folder on disk
 		/// </summary>
-		/// <param name="DirectoryInfo"></param>
+		/// <param name="directoryInfo"></param>
 		/// <returns></returns>
-		public static FileSet FromDirectory(DirectoryInfo DirectoryInfo)
+		public static FileSet FromDirectory(DirectoryInfo directoryInfo)
 		{
-			return new FileSetFromDirectory(DirectoryInfo);
+			return new FileSetFromDirectory(directoryInfo);
 		}
 
 		/// <summary>
 		/// Create a tree containing files filtered by any of the given wildcards
 		/// </summary>
-		/// <param name="Rules"></param>
+		/// <param name="rules"></param>
 		/// <returns></returns>
-		public FileSet Filter(string Rules)
+		public FileSet Filter(string rules)
 		{
-			return Filter(Rules.Split(';'));
+			return Filter(rules.Split(';'));
 		}
 
 		/// <summary>
 		/// Create a tree containing files filtered by any of the given wildcards
 		/// </summary>
-		/// <param name="Rules"></param>
+		/// <param name="rules"></param>
 		/// <returns></returns>
-		public FileSet Filter(params string[] Rules)
+		public FileSet Filter(params string[] rules)
 		{
-			return new FileSetFromFilter(this, new FileFilter(Rules));
+			return new FileSetFromFilter(this, new FileFilter(rules));
 		}
 
 		/// <summary>
 		/// Create a tree containing files filtered by any of the given file filter objects
 		/// </summary>
-		/// <param name="Filters"></param>
+		/// <param name="filters"></param>
 		/// <returns></returns>
-		public FileSet Filter(params FileFilter[] Filters)
+		public FileSet Filter(params FileFilter[] filters)
 		{
-			return new FileSetFromFilter(this, Filters);
+			return new FileSetFromFilter(this, filters);
 		}
 
 		/// <summary>
 		/// Create a tree containing the exception of files with another tree
 		/// </summary>
-		/// <param name="Filter">Files to exclude from the filter</param>
+		/// <param name="filter">Files to exclude from the filter</param>
 		/// <returns></returns>
-		public FileSet Except(string Filter)
+		public FileSet Except(string filter)
 		{
-			return Except(Filter.Split(';'));
+			return Except(filter.Split(';'));
 		}
 
 		/// <summary>
@@ -151,43 +150,43 @@ namespace EpicGames.Core
 		/// </summary>
 		/// <param name="Filter">Files to exclude from the filter</param>
 		/// <returns></returns>
-		public FileSet Except(params string[] Rules)
+		public FileSet Except(params string[] rules)
 		{
-			return new FileSetFromFilter(this, new FileFilter(Rules.Select(x => $"-{x}")));
+			return new FileSetFromFilter(this, new FileFilter(rules.Select(x => $"-{x}")));
 		}
 
 		/// <summary>
 		/// Create a tree containing the union of files with another tree
 		/// </summary>
-		/// <param name="Lhs"></param>
-		/// <param name="Rhs"></param>
+		/// <param name="lhs"></param>
+		/// <param name="rhs"></param>
 		/// <returns></returns>
-		public static FileSet Union(FileSet Lhs, FileSet Rhs)
+		public static FileSet Union(FileSet lhs, FileSet rhs)
 		{
-			return new FileSetFromUnion(Lhs, Rhs);
+			return new FileSetFromUnion(lhs, rhs);
 		}
 
 		/// <summary>
 		/// Create a tree containing the exception of files with another tree
 		/// </summary>
-		/// <param name="Lhs"></param>
-		/// <param name="Rhs"></param>
+		/// <param name="lhs"></param>
+		/// <param name="rhs"></param>
 		/// <returns></returns>
-		public static FileSet Except(FileSet Lhs, FileSet Rhs)
+		public static FileSet Except(FileSet lhs, FileSet rhs)
 		{
-			return new FileSetFromExcept(Lhs, Rhs);
+			return new FileSetFromExcept(lhs, rhs);
 		}
 
 		/// <inheritdoc cref="Union(FileSet, FileSet)"/>
-		public static FileSet operator +(FileSet Lhs, FileSet Rhs)
+		public static FileSet operator +(FileSet lhs, FileSet rhs)
 		{
-			return Union(Lhs, Rhs);
+			return Union(lhs, rhs);
 		}
 
 		/// <inheritdoc cref="Except(FileSet, FileSet)"/>
-		public static FileSet operator -(FileSet Lhs, FileSet Rhs)
+		public static FileSet operator -(FileSet lhs, FileSet rhs)
 		{
-			return Except(Lhs, Rhs);
+			return Except(lhs, rhs);
 		}
 
 		/// <summary>
@@ -196,20 +195,20 @@ namespace EpicGames.Core
 		/// <returns></returns>
 		public Dictionary<string, FileReference> Flatten()
 		{
-			Dictionary<string, FileReference> PathToSourceFile = new Dictionary<string, FileReference>(StringComparer.OrdinalIgnoreCase);
-			FlattenInternal(String.Empty, PathToSourceFile);
-			return PathToSourceFile;
+			Dictionary<string, FileReference> pathToSourceFile = new Dictionary<string, FileReference>(StringComparer.OrdinalIgnoreCase);
+			FlattenInternal(String.Empty, pathToSourceFile);
+			return pathToSourceFile;
 		}
 
-		private void FlattenInternal(string PathPrefix, Dictionary<string, FileReference> PathToSourceFile)
+		private void FlattenInternal(string pathPrefix, Dictionary<string, FileReference> pathToSourceFile)
 		{
-			foreach ((string Path, FileReference File) in EnumerateFiles())
+			foreach ((string path, FileReference file) in EnumerateFiles())
 			{
-				PathToSourceFile[PathPrefix + Path] = File;
+				pathToSourceFile[pathPrefix + path] = file;
 			}
-			foreach((string Path, FileSet FileSet) in EnumerateDirectories())
+			foreach((string path, FileSet fileSet) in EnumerateDirectories())
 			{
-				FileSet.FlattenInternal(PathPrefix + Path + "/", PathToSourceFile);
+				fileSet.FlattenInternal(pathPrefix + path + "/", pathToSourceFile);
 			}
 		}
 
@@ -217,15 +216,15 @@ namespace EpicGames.Core
 		/// Flatten to a map of files in a target directory
 		/// </summary>
 		/// <returns></returns>
-		public Dictionary<FileReference, FileReference> Flatten(DirectoryReference OutputDir)
+		public Dictionary<FileReference, FileReference> Flatten(DirectoryReference outputDir)
 		{
-			Dictionary<FileReference, FileReference> TargetToSourceFile = new Dictionary<FileReference, FileReference>();
-			foreach ((string Path, FileReference SourceFile) in Flatten())
+			Dictionary<FileReference, FileReference> targetToSourceFile = new Dictionary<FileReference, FileReference>();
+			foreach ((string path, FileReference sourceFile) in Flatten())
 			{
-				FileReference TargetFile = FileReference.Combine(OutputDir, Path);
-				TargetToSourceFile[TargetFile] = SourceFile;
+				FileReference targetFile = FileReference.Combine(outputDir, path);
+				targetToSourceFile[targetFile] = sourceFile;
 			}
-			return TargetToSourceFile;
+			return targetToSourceFile;
 		}
 
 		/// <inheritdoc/>
@@ -240,50 +239,50 @@ namespace EpicGames.Core
 	/// </summary>
 	class FileSetFromFiles : FileSet
 	{
-		Dictionary<string, FileReference> Files = new Dictionary<string, FileReference>();
-		Dictionary<string, FileSetFromFiles> SubTrees = new Dictionary<string, FileSetFromFiles>();
+		readonly Dictionary<string, FileReference> _files = new Dictionary<string, FileReference>();
+		readonly Dictionary<string, FileSetFromFiles> _subTrees = new Dictionary<string, FileSetFromFiles>();
 
 		/// <summary>
 		/// Private constructor
 		/// </summary>
-		/// <param name="Path"></param>
-		private FileSetFromFiles(string Path)
-			: base(Path)
+		/// <param name="path"></param>
+		private FileSetFromFiles(string path)
+			: base(path)
 		{
 		}
 
 		/// <summary>
 		/// Creates a tree from a given set of files
 		/// </summary>
-		/// <param name="InputFiles"></param>
-		public FileSetFromFiles(IEnumerable<(string, FileReference)> InputFiles)
+		/// <param name="inputFiles"></param>
+		public FileSetFromFiles(IEnumerable<(string, FileReference)> inputFiles)
 			: this(String.Empty)
 		{
-			foreach ((string Path, FileReference File) in InputFiles)
+			foreach ((string path, FileReference file) in inputFiles)
 			{
-				string[] Fragments = Path.Split(new[] { '/', '\\' }, StringSplitOptions.RemoveEmptyEntries);
+				string[] fragments = path.Split(new[] { '/', '\\' }, StringSplitOptions.RemoveEmptyEntries);
 
-				FileSetFromFiles Current = this;
-				for (int Idx = 0; Idx < Fragments.Length - 1; Idx++)
+				FileSetFromFiles current = this;
+				for (int idx = 0; idx < fragments.Length - 1; idx++)
 				{
-					FileSetFromFiles? Next;
-					if (!Current.SubTrees.TryGetValue(Fragments[Idx], out Next))
+					FileSetFromFiles? next;
+					if (!current._subTrees.TryGetValue(fragments[idx], out next))
 					{
-						Next = new FileSetFromFiles(Current.Path + Fragments[Idx] + "/");
-						Current.SubTrees.Add(Fragments[Idx], Next);
+						next = new FileSetFromFiles(current.Path + fragments[idx] + "/");
+						current._subTrees.Add(fragments[idx], next);
 					}
-					Current = Next;
+					current = next;
 				}
 
-				Current.Files.Add(Fragments[^1], File);
+				current._files.Add(fragments[^1], file);
 			}
 		}
 
 		/// <inheritdoc/>
-		public override IEnumerable<KeyValuePair<string, FileReference>> EnumerateFiles() => Files;
+		public override IEnumerable<KeyValuePair<string, FileReference>> EnumerateFiles() => _files;
 
 		/// <inheritdoc/>
-		public override IEnumerable<KeyValuePair<string, FileSet>> EnumerateDirectories() => SubTrees.Select(x => new KeyValuePair<string, FileSet>(x.Key, x.Value));
+		public override IEnumerable<KeyValuePair<string, FileSet>> EnumerateDirectories() => _subTrees.Select(x => new KeyValuePair<string, FileSet>(x.Key, x.Value));
 	}
 
 	/// <summary>
@@ -291,30 +290,30 @@ namespace EpicGames.Core
 	/// </summary>
 	sealed class FileSetFromDirectory : FileSet
 	{
-		DirectoryInfo DirectoryInfo;
+		readonly DirectoryInfo _directoryInfo;
 
 		/// <summary>
 		/// Constructor
 		/// </summary>
-		public FileSetFromDirectory(DirectoryInfo DirectoryInfo)
-			: this(DirectoryInfo, "/")
+		public FileSetFromDirectory(DirectoryInfo directoryInfo)
+			: this(directoryInfo, "/")
 		{
 		}
 
 		/// <summary>
 		/// Constructor
 		/// </summary>
-		public FileSetFromDirectory(DirectoryInfo DirectoryInfo, string Path)
-			: base(Path)
+		public FileSetFromDirectory(DirectoryInfo directoryInfo, string path)
+			: base(path)
 		{
-			this.DirectoryInfo = DirectoryInfo;
+			_directoryInfo = directoryInfo;
 		}
 
 		/// <inheritdoc/>
-		public override IEnumerable<KeyValuePair<string, FileReference>> EnumerateFiles() => DirectoryInfo.EnumerateFiles().Select(x => new KeyValuePair<string, FileReference>(x.Name, new FileReference(x)));
+		public override IEnumerable<KeyValuePair<string, FileReference>> EnumerateFiles() => _directoryInfo.EnumerateFiles().Select(x => new KeyValuePair<string, FileReference>(x.Name, new FileReference(x)));
 
 		/// <inheritdoc/>
-		public override IEnumerable<KeyValuePair<string, FileSet>> EnumerateDirectories() => DirectoryInfo.EnumerateDirectories().Select(x => KeyValuePair.Create<string, FileSet>(x.Name, new FileSetFromDirectory(x)));
+		public override IEnumerable<KeyValuePair<string, FileSet>> EnumerateDirectories() => _directoryInfo.EnumerateDirectories().Select(x => KeyValuePair.Create<string, FileSet>(x.Name, new FileSetFromDirectory(x)));
 	}
 
 	/// <summary>
@@ -322,57 +321,57 @@ namespace EpicGames.Core
 	/// </summary>
 	class FileSetFromUnion : FileSet
 	{
-		FileSet Lhs;
-		FileSet Rhs;
+		readonly FileSet _lhs;
+		readonly FileSet _rhs;
 
 		/// <summary>
 		/// Constructor
 		/// </summary>
-		/// <param name="Lhs">First file tree for the union</param>
-		/// <param name="Rhs">Other file tree for the union</param>
-		public FileSetFromUnion(FileSet Lhs, FileSet Rhs)
-			: base(Lhs.Path)
+		/// <param name="lhs">First file tree for the union</param>
+		/// <param name="rhs">Other file tree for the union</param>
+		public FileSetFromUnion(FileSet lhs, FileSet rhs)
+			: base(lhs.Path)
 		{
-			this.Lhs = Lhs;
-			this.Rhs = Rhs;
+			_lhs = lhs;
+			_rhs = rhs;
 		}
 
 		/// <inheritdoc/>
 		public override IEnumerable<KeyValuePair<string, FileReference>> EnumerateFiles()
 		{
-			Dictionary<string, FileReference> Files = new Dictionary<string, FileReference>(Lhs.EnumerateFiles(), StringComparer.OrdinalIgnoreCase);
-			foreach ((string Name, FileReference File) in Rhs.EnumerateFiles())
+			Dictionary<string, FileReference> files = new Dictionary<string, FileReference>(_lhs.EnumerateFiles(), StringComparer.OrdinalIgnoreCase);
+			foreach ((string name, FileReference file) in _rhs.EnumerateFiles())
 			{
-				FileReference? ExistingFile;
-				if (!Files.TryGetValue(Name, out ExistingFile))
+				FileReference? existingFile;
+				if (!files.TryGetValue(name, out existingFile))
 				{
-					Files.Add(Name, File);
+					files.Add(name, file);
 				}
-				else if (ExistingFile == null || !ExistingFile.Equals(File))
+				else if (existingFile == null || !existingFile.Equals(file))
 				{
-					throw new InvalidOperationException($"Conflict for contents of {Path}{Name} - could be {ExistingFile} or {File}");
+					throw new InvalidOperationException($"Conflict for contents of {Path}{name} - could be {existingFile} or {file}");
 				}
 			}
-			return Files;
+			return files;
 		}
 
 		/// <inheritdoc/>
 		public override IEnumerable<KeyValuePair<string, FileSet>> EnumerateDirectories()
 		{
-			Dictionary<string, FileSet> NameToSubTree = new Dictionary<string, FileSet>(Lhs.EnumerateDirectories(), StringComparer.OrdinalIgnoreCase);
-			foreach ((string Name, FileSet SubTree) in Rhs.EnumerateDirectories())
+			Dictionary<string, FileSet> nameToSubTree = new Dictionary<string, FileSet>(_lhs.EnumerateDirectories(), StringComparer.OrdinalIgnoreCase);
+			foreach ((string name, FileSet subTree) in _rhs.EnumerateDirectories())
 			{
-				FileSet? ExistingSubTree;
-				if (NameToSubTree.TryGetValue(Name, out ExistingSubTree))
+				FileSet? existingSubTree;
+				if (nameToSubTree.TryGetValue(name, out existingSubTree))
 				{
-					NameToSubTree[Name] = new FileSetFromUnion(ExistingSubTree, SubTree);
+					nameToSubTree[name] = new FileSetFromUnion(existingSubTree, subTree);
 				}
 				else
 				{
-					NameToSubTree[Name] = SubTree;
+					nameToSubTree[name] = subTree;
 				}
 			}
-			return NameToSubTree;
+			return nameToSubTree;
 		}
 	}
 
@@ -381,42 +380,42 @@ namespace EpicGames.Core
 	/// </summary>
 	class FileSetFromExcept : FileSet
 	{
-		FileSet Lhs;
-		FileSet Rhs;
+		readonly FileSet _lhs;
+		readonly FileSet _rhs;
 
 		/// <summary>
 		/// Constructor
 		/// </summary>
-		/// <param name="Lhs">First file tree for the union</param>
-		/// <param name="Rhs">Other file tree for the union</param>
-		public FileSetFromExcept(FileSet Lhs, FileSet Rhs)
-			: base(Lhs.Path)
+		/// <param name="lhs">First file tree for the union</param>
+		/// <param name="rhs">Other file tree for the union</param>
+		public FileSetFromExcept(FileSet lhs, FileSet rhs)
+			: base(lhs.Path)
 		{
-			this.Lhs = Lhs;
-			this.Rhs = Rhs;
+			_lhs = lhs;
+			_rhs = rhs;
 		}
 
 		/// <inheritdoc/>
 		public override IEnumerable<KeyValuePair<string, FileReference>> EnumerateFiles()
 		{
-			HashSet<string> RhsFiles = new HashSet<string>(Rhs.EnumerateFiles().Select(x => x.Key), StringComparer.OrdinalIgnoreCase);
-			return Lhs.EnumerateFiles().Where(x => !RhsFiles.Contains(x.Key));
+			HashSet<string> rhsFiles = new HashSet<string>(_rhs.EnumerateFiles().Select(x => x.Key), StringComparer.OrdinalIgnoreCase);
+			return _lhs.EnumerateFiles().Where(x => !rhsFiles.Contains(x.Key));
 		}
 
 		/// <inheritdoc/>
 		public override IEnumerable<KeyValuePair<string, FileSet>> EnumerateDirectories()
 		{
-			Dictionary<string, FileSet> RhsDirs = new Dictionary<string, FileSet>(Rhs.EnumerateDirectories(), StringComparer.OrdinalIgnoreCase);
-			foreach ((string Name, FileSet LhsSet) in Lhs.EnumerateDirectories())
+			Dictionary<string, FileSet> rhsDirs = new Dictionary<string, FileSet>(_rhs.EnumerateDirectories(), StringComparer.OrdinalIgnoreCase);
+			foreach ((string name, FileSet lhsSet) in _lhs.EnumerateDirectories())
 			{
-				FileSet? RhsSet;
-				if (RhsDirs.TryGetValue(Name, out RhsSet))
+				FileSet? rhsSet;
+				if (rhsDirs.TryGetValue(name, out rhsSet))
 				{
-					yield return KeyValuePair.Create<string, FileSet>(Name, new FileSetFromExcept(LhsSet, RhsSet));
+					yield return KeyValuePair.Create<string, FileSet>(name, new FileSetFromExcept(lhsSet, rhsSet));
 				}
 				else
 				{
-					yield return KeyValuePair.Create(Name, LhsSet);
+					yield return KeyValuePair.Create(name, lhsSet);
 				}
 			}
 		}
@@ -428,30 +427,30 @@ namespace EpicGames.Core
 	/// <typeparam name="T">Class containing information about a file</typeparam>
 	class FileSetFromFilter : FileSet
 	{
-		FileSet Inner;
-		FileFilter[] Filters;
+		readonly FileSet _inner;
+		readonly FileFilter[] _filters;
 
 		/// <summary>
 		/// Constructor
 		/// </summary>
-		/// <param name="Inner">The tree to filter</param>
-		/// <param name="Filters"></param>
-		public FileSetFromFilter(FileSet Inner, params FileFilter[] Filters)
-			: base(Inner.Path)
+		/// <param name="inner">The tree to filter</param>
+		/// <param name="filters"></param>
+		public FileSetFromFilter(FileSet inner, params FileFilter[] filters)
+			: base(inner.Path)
 		{
-			this.Inner = Inner;
-			this.Filters = Filters;
+			_inner = inner;
+			_filters = filters;
 		}
 
 		/// <inheritdoc/>
 		public override IEnumerable<KeyValuePair<string, FileReference>> EnumerateFiles()
 		{
-			foreach (KeyValuePair<string, FileReference> Item in Inner.EnumerateFiles())
+			foreach (KeyValuePair<string, FileReference> item in _inner.EnumerateFiles())
 			{
-				string FilterName = Inner.Path + Item.Key;
-				if (Filters.Any(x => x.Matches(FilterName)))
+				string filterName = _inner.Path + item.Key;
+				if (_filters.Any(x => x.Matches(filterName)))
 				{
-					yield return Item;
+					yield return item;
 				}
 			}
 		}
@@ -459,15 +458,15 @@ namespace EpicGames.Core
 		/// <inheritdoc/>
 		public override IEnumerable<KeyValuePair<string, FileSet>> EnumerateDirectories()
 		{
-			foreach (KeyValuePair<string, FileSet> Item in Inner.EnumerateDirectories())
+			foreach (KeyValuePair<string, FileSet> item in _inner.EnumerateDirectories())
 			{
-				string FilterName = Inner.Path + Item.Key;
+				string filterName = _inner.Path + item.Key;
 
-				FileFilter[] PossibleFilters = Filters.Where(x => x.PossiblyMatches(FilterName)).ToArray();
-				if (PossibleFilters.Length > 0)
+				FileFilter[] possibleFilters = _filters.Where(x => x.PossiblyMatches(filterName)).ToArray();
+				if (possibleFilters.Length > 0)
 				{
-					FileSetFromFilter SubTreeFilter = new FileSetFromFilter(Item.Value, PossibleFilters);
-					yield return new KeyValuePair<string, FileSet>(Item.Key, SubTreeFilter);
+					FileSetFromFilter subTreeFilter = new FileSetFromFilter(item.Value, possibleFilters);
+					yield return new KeyValuePair<string, FileSet>(item.Key, subTreeFilter);
 				}
 			}
 		}
diff --git a/Engine/Source/Programs/Shared/EpicGames.Core/FileSystemReference.cs b/Engine/Source/Programs/Shared/EpicGames.Core/FileSystemReference.cs
--- a/Engine/Source/Programs/Shared/EpicGames.Core/FileSystemReference.cs
+++ b/Engine/Source/Programs/Shared/EpicGames.Core/FileSystemReference.cs
@@ -6,7 +6,6 @@ using System.IO;
 using System.Linq;
 using System.Text;
 using System.Threading;
-using System.Threading.Tasks;
 
 namespace EpicGames.Core
 {
@@ -34,209 +33,201 @@ namespace EpicGames.Core
 		/// <summary>
 		/// Direct constructor for a path
 		/// </summary>
-		protected FileSystemReference(string InFullName)
+		protected FileSystemReference(string fullName)
 		{
-			FullName = InFullName;
-		}
-
-		/// <summary>
-		/// Direct constructor for a path
-		/// </summary>
-		protected FileSystemReference(string InFullName, string InCanonicalName)
-		{
-			FullName = InFullName;
+			FullName = fullName;
 		}
 
 		/// <summary>
 		/// Create a full path by concatenating multiple strings
 		/// </summary>
 		/// <returns></returns>
-		static ThreadLocal<StringBuilder> CombineStringsStringBuilder = new ThreadLocal<StringBuilder>(() => new StringBuilder(260));
-		static protected string CombineStrings(DirectoryReference BaseDirectory, params string[] Fragments)
+		static readonly ThreadLocal<StringBuilder> s_combineStringsStringBuilder = new ThreadLocal<StringBuilder>(() => new StringBuilder(260));
+		protected static string CombineStrings(DirectoryReference baseDirectory, params string[] fragments)
 		{
 			// Get the initial string to append to, and strip any root directory suffix from it
-			StringBuilder NewFullName = CombineStringsStringBuilder.Value!.Clear().Append(BaseDirectory.FullName);
-			if (NewFullName.Length > 0 && NewFullName[NewFullName.Length - 1] == Path.DirectorySeparatorChar)
+			StringBuilder newFullName = s_combineStringsStringBuilder.Value!.Clear().Append(baseDirectory.FullName);
+			if (newFullName.Length > 0 && newFullName[^1] == Path.DirectorySeparatorChar)
 			{
-				NewFullName.Remove(NewFullName.Length - 1, 1);
+				newFullName.Remove(newFullName.Length - 1, 1);
 			}
 
 			// Scan through the fragments to append, appending them to a string and updating the base length as we go
-			foreach (string Fragment in Fragments)
+			foreach (string fragment in fragments)
 			{
 				// Check if this fragment is an absolute path
-				if ((Fragment.Length >= 2 && Fragment[1] == ':') || (Fragment.Length >= 1 && (Fragment[0] == '\\' || Fragment[0] == '/')))
+				if ((fragment.Length >= 2 && fragment[1] == ':') || (fragment.Length >= 1 && (fragment[0] == '\\' || fragment[0] == '/')))
 				{
 					// It is. Reset the new name to the full version of this path.
-					NewFullName.Clear();
-					NewFullName.Append(Path.GetFullPath(Fragment).TrimEnd(Path.DirectorySeparatorChar));
+					newFullName.Clear();
+					newFullName.Append(Path.GetFullPath(fragment).TrimEnd(Path.DirectorySeparatorChar));
 				}
 				else
 				{
 					// Append all the parts of this fragment to the end of the existing path.
-					int StartIdx = 0;
-					while (StartIdx < Fragment.Length)
+					int startIdx = 0;
+					while (startIdx < fragment.Length)
 					{
 						// Find the end of this fragment. We may have been passed multiple paths in the same string.
-						int EndIdx = StartIdx;
-						while (EndIdx < Fragment.Length && Fragment[EndIdx] != '\\' && Fragment[EndIdx] != '/')
+						int endIdx = startIdx;
+						while (endIdx < fragment.Length && fragment[endIdx] != '\\' && fragment[endIdx] != '/')
 						{
-							EndIdx++;
+							endIdx++;
 						}
 
 						// Ignore any empty sections, like leading or trailing slashes, and '.' directory references.
-						int Length = EndIdx - StartIdx;
-						if (Length == 0)
+						int length = endIdx - startIdx;
+						if (length == 0)
 						{
 							// Multiple directory separators in a row; illegal.
-							throw new ArgumentException(String.Format("Path fragment '{0}' contains invalid directory separators.", Fragment));
+							throw new ArgumentException(String.Format("Path fragment '{0}' contains invalid directory separators.", fragment));
 						}
-						else if (Length == 2 && Fragment[StartIdx] == '.' && Fragment[StartIdx + 1] == '.')
+						else if (length == 2 && fragment[startIdx] == '.' && fragment[startIdx + 1] == '.')
 						{
 							// Remove the last directory name
-							for (int SeparatorIdx = NewFullName.Length - 1; SeparatorIdx >= 0; SeparatorIdx--)
+							for (int separatorIdx = newFullName.Length - 1; separatorIdx >= 0; separatorIdx--)
 							{
-								if (NewFullName[SeparatorIdx] == Path.DirectorySeparatorChar)
+								if (newFullName[separatorIdx] == Path.DirectorySeparatorChar)
 								{
-									NewFullName.Remove(SeparatorIdx, NewFullName.Length - SeparatorIdx);
+									newFullName.Remove(separatorIdx, newFullName.Length - separatorIdx);
 									break;
 								}
 							}
 						}
-						else if (Length != 1 || Fragment[StartIdx] != '.')
+						else if (length != 1 || fragment[startIdx] != '.')
 						{
 							// Append this fragment
-							NewFullName.Append(Path.DirectorySeparatorChar);
-							NewFullName.Append(Fragment, StartIdx, Length);
+							newFullName.Append(Path.DirectorySeparatorChar);
+							newFullName.Append(fragment, startIdx, length);
 						}
 
 						// Move to the next part
-						StartIdx = EndIdx + 1;
+						startIdx = endIdx + 1;
 					}
 				}
 			}
 
 			// Append the directory separator
-			if (NewFullName.Length == 0 || (NewFullName.Length == 2 && NewFullName[1] == ':'))
+			if (newFullName.Length == 0 || (newFullName.Length == 2 && newFullName[1] == ':'))
 			{
-				NewFullName.Append(Path.DirectorySeparatorChar);
+				newFullName.Append(Path.DirectorySeparatorChar);
 			}
 
 			// Set the new path variables
-			return NewFullName.ToString();
+			return newFullName.ToString();
 		}
 
 		/// <summary>
 		/// Checks whether this name has the given extension.
 		/// </summary>
-		/// <param name="Extension">The extension to check</param>
+		/// <param name="extension">The extension to check</param>
 		/// <returns>True if this name has the given extension, false otherwise</returns>
-		public bool HasExtension(string Extension)
+		public bool HasExtension(string extension)
 		{
-			if (Extension.Length > 0 && Extension[0] != '.')
+			if (extension.Length > 0 && extension[0] != '.')
 			{
-				return FullName.Length >= Extension.Length + 1 && FullName[FullName.Length - Extension.Length - 1] == '.' && FullName.EndsWith(Extension, Comparison);
+				return FullName.Length >= extension.Length + 1 && FullName[FullName.Length - extension.Length - 1] == '.' && FullName.EndsWith(extension, Comparison);
 			}
 			else
 			{
-				return FullName.EndsWith(Extension, Comparison);
+				return FullName.EndsWith(extension, Comparison);
 			}
 		}
 
 		/// <summary>
 		/// Determines if the given object is at or under the given directory
 		/// </summary>
-		/// <param name="Other">Directory to check against</param>
+		/// <param name="other">Directory to check against</param>
 		/// <returns>True if this path is under the given directory</returns>
-		public bool IsUnderDirectory(DirectoryReference Other)
+		public bool IsUnderDirectory(DirectoryReference other)
 		{
-			return FullName.StartsWith(Other.FullName, Comparison) && (FullName.Length == Other.FullName.Length || FullName[Other.FullName.Length] == Path.DirectorySeparatorChar || Other.IsRootDirectory());
+			return FullName.StartsWith(other.FullName, Comparison) && (FullName.Length == other.FullName.Length || FullName[other.FullName.Length] == Path.DirectorySeparatorChar || other.IsRootDirectory());
 		}
 
 		/// <summary>
 		/// Searches the path fragments for the given name. Only complete fragments are considered a match.
 		/// </summary>
-		/// <param name="Name">Name to check for</param>
-		/// <param name="Offset">Offset within the string to start the search</param>
+		/// <param name="name">Name to check for</param>
+		/// <param name="offset">Offset within the string to start the search</param>
 		/// <returns>True if the given name is found within the path</returns>
-		public bool ContainsName(string Name, int Offset)
+		public bool ContainsName(string name, int offset)
 		{
-			return ContainsName(Name, Offset, FullName.Length - Offset);
+			return ContainsName(name, offset, FullName.Length - offset);
 		}
 
 		/// <summary>
 		/// Searches the path fragments for the given name. Only complete fragments are considered a match.
 		/// </summary>
-		/// <param name="Name">Name to check for</param>
-		/// <param name="Offset">Offset within the string to start the search</param>
-		/// <param name="Length">Length of the substring to search</param>
+		/// <param name="name">Name to check for</param>
+		/// <param name="offset">Offset within the string to start the search</param>
+		/// <param name="length">Length of the substring to search</param>
 		/// <returns>True if the given name is found within the path</returns>
-		public bool ContainsName(string Name, int Offset, int Length)
+		public bool ContainsName(string name, int offset, int length)
 		{
 			// Check the substring to search is at least long enough to contain a match
-			if(Length < Name.Length)
+			if(length < name.Length)
 			{
 				return false;
 			}
 
 			// Find each occurence of the name within the remaining string, then test whether it's surrounded by directory separators
-			int MatchIdx = Offset;
+			int matchIdx = offset;
 			for(;;)
 			{
 				// Find the next occurrence
-				MatchIdx = FullName.IndexOf(Name, MatchIdx, Offset + Length - MatchIdx, Comparison);
-				if(MatchIdx == -1)
+				matchIdx = FullName.IndexOf(name, matchIdx, offset + length - matchIdx, Comparison);
+				if(matchIdx == -1)
 				{
 					return false;
 				}
 
 				// Check if the substring is a directory
-				int MatchEndIdx = MatchIdx + Name.Length;
-				if(FullName[MatchIdx - 1] == Path.DirectorySeparatorChar && (MatchEndIdx == FullName.Length || FullName[MatchEndIdx] == Path.DirectorySeparatorChar))
+				int matchEndIdx = matchIdx + name.Length;
+				if(FullName[matchIdx - 1] == Path.DirectorySeparatorChar && (matchEndIdx == FullName.Length || FullName[matchEndIdx] == Path.DirectorySeparatorChar))
 				{
 					return true;
 				}
 
 				// Move past the string that didn't match
-				MatchIdx += Name.Length;
+				matchIdx += name.Length;
 			}
 		}
 
 		/// <summary>
 		/// Determines if the given object is under the given directory, within a subfolder of the given name. Useful for masking out directories by name.
 		/// </summary>
-		/// <param name="Name">Name of a subfolder to also check for</param>
-		/// <param name="BaseDir">Base directory to check against</param>
+		/// <param name="name">Name of a subfolder to also check for</param>
+		/// <param name="baseDir">Base directory to check against</param>
 		/// <returns>True if the path is under the given directory</returns>
-		public bool ContainsName(string Name, DirectoryReference BaseDir)
+		public bool ContainsName(string name, DirectoryReference baseDir)
 		{
 			// Check that this is under the base directory
-			if(!IsUnderDirectory(BaseDir))
+			if(!IsUnderDirectory(baseDir))
 			{
 				return false;
 			}
 			else
 			{
-				return ContainsName(Name, BaseDir.FullName.Length);
+				return ContainsName(name, baseDir.FullName.Length);
 			}
 		}
 
 		/// <summary>
 		/// Determines if the given object is under the given directory, within a subfolder of the given name. Useful for masking out directories by name.
 		/// </summary>
-		/// <param name="Names">Names of subfolders to also check for</param>
-		/// <param name="BaseDir">Base directory to check against</param>
+		/// <param name="names">Names of subfolders to also check for</param>
+		/// <param name="baseDir">Base directory to check against</param>
 		/// <returns>True if the path is under the given directory</returns>
-		public bool ContainsAnyNames(IEnumerable<string> Names, DirectoryReference BaseDir)
+		public bool ContainsAnyNames(IEnumerable<string> names, DirectoryReference baseDir)
 		{
 			// Check that this is under the base directory
-			if(!IsUnderDirectory(BaseDir))
+			if(!IsUnderDirectory(baseDir))
 			{
 				return false;
 			}
 			else
 			{
-				return Names.Any(x => ContainsName(x, BaseDir.FullName.Length));
+				return names.Any(x => ContainsName(x, baseDir.FullName.Length));
 			}
 		}
 		
@@ -245,90 +236,90 @@ namespace EpicGames.Core
 		/// </summary>
 		/// <param name="Directory">The directory to create a relative path from</param>
 		/// <returns>A relative path from the given directory</returns>
-		static ThreadLocal<StringBuilder> MakeRelativeToStringBuilder = new ThreadLocal<StringBuilder>(() => new StringBuilder(260));
-		public string MakeRelativeTo(DirectoryReference Directory)
+		static readonly ThreadLocal<StringBuilder> s_makeRelativeToStringBuilder = new ThreadLocal<StringBuilder>(() => new StringBuilder(260));
+		public string MakeRelativeTo(DirectoryReference directory)
 		{
-			StringBuilder Result = MakeRelativeToStringBuilder.Value!.Clear();
-			WriteRelativeTo(Directory, ref Result);
-			return Result.ToString();
+			StringBuilder result = s_makeRelativeToStringBuilder.Value!.Clear();
+			WriteRelativeTo(directory, ref result);
+			return result.ToString();
 		}
 
-		public void WriteRelativeTo(DirectoryReference Directory, ref StringBuilder Result)
+		public void WriteRelativeTo(DirectoryReference directory, ref StringBuilder result)
 		{
 			// Find how much of the path is common between the two paths. This length does not include a trailing directory separator character.
-			int CommonDirectoryLength = -1;
-			for (int Idx = 0; ; Idx++)
+			int commonDirectoryLength = -1;
+			for (int idx = 0; ; idx++)
 			{
-				if (Idx == FullName.Length)
+				if (idx == FullName.Length)
 				{
 					// The two paths are identical. Just return the "." character.
-					if (Idx == Directory.FullName.Length)
+					if (idx == directory.FullName.Length)
 					{
-						Result.Append('.');
+						result.Append('.');
 						return;
 					}
 
 					// Check if we're finishing on a complete directory name
-					if (Directory.FullName[Idx] == Path.DirectorySeparatorChar)
+					if (directory.FullName[idx] == Path.DirectorySeparatorChar)
 					{
-						CommonDirectoryLength = Idx;
+						commonDirectoryLength = idx;
 					}
 					break;
 				}
-				else if (Idx == Directory.FullName.Length)
+				else if (idx == directory.FullName.Length)
 				{
 					// Check whether the end of the directory name coincides with a boundary for the current name.
-					if (FullName[Idx] == Path.DirectorySeparatorChar)
+					if (FullName[idx] == Path.DirectorySeparatorChar)
 					{
-						CommonDirectoryLength = Idx;
+						commonDirectoryLength = idx;
 					}
 					break;
 				}
 				else
 				{
 					// Check the two paths match, and bail if they don't. Increase the common directory length if we've reached a separator.
-					if(String.Compare(FullName, Idx, Directory.FullName, Idx, 1, Comparison) != 0)
+					if(String.Compare(FullName, idx, directory.FullName, idx, 1, Comparison) != 0)
 					{
 						break;
 					}
-					if (FullName[Idx] == Path.DirectorySeparatorChar)
+					if (FullName[idx] == Path.DirectorySeparatorChar)
 					{
-						CommonDirectoryLength = Idx;
+						commonDirectoryLength = idx;
 					}
 				}
 			}
 
 			// If there's no relative path, just return the absolute path
-			if (CommonDirectoryLength == -1)
+			if (commonDirectoryLength == -1)
 			{
-				Result.Append(FullName);
+				result.Append(FullName);
 				return;
 			}
 
 			// Append all the '..' separators to get back to the common directory, then the rest of the string to reach the target item
-			for (int Idx = CommonDirectoryLength + 1; Idx < Directory.FullName.Length; Idx++)
+			for (int idx = commonDirectoryLength + 1; idx < directory.FullName.Length; idx++)
 			{
 				// Move up a directory
-				if (Result.Length != 0)
+				if (result.Length != 0)
 				{
-					Result.Append(Path.DirectorySeparatorChar);
+					result.Append(Path.DirectorySeparatorChar);
 				}
-				Result.Append("..");
+				result.Append("..");
 
 				// Scan to the next directory separator
-				while (Idx < Directory.FullName.Length && Directory.FullName[Idx] != Path.DirectorySeparatorChar)
+				while (idx < directory.FullName.Length && directory.FullName[idx] != Path.DirectorySeparatorChar)
 				{
-					Idx++;
+					idx++;
 				}
 			}
 
-			if (CommonDirectoryLength + 1 < FullName.Length)
+			if (commonDirectoryLength + 1 < FullName.Length)
 			{
-				if (Result.Length != 0)
+				if (result.Length != 0)
 				{
-					Result.Append(Path.DirectorySeparatorChar);
+					result.Append(Path.DirectorySeparatorChar);
 				}
-				Result.Append(FullName, CommonDirectoryLength + 1, FullName.Length - CommonDirectoryLength - 1);
+				result.Append(FullName, commonDirectoryLength + 1, FullName.Length - commonDirectoryLength - 1);
 			}
 		}
 
diff --git a/Engine/Source/Programs/Shared/EpicGames.Core/FileUtils.cs b/Engine/Source/Programs/Shared/EpicGames.Core/FileUtils.cs
--- a/Engine/Source/Programs/Shared/EpicGames.Core/FileUtils.cs
+++ b/Engine/Source/Programs/Shared/EpicGames.Core/FileUtils.cs
@@ -1,16 +1,13 @@
 // Copyright Epic Games, Inc. All Rights Reserved.
 
-using Microsoft.Win32.SafeHandles;
 using System;
 using System.Collections.Generic;
 using System.ComponentModel;
-using System.Diagnostics;
 using System.IO;
-using System.Linq;
 using System.Runtime.InteropServices;
 using System.Runtime.InteropServices.ComTypes;
 using System.Text;
-using System.Threading.Tasks;
+using Microsoft.Win32.SafeHandles;
 
 namespace EpicGames.Core
 {
@@ -22,9 +19,9 @@ namespace EpicGames.Core
 		/// <summary>
 		/// Constructor
 		/// </summary>
-		/// <param name="Inner">Inner exception</param>
-		/// <param name="Message">Message to display</param>
-		public WrappedFileOrDirectoryException(Exception Inner, string Message) : base(Message)
+		/// <param name="inner">Inner exception</param>
+		/// <param name="message">Message to display</param>
+		public WrappedFileOrDirectoryException(Exception inner, string message) : base(message, inner)
 		{
 		}
 
@@ -56,106 +53,106 @@ namespace EpicGames.Core
 		/// <summary>
 		/// Read all text for a file
 		/// </summary>
-		/// <param name="File"></param>
+		/// <param name="file"></param>
 		/// <returns></returns>
-		public static string ReadAllText(FileReference File)
+		public static string ReadAllText(FileReference file)
 		{
 			try
 			{
-				return FileReference.ReadAllText(File);
+				return FileReference.ReadAllText(file);
 			}
-			catch (DirectoryNotFoundException Ex)
+			catch (DirectoryNotFoundException ex)
 			{
-				throw new WrappedFileOrDirectoryException(Ex, $"Unable to read file '{File}'. The directory does not exist.");
+				throw new WrappedFileOrDirectoryException(ex, $"Unable to read file '{file}'. The directory does not exist.");
 			}
-			catch (FileNotFoundException Ex)
+			catch (FileNotFoundException ex)
 			{
-				throw new WrappedFileOrDirectoryException(Ex, $"Unable to read file '{File}'. The file does not exist.");
+				throw new WrappedFileOrDirectoryException(ex, $"Unable to read file '{file}'. The file does not exist.");
 			}
-			catch (Exception Ex)
+			catch (Exception ex)
 			{
-				throw new WrappedFileOrDirectoryException(Ex, $"Unable to read file '{File}'");
+				throw new WrappedFileOrDirectoryException(ex, $"Unable to read file '{file}'");
 			}
 		}
 
 		/// <summary>
 		/// Finds the on-disk case of a a file
 		/// </summary>
-		/// <param name="Info">FileInfo instance describing the file</param>
+		/// <param name="info">FileInfo instance describing the file</param>
 		/// <returns>New FileInfo instance that represents the file with the correct case</returns>
-		public static FileInfo FindCorrectCase(FileInfo Info)
+		public static FileInfo FindCorrectCase(FileInfo info)
 		{
-			DirectoryInfo ParentInfo = DirectoryUtils.FindCorrectCase(Info.Directory);
-			if (Info.Exists)
+			DirectoryInfo parentInfo = DirectoryUtils.FindCorrectCase(info.Directory);
+			if (info.Exists)
 			{
-				foreach (FileInfo ChildInfo in ParentInfo.EnumerateFiles())
+				foreach (FileInfo childInfo in parentInfo.EnumerateFiles())
 				{
-					if (String.Equals(ChildInfo.Name, Info.Name, FileReference.Comparison))
+					if (String.Equals(childInfo.Name, info.Name, FileReference.Comparison))
 					{
-						return ChildInfo;
+						return childInfo;
 					}
 				}
 			}
-			return new FileInfo(Path.Combine(ParentInfo.FullName, Info.Name));
+			return new FileInfo(Path.Combine(parentInfo.FullName, info.Name));
 		}
 
 		/// <summary>
 		/// Creates a directory tree, with all intermediate branches
 		/// </summary>
-		/// <param name="Directory">The directory to create</param>
-		public static void CreateDirectoryTree(DirectoryReference Directory)
+		/// <param name="directory">The directory to create</param>
+		public static void CreateDirectoryTree(DirectoryReference directory)
 		{
-			if (!DirectoryReference.Exists(Directory))
+			if (!DirectoryReference.Exists(directory))
 			{
-				DirectoryReference? ParentDirectory = Directory.ParentDirectory;
-				if (ParentDirectory != null)
+				DirectoryReference? parentDirectory = directory.ParentDirectory;
+				if (parentDirectory != null)
 				{
-					CreateDirectoryTree(ParentDirectory);
+					CreateDirectoryTree(parentDirectory);
 				}
-				DirectoryReference.CreateDirectory(Directory);
+				DirectoryReference.CreateDirectory(directory);
 			}
 		}
 
 		/// <summary>
 		/// Deletes a file, whether it's read-only or not
 		/// </summary>
-		/// <param name="FileName">Name of the file to delete</param>
-		public static void ForceDeleteFile(string FileName)
+		/// <param name="fileName">Name of the file to delete</param>
+		public static void ForceDeleteFile(string fileName)
 		{
 			if (RuntimePlatform.IsWindows)
 			{
-				ForceDeleteFileWin32(FileName);
+				ForceDeleteFileWin32(fileName);
 			}
 			else
 			{
-				ForceDeleteFile(new FileInfo(FileName));
+				ForceDeleteFile(new FileInfo(fileName));
 			}
 		}
 
 		/// <summary>
 		/// Deletes a file, whether it's read-only or not
 		/// </summary>
-		/// <param name="File">The file to delete</param>
-		public static void ForceDeleteFile(FileInfo File)
+		/// <param name="file">The file to delete</param>
+		public static void ForceDeleteFile(FileInfo file)
 		{
 			if (RuntimePlatform.IsWindows)
 			{
-				ForceDeleteFileWin32(File.FullName);
-				File.Refresh();
+				ForceDeleteFileWin32(file.FullName);
+				file.Refresh();
 			}
 			else
 			{
 				try
 				{
-					if (File.Exists)
+					if (file.Exists)
 					{
-						File.Attributes = FileAttributes.Normal;
-						File.Delete();
+						file.Attributes = FileAttributes.Normal;
+						file.Delete();
 					}
 				}
-				catch (Exception Ex)
+				catch (Exception ex)
 				{
-					throw new WrappedFileOrDirectoryException(Ex, String.Format("Unable to delete '{0}': {1}", File.FullName, Ex.Message.TrimEnd()));
+					throw new WrappedFileOrDirectoryException(ex, String.Format("Unable to delete '{0}': {1}", file.FullName, ex.Message.TrimEnd()));
 				}
 			}
 		}
@@ -163,51 +160,51 @@ namespace EpicGames.Core
 		/// <summary>
 		/// Deletes a file, whether it's read-only or not
 		/// </summary>
-		/// <param name="Location">The file to delete</param>
-		public static void ForceDeleteFile(FileReference Location)
+		/// <param name="location">The file to delete</param>
+		public static void ForceDeleteFile(FileReference location)
 		{
 			if (RuntimePlatform.IsWindows)
 			{
-				ForceDeleteFileWin32(Location.FullName);
+				ForceDeleteFileWin32(location.FullName);
 			}
 			else
 			{
-				ForceDeleteFile(new FileInfo(Location.FullName));
+				ForceDeleteFile(new FileInfo(location.FullName));
 			}
 		}
 
 		/// <summary>
 		/// Deletes a directory and all its contents. Attempts to handle directories with long filenames (> 260 chars) on Windows.
 		/// </summary>
-		/// <param name="DirectoryName">Directory to delete</param>
-		public static void ForceDeleteDirectory(string DirectoryName)
+		/// <param name="directoryName">Directory to delete</param>
+		public static void ForceDeleteDirectory(string directoryName)
 		{
 			if (RuntimePlatform.IsWindows)
 			{
-				ForceDeleteLongDirectoryWin32("\\\\?\\" + DirectoryName);
+				ForceDeleteLongDirectoryWin32("\\\\?\\" + directoryName);
 			}
 			else
 			{
-				ForceDeleteDirectory(new DirectoryInfo(DirectoryName));
+				ForceDeleteDirectory(new DirectoryInfo(directoryName));
 			}
 		}
 
 		/// <summary>
 		/// Deletes a directory and all its contents. Attempts to handle directories with long filenames (> 260 chars) on Windows.
 		/// </summary>
-		/// <param name="Directory">Directory to delete</param>
-		public static void ForceDeleteDirectory(DirectoryInfo Directory)
+		/// <param name="directory">Directory to delete</param>
+		public static void ForceDeleteDirectory(DirectoryInfo directory)
 		{
-			if (Directory.Exists)
+			if (directory.Exists)
 			{
 				if (RuntimePlatform.IsWindows)
 				{
-					ForceDeleteLongDirectoryWin32("\\\\?\\" + Directory.FullName);
+					ForceDeleteLongDirectoryWin32("\\\\?\\" + directory.FullName);
 				}
 				else
 				{
-					ForceDeleteDirectoryContents(Directory);
-					ForceDeleteDirectoryInternal(Directory);
+					ForceDeleteDirectoryContents(directory);
+					ForceDeleteDirectoryInternal(directory);
 				}
 			}
 		}
@@ -215,64 +212,64 @@ namespace EpicGames.Core
 		/// <summary>
 		/// Deletes a directory and all its contents. Attempts to handle directories with long filenames (> 260 chars) on Windows.
 		/// </summary>
-		/// <param name="Location">Directory to delete</param>
-		public static void ForceDeleteDirectory(DirectoryReference Location)
+		/// <param name="location">Directory to delete</param>
+		public static void ForceDeleteDirectory(DirectoryReference location)
 		{
-			ForceDeleteDirectory(Location.FullName);
+			ForceDeleteDirectory(location.FullName);
 		}
 
 		/// <summary>
 		/// Helper method to delete a directory and throw a WrappedFileOrDirectoryException on failure.
 		/// </summary>
-		/// <param name="Directory">The directory to delete</param>
-		static void ForceDeleteDirectoryInternal(DirectoryInfo Directory)
+		/// <param name="directory">The directory to delete</param>
+		static void ForceDeleteDirectoryInternal(DirectoryInfo directory)
 		{
 			try
 			{
-				Directory.Delete(true);
+				directory.Delete(true);
 			}
-			catch (Exception Ex)
+			catch (Exception ex)
 			{
-				throw new WrappedFileOrDirectoryException(Ex, String.Format("Unable to delete '{0}': {1}", Directory.FullName, Ex.Message.TrimEnd()));
+				throw new WrappedFileOrDirectoryException(ex, String.Format("Unable to delete '{0}': {1}", directory.FullName, ex.Message.TrimEnd()));
 			}
 		}
 
 		/// <summary>
 		/// Deletes the contents of a directory, without deleting the directory itself. Attempts to handle directories with long filenames (> 260 chars) on Windows.
 		/// </summary>
-		/// <param name="Directory">Directory to delete</param>
-		public static void ForceDeleteDirectoryContents(string Directory)
+		/// <param name="directory">Directory to delete</param>
+		public static void ForceDeleteDirectoryContents(string directory)
 		{
-			ForceDeleteDirectoryContents(new DirectoryInfo(Directory));
+			ForceDeleteDirectoryContents(new DirectoryInfo(directory));
 		}
 
 		/// <summary>
 		/// Deletes the contents of a directory, without deleting the directory itself. Attempts to handle directories with long filenames (> 260 chars) on Windows.
 		/// </summary>
-		/// <param name="Directory">Directory to delete</param>
-		public static void ForceDeleteDirectoryContents(DirectoryInfo Directory)
+		/// <param name="directory">Directory to delete</param>
+		public static void ForceDeleteDirectoryContents(DirectoryInfo directory)
 		{
 			if (RuntimePlatform.IsWindows)
 			{
-				ForceDeleteLongDirectoryContentsWin32(Directory.FullName);
+				ForceDeleteLongDirectoryContentsWin32(directory.FullName);
 			}
 			else
 			{
-				if (Directory.Exists)
+				if (directory.Exists)
 				{
-					foreach (FileInfo File in Directory.EnumerateFiles())
+					foreach (FileInfo file in directory.EnumerateFiles())
 					{
-						ForceDeleteFile(File);
+						ForceDeleteFile(file);
 					}
-					foreach (DirectoryInfo SubDirectory in Directory.EnumerateDirectories())
+					foreach (DirectoryInfo subDirectory in directory.EnumerateDirectories())
 					{
-						if (SubDirectory.Attributes.HasFlag(FileAttributes.ReparsePoint))
+						if (subDirectory.Attributes.HasFlag(FileAttributes.ReparsePoint))
 						{
-							ForceDeleteDirectoryInternal(SubDirectory);
+							ForceDeleteDirectoryInternal(subDirectory);
 						}
 						else
 						{
-							ForceDeleteDirectory(SubDirectory);
+							ForceDeleteDirectory(subDirectory);
 						}
 					}
 				}
@@ -282,20 +279,20 @@ namespace EpicGames.Core
 		/// <summary>
 		/// Deletes the contents of a directory, without deleting the directory itself. Attempts to handle directories with long filenames (> 260 chars) on Windows.
 		/// </summary>
-		/// <param name="Directory">Directory to delete</param>
-		public static void ForceDeleteDirectoryContents(DirectoryReference Directory)
+		/// <param name="directory">Directory to delete</param>
+		public static void ForceDeleteDirectoryContents(DirectoryReference directory)
 		{
-			ForceDeleteDirectoryContents(new DirectoryInfo(Directory.FullName));
+			ForceDeleteDirectoryContents(new DirectoryInfo(directory.FullName));
 		}
 
 		/// <summary>
 		/// Moves a file from one location to another. Creates the destination directory, and removes read-only files in the target location if necessary.
 		/// </summary>
-		/// <param name="SourceFileName">Path to the source file</param>
-		/// <param name="TargetFileName">Path to the target file</param>
-		public static void ForceMoveFile(string SourceFileName, string TargetFileName)
+		/// <param name="sourceFileName">Path to the source file</param>
+		/// <param name="targetFileName">Path to the target file</param>
+		public static void ForceMoveFile(string sourceFileName, string targetFileName)
 		{
-			ForceMoveFile(new FileReference(SourceFileName), new FileReference(TargetFileName));
+			ForceMoveFile(new FileReference(sourceFileName), new FileReference(targetFileName));
 		}
 
 		/// <summary>
@@ -303,22 +300,22 @@ namespace EpicGames.Core
 		/// </summary>
 		/// <param name="SourceFileName">Path to the source file</param>
 		/// <param name="TargetFileName">Path to the target file</param>
-		public static void ForceMoveFile(FileReference SourceLocation, FileReference TargetLocation)
+		public static void ForceMoveFile(FileReference sourceLocation, FileReference targetLocation)
 		{
 			// Try to move the file into place
 			try
 			{
-				FileReference.Move(SourceLocation, TargetLocation);
+				FileReference.Move(sourceLocation, targetLocation);
 				return;
 			}
-			catch (Exception Ex)
+			catch (Exception ex)
 			{
 				// Try to create the target directory
 				try
 				{
-					if (!DirectoryReference.Exists(TargetLocation.Directory))
+					if (!DirectoryReference.Exists(targetLocation.Directory))
 					{
-						CreateDirectoryTree(TargetLocation.Directory);
+						CreateDirectoryTree(targetLocation.Directory);
 					}
 				}
 				catch
@@ -328,69 +325,70 @@ namespace EpicGames.Core
 				// Try to delete an existing file at the target location
 				try
 				{
-					if (FileReference.Exists(TargetLocation))
+					if (FileReference.Exists(targetLocation))
 					{
-						FileReference.SetAttributes(TargetLocation, FileAttributes.Normal);
-						FileReference.Delete(TargetLocation);
-						FileReference.Move(SourceLocation, TargetLocation);
+						FileReference.SetAttributes(targetLocation, FileAttributes.Normal);
+						FileReference.Delete(targetLocation);
+						FileReference.Move(sourceLocation, targetLocation);
 						return;
 					}
 				}
-				catch (Exception DeleteEx)
+				catch (Exception deleteEx)
 				{
-					throw new WrappedFileOrDirectoryException(new AggregateException(Ex, DeleteEx), String.Format("Unable to move {0} to {1} (also tried delete/move)", SourceLocation, TargetLocation));
+					throw new WrappedFileOrDirectoryException(new AggregateException(ex, deleteEx), String.Format("Unable to move {0} to {1} (also tried delete/move)", sourceLocation, targetLocation));
 				}
 
 				// Throw the original exception
-				throw new WrappedFileOrDirectoryException(Ex, String.Format("Unable to move {0} to {1}", SourceLocation, TargetLocation));
+				throw new WrappedFileOrDirectoryException(ex, String.Format("Unable to move {0} to {1}", sourceLocation, targetLocation));
 			}
 		}
 
 		/// <summary>
 		/// Gets the file mode on Mac
 		/// </summary>
-		/// <param name="FileName"></param>
+		/// <param name="fileName"></param>
 		/// <returns></returns>
-		public static int GetFileMode_Mac(string FileName)
+		public static int GetFileMode_Mac(string fileName)
 		{
 			stat64_t stat = new stat64_t();
-			int Result = stat64(FileName, stat);
-			return (Result >= 0)? stat.st_mode : -1;
+			int result = stat64(fileName, stat);
+			return (result >= 0)? stat.st_mode : -1;
 		}
 
 		/// <summary>
 		/// Sets the file mode on Mac
 		/// </summary>
-		/// <param name="FileName"></param>
-		/// <param name="Mode"></param>
-		public static void SetFileMode_Mac(string FileName, ushort Mode)
+		/// <param name="fileName"></param>
+		/// <param name="mode"></param>
+		public static void SetFileMode_Mac(string fileName, ushort mode)
 		{
-			chmod(FileName, Mode);
+			chmod(fileName, mode);
 		}
 
 		/// <summary>
 		/// Gets the file mode on Linux
 		/// </summary>
-		/// <param name="FileName"></param>
+		/// <param name="fileName"></param>
 		/// <returns></returns>
-		public static int GetFileMode_Linux(string FileName)
+		public static int GetFileMode_Linux(string fileName)
 		{
 			stat64_linux_t stat = new stat64_linux_t();
-			int Result = stat64_linux(1, FileName, stat);
-			return (Result >= 0)? (int)stat.st_mode : -1;
+			int result = stat64_linux(1, fileName, stat);
+			return (result >= 0)? (int)stat.st_mode : -1;
 		}
 
 		/// <summary>
 		/// Sets the file mode on Linux
 		/// </summary>
-		/// <param name="FileName"></param>
-		/// <param name="Mode"></param>
-		public static void SetFileMode_Linux(string FileName, ushort Mode)
+		/// <param name="fileName"></param>
+		/// <param name="mode"></param>
+		public static void SetFileMode_Linux(string fileName, ushort mode)
 		{
-			chmod_linux(FileName, Mode);
+			chmod_linux(fileName, mode);
 		}
 
 		#region Win32 Native File Methods
+#pragma warning disable IDE1006 // Naming Styles
 
 		[StructLayout(LayoutKind.Sequential, CharSet = CharSet.Unicode)]
 		struct WIN32_FIND_DATA
@@ -422,7 +420,7 @@ namespace EpicGames.Core
 
 		[DllImport("kernel32.dll", SetLastError = true, CharSet = CharSet.Unicode)]
 		[return: MarshalAs(UnmanagedType.Bool)]
-		static extern bool FindClose(IntPtr FindHandle);
+		static extern bool FindClose(IntPtr findHandle);
 
 		[DllImport("kernel32.dll", SetLastError = true, CharSet = CharSet.Unicode)]
 		[return: MarshalAs(UnmanagedType.Bool)]
@@ -441,36 +439,37 @@ namespace EpicGames.Core
 		const int ERROR_FILE_NOT_FOUND = 2;
 		const int ERROR_PATH_NOT_FOUND = 3;
 		const int ERROR_ACCESS_DENIED = 5;
+#pragma warning restore IDE1006 // Naming Styles
 
-		private static void ForceDeleteLongDirectoryContentsWin32(string DirName)
+		private static void ForceDeleteLongDirectoryContentsWin32(string dirName)
 		{
-			WIN32_FIND_DATA FindData = new WIN32_FIND_DATA();
+			WIN32_FIND_DATA findData = new WIN32_FIND_DATA();
 
-			IntPtr hFind = FindFirstFileW(DirName + "\\*", ref FindData);
+			IntPtr hFind = FindFirstFileW(dirName + "\\*", ref findData);
 			if (hFind != INVALID_HANDLE_VALUE)
 			{
 				try
 				{
 					for (; ; )
 					{
-						string FullName = DirName + "\\" + FindData.cFileName;
-						if ((FindData.dwFileAttributes & FILE_ATTRIBUTE_DIRECTORY) != 0)
+						string fullName = dirName + "\\" + findData.cFileName;
+						if ((findData.dwFileAttributes & FILE_ATTRIBUTE_DIRECTORY) != 0)
 						{
-							if (FindData.cFileName != "." && FindData.cFileName != "..")
+							if (findData.cFileName != "." && findData.cFileName != "..")
 							{
-								ForceDeleteLongDirectoryWin32(FullName);
+								ForceDeleteLongDirectoryWin32(fullName);
 							}
 						}
 						else
 						{
-							if ((FindData.dwFileAttributes & FILE_ATTRIBUTE_READONLY) != 0)
+							if ((findData.dwFileAttributes & FILE_ATTRIBUTE_READONLY) != 0)
 							{
-								SetFileAttributesW(FullName, FILE_ATTRIBUTE_NORMAL);
+								SetFileAttributesW(fullName, FILE_ATTRIBUTE_NORMAL);
 							}
-							ForceDeleteFileWin32(FullName);
+							ForceDeleteFileWin32(fullName);
 						}
 
-						if (!FindNextFileW(hFind, ref FindData))
+						if (!FindNextFileW(hFind, ref findData))
 						{
 							break;
 						}
@@ -483,16 +482,16 @@ namespace EpicGames.Core
 			}
 		}
 
-		private static void ForceDeleteLongDirectoryWin32(string DirName)
+		private static void ForceDeleteLongDirectoryWin32(string dirName)
 		{
-			ForceDeleteLongDirectoryContentsWin32(DirName);
+			ForceDeleteLongDirectoryContentsWin32(dirName);
 
-			if (!RemoveDirectory(DirName))
+			if (!RemoveDirectory(dirName))
 			{
-				int ErrorCode = Marshal.GetLastWin32Error();
-				if (ErrorCode != ERROR_FILE_NOT_FOUND && ErrorCode != ERROR_PATH_NOT_FOUND)
+				int errorCode = Marshal.GetLastWin32Error();
+				if (errorCode != ERROR_FILE_NOT_FOUND && errorCode != ERROR_PATH_NOT_FOUND)
 				{
-					throw new WrappedFileOrDirectoryException(new Win32Exception(ErrorCode), "Unable to delete " + DirName);
+					throw new WrappedFileOrDirectoryException(new Win32Exception(errorCode), "Unable to delete " + dirName);
 				}
 			}
 		}
@@ -500,32 +499,32 @@ namespace EpicGames.Core
 		[DllImport("kernel32.dll", CharSet = CharSet.Unicode, SetLastError = true)]
 		internal static extern int GetFileAttributesW(string lpFileName);
 
-		public static void ForceDeleteFileWin32(string FileName)
+		public static void ForceDeleteFileWin32(string fileName)
 		{
 			// Try to delete the file normally
-			if (DeleteFileW(FileName))
+			if (DeleteFileW(fileName))
 			{
 				return;
 			}
 
 			// Capture the exception for failing to delete the file
-			Win32Exception Ex = new Win32Exception();
+			Win32Exception ex = new Win32Exception();
 
 			// Check the file exists and is not readonly
-			int Attributes = GetFileAttributesW(FileName);
-			if (Attributes == -1)
+			int attributes = GetFileAttributesW(fileName);
+			if (attributes == -1)
 			{
-				int ErrorCode = Marshal.GetLastWin32Error();
-				if (ErrorCode == ERROR_PATH_NOT_FOUND || ErrorCode == ERROR_FILE_NOT_FOUND)
+				int errorCode = Marshal.GetLastWin32Error();
+				if (errorCode == ERROR_PATH_NOT_FOUND || errorCode == ERROR_FILE_NOT_FOUND)
 				{
 					return;
 				}
 			}
 			else
 			{
-				if ((Attributes & (int)FileAttributes.ReadOnly) != 0)
+				if ((attributes & (int)FileAttributes.ReadOnly) != 0)
 				{
-					if (SetFileAttributesW(FileName, (int)FileAttributes.Normal) && DeleteFileW(FileName))
+					if (SetFileAttributesW(fileName, (int)FileAttributes.Normal) && DeleteFileW(fileName))
 					{
 						return;
 					}
@@ -533,35 +532,35 @@ namespace EpicGames.Core
 			}
 
 			// Get a useful error message about why the delete failed
-			StringBuilder Message = new StringBuilder($"Unable to delete {FileName} - {Ex.Message}");
-			if (Ex.NativeErrorCode == ERROR_ACCESS_DENIED)
+			StringBuilder message = new StringBuilder($"Unable to delete {fileName} - {ex.Message}");
+			if (ex.NativeErrorCode == ERROR_ACCESS_DENIED)
 			{
-				List<FileLockInfo_Win32>? LockInfoList;
+				List<FileLockInfoWin32>? lockInfoList;
 				try
 				{
-					LockInfoList = GetFileLockInfo_Win32(FileName);
+					lockInfoList = GetFileLockInfo_Win32(fileName);
 				}
 				catch
 				{
-					LockInfoList = null;
+					lockInfoList = null;
 				}
 
-				if (LockInfoList != null && LockInfoList.Count > 0)
+				if (lockInfoList != null && lockInfoList.Count > 0)
 				{
-					Message.Append("\nProcesses with open handles to file:");
-					foreach (FileLockInfo_Win32 LockInfo in LockInfoList)
+					message.Append("\nProcesses with open handles to file:");
+					foreach (FileLockInfoWin32 lockInfo in lockInfoList)
 					{
-						Message.Append($"\n  {LockInfo}");
+						message.Append($"\n  {lockInfo}");
 					}
 				}
 			}
-			throw new WrappedFileOrDirectoryException(Ex, Message.ToString());
+			throw new WrappedFileOrDirectoryException(ex, message.ToString());
 		}
 
 		#endregion
 
 		#region Win32 Restart Manager API
-
+#pragma warning disable IDE1006
 		[StructLayout(LayoutKind.Sequential)]
 		struct RM_UNIQUE_PROCESS
 		{
@@ -638,45 +637,46 @@ namespace EpicGames.Core
 
 		[DllImport("kernel32.dll", SetLastError = true)]
 		static extern bool QueryFullProcessImageName([In]SafeProcessHandle hProcess, [In]int dwFlags, [Out]StringBuilder lpExeName, ref int lpdwSize);
+#pragma warning restore IDE1006
 
 		/// <summary>
 		/// Information about a locked file
 		/// </summary>
-		public class FileLockInfo_Win32
+		public class FileLockInfoWin32
 		{
 			/// <summary>
 			/// Process id
 			/// </summary>
-			public int ProcessId;
+			public int ProcessId { get; }
 
 			/// <summary>
 			/// Path to the process holding the lock
 			/// </summary>
-			public string? FileName;
+			public string? FileName { get; }
 
 			/// <summary>
 			/// Name of the application
 			/// </summary>
-			public string AppName;
+			public string AppName { get; }
 
 			/// <summary>
 			/// Time at which the process started
 			/// </summary>
-			public DateTime StartTime;
+			public DateTime StartTime { get; }
 
 			/// <summary>
 			/// Constructor
 			/// </summary>
-			/// <param name="ProcessId">Process id</param>
-			/// <param name="FileName">Path to the process holding the lock</param>
-			/// <param name="AppName">Name of the application</param>
-			/// <param name="StartTime">Time at which the process started</param>
-			public FileLockInfo_Win32(int ProcessId, string? FileName, string AppName, DateTime StartTime)
+			/// <param name="processId">Process id</param>
+			/// <param name="fileName">Path to the process holding the lock</param>
+			/// <param name="appName">Name of the application</param>
+			/// <param name="startTime">Time at which the process started</param>
+			public FileLockInfoWin32(int processId, string? fileName, string appName, DateTime startTime)
 			{
-				this.ProcessId = ProcessId;
-				this.FileName = FileName;
-				this.AppName = AppName;
-				this.StartTime = StartTime;
+				ProcessId = processId;
+				FileName = fileName;
+				AppName = appName;
+				StartTime = startTime;
 			}
 
 			/// <inheritdoc/>
@@ -689,83 +689,84 @@ namespace EpicGames.Core
 		/// <summary>
 		/// Gets a list of processes that have a handle to the given file open
 		/// </summary>
-		/// <param name="FileName">File to check</param>
+		/// <param name="fileName">File to check</param>
 		/// <returns>List of processes with a lock open</returns>
-		public static List<FileLockInfo_Win32> GetFileLockInfo_Win32(string FileName)
+		public static List<FileLockInfoWin32> GetFileLockInfo_Win32(string fileName)
 		{
-			uint SessionHandle = 0;
+			uint sessionHandle = 0;
 			try
 			{
-				StringBuilder SessionKey = new StringBuilder(CCH_RM_SESSION_KEY + 1);
+				StringBuilder sessionKey = new StringBuilder(CCH_RM_SESSION_KEY + 1);
 
-				int Result = RmStartSession(out SessionHandle, 0, SessionKey);
-				if (Result != 0)
+				int result = RmStartSession(out sessionHandle, 0, sessionKey);
+				if (result != 0)
 				{
-					throw new Win32Exception(Result, "Unable to open restart manager session");
+					throw new Win32Exception(result, "Unable to open restart manager session");
 				}
 
-				Result = RmRegisterResources(SessionHandle, 1, new string[] { FileName }, 0, null, 0, null);
-				if (Result != 0)
+				result = RmRegisterResources(sessionHandle, 1, new string[] { fileName }, 0, null, 0, null);
+				if (result != 0)
 				{
-					throw new Win32Exception(Result, "Unable to register resource with restart manager");
+					throw new Win32Exception(result, "Unable to register resource with restart manager");
 				}
 
 				uint nProcInfoNeeded = 0;
 				uint nProcInfo = 10;
-				uint Reason = 0;
-				RM_PROCESS_INFO[] ProcessInfoArray = new RM_PROCESS_INFO[nProcInfo];
-				Result = RmGetList(SessionHandle, out nProcInfoNeeded, ref nProcInfo, ProcessInfoArray, ref Reason);
-				if (Result != 0)
+				uint reason = 0;
+				RM_PROCESS_INFO[] processInfoArray = new RM_PROCESS_INFO[nProcInfo];
+				result = RmGetList(sessionHandle, out nProcInfoNeeded, ref nProcInfo, processInfoArray, ref reason);
+				if (result != 0)
 				{
-					throw new Win32Exception(Result, "Unable to query processes with file handle open");
+					throw new Win32Exception(result, "Unable to query processes with file handle open");
 				}
 
-				List<FileLockInfo_Win32> FileLocks = new List<FileLockInfo_Win32>();
-				for (int Idx = 0; Idx < nProcInfo; Idx++)
+				List<FileLockInfoWin32> fileLocks = new List<FileLockInfoWin32>();
+				for (int idx = 0; idx < nProcInfo; idx++)
 				{
-					RM_PROCESS_INFO ProcessInfo = ProcessInfoArray[Idx];
-					long StartTimeTicks = FileTimeToTicks(ProcessInfo.Process.ProcessStartTime);
+					RM_PROCESS_INFO processInfo = processInfoArray[idx];
+					long startTimeTicks = FileTimeToTicks(processInfo.Process.ProcessStartTime);
 
-					string? ImageName = null;
-					using (SafeProcessHandle hProcess = OpenProcess(PROCESS_QUERY_LIMITED_INFORMATION, false, ProcessInfo.Process.dwProcessId))
+					string? imageName = null;
+					using (SafeProcessHandle hProcess = OpenProcess(PROCESS_QUERY_LIMITED_INFORMATION, false, processInfo.Process.dwProcessId))
 					{
 						if (hProcess != null)
 						{
-							FILETIME CreateTime, ExitTime, KernelTime, UserTime;
-							if (GetProcessTimes(hProcess, out CreateTime, out ExitTime, out KernelTime, out UserTime) && FileTimeToTicks(CreateTime) == StartTimeTicks)
+							FILETIME createTime, exitTime, kernelTime, userTime;
+							if (GetProcessTimes(hProcess, out createTime, out exitTime, out kernelTime, out userTime) && FileTimeToTicks(createTime) == startTimeTicks)
 							{
-								int Capacity = 260;
-								StringBuilder ImageNameBuilder = new StringBuilder(Capacity);
-								if (QueryFullProcessImageName(hProcess, 0, ImageNameBuilder, ref Capacity))
+								int capacity = 260;
+								StringBuilder imageNameBuilder = new StringBuilder(capacity);
+								if (QueryFullProcessImageName(hProcess, 0, imageNameBuilder, ref capacity))
 								{
-									ImageName = ImageNameBuilder.ToString(0, Capacity);
+									imageName = imageNameBuilder.ToString(0, capacity);
 								}
 							}
 						}
 					}
 
-					FileLocks.Add(new FileLockInfo_Win32(ProcessInfo.Process.dwProcessId, ImageName, ProcessInfo.strAppName, DateTime.FromFileTime(StartTimeTicks)));
+					fileLocks.Add(new FileLockInfoWin32(processInfo.Process.dwProcessId, imageName, processInfo.strAppName, DateTime.FromFileTime(startTimeTicks)));
 				}
-				return FileLocks;
+				return fileLocks;
 			}
 			finally
 			{
-				if (SessionHandle != 0)
+				if (sessionHandle != 0)
 				{
-					RmEndSession(SessionHandle);
+					RmEndSession(sessionHandle);
 				}
 			}
 		}
 
-		private static long FileTimeToTicks(FILETIME FileTime)
+		private static long FileTimeToTicks(FILETIME fileTime)
 		{
-			return (long)(uint)FileTime.dwLowDateTime | ((long)(uint)FileTime.dwHighDateTime << 32);
+			return (long)(uint)fileTime.dwLowDateTime | ((long)(uint)fileTime.dwHighDateTime << 32);
 		}
 
 		#endregion
 
 		#region Mac Native File Methods
 #pragma warning disable CS0649
+#pragma warning disable IDE1006
 		struct timespec_t
 		{
 			public ulong tv_sec;
@@ -802,11 +803,13 @@ namespace EpicGames.Core
 		[DllImport("libSystem.dylib")]
 		static extern int chmod(string path, ushort mode);
 
+#pragma warning restore IDE1006
 #pragma warning restore CS0649
 		#endregion
 
 		#region Linux Native File Methods
 #pragma warning disable CS0649
+#pragma warning disable IDE1006
 
 		[StructLayout(LayoutKind.Sequential)]
 		class stat64_linux_t
@@ -837,6 +840,7 @@ namespace EpicGames.Core
 		[DllImport("libc", EntryPoint="chmod")]
 		static extern int chmod_linux(string path, ushort mode);
 
+#pragma warning restore IDE1006
 #pragma warning restore CS0649
 		#endregion
 	}
diff --git a/Engine/Source/Programs/Shared/EpicGames.Core/ForwardingLogger.cs b/Engine/Source/Programs/Shared/EpicGames.Core/ForwardingLogger.cs
--- a/Engine/Source/Programs/Shared/EpicGames.Core/ForwardingLogger.cs
+++ b/Engine/Source/Programs/Shared/EpicGames.Core/ForwardingLogger.cs
@@ -1,10 +1,8 @@
 // Copyright Epic Games, Inc. All Rights Reserved.
 
-using Microsoft.Extensions.Logging;
 using System;
 using System.Collections.Generic;
-using System.Linq;
-using System.Text;
+using Microsoft.Extensions.Logging;
 
 namespace EpicGames.Core
 {
@@ -21,9 +19,9 @@ namespace EpicGames.Core
 			/// <inheritdoc/>
 			public void Dispose()
 			{
-				foreach(IDisposable Child in this)
+				foreach(IDisposable child in this)
 				{
-					Child.Dispose();
+					child.Dispose();
 				}
 			}
 		}
@@ -39,50 +37,50 @@ namespace EpicGames.Core
 		/// <summary>
 		/// Constructor
 		/// </summary>
-		/// <param name="Loggers">Array of loggers to forward to</param>
-		public ForwardingLogger(params ILogger[] Loggers)
+		/// <param name="loggers">Array of loggers to forward to</param>
+		public ForwardingLogger(params ILogger[] loggers)
 		{
-			this.Loggers = Loggers;
+			Loggers = loggers;
 		}
 
 		/// <inheritdoc/>
-		public IDisposable? BeginScope<TState>(TState State)
+		public IDisposable? BeginScope<TState>(TState state)
 		{
-			IDisposable? Disposable = null;
+			IDisposable? disposable = null;
 
-			DisposableList? Scope = null;
-			foreach (ILogger Logger in Loggers)
+			DisposableList? scope = null;
+			foreach (ILogger logger in Loggers)
 			{
-				IDisposable? Next = Logger.BeginScope(State);
-				if(Next != null)
+				IDisposable? next = logger.BeginScope(state);
+				if(next != null)
 				{
-					if (Scope != null)
+					if (scope != null)
 					{
-						Scope.Add(Next);
+						scope.Add(next);
 					}
-					else if (Disposable != null)
+					else if (disposable != null)
 					{
-						Scope = new DisposableList();
-						Scope.Add(Disposable);
-						Scope.Add(Next);
-						Disposable = Scope;
+						scope = new DisposableList();
+						scope.Add(disposable);
+						scope.Add(next);
+						disposable = scope;
 					}
 					else
 					{
-						Disposable = Next;
+						disposable = next;
 					}
 				}
 			}
 
-			return Disposable;
+			return disposable;
 		}
 
 		/// <inheritdoc/>
-		public bool IsEnabled(LogLevel LogLevel)
+		public bool IsEnabled(LogLevel logLevel)
 		{
-			foreach(ILogger Logger in Loggers)
+			foreach(ILogger logger in Loggers)
 			{
-				if(Logger.IsEnabled(LogLevel))
+				if(logger.IsEnabled(logLevel))
 				{
 					return true;
 				}
@@ -91,11 +89,11 @@ namespace EpicGames.Core
 		}
 
 		/// <inheritdoc/>
-		public void Log<TState>(LogLevel LogLevel, EventId EventId, TState State, Exception Exception, Func<TState, Exception, string> Formatter)
+		public void Log<TState>(LogLevel logLevel, EventId eventId, TState state, Exception exception, Func<TState, Exception, string> formatter)
 		{
-			foreach (ILogger Logger in Loggers)
+			foreach (ILogger logger in Loggers)
 			{
-				Logger.Log(LogLevel, EventId, State, Exception, Formatter);
+				logger.Log(logLevel, eventId, state, exception, formatter);
 			}
 		}
 	}
diff --git a/Engine/Source/Programs/Shared/EpicGames.Core/HelpUtils.cs b/Engine/Source/Programs/Shared/EpicGames.Core/HelpUtils.cs
--- a/Engine/Source/Programs/Shared/EpicGames.Core/HelpUtils.cs
+++ b/Engine/Source/Programs/Shared/EpicGames.Core/HelpUtils.cs
@@ -1,12 +1,11 @@
 // Copyright Epic Games, Inc. All Rights Reserved.
 
-using Microsoft.Extensions.Logging;
 using System;
 using System.Collections.Generic;
 using System.ComponentModel;
 using System.Linq;
 using System.Text;
-using System.Threading.Tasks;
+using Microsoft.Extensions.Logging;
 
 namespace EpicGames.Core
 {
@@ -23,38 +22,38 @@ namespace EpicGames.Core
 		/// <summary>
 		/// Prints help for the given object type
 		/// </summary>
-		/// <param name="Type">Type to print help for</param>
-		public static void PrintHelp(string Title, Type Type)
+		/// <param name="type">Type to print help for</param>
+		public static void PrintHelp(string title, Type type)
 		{
-			PrintHelp(Title, GetDescription(Type), CommandLineArguments.GetParameters(Type));
+			PrintHelp(title, GetDescription(type), CommandLineArguments.GetParameters(type));
 		}
 
 		/// <summary>
 		/// Prints help for a command
 		/// </summary>
-		/// <param name="Title">Title for the help text</param>
-		/// <param name="Description">Description for the command</param>
-		/// <param name="Parameters">List of parameters</param>
-		public static void PrintHelp(string? Title, string? Description, List<KeyValuePair<string, string>> Parameters)
+		/// <param name="title">Title for the help text</param>
+		/// <param name="description">Description for the command</param>
+		/// <param name="parameters">List of parameters</param>
+		public static void PrintHelp(string? title, string? description, List<KeyValuePair<string, string>> parameters)
 		{
 			bool bFirstLine = true;
-			if (!String.IsNullOrEmpty(Title))
+			if (!String.IsNullOrEmpty(title))
 			{
-				PrintParagraph(Title);
+				PrintParagraph(title);
 				bFirstLine = false;
 			}
 
-			if (!String.IsNullOrEmpty(Description))
+			if (!String.IsNullOrEmpty(description))
 			{
 				if (!bFirstLine)
 				{
 					Console.WriteLine("");
 				}
-				PrintParagraph(Description);
+				PrintParagraph(description);
 				bFirstLine = false;
 			}
 
-			if (Parameters.Count > 0)
+			if (parameters.Count > 0)
 			{
 				if (!bFirstLine)
 				{
@@ -62,87 +61,87 @@ namespace EpicGames.Core
 				}
 
 				Console.WriteLine("Parameters:");
-				PrintTable(Parameters, 4, 24);
+				PrintTable(parameters, 4, 24);
 			}
 		}
 
 		/// <summary>
 		/// Gets the description from a type
 		/// </summary>
-		/// <param name="Type">The type to get a description for</param>
+		/// <param name="type">The type to get a description for</param>
 		/// <returns>The description text</returns>
-		public static string GetDescription(Type Type)
+		public static string GetDescription(Type type)
 		{
-			StringBuilder DescriptionText = new StringBuilder();
-			foreach (DescriptionAttribute Attribute in Type.GetCustomAttributes(typeof(DescriptionAttribute), false))
+			StringBuilder descriptionText = new StringBuilder();
+			foreach (DescriptionAttribute attribute in type.GetCustomAttributes(typeof(DescriptionAttribute), false))
 			{
-				if (DescriptionText.Length > 0)
+				if (descriptionText.Length > 0)
 				{
-					DescriptionText.AppendLine();
+					descriptionText.AppendLine();
 				}
-				DescriptionText.AppendLine(Attribute.Description);
+				descriptionText.AppendLine(attribute.Description);
 			}
-			return DescriptionText.ToString();
+			return descriptionText.ToString();
 		}
 
 		/// <summary>
 		/// Prints a paragraph of text using word wrapping
 		/// </summary>
-		/// <param name="Text">Text to print</param>
+		/// <param name="text">Text to print</param>
 		/// <param name="Logger">Logger implementation to write to</param>
-		public static void PrintParagraph(string Text)
+		public static void PrintParagraph(string text)
 		{
-			PrintParagraph(Text, WindowWidth - 1);
+			PrintParagraph(text, WindowWidth - 1);
 		}
 
 		/// <summary>
 		/// Prints a paragraph of text using word wrapping
 		/// </summary>
-		/// <param name="Text">Text to print</param>
-		/// <param name="MaxWidth">Maximum width for each line</param>
-		public static void PrintParagraph(string Text, int MaxWidth)
+		/// <param name="text">Text to print</param>
+		/// <param name="maxWidth">Maximum width for each line</param>
+		public static void PrintParagraph(string text, int maxWidth)
 		{
-			IEnumerable<string> Lines = StringUtils.WordWrap(Text, MaxWidth);
-			foreach (string Line in Lines)
+			IEnumerable<string> lines = StringUtils.WordWrap(text, maxWidth);
+			foreach (string line in lines)
 			{
-				Console.WriteLine(Line);
+				Console.WriteLine(line);
 			}
 		}
 
 		/// <summary>
 		/// Prints an argument list to the console
 		/// </summary>
-		/// <param name="Items">List of parameters arranged as "-ParamName Param Description"</param>
-		/// <param name="Indent">Indent from the left hand side</param>
-		/// <param name="MinFirstColumnWidth">The minimum padding from the start of the param name to the start of the description (resizes with larger param names)</param>
+		/// <param name="items">List of parameters arranged as "-ParamName Param Description"</param>
+		/// <param name="indent">Indent from the left hand side</param>
+		/// <param name="minFirstColumnWidth">The minimum padding from the start of the param name to the start of the description (resizes with larger param names)</param>
 		/// <returns></returns>
-		public static void PrintTable(List<KeyValuePair<string, string>> Items, int Indent, int MinFirstColumnWidth)
+		public static void PrintTable(List<KeyValuePair<string, string>> items, int indent, int minFirstColumnWidth)
 		{
-			List<string> Lines = new List<string>();
-			FormatTable(Items, Indent, MinFirstColumnWidth, WindowWidth - 1, Lines);
+			List<string> lines = new List<string>();
+			FormatTable(items, indent, minFirstColumnWidth, WindowWidth - 1, lines);
 
-			foreach (string Line in Lines)
+			foreach (string line in lines)
 			{
-				Console.WriteLine(Line);
+				Console.WriteLine(line);
 			}
 		}
 
 		/// <summary>
 		/// Prints a table of items to a logging device
 		/// </summary>
-		/// <param name="Items"></param>
-		/// <param name="Indent"></param>
-		/// <param name="MinFirstColumnWidth"></param>
-		/// <param name="MaxWidth"></param>
-		/// <param name="Logger"></param>
-		public static void PrintTable(List<KeyValuePair<string, string>> Items, int Indent, int MinFirstColumnWidth, int MaxWidth, ILogger Logger)
+		/// <param name="items"></param>
+		/// <param name="indent"></param>
+		/// <param name="minFirstColumnWidth"></param>
+		/// <param name="maxWidth"></param>
+		/// <param name="logger"></param>
+		public static void PrintTable(List<KeyValuePair<string, string>> items, int indent, int minFirstColumnWidth, int maxWidth, ILogger logger)
 		{
-			List<string> Lines = new List<string>();
-			FormatTable(Items, Indent, MinFirstColumnWidth, MaxWidth, Lines);
+			List<string> lines = new List<string>();
+			FormatTable(items, indent, minFirstColumnWidth, maxWidth, lines);
 
-			foreach (string Line in Lines)
+			foreach (string line in lines)
 			{
-				Logger.LogInformation("{Line)", Line);
+				logger.LogInformation("{Line)", line);
 			}
 		}
 
@@ -155,37 +154,37 @@ namespace EpicGames.Core
 		///
 		///     -Param3      Param3 Description continues as before. 
 		/// </summary>
-		/// <param name="Items">List of parameters arranged as "-ParamName Param Description"</param>
-		/// <param name="Indent">Indent from the left hand side</param>
-		/// <param name="MinFirstColumnWidth">The minimum padding from the start of the param name to the start of the description (resizes with larger param names)</param>
+		/// <param name="items">List of parameters arranged as "-ParamName Param Description"</param>
+		/// <param name="indent">Indent from the left hand side</param>
+		/// <param name="minFirstColumnWidth">The minimum padding from the start of the param name to the start of the description (resizes with larger param names)</param>
 		/// <returns>Sequence of formatted lines in the table</returns>
-		public static void FormatTable(IReadOnlyList<KeyValuePair<string, string>> Items, int Indent, int MinFirstColumnWidth, int MaxWidth, List<string> Lines)
+		public static void FormatTable(IReadOnlyList<KeyValuePair<string, string>> items, int indent, int minFirstColumnWidth, int maxWidth, List<string> lines)
 		{
-			if(Items.Count > 0)
+			if(items.Count > 0)
 			{
 				// string used to intent the param
-				string IndentString = new string(' ', Indent);
+				string indentString = new string(' ', indent);
 
 				// default the padding value
-				int RightPadding = Math.Max(MinFirstColumnWidth, Items.Max(x => x.Key.Length + 2));
+				int rightPadding = Math.Max(minFirstColumnWidth, items.Max(x => x.Key.Length + 2));
 
 				// Build the formatted params
-				foreach(KeyValuePair<string, string> Item in Items)
+				foreach(KeyValuePair<string, string> item in items)
 				{
 					// build the param first, including intend and padding on the rights size
-					string ParamString = IndentString + Item.Key.PadRight(RightPadding);
+					string paramString = indentString + item.Key.PadRight(rightPadding);
 
 					// Build the description line by line, adding the same amount of intending each time. 
-					IEnumerable<string> DescriptionLines = StringUtils.WordWrap(Item.Value, MaxWidth - ParamString.Length);
+					IEnumerable<string> descriptionLines = StringUtils.WordWrap(item.Value, maxWidth - paramString.Length);
 
-					foreach(string DescriptionLine in DescriptionLines)
+					foreach(string descriptionLine in descriptionLines)
 					{
 						// Formatting as following:
 						// <Indent>-param<Right Padding>Description<New line>
-						Lines.Add(ParamString + DescriptionLine);
+						lines.Add(paramString + descriptionLine);
 
 						// we replace the param string on subsequent lines with white space of the same length
-						ParamString = string.Empty.PadRight(IndentString.Length + RightPadding);
+						paramString = string.Empty.PadRight(indentString.Length + rightPadding);
 					}
 				}
 			}
diff --git a/Engine/Source/Programs/Shared/EpicGames.Core/HttpClientExtensions.cs b/Engine/Source/Programs/Shared/EpicGames.Core/HttpClientExtensions.cs
--- a/Engine/Source/Programs/Shared/EpicGames.Core/HttpClientExtensions.cs
+++ b/Engine/Source/Programs/Shared/EpicGames.Core/HttpClientExtensions.cs
@@ -1,16 +1,10 @@
 // Copyright Epic Games, Inc. All Rights Reserved.
 
-using System;
-using System.Collections.Generic;
-using System.IO;
-using System.Net;
 using System.Net.Http;
 using System.Text;
 using System.Text.Json;
-using System.Text.Json.Serialization;
 using System.Threading;
 using System.Threading.Tasks;
-using System.Web;
 
 namespace EpicGames.Core
 {
@@ -23,16 +17,16 @@ namespace EpicGames.Core
 		/// Gets a resource from an HTTP endpoint and parses it as a JSON object
 		/// </summary>
 		/// <typeparam name="TResponse">The object type to return</typeparam>
-		/// <param name="Client">The http client instance</param>
-		/// <param name="Url">The url to retrieve</param>
-		/// <param name="CancellationToken">Cancels the request</param>
+		/// <param name="client">The http client instance</param>
+		/// <param name="url">The url to retrieve</param>
+		/// <param name="cancellationToken">Cancels the request</param>
 		/// <returns>New instance of the object</returns>
-		public static async Task<TResponse> GetAsync<TResponse>(this HttpClient Client, string Url, CancellationToken CancellationToken)
+		public static async Task<TResponse> GetAsync<TResponse>(this HttpClient client, string url, CancellationToken cancellationToken)
 		{
-			using (HttpResponseMessage Response = await Client.GetAsync(Url, CancellationToken))
+			using (HttpResponseMessage response = await client.GetAsync(url, cancellationToken))
 			{
-				Response.EnsureSuccessStatusCode();
-				return await ParseJsonContent<TResponse>(Response);
+				response.EnsureSuccessStatusCode();
+				return await ParseJsonContent<TResponse>(response);
 			}
 		}
 
@@ -40,14 +34,14 @@ namespace EpicGames.Core
 		/// Posts an object to an HTTP endpoint as a JSON object
 		/// </summary>
 		/// <typeparam name="TRequest">The object type to post</typeparam>
-		/// <param name="Client">The http client instance</param>
-		/// <param name="Url">The url to post to</param>
-		/// <param name="Object">The object to post</param>
-		/// <param name="CancellationToken">Cancels the request</param>
+		/// <param name="client">The http client instance</param>
+		/// <param name="url">The url to post to</param>
+		/// <param name="request">The object to post</param>
+		/// <param name="cancellationToken">Cancels the request</param>
 		/// <returns>Response message</returns>
-		public static async Task<HttpResponseMessage> PostAsync<TRequest>(this HttpClient Client, string Url, TRequest Object, CancellationToken CancellationToken)
+		public static async Task<HttpResponseMessage> PostAsync<TRequest>(this HttpClient client, string url, TRequest request, CancellationToken cancellationToken)
 		{
-			return await Client.PostAsync(Url, ToJsonContent(Object), CancellationToken);
+			return await client.PostAsync(url, ToJsonContent(request), cancellationToken);
 		}
 
 		/// <summary>
@@ -55,17 +49,17 @@ namespace EpicGames.Core
 		/// </summary>
 		/// <typeparam name="TResponse">The object type to return</typeparam>
 		/// <typeparam name="TRequest">The object type to post</typeparam>
-		/// <param name="Client">The http client instance</param>
-		/// <param name="Url">The url to post to</param>
-		/// <param name="Request">The object to post</param>
-		/// <param name="CancellationToken">Cancels the request</param>
+		/// <param name="client">The http client instance</param>
+		/// <param name="url">The url to post to</param>
+		/// <param name="request">The object to post</param>
+		/// <param name="cancellationToken">Cancels the request</param>
 		/// <returns>The response parsed into the requested type</returns>
-		public static async Task<TResponse> PostAsync<TResponse, TRequest>(this HttpClient Client, string Url, TRequest Request, CancellationToken CancellationToken)
+		public static async Task<TResponse> PostAsync<TResponse, TRequest>(this HttpClient client, string url, TRequest request, CancellationToken cancellationToken)
 		{
-			using (HttpResponseMessage Response = await PostAsync(Client, Url, Request, CancellationToken))
+			using (HttpResponseMessage response = await PostAsync(client, url, request, cancellationToken))
 			{
-				Response.EnsureSuccessStatusCode();
-				return await ParseJsonContent<TResponse>(Response);
+				response.EnsureSuccessStatusCode();
+				return await ParseJsonContent<TResponse>(response);
 			}
 		}
 
@@ -73,14 +67,14 @@ namespace EpicGames.Core
 		/// Puts an object to an HTTP endpoint as a JSON object
 		/// </summary>
 		/// <typeparam name="TRequest">The object type to post</typeparam>
-		/// <param name="Client">The http client instance</param>
-		/// <param name="Url">The url to post to</param>
-		/// <param name="Object">The object to post</param>
-		/// <param name="CancellationToken">Cancels the request</param>
+		/// <param name="client">The http client instance</param>
+		/// <param name="url">The url to post to</param>
+		/// <param name="obj">The object to post</param>
+		/// <param name="cancellationToken">Cancels the request</param>
 		/// <returns>Response message</returns>
-		public static async Task<HttpResponseMessage> PutAsync<TRequest>(this HttpClient Client, string Url, TRequest Object, CancellationToken CancellationToken)
+		public static async Task<HttpResponseMessage> PutAsync<TRequest>(this HttpClient client, string url, TRequest obj, CancellationToken cancellationToken)
 		{
-			return await Client.PutAsync(Url, ToJsonContent(Object), CancellationToken);
+			return await client.PutAsync(url, ToJsonContent(obj), cancellationToken);
 		}
 
 		/// <summary>
@@ -88,17 +82,17 @@ namespace EpicGames.Core
 		/// </summary>
 		/// <typeparam name="TResponse">The object type to return</typeparam>
 		/// <typeparam name="TRequest">The object type to post</typeparam>
-		/// <param name="Client">The http client instance</param>
-		/// <param name="Url">The url to post to</param>
-		/// <param name="Request">The object to post</param>
-		/// <param name="CancellationToken">Cancels the request</param>
+		/// <param name="client">The http client instance</param>
+		/// <param name="url">The url to post to</param>
+		/// <param name="request">The object to post</param>
+		/// <param name="cancellationToken">Cancels the request</param>
 		/// <returns>The response parsed into the requested type</returns>
-		public static async Task<TResponse> PutAsync<TResponse, TRequest>(this HttpClient Client, string Url, TRequest Request, CancellationToken CancellationToken)
+		public static async Task<TResponse> PutAsync<TResponse, TRequest>(this HttpClient client, string url, TRequest request, CancellationToken cancellationToken)
 		{
-			using (HttpResponseMessage Response = await PutAsync(Client, Url, Request, CancellationToken))
+			using (HttpResponseMessage response = await PutAsync(client, url, request, cancellationToken))
 			{
-				Response.EnsureSuccessStatusCode();
-				return await ParseJsonContent<TResponse>(Response);
+				response.EnsureSuccessStatusCode();
+				return await ParseJsonContent<TResponse>(response);
 			}
 		}
 
@@ -106,23 +100,23 @@ namespace EpicGames.Core
 		/// Converts an object to a JSON http content object
 		/// </summary>
 		/// <typeparam name="T">Type of the object to parse</typeparam>
-		/// <param name="Object">The object instance</param>
+		/// <param name="obj">The object instance</param>
 		/// <returns>Http content object</returns>
-		private static HttpContent ToJsonContent<T>(T Object)
+		private static HttpContent ToJsonContent<T>(T obj)
 		{
-			return new StringContent(JsonSerializer.Serialize<T>(Object), Encoding.UTF8, "application/json");
+			return new StringContent(JsonSerializer.Serialize<T>(obj), Encoding.UTF8, "application/json");
 		}
 
 		/// <summary>
 		/// Parses a HTTP response as a JSON object
 		/// </summary>
 		/// <typeparam name="T">Type of the object to parse</typeparam>
-		/// <param name="Message">The message received</param>
+		/// <param name="message">The message received</param>
 		/// <returns>Parsed object instance</returns>
-		private static async Task<T> ParseJsonContent<T>(HttpResponseMessage Message)
+		private static async Task<T> ParseJsonContent<T>(HttpResponseMessage message)
 		{
-			byte[] Bytes = await Message.Content.ReadAsByteArrayAsync();
-			return JsonSerializer.Deserialize<T>(Bytes, new JsonSerializerOptions { PropertyNameCaseInsensitive = true });
+			byte[] bytes = await message.Content.ReadAsByteArrayAsync();
+			return JsonSerializer.Deserialize<T>(bytes, new JsonSerializerOptions { PropertyNameCaseInsensitive = true });
 		}
 	}
 }
