Reformat EpicGames.Core according to standard coding conventions.

diff --git a/Engine/Source/Programs/Shared/EpicGames.Core/ILogCursor.cs b/Engine/Source/Programs/Shared/EpicGames.Core/ILogCursor.cs
--- a/Engine/Source/Programs/Shared/EpicGames.Core/ILogCursor.cs
+++ b/Engine/Source/Programs/Shared/EpicGames.Core/ILogCursor.cs
@@ -1,9 +1,6 @@
 // Copyright Epic Games, Inc. All Rights Reserved.
 
-using System;
-using System.Collections.Generic;
 using System.Diagnostics.CodeAnalysis;
-using System.Text;
 using System.Text.RegularExpressions;
 
 namespace EpicGames.Core
@@ -32,9 +29,9 @@ namespace EpicGames.Core
 		/// <summary>
 		/// Index to find the string at the given offset
 		/// </summary>
-		/// <param name="Offset"></param>
+		/// <param name="offset"></param>
 		/// <returns></returns>
-		string? this[int Offset]
+		string? this[int offset]
 		{
 			get;
 		}
@@ -50,142 +47,142 @@ namespace EpicGames.Core
 		/// </summary>
 		class RebasedLogCursor : ILogCursor
 		{
-			ILogCursor Inner { get; }
-			int BaseLineNumber { get; }
+			readonly ILogCursor _inner;
+			readonly int _baseLineNumber;
 
-			public RebasedLogCursor(ILogCursor Inner, int BaseLineNumber)
+			public RebasedLogCursor(ILogCursor inner, int baseLineNumber)
 			{
-				this.Inner = Inner;
-				this.BaseLineNumber = BaseLineNumber;
+				_inner = inner;
+				_baseLineNumber = baseLineNumber;
 			}
 
-			public string? this[int Index] => Inner[(BaseLineNumber + Index) - Inner.CurrentLineNumber];
-			public string? CurrentLine => Inner[BaseLineNumber - Inner.CurrentLineNumber];
-			public int CurrentLineNumber => BaseLineNumber;
+			public string? this[int index] => _inner[(_baseLineNumber + index) - _inner.CurrentLineNumber];
+			public string? CurrentLine => _inner[_baseLineNumber - _inner.CurrentLineNumber];
+			public int CurrentLineNumber => _baseLineNumber;
 		}
 
 		/// <summary>
 		/// Creates a new log cursor based at an offset from the current line
 		/// </summary>
-		/// <param name="Cursor">The current log cursor instance</param>
-		/// <param name="Offset">Line number offset from the current</param>
+		/// <param name="cursor">The current log cursor instance</param>
+		/// <param name="offset">Line number offset from the current</param>
 		/// <returns>New log cursor instance</returns>
-		public static ILogCursor Rebase(this ILogCursor Cursor, int Offset)
+		public static ILogCursor Rebase(this ILogCursor cursor, int offset)
 		{
-			return new RebasedLogCursor(Cursor, Cursor.CurrentLineNumber + Offset);
+			return new RebasedLogCursor(cursor, cursor.CurrentLineNumber + offset);
 		}
 
 		/// <summary>
 		/// Attempts to get a line at the given offset
 		/// </summary>
-		/// <param name="Cursor">The log cursor instance</param>
-		/// <param name="Offset">Offset of the line to retrieve</param>
-		/// <param name="NextLine">On success, receives the matched line</param>
+		/// <param name="cursor">The log cursor instance</param>
+		/// <param name="offset">Offset of the line to retrieve</param>
+		/// <param name="nextLine">On success, receives the matched line</param>
 		/// <returns>True if the line was retrieved</returns>
-		public static bool TryGetLine(this ILogCursor Cursor, int Offset, [NotNullWhen(true)] out string? NextLine)
+		public static bool TryGetLine(this ILogCursor cursor, int offset, [NotNullWhen(true)] out string? nextLine)
 		{
-			NextLine = Cursor[Offset];
-			return NextLine != null;
+			nextLine = cursor[offset];
+			return nextLine != null;
 		}
 
 		/// <summary>
 		/// Determines if the current line matches the given regex
 		/// </summary>
-		/// <param name="Cursor">The log cursor instance</param>
-		/// <param name="Pattern">The regex pattern to match</param>
+		/// <param name="cursor">The log cursor instance</param>
+		/// <param name="pattern">The regex pattern to match</param>
 		/// <returns>True if the current line matches the given patter</returns>
-		public static bool IsMatch(this ILogCursor Cursor, string Pattern)
+		public static bool IsMatch(this ILogCursor cursor, string pattern)
 		{
-			return IsMatch(Cursor, 0, Pattern);
+			return IsMatch(cursor, 0, pattern);
 		}
 
 		/// <summary>
 		/// Determines if the line at the given offset matches the given regex
 		/// </summary>
-		/// <param name="Cursor">The log cursor instance</param>
-		/// <param name="Offset">Offset of the line to match</param>
-		/// <param name="Pattern">The regex pattern to match</param>
+		/// <param name="cursor">The log cursor instance</param>
+		/// <param name="offset">Offset of the line to match</param>
+		/// <param name="pattern">The regex pattern to match</param>
 		/// <returns>True if the requested line matches the given patter</returns>
-		public static bool IsMatch(this ILogCursor Cursor, int Offset, string Pattern)
+		public static bool IsMatch(this ILogCursor cursor, int offset, string pattern)
 		{
-			string? Line;
-			return Cursor.TryGetLine(Offset, out Line) && Regex.IsMatch(Line!, Pattern);
+			string? line;
+			return cursor.TryGetLine(offset, out line) && Regex.IsMatch(line!, pattern);
 		}
 
 		/// <summary>
 		/// Determines if the current line matches the given regex
 		/// </summary>
-		/// <param name="Cursor">The log cursor instance</param>
-		/// <param name="Pattern">The regex pattern to match</param>
-		/// <param name="OutMatch">On success, receives the match result</param>
+		/// <param name="cursor">The log cursor instance</param>
+		/// <param name="pattern">The regex pattern to match</param>
+		/// <param name="outMatch">On success, receives the match result</param>
 		/// <returns>True if the current line matches the given patter</returns>
-		public static bool TryMatch(this ILogCursor Cursor, string Pattern, [NotNullWhen(true)] out Match? OutMatch)
+		public static bool TryMatch(this ILogCursor cursor, string pattern, [NotNullWhen(true)] out Match? outMatch)
 		{
-			return TryMatch(Cursor, 0, Pattern, out OutMatch);
+			return TryMatch(cursor, 0, pattern, out outMatch);
 		}
 
 		/// <summary>
 		/// Determines if the line at the given offset matches the given regex
 		/// </summary>
-		/// <param name="Cursor">The log cursor instance</param>
-		/// <param name="Offset">The line offset to check</param>
-		/// <param name="Pattern">The regex pattern to match</param>
-		/// <param name="OutMatch">On success, receives the match result</param>
+		/// <param name="cursor">The log cursor instance</param>
+		/// <param name="offset">The line offset to check</param>
+		/// <param name="pattern">The regex pattern to match</param>
+		/// <param name="outMatch">On success, receives the match result</param>
 		/// <returns>True if the current line matches the given patter</returns>
-		public static bool TryMatch(this ILogCursor Cursor, int Offset, string Pattern, [NotNullWhen(true)] out Match? OutMatch)
+		public static bool TryMatch(this ILogCursor cursor, int offset, string pattern, [NotNullWhen(true)] out Match? outMatch)
 		{
-			string? Line;
-			if (!Cursor.TryGetLine(Offset, out Line))
+			string? line;
+			if (!cursor.TryGetLine(offset, out line))
 			{
-				OutMatch = null;
+				outMatch = null;
 				return false;
 			}
 
-			Match Match = Regex.Match(Line, Pattern);
-			if (!Match.Success)
+			Match match = Regex.Match(line, pattern);
+			if (!match.Success)
 			{
-				OutMatch = null;
+				outMatch = null;
 				return false;
 			}
 
-			OutMatch = Match;
+			outMatch = match;
 			return true;
 		}
 
 		/// <summary>
 		/// Matches lines forward from the given offset while the given pattern matches
 		/// </summary>
-		/// <param name="Cursor">The log cursor instance</param>
-		/// <param name="Offset">Initial offset</param>
-		/// <param name="Pattern">Pattern to match</param>
+		/// <param name="cursor">The log cursor instance</param>
+		/// <param name="offset">Initial offset</param>
+		/// <param name="pattern">Pattern to match</param>
 		/// <returns>Offset of the last line that still matches the pattern (inclusive)</returns>
-		public static int MatchForwards(this ILogCursor Cursor, int Offset, string Pattern)
+		public static int MatchForwards(this ILogCursor cursor, int offset, string pattern)
 		{
-			while (IsMatch(Cursor, Offset + 1, Pattern))
+			while (IsMatch(cursor, offset + 1, pattern))
 			{
-				Offset++;
+				offset++;
 			}
-			return Offset;
+			return offset;
 		}
 
 		/// <summary>
 		/// Matches lines forwards from the given offset until the given pattern matches
 		/// </summary>
-		/// <param name="Cursor">The log cursor</param>
-		/// <param name="Offset">Initial offset</param>
-		/// <param name="Pattern">Pattern to match</param>
+		/// <param name="cursor">The log cursor</param>
+		/// <param name="offset">Initial offset</param>
+		/// <param name="pattern">Pattern to match</param>
 		/// <returns>Offset of the line that matches the pattern (inclusive), or EOF is encountered</returns>
-		public static int MatchForwardsUntil(this ILogCursor Cursor, int Offset, string Pattern)
+		public static int MatchForwardsUntil(this ILogCursor cursor, int offset, string pattern)
 		{
-			string? NextLine;
-			for (int NextOffset = Offset + 1; Cursor.TryGetLine(NextOffset, out NextLine); NextOffset++)
+			string? nextLine;
+			for (int nextOffset = offset + 1; cursor.TryGetLine(nextOffset, out nextLine); nextOffset++)
 			{
-				if (Regex.IsMatch(NextLine, Pattern))
+				if (Regex.IsMatch(nextLine, pattern))
 				{
-					return NextOffset;
+					return nextOffset;
 				}
 			}
-			return Offset;
+			return offset;
 		}
 	}
 }
diff --git a/Engine/Source/Programs/Shared/EpicGames.Core/ILoggerIndent.cs b/Engine/Source/Programs/Shared/EpicGames.Core/ILoggerIndent.cs
--- a/Engine/Source/Programs/Shared/EpicGames.Core/ILoggerIndent.cs
+++ b/Engine/Source/Programs/Shared/EpicGames.Core/ILoggerIndent.cs
@@ -1,12 +1,10 @@
 // Copyright Epic Games, Inc. All Rights Reserved.
 
-using Microsoft.Extensions.Logging;
 using System;
 using System.Collections;
 using System.Collections.Generic;
 using System.Linq;
-using System.Text;
-using System.Threading.Tasks;
+using Microsoft.Extensions.Logging;
 
 namespace EpicGames.Core
 {
@@ -34,7 +32,7 @@ namespace EpicGames.Core
 			/// <summary>
 			/// Owning object
 			/// </summary>
-			DefaultLoggerIndentHandler Owner;
+			readonly DefaultLoggerIndentHandler _owner;
 
 			/// <summary>
 			/// The indent scope object
@@ -44,14 +42,14 @@ namespace EpicGames.Core
 			/// <summary>
 			/// Constructor
 			/// </summary>
-			public Scope(DefaultLoggerIndentHandler Owner, ILoggerIndent Indent)
+			public Scope(DefaultLoggerIndentHandler owner, ILoggerIndent indent)
 			{
-				this.Owner = Owner;
-				this.Indent = Indent;
+				_owner = owner;
+				Indent = indent;
 
-				lock (Owner.Scopes)
+				lock (owner._scopes)
 				{
-					Owner.Scopes.Add(this);
+					owner._scopes.Add(this);
 				}
 			}
 
@@ -60,9 +58,9 @@ namespace EpicGames.Core
 			/// </summary>
 			public void Dispose()
 			{
-				lock (Owner.Scopes)
+				lock (_owner._scopes)
 				{
-					Owner.Scopes.Remove(this);
+					_owner._scopes.Remove(this);
 				}
 			}
 		}
@@ -76,46 +74,46 @@ namespace EpicGames.Core
 			/// <summary>
 			/// The indent to apply
 			/// </summary>
-			string Indent;
+			readonly string _indent;
 
 			/// <summary>
 			/// The inner state
 			/// </summary>
-			TState State;
+			readonly TState _state;
 
 			/// <summary>
 			/// Formatter for the inner state
 			/// </summary>
-			Func<TState, Exception, string> Formatter;
+			readonly Func<TState, Exception, string> _formatter;
 
 			/// <summary>
 			/// Constructor
 			/// </summary>
-			/// <param name="Indent">The indent to apply</param>
-			/// <param name="State">The inner state</param>
-			/// <param name="Formatter">Formatter for the inner state</param>
-			public FormattedLogValues(string Indent, TState State, Func<TState, Exception, string> Formatter)
+			/// <param name="indent">The indent to apply</param>
+			/// <param name="state">The inner state</param>
+			/// <param name="formatter">Formatter for the inner state</param>
+			public FormattedLogValues(string indent, TState state, Func<TState, Exception, string> formatter)
 			{
-				this.Indent = Indent;
-				this.State = State;
-				this.Formatter = Formatter;
+				_indent = indent;
+				_state = state;
+				_formatter = formatter;
 			}
 
 			/// <inheritdoc/>
 			public IEnumerator<KeyValuePair<string, object>> GetEnumerator()
 			{
-				IEnumerable<KeyValuePair<string, object>>? InnerEnumerable = State as IEnumerable<KeyValuePair<string, object>>;
-				if (InnerEnumerable != null)
+				IEnumerable<KeyValuePair<string, object>>? innerEnumerable = _state as IEnumerable<KeyValuePair<string, object>>;
+				if (innerEnumerable != null)
 				{
-					foreach (KeyValuePair<string, object> Pair in InnerEnumerable)
+					foreach (KeyValuePair<string, object> pair in innerEnumerable)
 					{
-						if (Pair.Key.Equals("{OriginalFormat}", StringComparison.Ordinal))
+						if (pair.Key.Equals("{OriginalFormat}", StringComparison.Ordinal))
 						{
-							yield return new KeyValuePair<string, object>(Pair.Key, Indent + Pair.Value.ToString());
+							yield return new KeyValuePair<string, object>(pair.Key, _indent + pair.Value.ToString());
 						}
 						else
 						{
-							yield return Pair;
+							yield return pair;
 						}
 					}
 				}
@@ -130,24 +128,24 @@ namespace EpicGames.Core
 			/// <summary>
 			/// Formats an instance of this object
 			/// </summary>
-			/// <param name="Values">The object instance</param>
-			/// <param name="Exception">The exception to format</param>
+			/// <param name="values">The object instance</param>
+			/// <param name="exception">The exception to format</param>
 			/// <returns>The formatted string</returns>
-			public static string Format(FormattedLogValues<TState> Values, Exception Exception)
+			public static string Format(FormattedLogValues<TState> values, Exception exception)
 			{
-				return Values.Indent + Values.Formatter(Values.State, Exception);
+				return values._indent + values._formatter(values._state, exception);
 			}
 		}
 
 		/// <summary>
 		/// The internal logger
 		/// </summary>
-		ILogger Inner;
+		readonly ILogger _inner;
 
 		/// <summary>
 		/// Current list of indents
 		/// </summary>
-		List<Scope> Scopes = new List<Scope>();
+		readonly List<Scope> _scopes = new List<Scope>();
 
 		/// <summary>
 		/// The current indent text
@@ -161,42 +159,42 @@ namespace EpicGames.Core
 		/// <summary>
 		/// Constructor
 		/// </summary>
-		/// <param name="Inner">The logger to wrap</param>
-		public DefaultLoggerIndentHandler(ILogger Inner)
+		/// <param name="inner">The logger to wrap</param>
+		public DefaultLoggerIndentHandler(ILogger inner)
 		{
-			this.Inner = Inner;
-			this.Indent = "";
+			_inner = inner;
+			Indent = "";
 		}
 
 		/// <inheritdoc/>
-		public IDisposable BeginScope<TState>(TState State)
+		public IDisposable BeginScope<TState>(TState state)
 		{
-			ILoggerIndent? Indent = State as ILoggerIndent;
-			if (Indent != null)
+			ILoggerIndent? indent = state as ILoggerIndent;
+			if (indent != null)
 			{
-				return new Scope(this, Indent);
+				return new Scope(this, indent);
 			}
 
-			return Inner.BeginScope(State);
+			return _inner.BeginScope(state);
 		}
 
 		/// <inheritdoc/>
-		public bool IsEnabled(LogLevel LogLevel)
+		public bool IsEnabled(LogLevel logLevel)
 		{
-			return Inner.IsEnabled(LogLevel);
+			return _inner.IsEnabled(logLevel);
 		}
 
 		/// <inheritdoc/>
-		public void Log<TState>(LogLevel LogLevel, EventId EventId, TState State, Exception Exception, Func<TState, Exception, string> Formatter)
+		public void Log<TState>(LogLevel logLevel, EventId eventId, TState state, Exception exception, Func<TState, Exception, string> formatter)
 		{
-			if (Scopes.Count > 0)
+			if (_scopes.Count > 0)
 			{
-				string Indent = String.Join("", Scopes.Select(x => x.Indent.Indent));
-				Inner.Log(LogLevel, EventId, new FormattedLogValues<TState>(Indent, State, Formatter), Exception, FormattedLogValues<TState>.Format);
+				string indent = String.Join("", _scopes.Select(x => x.Indent.Indent));
+				_inner.Log(logLevel, eventId, new FormattedLogValues<TState>(indent, state, formatter), exception, FormattedLogValues<TState>.Format);
 				return;
 			}
 
-			Inner.Log(LogLevel, EventId, State, Exception, Formatter);
+			_inner.Log(logLevel, eventId, state, exception, formatter);
 		}
 	}
 
@@ -218,22 +216,22 @@ namespace EpicGames.Core
 			/// <summary>
 			/// Constructor
 			/// </summary>
-			/// <param name="Indent">Indent to append to the existing indent</param>
-			public LoggerIndent(string Indent)
+			/// <param name="indent">Indent to append to the existing indent</param>
+			public LoggerIndent(string indent)
 			{
-				this.Indent = Indent;
+				Indent = indent;
 			}
 		}
 
 		/// <summary>
 		/// Create an indent
 		/// </summary>
-		/// <param name="Logger">Logger interface</param>
-		/// <param name="Indent">The indent to apply</param>
+		/// <param name="logger">Logger interface</param>
+		/// <param name="indent">The indent to apply</param>
 		/// <returns>Disposable object</returns>
-		public static IDisposable BeginIndentScope(this ILogger Logger, string Indent)
+		public static IDisposable BeginIndentScope(this ILogger logger, string indent)
 		{
-			return Logger.BeginScope(new LoggerIndent(Indent));
+			return logger.BeginScope(new LoggerIndent(indent));
 		}
 	}
 }
diff --git a/Engine/Source/Programs/Shared/EpicGames.Core/ILoggerProgress.cs b/Engine/Source/Programs/Shared/EpicGames.Core/ILoggerProgress.cs
--- a/Engine/Source/Programs/Shared/EpicGames.Core/ILoggerProgress.cs
+++ b/Engine/Source/Programs/Shared/EpicGames.Core/ILoggerProgress.cs
@@ -1,12 +1,8 @@
 // Copyright Epic Games, Inc. All Rights Reserved.
 
-using Microsoft.Extensions.Logging;
 using System;
-using System.Collections.Generic;
 using System.Diagnostics;
-using System.Linq;
-using System.Text;
-using System.Threading.Tasks;
+using Microsoft.Extensions.Logging;
 
 namespace EpicGames.Core
 {
@@ -45,7 +41,7 @@ namespace EpicGames.Core
 			/// <summary>
 			/// The logger to output to
 			/// </summary>
-			private ILogger Logger;
+			private readonly ILogger _logger;
 
 			/// <summary>
 			/// Prefix message for the status
@@ -60,16 +56,16 @@ namespace EpicGames.Core
 			/// </summary>
 			public string Progress
 			{
-				get { return ProgressInternal; }
+				get => _progressInternal;
 				set
 				{
-					ProgressInternal = value;
+					_progressInternal = value;
 
-					if (Timer.Elapsed > TimeSpan.FromSeconds(3.0))
+					if (_timer.Elapsed > TimeSpan.FromSeconds(3.0))
 					{
-						LastOutput = String.Empty;
+						_lastOutput = String.Empty;
 						Flush();
-						Timer.Restart();
+						_timer.Restart();
 					}
 				}
 			}
@@ -77,30 +73,30 @@ namespace EpicGames.Core
 			/// <summary>
 			/// The last string that was output
 			/// </summary>
-			string LastOutput = String.Empty;
+			string _lastOutput = String.Empty;
 
 			/// <summary>
 			/// Backing storage for the Progress string
 			/// </summary>
-			string ProgressInternal = String.Empty;
+			string _progressInternal = String.Empty;
 
 			/// <summary>
 			/// Timer since the last update
 			/// </summary>
-			Stopwatch Timer = Stopwatch.StartNew();
+			readonly Stopwatch _timer = Stopwatch.StartNew();
 
 			/// <summary>
 			/// Constructor
 			/// </summary>
-			/// <param name="Logger">The logger to write to</param>
-			/// <param name="Message">The base message to display</param>
-			public LoggerProgress(ILogger Logger, string Message)
+			/// <param name="logger">The logger to write to</param>
+			/// <param name="message">The base message to display</param>
+			public LoggerProgress(ILogger logger, string message)
 			{
-				this.Logger = Logger;
-				this.Message = Message;
-				this.ProgressInternal = String.Empty;
+				_logger = logger;
+				Message = message;
+				_progressInternal = String.Empty;
 
-				Logger.LogInformation(Message);
+				logger.LogInformation(message);
 			}
 
 			/// <summary>
@@ -116,15 +112,15 @@ namespace EpicGames.Core
 			/// </summary>
 			void Flush()
 			{
-				string Output = Message;
+				string output = Message;
 				if (!String.IsNullOrEmpty(Progress))
 				{
-					Output += $" {Progress}";
+					output += $" {Progress}";
 				}
-				if (!String.Equals(Output, LastOutput, StringComparison.Ordinal))
+				if (!String.Equals(output, _lastOutput, StringComparison.Ordinal))
 				{
-					Logger.LogInformation(Output);
-					LastOutput = Output;
+					_logger.LogInformation(output);
+					_lastOutput = output;
 				}
 			}
 		}
@@ -132,12 +128,12 @@ namespace EpicGames.Core
 		/// <summary>
 		/// Begins a new progress scope
 		/// </summary>
-		/// <param name="Logger">The logger being written to</param>
-		/// <param name="Message">The message prefix</param>
+		/// <param name="logger">The logger being written to</param>
+		/// <param name="message">The message prefix</param>
 		/// <returns>Scope object, which should be disposed when finished</returns>
-		public static ILoggerProgress BeginProgressScope(this ILogger Logger, string Message)
+		public static ILoggerProgress BeginProgressScope(this ILogger logger, string message)
 		{
-			return new LoggerProgress(Logger, Message);
+			return new LoggerProgress(logger, message);
 		}
 	}
 }
diff --git a/Engine/Source/Programs/Shared/EpicGames.Core/JsonExtensions.cs b/Engine/Source/Programs/Shared/EpicGames.Core/JsonExtensions.cs
--- a/Engine/Source/Programs/Shared/EpicGames.Core/JsonExtensions.cs
+++ b/Engine/Source/Programs/Shared/EpicGames.Core/JsonExtensions.cs
@@ -1,8 +1,6 @@
 // Copyright Epic Games, Inc. All Rights Reserved.
 
 using System;
-using System.Collections.Generic;
-using System.Text;
 using System.Text.Json;
 
 namespace EpicGames.Core
@@ -15,19 +13,19 @@ namespace EpicGames.Core
 		/// <summary>
 		/// Tries to read a property name then move to the value token
 		/// </summary>
-		/// <param name="Reader">Token reader</param>
-		/// <param name="PropertyName">Receives the property name on success</param>
+		/// <param name="reader">Token reader</param>
+		/// <param name="propertyName">Receives the property name on success</param>
 		/// <returns>True if the read succeeded</returns>
-		public static bool TryReadNextPropertyName(ref Utf8JsonReader Reader, out ReadOnlySpan<byte> PropertyName)
+		public static bool TryReadNextPropertyName(ref Utf8JsonReader reader, out ReadOnlySpan<byte> propertyName)
 		{
-			if (Reader.Read() && Reader.TokenType == JsonTokenType.PropertyName)
+			if (reader.Read() && reader.TokenType == JsonTokenType.PropertyName)
 			{
-				PropertyName = Reader.ValueSpan;
-				return Reader.Read();
+				propertyName = reader.ValueSpan;
+				return reader.Read();
 			}
 			else
 			{
-				PropertyName = ReadOnlySpan<byte>.Empty;
+				propertyName = ReadOnlySpan<byte>.Empty;
 				return false;
 			}
 		}
@@ -35,9 +33,9 @@ namespace EpicGames.Core
 		/// <summary>
 		/// Reads a Utf8 string from the current json token
 		/// </summary>
-		/// <param name="Reader"></param>
+		/// <param name="reader"></param>
 		/// <returns></returns>
-		public static Utf8String GetUtf8String(this Utf8JsonReader Reader)
+		public static Utf8String GetUtf8String(this Utf8JsonReader reader)
 		{
 			throw new NotImplementedException();
 		}
diff --git a/Engine/Source/Programs/Shared/EpicGames.Core/JsonKnownTypes.cs b/Engine/Source/Programs/Shared/EpicGames.Core/JsonKnownTypes.cs
--- a/Engine/Source/Programs/Shared/EpicGames.Core/JsonKnownTypes.cs
+++ b/Engine/Source/Programs/Shared/EpicGames.Core/JsonKnownTypes.cs
@@ -1,13 +1,10 @@
 // Copyright Epic Games, Inc. All Rights Reserved.
 
 using System;
-using System.Buffers;
 using System.Collections.Generic;
-using System.Diagnostics.CodeAnalysis;
 using System.Linq;
 using System.Text.Json;
 using System.Text.Json.Serialization;
-using System.Threading.Tasks;
 
 namespace EpicGames.Core
 {
@@ -25,10 +22,10 @@ namespace EpicGames.Core
 		/// <summary>
 		/// Constructor
 		/// </summary>
-		/// <param name="Name">Name to use to discriminate between different classes</param>
-		public JsonDiscriminatorAttribute(string Name)
+		/// <param name="name">Name to use to discriminate between different classes</param>
+		public JsonDiscriminatorAttribute(string name)
 		{
-			this.Name = Name;
+			Name = name;
 		}
 	}
 
@@ -46,10 +43,10 @@ namespace EpicGames.Core
 		/// <summary>
 		/// Constructor
 		/// </summary>
-		/// <param name="Types">Array of derived classes</param>
-		public JsonKnownTypesAttribute(params Type[] Types)
+		/// <param name="types">Array of derived classes</param>
+		public JsonKnownTypesAttribute(params Type[] types)
 		{
-			this.Types = Types;
+			Types = types;
 		}
 	}
 
@@ -66,136 +63,136 @@ namespace EpicGames.Core
 		/// <summary>
 		/// Mapping from discriminator to class type
 		/// </summary>
-		Dictionary<string, Type> DiscriminatorToType;
+		readonly Dictionary<string, Type> _discriminatorToType;
 
 		/// <summary>
 		/// Mapping from class type to discriminator
 		/// </summary>
-		Dictionary<Type, JsonEncodedText> TypeToDiscriminator;
+		readonly Dictionary<Type, JsonEncodedText> _typeToDiscriminator;
 
 		/// <summary>
 		/// Constructor
 		/// </summary>
-		/// <param name="KnownTypes">Set of derived class types</param>
-		public JsonKnownTypesConverter(IEnumerable<Type> KnownTypes)
+		/// <param name="knownTypes">Set of derived class types</param>
+		public JsonKnownTypesConverter(IEnumerable<Type> knownTypes)
 		{
-			DiscriminatorToType = new Dictionary<string, Type>(StringComparer.OrdinalIgnoreCase);
-			foreach (Type KnownType in KnownTypes)
+			_discriminatorToType = new Dictionary<string, Type>(StringComparer.OrdinalIgnoreCase);
+			foreach (Type knownType in knownTypes)
 			{
-				if (!typeof(T).IsAssignableFrom(KnownType))
+				if (!typeof(T).IsAssignableFrom(knownType))
 				{
-					throw new InvalidOperationException($"{KnownType} is not derived from {typeof(T)}");
+					throw new InvalidOperationException($"{knownType} is not derived from {typeof(T)}");
 				}
-				foreach (JsonDiscriminatorAttribute DiscriminatorAttribute in KnownType.GetCustomAttributes(typeof(JsonDiscriminatorAttribute), true))
+				foreach (JsonDiscriminatorAttribute discriminatorAttribute in knownType.GetCustomAttributes(typeof(JsonDiscriminatorAttribute), true))
 				{
-					DiscriminatorToType.Add(DiscriminatorAttribute.Name, KnownType);
+					_discriminatorToType.Add(discriminatorAttribute.Name, knownType);
 				}
 			}
 
-			TypeToDiscriminator = DiscriminatorToType.ToDictionary(x => x.Value, x => JsonEncodedText.Encode(x.Key));
+			_typeToDiscriminator = _discriminatorToType.ToDictionary(x => x.Value, x => JsonEncodedText.Encode(x.Key));
 		}
 
 		/// <summary>
 		/// Determines whether the given type can be converted
 		/// </summary>
-		/// <param name="TypeToConvert">Type to convert</param>
+		/// <param name="typeToConvert">Type to convert</param>
 		/// <returns>True if the type can be converted</returns>
-		public override bool CanConvert(Type TypeToConvert)
+		public override bool CanConvert(Type typeToConvert)
 		{
-			return TypeToConvert == typeof(T) || TypeToDiscriminator.ContainsKey(TypeToConvert);
+			return typeToConvert == typeof(T) || _typeToDiscriminator.ContainsKey(typeToConvert);
 		}
 
 		/// <summary>
 		/// Reads an object from the given input stream
 		/// </summary>
-		/// <param name="Reader">UTF8 reader</param>
-		/// <param name="TypeToConvert">Type to be read from the stream</param>
-		/// <param name="Options">Options for serialization</param>
+		/// <param name="reader">UTF8 reader</param>
+		/// <param name="typeToConvert">Type to be read from the stream</param>
+		/// <param name="options">Options for serialization</param>
 		/// <returns>Instance of the serialized object</returns>
-		public override T Read(ref Utf8JsonReader Reader, Type TypeToConvert, JsonSerializerOptions Options)
+		public override T Read(ref Utf8JsonReader reader, Type typeToConvert, JsonSerializerOptions options)
 		{
-			using (JsonDocument Document = JsonDocument.ParseValue(ref Reader))
+			using (JsonDocument document = JsonDocument.ParseValue(ref reader))
 			{
-				JsonElement Element;
-				if (!TryGetPropertyRespectingCase(Document.RootElement, TypePropertyName, Options, out Element))
+				JsonElement element;
+				if (!TryGetPropertyRespectingCase(document.RootElement, TypePropertyName, options, out element))
 				{
-					throw new JsonException($"Missing '{TypePropertyName}' field for serializing {TypeToConvert.Name}");
+					throw new JsonException($"Missing '{TypePropertyName}' field for serializing {typeToConvert.Name}");
 				}
 
-				Type? TargetType;
-				if (!DiscriminatorToType.TryGetValue(Element.GetString(), out TargetType))
+				Type? targetType;
+				if (!_discriminatorToType.TryGetValue(element.GetString(), out targetType))
 				{
-					throw new JsonException($"Invalid '{TypePropertyName}' field ('{Element.GetString()}') for serializing {TypeToConvert.Name}");
+					throw new JsonException($"Invalid '{TypePropertyName}' field ('{element.GetString()}') for serializing {typeToConvert.Name}");
 				}
 
-				string Text = Document.RootElement.GetRawText();
-				T Result = (T)JsonSerializer.Deserialize(Text, TargetType, Options);
-				return Result;
+				string text = document.RootElement.GetRawText();
+				T result = (T)JsonSerializer.Deserialize(text, targetType, options);
+				return result;
 			}
 		}
 
 		/// <summary>
 		/// Finds a property within an object, ignoring case
 		/// </summary>
-		/// <param name="Object">The object to search</param>
-		/// <param name="PropertyName">Name of the property to search</param>
-		/// <param name="Options">Options for serialization. The <see cref="JsonSerializerOptions.PropertyNameCaseInsensitive"/> property defines whether a case insensitive scan will be performed.</param>
-		/// <param name="OutElement">The property value, if successful</param>
+		/// <param name="element">The object to search</param>
+		/// <param name="propertyName">Name of the property to search</param>
+		/// <param name="options">Options for serialization. The <see cref="JsonSerializerOptions.PropertyNameCaseInsensitive"/> property defines whether a case insensitive scan will be performed.</param>
+		/// <param name="outElement">The property value, if successful</param>
 		/// <returns>True if the property was found, false otherwise</returns>
-		static bool TryGetPropertyRespectingCase(JsonElement Object, string PropertyName, JsonSerializerOptions Options, out JsonElement OutElement)
+		static bool TryGetPropertyRespectingCase(JsonElement element, string propertyName, JsonSerializerOptions options, out JsonElement outElement)
 		{
-			JsonElement NewElement;
-			if (Object.TryGetProperty(PropertyName, out NewElement))
+			JsonElement newElement;
+			if (element.TryGetProperty(propertyName, out newElement))
 			{
-				OutElement = NewElement;
+				outElement = newElement;
 				return true;
 			}
 
-			if (Options.PropertyNameCaseInsensitive)
+			if (options.PropertyNameCaseInsensitive)
 			{
-				foreach (JsonProperty Property in Object.EnumerateObject())
+				foreach (JsonProperty property in element.EnumerateObject())
 				{
-					if (String.Equals(Property.Name, PropertyName, StringComparison.OrdinalIgnoreCase))
+					if (String.Equals(property.Name, propertyName, StringComparison.OrdinalIgnoreCase))
 					{
-						OutElement = Property.Value;
+						outElement = property.Value;
 						return true;
 					}
 				}
 			}
 
-			OutElement = new JsonElement();
+			outElement = new JsonElement();
 			return false;
 		}
 
 		/// <summary>
 		/// Writes an object to the given output stream
 		/// </summary>
-		/// <param name="Writer">UTF8 writer</param>
-		/// <param name="Value">Object to be serialized to the stream</param>
-		/// <param name="Options">Options for serialization</param>
-		public override void Write(Utf8JsonWriter Writer, T Value, JsonSerializerOptions Options)
+		/// <param name="writer">UTF8 writer</param>
+		/// <param name="value">Object to be serialized to the stream</param>
+		/// <param name="options">Options for serialization</param>
+		public override void Write(Utf8JsonWriter writer, T value, JsonSerializerOptions options)
 		{
-			if (Value == null)
+			if (value == null)
 			{
-				Writer.WriteNullValue();
+				writer.WriteNullValue();
 			}
 			else
 			{
-				Type ValueType = Value.GetType();
+				Type valueType = value.GetType();
 
-				Writer.WriteStartObject();
-				Writer.WriteString(Options.PropertyNamingPolicy?.ConvertName(TypePropertyName) ?? TypePropertyName, TypeToDiscriminator[ValueType]);
+				writer.WriteStartObject();
+				writer.WriteString(options.PropertyNamingPolicy?.ConvertName(TypePropertyName) ?? TypePropertyName, _typeToDiscriminator[valueType]);
 
-				byte[] Bytes = JsonSerializer.SerializeToUtf8Bytes(Value, ValueType, Options);
-				using (JsonDocument Document = JsonDocument.Parse(Bytes))
+				byte[] bytes = JsonSerializer.SerializeToUtf8Bytes(value, valueType, options);
+				using (JsonDocument document = JsonDocument.Parse(bytes))
 				{
-					foreach (JsonProperty Property in Document.RootElement.EnumerateObject())
+					foreach (JsonProperty property in document.RootElement.EnumerateObject())
 					{
-						Property.WriteTo(Writer);
+						property.WriteTo(writer);
 					}
 				}
 
-				Writer.WriteEndObject();
+				writer.WriteEndObject();
 			}
 		}
 	}
@@ -208,25 +205,25 @@ namespace EpicGames.Core
 		/// <summary>
 		/// Determines whether it's possible to create a converter for the given type
 		/// </summary>
-		/// <param name="TypeToConvert">The type being converted</param>
+		/// <param name="typeToConvert">The type being converted</param>
 		/// <returns>True if the type can be converted</returns>
-		public override bool CanConvert(Type TypeToConvert)
+		public override bool CanConvert(Type typeToConvert)
 		{
-			object[] Attributes = TypeToConvert.GetCustomAttributes(typeof(JsonKnownTypesAttribute), false);
-			return Attributes.Length > 0;
+			object[] attributes = typeToConvert.GetCustomAttributes(typeof(JsonKnownTypesAttribute), false);
+			return attributes.Length > 0;
 		}
 
 		/// <summary>
 		/// Creates a converter for the given type
 		/// </summary>
-		/// <param name="TypeToConvert">The type being converted</param>
-		/// <param name="Options">Options for serialization</param>
+		/// <param name="typeToConvert">The type being converted</param>
+		/// <param name="options">Options for serialization</param>
 		/// <returns>New converter for the requested type</returns>
-		public override JsonConverter CreateConverter(Type TypeToConvert, JsonSerializerOptions Options)
+		public override JsonConverter CreateConverter(Type typeToConvert, JsonSerializerOptions options)
 		{
-			object[] Attributes = TypeToConvert.GetCustomAttributes(typeof(JsonKnownTypesAttribute), false);
-			Type GenericType = typeof(JsonKnownTypesConverter<>).MakeGenericType(TypeToConvert);
-			return (JsonConverter)Activator.CreateInstance(GenericType, Attributes.SelectMany(x => ((JsonKnownTypesAttribute)x).Types))!;
+			object[] attributes = typeToConvert.GetCustomAttributes(typeof(JsonKnownTypesAttribute), false);
+			Type genericType = typeof(JsonKnownTypesConverter<>).MakeGenericType(typeToConvert);
+			return (JsonConverter)Activator.CreateInstance(genericType, attributes.SelectMany(x => ((JsonKnownTypesAttribute)x).Types))!;
 		}
 	}
 }
diff --git a/Engine/Source/Programs/Shared/EpicGames.Core/JsonLogEvent.cs b/Engine/Source/Programs/Shared/EpicGames.Core/JsonLogEvent.cs
--- a/Engine/Source/Programs/Shared/EpicGames.Core/JsonLogEvent.cs
+++ b/Engine/Source/Programs/Shared/EpicGames.Core/JsonLogEvent.cs
@@ -1,10 +1,9 @@
 // Copyright Epic Games, Inc. All Rights Reserved.
 
-using Microsoft.Extensions.Logging;
 using System;
-using System.Collections.Generic;
 using System.Text;
 using System.Text.Json;
+using Microsoft.Extensions.Logging;
 
 namespace EpicGames.Core
 {
@@ -49,63 +48,63 @@ namespace EpicGames.Core
 		/// <summary>
 		/// Constructor
 		/// </summary>
-		public JsonLogEvent(LogLevel Level, EventId EventId, int LineIndex, int LineCount, ReadOnlyMemory<byte> Data)
+		public JsonLogEvent(LogLevel level, EventId eventId, int lineIndex, int lineCount, ReadOnlyMemory<byte> data)
 		{
-			this.Level = Level;
-			this.EventId = EventId;
-			this.LineIndex = LineIndex;
-			this.LineCount = LineCount;
-			this.Data = Data;
+			Level = level;
+			EventId = eventId;
+			LineIndex = lineIndex;
+			LineCount = lineCount;
+			Data = data;
 		}
 
 		/// <summary>
 		/// Tries to parse a Json log event from the given 
 		/// </summary>
-		/// <param name="Data"></param>
-		/// <param name="Event"></param>
+		/// <param name="data"></param>
+		/// <param name="logEvent"></param>
 		/// <returns></returns>
-		public static bool TryParse(ReadOnlyMemory<byte> Data, out JsonLogEvent Event)
+		public static bool TryParse(ReadOnlyMemory<byte> data, out JsonLogEvent logEvent)
 		{
-			Utf8StringComparer StringComparer = Utf8StringComparer.OrdinalIgnoreCase;
+			Utf8StringComparer stringComparer = Utf8StringComparer.OrdinalIgnoreCase;
 			try
 			{
-				LogLevel Level = LogLevel.None;
-				int EventId = 0;
-				int LineIndex = 0;
-				int LineCount = 1;
+				LogLevel level = LogLevel.None;
+				int eventId = 0;
+				int lineIndex = 0;
+				int lineCount = 1;
 
-				Utf8JsonReader Reader = new Utf8JsonReader(Data.Span);
-				if (Reader.Read() && Reader.TokenType == JsonTokenType.StartObject)
+				Utf8JsonReader reader = new Utf8JsonReader(data.Span);
+				if (reader.Read() && reader.TokenType == JsonTokenType.StartObject)
 				{
-					ReadOnlySpan<byte> PropertyName;
-					for (; ReadNextPropertyName(ref Reader, out PropertyName); Reader.Skip())
+					ReadOnlySpan<byte> propertyName;
+					for (; ReadNextPropertyName(ref reader, out propertyName); reader.Skip())
 					{
-						if (StringComparer.Equals(PropertyName, LevelString.Span))
+						if (stringComparer.Equals(propertyName, LevelString.Span))
 						{
-							if (!Enum.TryParse(Reader.GetString(), out Level))
+							if (!Enum.TryParse(reader.GetString(), out level))
 							{
 								break;
 							}
 						}
-						else if (StringComparer.Equals(PropertyName, LineIndexString.Span))
+						else if (stringComparer.Equals(propertyName, LineIndexString.Span))
 						{
-							LineIndex = Reader.GetInt32();
+							lineIndex = reader.GetInt32();
 						}
-						else if (StringComparer.Equals(PropertyName, LineCountString.Span))
+						else if (stringComparer.Equals(propertyName, LineCountString.Span))
 						{
-							LineCount = Reader.GetInt32();
+							lineCount = reader.GetInt32();
 						}
-						else if (StringComparer.Equals(PropertyName, PropertiesString.Span) && Reader.TokenType == JsonTokenType.StartObject)
+						else if (stringComparer.Equals(propertyName, PropertiesString.Span) && reader.TokenType == JsonTokenType.StartObject)
 						{
-							for (; ReadNextPropertyName(ref Reader, out PropertyName); Reader.Skip())
+							for (; ReadNextPropertyName(ref reader, out propertyName); reader.Skip())
 							{
-								if (StringComparer.Equals(EventIdString.Span, PropertyName) && Reader.TokenType == JsonTokenType.StartObject)
+								if (stringComparer.Equals(EventIdString.Span, propertyName) && reader.TokenType == JsonTokenType.StartObject)
 								{
-									for (; ReadNextPropertyName(ref Reader, out PropertyName); Reader.Skip())
+									for (; ReadNextPropertyName(ref reader, out propertyName); reader.Skip())
 									{
-										if (StringComparer.Equals(IdString.Span, PropertyName) && Reader.TryGetInt32(out int NewEventId))
+										if (stringComparer.Equals(IdString.Span, propertyName) && reader.TryGetInt32(out int newEventId))
 										{
-											EventId = NewEventId;
+											eventId = newEventId;
 										}
 									}
 								}
@@ -114,9 +113,9 @@ namespace EpicGames.Core
 					}
 				}
 
-				if (Reader.TokenType == JsonTokenType.EndObject && Level != LogLevel.None && Reader.BytesConsumed == Data.Length)
+				if (reader.TokenType == JsonTokenType.EndObject && level != LogLevel.None && reader.BytesConsumed == data.Length)
 				{
-					Event = new JsonLogEvent(Level, new EventId(EventId), LineIndex, LineCount, Data);
+					logEvent = new JsonLogEvent(level, new EventId(eventId), lineIndex, lineCount, data);
 					return true;
 				}
 			}
@@ -124,20 +123,20 @@ namespace EpicGames.Core
 			{
 			}
 
-			Event = default;
+			logEvent = default;
 			return false;
 		}
 
-		static bool ReadNextPropertyName(ref Utf8JsonReader Reader, out ReadOnlySpan<byte> PropertyName)
+		static bool ReadNextPropertyName(ref Utf8JsonReader reader, out ReadOnlySpan<byte> propertyName)
 		{
-			if (Reader.Read() && Reader.TokenType == JsonTokenType.PropertyName)
+			if (reader.Read() && reader.TokenType == JsonTokenType.PropertyName)
 			{
-				PropertyName = Reader.ValueSpan;
-				return Reader.Read();
+				propertyName = reader.ValueSpan;
+				return reader.Read();
 			}
 			else
 			{
-				PropertyName = ReadOnlySpan<byte>.Empty;
+				propertyName = ReadOnlySpan<byte>.Empty;
 				return false;
 			}
 		}
@@ -145,30 +144,30 @@ namespace EpicGames.Core
 		/// <summary>
 		/// Formats an event as a string
 		/// </summary>
-		/// <param name="State"></param>
-		/// <param name="Ex"></param>
+		/// <param name="state"></param>
+		/// <param name="ex"></param>
 		/// <returns></returns>
 		/// <exception cref="NotImplementedException"></exception>
-		public static string Format(JsonLogEvent State, Exception? Ex)
+		public static string Format(JsonLogEvent state, Exception? ex)
 		{
-			Utf8JsonReader Reader = new Utf8JsonReader(State.Data.Span);
-			if (Reader.Read() && Reader.TokenType == JsonTokenType.StartObject)
+			Utf8JsonReader reader = new Utf8JsonReader(state.Data.Span);
+			if (reader.Read() && reader.TokenType == JsonTokenType.StartObject)
 			{
 				for (; ; )
 				{
-					if (!Reader.Read() || Reader.TokenType != JsonTokenType.PropertyName)
+					if (!reader.Read() || reader.TokenType != JsonTokenType.PropertyName)
 					{
 						break;
 					}
 
-					bool IsMessage = Utf8StringComparer.OrdinalIgnoreCase.Equals(Reader.ValueSpan, RenderedMessageString.Span);
-					if (!Reader.TrySkip())
+					bool isMessage = Utf8StringComparer.OrdinalIgnoreCase.Equals(reader.ValueSpan, RenderedMessageString.Span);
+					if (!reader.TrySkip())
 					{
 						break;
 					}
-					if (IsMessage && Reader.TokenType == JsonTokenType.String)
+					if (isMessage && reader.TokenType == JsonTokenType.String)
 					{
-						return Reader.GetString();
+						return reader.GetString();
 					}
 				}
 			}
diff --git a/Engine/Source/Programs/Shared/EpicGames.Core/JsonObject.cs b/Engine/Source/Programs/Shared/EpicGames.Core/JsonObject.cs
--- a/Engine/Source/Programs/Shared/EpicGames.Core/JsonObject.cs
+++ b/Engine/Source/Programs/Shared/EpicGames.Core/JsonObject.cs
@@ -3,9 +3,7 @@
 using System;
 using System.Collections.Generic;
 using System.Diagnostics.CodeAnalysis;
-using System.IO;
 using System.Linq;
-using System.Text;
 using System.Text.Json;
 
 namespace EpicGames.Core
@@ -18,10 +16,10 @@ namespace EpicGames.Core
 		/// <summary>
 		/// Constructor
 		/// </summary>
-		/// <param name="Format">Format string</param>
-		/// <param name="Args">Optional arguments</param>
-		public JsonParseException(string Format, params object[] Args)
-			: base(String.Format(Format, Args))
+		/// <param name="format">Format string</param>
+		/// <param name="args">Optional arguments</param>
+		public JsonParseException(string format, params object[] args)
+			: base(String.Format(format, args))
 		{
 		}
 	}
@@ -31,47 +29,47 @@ namespace EpicGames.Core
 	/// </summary>
 	public class JsonObject
 	{
-		Dictionary<string, object?> RawObject;
+		readonly Dictionary<string, object?> _rawObject;
 
 		/// <summary>
 		/// Construct a JSON object from the raw string -> object dictionary
 		/// </summary>
-		/// <param name="InRawObject">Raw object parsed from disk</param>
-		public JsonObject(Dictionary<string, object?> InRawObject)
+		/// <param name="inRawObject">Raw object parsed from disk</param>
+		public JsonObject(Dictionary<string, object?> inRawObject)
 		{
-			RawObject = new Dictionary<string, object?>(InRawObject, StringComparer.InvariantCultureIgnoreCase);
+			_rawObject = new Dictionary<string, object?>(inRawObject, StringComparer.InvariantCultureIgnoreCase);
 		}
 
 		/// <summary>
 		/// Constructor
 		/// </summary>
-		/// <param name="Element"></param>
-		public JsonObject(JsonElement Element)
+		/// <param name="element"></param>
+		public JsonObject(JsonElement element)
 		{
-			RawObject = new Dictionary<string, object?>(StringComparer.InvariantCultureIgnoreCase);
-			foreach (JsonProperty Property in Element.EnumerateObject())
+			_rawObject = new Dictionary<string, object?>(StringComparer.InvariantCultureIgnoreCase);
+			foreach (JsonProperty property in element.EnumerateObject())
 			{
-				RawObject[Property.Name] = ParseElement(Property.Value);
+				_rawObject[property.Name] = ParseElement(property.Value);
 			}
 		}
 
 		/// <summary>
 		/// Parse an individual element
 		/// </summary>
-		/// <param name="Element"></param>
+		/// <param name="element"></param>
 		/// <returns></returns>
-		public static object? ParseElement(JsonElement Element)
+		public static object? ParseElement(JsonElement element)
 		{
-			switch(Element.ValueKind)
+			switch(element.ValueKind)
 			{
 				case JsonValueKind.Array:
-					return Element.EnumerateArray().Select(x => ParseElement(x)).ToArray();
+					return element.EnumerateArray().Select(x => ParseElement(x)).ToArray();
 				case JsonValueKind.Number:
-					return Element.GetDouble();
+					return element.GetDouble();
 				case JsonValueKind.Object:
-					return Element.EnumerateObject().ToDictionary(x => x.Name, x => ParseElement(x.Value));
+					return element.EnumerateObject().ToDictionary(x => x.Name, x => ParseElement(x.Value));
 				case JsonValueKind.String:
-					return Element.GetString();
+					return element.GetString();
 				case JsonValueKind.False:
 					return false;
 				case JsonValueKind.True:
@@ -86,66 +84,66 @@ namespace EpicGames.Core
 		/// <summary>
 		/// Read a JSON file from disk and construct a JsonObject from it
 		/// </summary>
-		/// <param name="File">File to read from</param>
+		/// <param name="file">File to read from</param>
 		/// <returns>New JsonObject instance</returns>
-		public static JsonObject Read(FileReference File)
+		public static JsonObject Read(FileReference file)
 		{
-			string Text = FileReference.ReadAllText(File);
+			string text = FileReference.ReadAllText(file);
 			try
 			{
-				return Parse(Text);
+				return Parse(text);
 			}
-			catch(Exception Ex)
+			catch(Exception ex)
 			{
-				throw new JsonParseException("Unable to parse {0}: {1}", File, Ex.Message);
+				throw new JsonParseException("Unable to parse {0}: {1}", file, ex.Message);
 			}
 		}
 
 		/// <summary>
 		/// Tries to read a JSON file from disk
 		/// </summary>
-		/// <param name="FileName">File to read from</param>
-		/// <param name="Result">On success, receives the parsed object</param>
+		/// <param name="fileName">File to read from</param>
+		/// <param name="result">On success, receives the parsed object</param>
 		/// <returns>True if the file was read, false otherwise</returns>
-		public static bool TryRead(FileReference FileName, [NotNullWhen(true)] out JsonObject? Result)
+		public static bool TryRead(FileReference fileName, [NotNullWhen(true)] out JsonObject? result)
 		{
-			if (!FileReference.Exists(FileName))
+			if (!FileReference.Exists(fileName))
 			{
-				Result = null;
+				result = null;
 				return false;
 			}
 
-			string Text = FileReference.ReadAllText(FileName);
-			return TryParse(Text, out Result);
+			string text = FileReference.ReadAllText(fileName);
+			return TryParse(text, out result);
 		}
 
 		/// <summary>
 		/// Parse a JsonObject from the given raw text string
 		/// </summary>
-		/// <param name="Text">The text to parse</param>
+		/// <param name="text">The text to parse</param>
 		/// <returns>New JsonObject instance</returns>
-		public static JsonObject Parse(string Text)
+		public static JsonObject Parse(string text)
 		{
-			JsonDocument Document = JsonDocument.Parse(Text, new JsonDocumentOptions { AllowTrailingCommas = true });
-			return new JsonObject(Document.RootElement);
+			JsonDocument document = JsonDocument.Parse(text, new JsonDocumentOptions { AllowTrailingCommas = true });
+			return new JsonObject(document.RootElement);
 		}
 
 		/// <summary>
 		/// Try to parse a JsonObject from the given raw text string
 		/// </summary>
-		/// <param name="Text">The text to parse</param>
-		/// <param name="Result">On success, receives the new JsonObject</param>
+		/// <param name="text">The text to parse</param>
+		/// <param name="result">On success, receives the new JsonObject</param>
 		/// <returns>True if the object was parsed</returns>
-		public static bool TryParse(string Text, [NotNullWhen(true)] out JsonObject? Result)
+		public static bool TryParse(string text, [NotNullWhen(true)] out JsonObject? result)
 		{
 			try
 			{
-				Result = Parse(Text);
+				result = Parse(text);
 				return true;
 			}
 			catch (Exception)
 			{
-				Result = null;
+				result = null;
 				return false;
 			}
 		}
@@ -153,43 +151,40 @@ namespace EpicGames.Core
 		/// <summary>
 		/// List of key names in this object
 		/// </summary>
-		public IEnumerable<string> KeyNames
-		{
-			get { return RawObject.Keys; }
-		}
+		public IEnumerable<string> KeyNames => _rawObject.Keys;
 
 		/// <summary>
 		/// Gets a string field by the given name from the object, throwing an exception if it is not there or cannot be parsed.
 		/// </summary>
-		/// <param name="FieldName">Name of the field to get</param>
+		/// <param name="fieldName">Name of the field to get</param>
 		/// <returns>The field value</returns>
-		public string GetStringField(string FieldName)
+		public string GetStringField(string fieldName)
 		{
-			string? StringValue;
-			if (!TryGetStringField(FieldName, out StringValue))
+			string? stringValue;
+			if (!TryGetStringField(fieldName, out stringValue))
 			{
-				throw new JsonParseException("Missing or invalid '{0}' field", FieldName);
+				throw new JsonParseException("Missing or invalid '{0}' field", fieldName);
 			}
-			return StringValue;
+			return stringValue;
 		}
 
 		/// <summary>
 		/// Tries to read a string field by the given name from the object
 		/// </summary>
-		/// <param name="FieldName">Name of the field to get</param>
-		/// <param name="Result">On success, receives the field value</param>
+		/// <param name="fieldName">Name of the field to get</param>
+		/// <param name="result">On success, receives the field value</param>
 		/// <returns>True if the field could be read, false otherwise</returns>
-		public bool TryGetStringField(string FieldName, [NotNullWhen(true)] out string? Result)
+		public bool TryGetStringField(string fieldName, [NotNullWhen(true)] out string? result)
 		{
-			object? RawValue;
-			if (RawObject.TryGetValue(FieldName, out RawValue) && (RawValue is string))
+			object? rawValue;
+			if (_rawObject.TryGetValue(fieldName, out rawValue) && (rawValue is string strValue))
 			{
-				Result = (string)RawValue;
+				result = strValue;
 				return true;
 			}
 			else
 			{
-				Result = null;
+				result = null;
 				return false;
 			}
 		}
@@ -197,35 +192,35 @@ namespace EpicGames.Core
 		/// <summary>
 		/// Gets a string array field by the given name from the object, throwing an exception if it is not there or cannot be parsed.
 		/// </summary>
-		/// <param name="FieldName">Name of the field to get</param>
+		/// <param name="fieldName">Name of the field to get</param>
 		/// <returns>The field value</returns>
-		public string[] GetStringArrayField(string FieldName)
+		public string[] GetStringArrayField(string fieldName)
 		{
-			string[]? StringValues;
-			if (!TryGetStringArrayField(FieldName, out StringValues))
+			string[]? stringValues;
+			if (!TryGetStringArrayField(fieldName, out stringValues))
 			{
-				throw new JsonParseException("Missing or invalid '{0}' field", FieldName);
+				throw new JsonParseException("Missing or invalid '{0}' field", fieldName);
 			}
-			return StringValues;
+			return stringValues;
 		}
 
 		/// <summary>
 		/// Tries to read a string array field by the given name from the object
 		/// </summary>
-		/// <param name="FieldName">Name of the field to get</param>
-		/// <param name="Result">On success, receives the field value</param>
+		/// <param name="fieldName">Name of the field to get</param>
+		/// <param name="result">On success, receives the field value</param>
 		/// <returns>True if the field could be read, false otherwise</returns>
-		public bool TryGetStringArrayField(string FieldName, [NotNullWhen(true)] out string[]? Result)
+		public bool TryGetStringArrayField(string fieldName, [NotNullWhen(true)] out string[]? result)
 		{
-			object? RawValue;
-			if (RawObject.TryGetValue(FieldName, out RawValue) && (RawValue is IEnumerable<object>) && ((IEnumerable<object>)RawValue).All(x => x is string))
+			object? rawValue;
+			if (_rawObject.TryGetValue(fieldName, out rawValue) && (rawValue is IEnumerable<object> enumValue) && enumValue.All(x => x is string))
 			{
-				Result = ((IEnumerable<object>)RawValue).Select(x => (string)x).ToArray();
+				result = enumValue.Select(x => (string)x).ToArray();
 				return true;
 			}
 			else
 			{
-				Result = null;
+				result = null;
 				return false;
 			}
 		}
@@ -233,35 +228,35 @@ namespace EpicGames.Core
 		/// <summary>
 		/// Gets a boolean field by the given name from the object, throwing an exception if it is not there or cannot be parsed.
 		/// </summary>
-		/// <param name="FieldName">Name of the field to get</param>
+		/// <param name="fieldName">Name of the field to get</param>
 		/// <returns>The field value</returns>
-		public bool GetBoolField(string FieldName)
+		public bool GetBoolField(string fieldName)
 		{
-			bool BoolValue;
-			if (!TryGetBoolField(FieldName, out BoolValue))
+			bool boolValue;
+			if (!TryGetBoolField(fieldName, out boolValue))
 			{
-				throw new JsonParseException("Missing or invalid '{0}' field", FieldName);
+				throw new JsonParseException("Missing or invalid '{0}' field", fieldName);
 			}
-			return BoolValue;
+			return boolValue;
 		}
 
 		/// <summary>
 		/// Tries to read a bool field by the given name from the object
 		/// </summary>
-		/// <param name="FieldName">Name of the field to get</param>
-		/// <param name="Result">On success, receives the field value</param>
+		/// <param name="fieldName">Name of the field to get</param>
+		/// <param name="result">On success, receives the field value</param>
 		/// <returns>True if the field could be read, false otherwise</returns>
-		public bool TryGetBoolField(string FieldName, out bool Result)
+		public bool TryGetBoolField(string fieldName, out bool result)
 		{
-			object? RawValue;
-			if (RawObject.TryGetValue(FieldName, out RawValue) && (RawValue is Boolean))
+			object? rawValue;
+			if (_rawObject.TryGetValue(fieldName, out rawValue) && (rawValue is bool boolValue))
 			{
-				Result = (bool)RawValue;
+				result = boolValue;
 				return true;
 			}
 			else
 			{
-				Result = false;
+				result = false;
 				return false;
 			}
 		}
@@ -269,30 +264,30 @@ namespace EpicGames.Core
 		/// <summary>
 		/// Gets an integer field by the given name from the object, throwing an exception if it is not there or cannot be parsed.
 		/// </summary>
-		/// <param name="FieldName">Name of the field to get</param>
+		/// <param name="fieldName">Name of the field to get</param>
 		/// <returns>The field value</returns>
-		public int GetIntegerField(string FieldName)
+		public int GetIntegerField(string fieldName)
 		{
-			int IntegerValue;
-			if (!TryGetIntegerField(FieldName, out IntegerValue))
+			int integerValue;
+			if (!TryGetIntegerField(fieldName, out integerValue))
 			{
-				throw new JsonParseException("Missing or invalid '{0}' field", FieldName);
+				throw new JsonParseException("Missing or invalid '{0}' field", fieldName);
 			}
-			return IntegerValue;
+			return integerValue;
 		}
 
 		/// <summary>
 		/// Tries to read an integer field by the given name from the object
 		/// </summary>
-		/// <param name="FieldName">Name of the field to get</param>
-		/// <param name="Result">On success, receives the field value</param>
+		/// <param name="fieldName">Name of the field to get</param>
+		/// <param name="result">On success, receives the field value</param>
 		/// <returns>True if the field could be read, false otherwise</returns>
-		public bool TryGetIntegerField(string FieldName, out int Result)
+		public bool TryGetIntegerField(string fieldName, out int result)
 		{
-			object? RawValue;
-			if (!RawObject.TryGetValue(FieldName, out RawValue) || !int.TryParse(RawValue?.ToString(), out Result))
+			object? rawValue;
+			if (!_rawObject.TryGetValue(fieldName, out rawValue) || !Int32.TryParse(rawValue?.ToString(), out result))
 			{
-				Result = 0;
+				result = 0;
 				return false;
 			}
 			return true;
@@ -301,15 +296,15 @@ namespace EpicGames.Core
 		/// <summary>
 		/// Tries to read an unsigned integer field by the given name from the object
 		/// </summary>
-		/// <param name="FieldName">Name of the field to get</param>
-		/// <param name="Result">On success, receives the field value</param>
+		/// <param name="fieldName">Name of the field to get</param>
+		/// <param name="result">On success, receives the field value</param>
 		/// <returns>True if the field could be read, false otherwise</returns>
-		public bool TryGetUnsignedIntegerField(string FieldName, out uint Result)
+		public bool TryGetUnsignedIntegerField(string fieldName, out uint result)
 		{
-			object? RawValue;
-			if (!RawObject.TryGetValue(FieldName, out RawValue) || !uint.TryParse(RawValue?.ToString(), out Result))
+			object? rawValue;
+			if (!_rawObject.TryGetValue(fieldName, out rawValue) || !UInt32.TryParse(rawValue?.ToString(), out result))
 			{
-				Result = 0;
+				result = 0;
 				return false;
 			}
 			return true;
@@ -318,30 +313,30 @@ namespace EpicGames.Core
 		/// <summary>
 		/// Gets a double field by the given name from the object, throwing an exception if it is not there or cannot be parsed.
 		/// </summary>
-		/// <param name="FieldName">Name of the field to get</param>
+		/// <param name="fieldName">Name of the field to get</param>
 		/// <returns>The field value</returns>
-		public double GetDoubleField(string FieldName)
+		public double GetDoubleField(string fieldName)
 		{
-			double DoubleValue;
-			if (!TryGetDoubleField(FieldName, out DoubleValue))
+			double doubleValue;
+			if (!TryGetDoubleField(fieldName, out doubleValue))
 			{
-				throw new JsonParseException("Missing or invalid '{0}' field", FieldName);
+				throw new JsonParseException("Missing or invalid '{0}' field", fieldName);
 			}
-			return DoubleValue;
+			return doubleValue;
 		}
 
 		/// <summary>
 		/// Tries to read a double field by the given name from the object
 		/// </summary>
-		/// <param name="FieldName">Name of the field to get</param>
-		/// <param name="Result">On success, receives the field value</param>
+		/// <param name="fieldName">Name of the field to get</param>
+		/// <param name="result">On success, receives the field value</param>
 		/// <returns>True if the field could be read, false otherwise</returns>
-		public bool TryGetDoubleField(string FieldName, out double Result)
+		public bool TryGetDoubleField(string fieldName, out double result)
 		{
-			object? RawValue;
-			if (!RawObject.TryGetValue(FieldName, out RawValue) || !double.TryParse(RawValue?.ToString(), out Result))
+			object? rawValue;
+			if (!_rawObject.TryGetValue(fieldName, out rawValue) || !Double.TryParse(rawValue?.ToString(), out result))
 			{
-				Result = 0.0;
+				result = 0.0;
 				return false;
 			}
 			return true;
@@ -350,30 +345,30 @@ namespace EpicGames.Core
 		/// <summary>
 		/// Gets an enum field by the given name from the object, throwing an exception if it is not there or cannot be parsed.
 		/// </summary>
-		/// <param name="FieldName">Name of the field to get</param>
+		/// <param name="fieldName">Name of the field to get</param>
 		/// <returns>The field value</returns>
-		public T GetEnumField<T>(string FieldName) where T : struct
+		public T GetEnumField<T>(string fieldName) where T : struct
 		{
-			T EnumValue;
-			if (!TryGetEnumField(FieldName, out EnumValue))
+			T enumValue;
+			if (!TryGetEnumField(fieldName, out enumValue))
 			{
-				throw new JsonParseException("Missing or invalid '{0}' field", FieldName);
+				throw new JsonParseException("Missing or invalid '{0}' field", fieldName);
 			}
-			return EnumValue;
+			return enumValue;
 		}
 
 		/// <summary>
 		/// Tries to read an enum field by the given name from the object
 		/// </summary>
-		/// <param name="FieldName">Name of the field to get</param>
-		/// <param name="Result">On success, receives the field value</param>
+		/// <param name="fieldName">Name of the field to get</param>
+		/// <param name="result">On success, receives the field value</param>
 		/// <returns>True if the field could be read, false otherwise</returns>
-		public bool TryGetEnumField<T>(string FieldName, out T Result) where T : struct
+		public bool TryGetEnumField<T>(string fieldName, out T result) where T : struct
 		{
-			string? StringValue;
-			if (!TryGetStringField(FieldName, out StringValue) || !Enum.TryParse<T>(StringValue, true, out Result))
+			string? stringValue;
+			if (!TryGetStringField(fieldName, out stringValue) || !Enum.TryParse<T>(stringValue, true, out result))
 			{
-				Result = default(T);
+				result = default;
 				return false;
 			}
 			return true;
@@ -382,64 +377,64 @@ namespace EpicGames.Core
 		/// <summary>
 		/// Tries to read an enum array field by the given name from the object
 		/// </summary>
-		/// <param name="FieldName">Name of the field to get</param>
-		/// <param name="Result">On success, receives the field value</param>
+		/// <param name="fieldName">Name of the field to get</param>
+		/// <param name="result">On success, receives the field value</param>
 		/// <returns>True if the field could be read, false otherwise</returns>
-		public bool TryGetEnumArrayField<T>(string FieldName, [NotNullWhen(true)] out T[]? Result) where T : struct
+		public bool TryGetEnumArrayField<T>(string fieldName, [NotNullWhen(true)] out T[]? result) where T : struct
 		{
-			string[]? StringValues;
-			if (!TryGetStringArrayField(FieldName, out StringValues))
+			string[]? stringValues;
+			if (!TryGetStringArrayField(fieldName, out stringValues))
 			{
-				Result = null;
+				result = null;
 				return false;
 			}
 
-			T[] EnumValues = new T[StringValues.Length];
-			for (int Idx = 0; Idx < StringValues.Length; Idx++)
+			T[] enumValues = new T[stringValues.Length];
+			for (int idx = 0; idx < stringValues.Length; idx++)
 			{
-				if (!Enum.TryParse<T>(StringValues[Idx], true, out EnumValues[Idx]))
+				if (!Enum.TryParse<T>(stringValues[idx], true, out enumValues[idx]))
 				{
-					Result = null;
+					result = null;
 					return false;
 				}
 			}
 
-			Result = EnumValues;
+			result = enumValues;
 			return true;
 		}
 
 		/// <summary>
 		/// Gets an object field by the given name from the object, throwing an exception if it is not there or cannot be parsed.
 		/// </summary>
-		/// <param name="FieldName">Name of the field to get</param>
+		/// <param name="fieldName">Name of the field to get</param>
 		/// <returns>The field value</returns>
-		public JsonObject GetObjectField(string FieldName)
+		public JsonObject GetObjectField(string fieldName)
 		{
-			JsonObject? Result;
-			if (!TryGetObjectField(FieldName, out Result))
+			JsonObject? result;
+			if (!TryGetObjectField(fieldName, out result))
 			{
-				throw new JsonParseException("Missing or invalid '{0}' field", FieldName);
+				throw new JsonParseException("Missing or invalid '{0}' field", fieldName);
 			}
-			return Result;
+			return result;
 		}
 
 		/// <summary>
 		/// Tries to read an object field by the given name from the object
 		/// </summary>
-		/// <param name="FieldName">Name of the field to get</param>
-		/// <param name="Result">On success, receives the field value</param>
+		/// <param name="fieldName">Name of the field to get</param>
+		/// <param name="result">On success, receives the field value</param>
 		/// <returns>True if the field could be read, false otherwise</returns>
-		public bool TryGetObjectField(string FieldName, [NotNullWhen(true)] out JsonObject? Result)
+		public bool TryGetObjectField(string fieldName, [NotNullWhen(true)] out JsonObject? result)
 		{
-			object? RawValue;
-			if (RawObject.TryGetValue(FieldName, out RawValue) && (RawValue is Dictionary<string, object?>))
+			object? rawValue;
+			if (_rawObject.TryGetValue(fieldName, out rawValue) && (rawValue is Dictionary<string, object?> dictValue))
 			{
-				Result = new JsonObject((Dictionary<string, object?>)RawValue);
+				result = new JsonObject(dictValue);
 				return true;
 			}
 			else
 			{
-				Result = null;
+				result = null;
 				return false;
 			}
 		}
@@ -447,35 +442,35 @@ namespace EpicGames.Core
 		/// <summary>
 		/// Gets an object array field by the given name from the object, throwing an exception if it is not there or cannot be parsed.
 		/// </summary>
-		/// <param name="FieldName">Name of the field to get</param>
+		/// <param name="fieldName">Name of the field to get</param>
 		/// <returns>The field value</returns>
-		public JsonObject[] GetObjectArrayField(string FieldName)
+		public JsonObject[] GetObjectArrayField(string fieldName)
 		{
-			JsonObject[]? Result;
-			if (!TryGetObjectArrayField(FieldName, out Result))
+			JsonObject[]? result;
+			if (!TryGetObjectArrayField(fieldName, out result))
 			{
-				throw new JsonParseException("Missing or invalid '{0}' field", FieldName);
+				throw new JsonParseException("Missing or invalid '{0}' field", fieldName);
 			}
-			return Result;
+			return result;
 		}
 
 		/// <summary>
 		/// Tries to read an object array field by the given name from the object
 		/// </summary>
-		/// <param name="FieldName">Name of the field to get</param>
-		/// <param name="Result">On success, receives the field value</param>
+		/// <param name="fieldName">Name of the field to get</param>
+		/// <param name="result">On success, receives the field value</param>
 		/// <returns>True if the field could be read, false otherwise</returns>
-		public bool TryGetObjectArrayField(string FieldName, [NotNullWhen(true)] out JsonObject[]? Result)
+		public bool TryGetObjectArrayField(string fieldName, [NotNullWhen(true)] out JsonObject[]? result)
 		{
-			object? RawValue;
-			if (RawObject.TryGetValue(FieldName, out RawValue) && (RawValue is IEnumerable<object>) && ((IEnumerable<object>)RawValue).All(x => x is Dictionary<string, object>))
+			object? rawValue;
+			if (_rawObject.TryGetValue(fieldName, out rawValue) && (rawValue is IEnumerable<object> enumValue) && enumValue.All(x => x is Dictionary<string, object>))
 			{
-				Result = ((IEnumerable<object>)RawValue).Select(x => new JsonObject((Dictionary<string, object?>)x)).ToArray();
+				result = enumValue.Select(x => new JsonObject((Dictionary<string, object?>)x)).ToArray();
 				return true;
 			}
 			else
 			{
-				Result = null;
+				result = null;
 				return false;
 			}
 		}
diff --git a/Engine/Source/Programs/Shared/EpicGames.Core/JsonSchema.cs b/Engine/Source/Programs/Shared/EpicGames.Core/JsonSchema.cs
--- a/Engine/Source/Programs/Shared/EpicGames.Core/JsonSchema.cs
+++ b/Engine/Source/Programs/Shared/EpicGames.Core/JsonSchema.cs
@@ -3,11 +3,8 @@
 using System;
 using System.Collections.Generic;
 using System.IO;
-using System.Linq;
 using System.Reflection;
-using System.Text;
 using System.Text.Json;
-using System.Threading.Tasks;
 using System.Xml;
 
 namespace EpicGames.Core
@@ -20,8 +17,8 @@ namespace EpicGames.Core
 		/// <inheritdoc cref="Utf8JsonWriter.WriteStartObject()"/>
 		void WriteStartObject();
 
-		/// <inheritdoc cref="Utf8JsonWriter.WriteStartObject(string)"/>
-		void WriteStartObject(string Name);
+		/// <inheritdoc cref="Utf8JsonWriter.WriteStartObject(String)"/>
+		void WriteStartObject(string name);
 
 		/// <inheritdoc cref="Utf8JsonWriter.WriteEndObject()"/>
 		void WriteEndObject();
@@ -29,26 +26,26 @@ namespace EpicGames.Core
 		/// <inheritdoc cref="Utf8JsonWriter.WriteStartArray()"/>
 		void WriteStartArray();
 
-		/// <inheritdoc cref="Utf8JsonWriter.WriteStartArray(string)"/>
-		void WriteStartArray(string Name);
+		/// <inheritdoc cref="Utf8JsonWriter.WriteStartArray(String)"/>
+		void WriteStartArray(string name);
 
 		/// <inheritdoc cref="Utf8JsonWriter.WriteEndObject()"/>
 		void WriteEndArray();
 
-		/// <inheritdoc cref="Utf8JsonWriter.WriteBoolean(string, bool)"/>
-		void WriteBoolean(string Name, bool Value);
+		/// <inheritdoc cref="Utf8JsonWriter.WriteBoolean(String, Bool)"/>
+		void WriteBoolean(string name, bool value);
 
-		/// <inheritdoc cref="Utf8JsonWriter.WriteString(string, string)"/>
-		void WriteString(string Key, string Value);
+		/// <inheritdoc cref="Utf8JsonWriter.WriteString(String, String)"/>
+		void WriteString(string key, string value);
 
-		/// <inheritdoc cref="Utf8JsonWriter.WriteStringValue(string)"/>
-		void WriteStringValue(string Name);
+		/// <inheritdoc cref="Utf8JsonWriter.WriteStringValue(String)"/>
+		void WriteStringValue(string name);
 
 		/// <summary>
 		/// Serialize a type to JSON
 		/// </summary>
-		/// <param name="Type"></param>
-		void WriteType(JsonSchemaType Type);
+		/// <param name="type"></param>
+		void WriteType(JsonSchemaType type);
 	}
 
 	/// <summary>
@@ -64,19 +61,19 @@ namespace EpicGames.Core
 			public Dictionary<JsonSchemaType, int> TypeRefCount { get; } = new Dictionary<JsonSchemaType, int>();
 
 			/// <inheritdoc/>
-			public void WriteBoolean(string Name, bool Value) { }
+			public void WriteBoolean(string name, bool value) { }
 
 			/// <inheritdoc/>
-			public void WriteString(string Key, string Value) { }
+			public void WriteString(string key, string value) { }
 
 			/// <inheritdoc/>
-			public void WriteStringValue(string Name) { }
+			public void WriteStringValue(string name) { }
 
 			/// <inheritdoc/>
 			public void WriteStartObject() { }
 
 			/// <inheritdoc/>
-			public void WriteStartObject(string Name) { }
+			public void WriteStartObject(string name) { }
 
 			/// <inheritdoc/>
 			public void WriteEndObject() { }
@@ -85,24 +82,24 @@ namespace EpicGames.Core
 			public void WriteStartArray() { }
 
 			/// <inheritdoc/>
-			public void WriteStartArray(string Name) { }
+			public void WriteStartArray(string name) { }
 
 			/// <inheritdoc/>
 			public void WriteEndArray() { }
 
 			/// <inheritdoc/>
-			public void WriteType(JsonSchemaType Type)
+			public void WriteType(JsonSchemaType type)
 			{
-				if (!(Type is JsonSchemaPrimitiveType))
+				if (!(type is JsonSchemaPrimitiveType))
 				{
-					TypeRefCount.TryGetValue(Type, out int RefCount);
-					if (RefCount < 2)
+					TypeRefCount.TryGetValue(type, out int refCount);
+					if (refCount < 2)
 					{
-						TypeRefCount[Type] = ++RefCount;
+						TypeRefCount[type] = ++refCount;
 					}
-					if (RefCount < 2)
+					if (refCount < 2)
 					{
-						Type.Write(this);
+						type.Write(this);
 					}
 				}
 			}
@@ -116,64 +113,64 @@ namespace EpicGames.Core
 			/// <summary>
 			/// Raw Json output
 			/// </summary>
-			Utf8JsonWriter JsonWriter;
+			readonly Utf8JsonWriter _jsonWriter;
 
 			/// <summary>
 			/// Mapping of type to definition name
 			/// </summary>
-			Dictionary<JsonSchemaType, string> TypeToDefinition { get; }
+			readonly Dictionary<JsonSchemaType, string> _typeToDefinition;
 
 			/// <summary>
 			/// Constructor
 			/// </summary>
-			/// <param name="Writer"></param>
-			/// <param name="TypeToDefinition"></param>
-			public JsonSchemaWriter(Utf8JsonWriter Writer, Dictionary<JsonSchemaType, string> TypeToDefinition)
+			/// <param name="writer"></param>
+			/// <param name="typeToDefinition"></param>
+			public JsonSchemaWriter(Utf8JsonWriter writer, Dictionary<JsonSchemaType, string> typeToDefinition)
 			{
-				this.JsonWriter = Writer;
-				this.TypeToDefinition = TypeToDefinition;
+				_jsonWriter = writer;
+				_typeToDefinition = typeToDefinition;
 			}
 
 			/// <inheritdoc/>
-			public void WriteBoolean(string Name, bool Value) => JsonWriter.WriteBoolean(Name, Value);
+			public void WriteBoolean(string name, bool value) => _jsonWriter.WriteBoolean(name, value);
 
 			/// <inheritdoc/>
-			public void WriteString(string Key, string Value) => JsonWriter.WriteString(Key, Value);
+			public void WriteString(string key, string value) => _jsonWriter.WriteString(key, value);
 
 			/// <inheritdoc/>
-			public void WriteStringValue(string Name) => JsonWriter.WriteStringValue(Name);
+			public void WriteStringValue(string name) => _jsonWriter.WriteStringValue(name);
 
 			/// <inheritdoc/>
-			public void WriteStartObject() => JsonWriter.WriteStartObject();
+			public void WriteStartObject() => _jsonWriter.WriteStartObject();
 
 			/// <inheritdoc/>
-			public void WriteStartObject(string Name) => JsonWriter.WriteStartObject(Name);
+			public void WriteStartObject(string name) => _jsonWriter.WriteStartObject(name);
 
 			/// <inheritdoc/>
-			public void WriteEndObject() => JsonWriter.WriteEndObject();
+			public void WriteEndObject() => _jsonWriter.WriteEndObject();
 
 			/// <inheritdoc/>
-			public void WriteStartArray() => JsonWriter.WriteStartArray();
+			public void WriteStartArray() => _jsonWriter.WriteStartArray();
 
 			/// <inheritdoc/>
-			public void WriteStartArray(string Name) => JsonWriter.WriteStartArray(Name);
+			public void WriteStartArray(string name) => _jsonWriter.WriteStartArray(name);
 
 			/// <inheritdoc/>
-			public void WriteEndArray() => JsonWriter.WriteEndArray();
+			public void WriteEndArray() => _jsonWriter.WriteEndArray();
 
 			/// <summary>
 			/// Writes a type, either inline or as a reference to a definition elsewhere
 			/// </summary>
-			/// <param name="Type"></param>
-			public void WriteType(JsonSchemaType Type)
+			/// <param name="type"></param>
+			public void WriteType(JsonSchemaType type)
 			{
-				if (TypeToDefinition.TryGetValue(Type, out string? Definition))
+				if (_typeToDefinition.TryGetValue(type, out string? definition))
 				{
-					JsonWriter.WriteString("$ref", $"#/definitions/{Definition}");
+					_jsonWriter.WriteString("$ref", $"#/definitions/{definition}");
 				}
 				else
 				{
-					Type.Write(this);
+					type.Write(this);
 				}
 			}
 		}
@@ -191,163 +188,139 @@ namespace EpicGames.Core
 		/// <summary>
 		/// Constructor
 		/// </summary>
-		/// <param name="Id">Id for the schema</param>
-		/// <param name="RootType"></param>
-		public JsonSchema(string? Id, JsonSchemaType RootType)
+		/// <param name="id">Id for the schema</param>
+		/// <param name="rootType"></param>
+		public JsonSchema(string? id, JsonSchemaType rootType)
 		{
-			this.Id = Id;
-			this.RootType = RootType;
+			Id = id;
+			RootType = rootType;
 		}
 
 		/// <summary>
 		/// Write this schema to a byte array
 		/// </summary>
-		/// <param name="Writer"></param>
-		public void Write(Utf8JsonWriter Writer)
+		/// <param name="writer"></param>
+		public void Write(Utf8JsonWriter writer)
 		{
 			// Determine reference counts for each type. Any type referenced at least twice will be split off into a separate definition.
-			JsonTypeRefCollector RefCollector = new JsonTypeRefCollector();
-			RefCollector.WriteType(RootType);
+			JsonTypeRefCollector refCollector = new JsonTypeRefCollector();
+			refCollector.WriteType(RootType);
 
 			// Assign names to each type definition
-			HashSet<string> DefinitionNames = new HashSet<string>();
-			Dictionary<JsonSchemaType, string> TypeToDefinition = new Dictionary<JsonSchemaType, string>();
-			foreach ((JsonSchemaType Type, int RefCount) in RefCollector.TypeRefCount)
+			HashSet<string> definitionNames = new HashSet<string>();
+			Dictionary<JsonSchemaType, string> typeToDefinition = new Dictionary<JsonSchemaType, string>();
+			foreach ((JsonSchemaType type, int refCount) in refCollector.TypeRefCount)
 			{
-				if (RefCount > 1)
+				if (refCount > 1)
 				{
-					string BaseName = Type.Name ?? "unnamed";
+					string baseName = type.Name ?? "unnamed";
 
-					string Name = BaseName;
-					for (int Idx = 1; !DefinitionNames.Add(Name); Idx++)
+					string name = baseName;
+					for (int idx = 1; !definitionNames.Add(name); idx++)
 					{
-						Name = $"{BaseName}{Idx}";
+						name = $"{baseName}{idx}";
 					}
 
-					TypeToDefinition[Type] = Name;
+					typeToDefinition[type] = name;
 				}
 			}
 
 			// Write the schema
-			Writer.WriteStartObject();
-			Writer.WriteString("$schema", "http://json-schema.org/draft-04/schema#");
+			writer.WriteStartObject();
+			writer.WriteString("$schema", "http://json-schema.org/draft-04/schema#");
 			if (Id != null)
 			{
-				Writer.WriteString("$id", Id);
+				writer.WriteString("$id", Id);
 			}
 
-			JsonSchemaWriter SchemaWriter = new JsonSchemaWriter(Writer, TypeToDefinition);
-			RootType.Write(SchemaWriter);
+			JsonSchemaWriter schemaWriter = new JsonSchemaWriter(writer, typeToDefinition);
+			RootType.Write(schemaWriter);
 
-			if (TypeToDefinition.Count > 0)
+			if (typeToDefinition.Count > 0)
 			{
-				Writer.WriteStartObject("definitions");
-				foreach ((JsonSchemaType Type, string RefName) in TypeToDefinition)
+				writer.WriteStartObject("definitions");
+				foreach ((JsonSchemaType type, string refName) in typeToDefinition)
 				{
-					Writer.WriteStartObject(RefName);
-					Type.Write(SchemaWriter);
-					Writer.WriteEndObject();
+					writer.WriteStartObject(refName);
+					type.Write(schemaWriter);
+					writer.WriteEndObject();
 				}
-				Writer.WriteEndObject();
+				writer.WriteEndObject();
 			}
 
-			Writer.WriteEndObject();
+			writer.WriteEndObject();
 		}
 
 		/// <summary>
 		/// Write this schema to a stream
 		/// </summary>
-		/// <param name="Stream">The output stream</param>
-		public void Write(Stream Stream)
+		/// <param name="stream">The output stream</param>
+		public void Write(Stream stream)
 		{
-			using (Utf8JsonWriter Writer = new Utf8JsonWriter(Stream, new JsonWriterOptions { Indented = true }))
+			using (Utf8JsonWriter writer = new Utf8JsonWriter(stream, new JsonWriterOptions { Indented = true }))
 			{
-				Write(Writer);
+				Write(writer);
 			}
 		}
 
 		/// <summary>
 		/// Writes this schema to a file
 		/// </summary>
-		/// <param name="File">The output file</param>
-		public void Write(FileReference File)
+		/// <param name="file">The output file</param>
+		public void Write(FileReference file)
 		{
-			using (FileStream Stream = FileReference.Open(File, FileMode.Create))
+			using (FileStream stream = FileReference.Open(file, FileMode.Create))
 			{
-				Write(Stream);
+				Write(stream);
 			}
 		}
 
 		/// <summary>
 		/// Constructs a Json schema from a type
 		/// </summary>
-		/// <param name="Type">The type to construct from</param>
-		/// <param name="XmlDoc">C# Xml documentation file, to use for property descriptions</param>
+		/// <param name="type">The type to construct from</param>
+		/// <param name="xmlDoc">C# Xml documentation file, to use for property descriptions</param>
 		/// <returns>New schema object</returns>
-		public static JsonSchema FromType(Type Type, XmlDocument? XmlDoc)
+		public static JsonSchema FromType(Type type, XmlDocument? xmlDoc)
 		{
-			JsonSchemaAttribute? SchemaAttribute = Type.GetCustomAttribute<JsonSchemaAttribute>();
-			return new JsonSchema(SchemaAttribute?.Id, CreateSchemaType(Type, new Dictionary<Type, JsonSchemaType>(), XmlDoc));
-		}
-
-		/// <summary>
-		/// Gets the description for a type
-		/// </summary>
-		/// <param name="Type"></param>
-		/// <param name="XmlDoc"></param>
-		/// <returns></returns>
-		static string? GetTypeDescription(Type Type, XmlDocument? XmlDoc)
-		{
-			if (XmlDoc == null)
-			{
-				return null;
-			}
-
-			string Selector = $"//member[@name='T:{Type.FullName}']/summary";
-
-			XmlNode Node = XmlDoc.SelectSingleNode(Selector);
-			if (Node == null)
-			{
-				return null;
-			}
-
-			return Node.InnerText.Trim().Replace("\r\n", "\n", StringComparison.Ordinal);
+			JsonSchemaAttribute? schemaAttribute = type.GetCustomAttribute<JsonSchemaAttribute>();
+			return new JsonSchema(schemaAttribute?.Id, CreateSchemaType(type, new Dictionary<Type, JsonSchemaType>(), xmlDoc));
 		}
 
 		/// <summary>
 		/// Gets a property description from Xml documentation file
 		/// </summary>
-		/// <param name="Type"></param>
-		/// <param name="Name"></param>
-		/// <param name="XmlDoc"></param>
+		/// <param name="type"></param>
+		/// <param name="name"></param>
+		/// <param name="xmlDoc"></param>
 		/// <returns></returns>
-		static string? GetPropertyDescription(Type Type, string Name, XmlDocument? XmlDoc)
+		static string? GetPropertyDescription(Type type, string name, XmlDocument? xmlDoc)
 		{
-			if (XmlDoc == null)
+			if (xmlDoc == null)
 			{
 				return null;
 			}
 
-			string Selector = $"//member[@name='P:{Type.FullName}.{Name}']/summary";
+			string selector = $"//member[@name='P:{type.FullName}.{name}']/summary";
 
-			XmlNode Node = XmlDoc.SelectSingleNode(Selector);
-			if (Node == null)
+			XmlNode node = xmlDoc.SelectSingleNode(selector);
+			if (node == null)
 			{
 				return null;
 			}
 
-			return Node.InnerText.Trim().Replace("\r\n", "\n", StringComparison.Ordinal);
+			return node.InnerText.Trim().Replace("\r\n", "\n", StringComparison.Ordinal);
 		}
 
 		/// <summary>
 		/// Constructs a schema type from the given type object
 		/// </summary>
-		/// <param name="Type"></param>
-		/// <param name="XmlDoc"></param>
+		/// <param name="type"></param>
+		/// <param name="xmlDoc"></param>
 		/// <returns></returns>
-		static JsonSchemaType CreateSchemaType(Type Type, Dictionary<Type, JsonSchemaType> TypeCache, XmlDocument? XmlDoc)
+		static JsonSchemaType CreateSchemaType(Type type, Dictionary<Type, JsonSchemaType> typeCache, XmlDocument? xmlDoc)
 		{
-			switch (Type.GetTypeCode(Type))
+			switch (Type.GetTypeCode(type))
 			{
 				case TypeCode.Boolean:
 					return new JsonSchemaBoolean();
@@ -367,99 +340,99 @@ namespace EpicGames.Core
 					return new JsonSchemaString();
 			}
 
-			JsonSchemaTypeAttribute? Attribute = Type.GetCustomAttribute<JsonSchemaTypeAttribute>();
-			switch (Attribute)
+			JsonSchemaTypeAttribute? attribute = type.GetCustomAttribute<JsonSchemaTypeAttribute>();
+			switch (attribute)
 			{
-				case JsonSchemaStringAttribute String:
-					return new JsonSchemaString(String.Format);
+				case JsonSchemaStringAttribute str:
+					return new JsonSchemaString(str.Format);
 			}
 
-			if (Type.IsEnum)
+			if (type.IsEnum)
 			{
-				return new JsonSchemaEnum(Enum.GetNames(Type)) { Name = Type.Name };
+				return new JsonSchemaEnum(Enum.GetNames(type)) { Name = type.Name };
 			}
-			if (Type.IsGenericType && Type.GetGenericTypeDefinition() == typeof(Nullable<>))
+			if (type.IsGenericType && type.GetGenericTypeDefinition() == typeof(Nullable<>))
 			{
-				return CreateSchemaType(Type.GetGenericArguments()[0], TypeCache, XmlDoc);
+				return CreateSchemaType(type.GetGenericArguments()[0], typeCache, xmlDoc);
 			}
-			if (Type == typeof(DateTime) || Type == typeof(DateTimeOffset))
+			if (type == typeof(DateTime) || type == typeof(DateTimeOffset))
 			{
 				return new JsonSchemaString(JsonSchemaStringFormat.DateTime);
 			}
 
-			Type[] Interfaces = Type.GetInterfaces();
-			foreach (Type Interface in Interfaces)
+			Type[] interfaceTypes = type.GetInterfaces();
+			foreach (Type interfaceType in interfaceTypes)
 			{
-				if (Interface.IsGenericType)
+				if (interfaceType.IsGenericType)
 				{
-					Type[] Arguments = Interface.GetGenericArguments();
-					if (Interface.GetGenericTypeDefinition() == typeof(IList<>))
+					Type[] arguments = interfaceType.GetGenericArguments();
+					if (interfaceType.GetGenericTypeDefinition() == typeof(IList<>))
 					{
-						return new JsonSchemaArray(CreateSchemaType(Arguments[0], TypeCache, XmlDoc));
+						return new JsonSchemaArray(CreateSchemaType(arguments[0], typeCache, xmlDoc));
 					}
-					if (Interface.GetGenericTypeDefinition() == typeof(IDictionary<,>))
+					if (interfaceType.GetGenericTypeDefinition() == typeof(IDictionary<,>))
 					{
-						JsonSchemaObject Object = new JsonSchemaObject();
-						Object.AdditionalProperties = CreateSchemaType(Arguments[1], TypeCache, XmlDoc);
-						return Object;
+						JsonSchemaObject obj = new JsonSchemaObject();
+						obj.AdditionalProperties = CreateSchemaType(arguments[1], typeCache, xmlDoc);
+						return obj;
 					}
 				}
 			}
 
-			if (Type.IsClass)
+			if (type.IsClass)
 			{
-				if (TypeCache.TryGetValue(Type, out JsonSchemaType? SchemaType))
+				if (typeCache.TryGetValue(type, out JsonSchemaType? schemaType))
 				{
-					return SchemaType;
+					return schemaType;
 				}
 
-				JsonKnownTypesAttribute? KnownTypes = Type.GetCustomAttribute<JsonKnownTypesAttribute>(false);
-				if (KnownTypes != null)
+				JsonKnownTypesAttribute? knownTypes = type.GetCustomAttribute<JsonKnownTypesAttribute>(false);
+				if (knownTypes != null)
 				{
-					JsonSchemaOneOf Object = new JsonSchemaOneOf();
-					TypeCache[Type] = Object;
-					SetOneOfProperties(Object, Type, KnownTypes.Types, TypeCache, XmlDoc);
-					return Object;
+					JsonSchemaOneOf obj = new JsonSchemaOneOf();
+					typeCache[type] = obj;
+					SetOneOfProperties(obj, type, knownTypes.Types, typeCache, xmlDoc);
+					return obj;
 				}
 				else
 				{
-					JsonSchemaObject Object = new JsonSchemaObject();
-					TypeCache[Type] = Object;
-					SetObjectProperties(Object, Type, TypeCache, XmlDoc);
-					return Object;
+					JsonSchemaObject obj = new JsonSchemaObject();
+					typeCache[type] = obj;
+					SetObjectProperties(obj, type, typeCache, xmlDoc);
+					return obj;
 				}
 			}
 
-			throw new Exception($"Unknown type for schema generation: {Type}");
+			throw new Exception($"Unknown type for schema generation: {type}");
 		}
 
-		static void SetOneOfProperties(JsonSchemaOneOf Object, Type Type, Type[] KnownTypes, Dictionary<Type, JsonSchemaType> TypeCache, XmlDocument? XmlDoc)
+		static void SetOneOfProperties(JsonSchemaOneOf obj, Type type, Type[] knownTypes, Dictionary<Type, JsonSchemaType> typeCache, XmlDocument? xmlDoc)
 		{
-			Object.Name = Type.Name;
+			obj.Name = type.Name;
 
-			foreach (Type KnownType in KnownTypes)
+			foreach (Type knownType in knownTypes)
 			{
-				JsonDiscriminatorAttribute? Attribute = KnownType.GetCustomAttribute<JsonDiscriminatorAttribute>();
-				if (Attribute != null)
+				JsonDiscriminatorAttribute? attribute = knownType.GetCustomAttribute<JsonDiscriminatorAttribute>();
+				if (attribute != null)
 				{
-					JsonSchemaObject KnownObject = new JsonSchemaObject();
-					KnownObject.Properties.Add(new JsonSchemaProperty("type", "Type discriminator", new JsonSchemaEnum(new[] { Attribute.Name })));
-					SetObjectProperties(KnownObject, KnownType, TypeCache, XmlDoc);
-					Object.Types.Add(KnownObject);
+					JsonSchemaObject knownObject = new JsonSchemaObject();
+					knownObject.Properties.Add(new JsonSchemaProperty("type", "Type discriminator", new JsonSchemaEnum(new[] { attribute.Name })));
+					SetObjectProperties(knownObject, knownType, typeCache, xmlDoc);
+					obj.Types.Add(knownObject);
 				}
 			}
 		}
 
-		static void SetObjectProperties(JsonSchemaObject Object, Type Type, Dictionary<Type, JsonSchemaType> TypeCache, XmlDocument? XmlDoc)
+		static void SetObjectProperties(JsonSchemaObject obj, Type type, Dictionary<Type, JsonSchemaType> typeCache, XmlDocument? xmlDoc)
 		{
-			Object.Name = Type.Name;
+			obj.Name = type.Name;
 
-			PropertyInfo[] Properties = Type.GetProperties(BindingFlags.Instance | BindingFlags.Public);
-			foreach (PropertyInfo Property in Properties)
+			PropertyInfo[] properties = type.GetProperties(BindingFlags.Instance | BindingFlags.Public);
+			foreach (PropertyInfo property in properties)
 			{
-				string? Description = GetPropertyDescription(Type, Property.Name, XmlDoc);
-				JsonSchemaType PropertyType = CreateSchemaType(Property.PropertyType, TypeCache, XmlDoc);
-				Object.Properties.Add(new JsonSchemaProperty(Property.Name, Description, PropertyType));
+				string? description = GetPropertyDescription(type, property.Name, xmlDoc);
+				JsonSchemaType propertyType = CreateSchemaType(property.PropertyType, typeCache, xmlDoc);
+				obj.Properties.Add(new JsonSchemaProperty(property.Name, description, propertyType));
 			}
 		}
 	}
diff --git a/Engine/Source/Programs/Shared/EpicGames.Core/JsonSchemaAttributes.cs b/Engine/Source/Programs/Shared/EpicGames.Core/JsonSchemaAttributes.cs
--- a/Engine/Source/Programs/Shared/EpicGames.Core/JsonSchemaAttributes.cs
+++ b/Engine/Source/Programs/Shared/EpicGames.Core/JsonSchemaAttributes.cs
@@ -1,8 +1,6 @@
 // Copyright Epic Games, Inc. All Rights Reserved.
 
 using System;
-using System.Collections.Generic;
-using System.Text;
 
 namespace EpicGames.Core
 {
@@ -20,10 +18,10 @@ namespace EpicGames.Core
 		/// <summary>
 		/// Constructor
 		/// </summary>
-		/// <param name="Id"></param>
-		public JsonSchemaAttribute(string Id)
+		/// <param name="id"></param>
+		public JsonSchemaAttribute(string id)
 		{
-			this.Id = Id;
+			Id = id;
 		}
 	}
 
@@ -51,25 +49,25 @@ namespace EpicGames.Core
 		/// <summary>
 		/// Constructor
 		/// </summary>
-		/// <param name="Name"></param>
-		/// <param name="Description"></param>
-		/// <param name="FileMatch">File patterns to match</param>
-		public JsonSchemaCatalogAttribute(string Name, string Description, string? FileMatch)
-			: this(Name, Description, (FileMatch == null) ? (string[]?)null : new[] { FileMatch })
+		/// <param name="name"></param>
+		/// <param name="description"></param>
+		/// <param name="fileMatch">File patterns to match</param>
+		public JsonSchemaCatalogAttribute(string name, string description, string? fileMatch)
+			: this(name, description, (fileMatch == null) ? (string[]?)null : new[] { fileMatch })
 		{
 		}
 
 		/// <summary>
 		/// Constructor
 		/// </summary>
-		/// <param name="Name"></param>
-		/// <param name="Description"></param>
-		/// <param name="FileMatch">File patterns to match</param>
-		public JsonSchemaCatalogAttribute(string Name, string Description, string[]? FileMatch)
+		/// <param name="name"></param>
+		/// <param name="description"></param>
+		/// <param name="fileMatch">File patterns to match</param>
+		public JsonSchemaCatalogAttribute(string name, string description, string[]? fileMatch)
 		{
-			this.Name = Name;
-			this.Description = Description;
-			this.FileMatch = FileMatch;
+			Name = name;
+			Description = description;
+			FileMatch = fileMatch;
 		}
 	}
 
diff --git a/Engine/Source/Programs/Shared/EpicGames.Core/JsonSchemaType.cs b/Engine/Source/Programs/Shared/EpicGames.Core/JsonSchemaType.cs
--- a/Engine/Source/Programs/Shared/EpicGames.Core/JsonSchemaType.cs
+++ b/Engine/Source/Programs/Shared/EpicGames.Core/JsonSchemaType.cs
@@ -2,11 +2,6 @@
 
 using System;
 using System.Collections.Generic;
-using System.IO;
-using System.Linq;
-using System.Text;
-using System.Text.Json;
-using System.Threading.Tasks;
 
 namespace EpicGames.Core
 {
@@ -23,8 +18,8 @@ namespace EpicGames.Core
 		/// <summary>
 		/// Write this type to the given archive
 		/// </summary>
-		/// <param name="Writer"></param>
-		public abstract void Write(IJsonSchemaWriter Writer);
+		/// <param name="writer"></param>
+		public abstract void Write(IJsonSchemaWriter writer);
 	}
 
 	/// <summary>
@@ -40,9 +35,9 @@ namespace EpicGames.Core
 	public class JsonSchemaBoolean : JsonSchemaPrimitiveType
 	{
 		/// <inheritdoc/>
-		public override void Write(IJsonSchemaWriter Writer)
+		public override void Write(IJsonSchemaWriter writer)
 		{
-			Writer.WriteString("type", "boolean");
+			writer.WriteString("type", "boolean");
 		}
 	}
 
@@ -52,9 +47,9 @@ namespace EpicGames.Core
 	public class JsonSchemaInteger : JsonSchemaPrimitiveType
 	{
 		/// <inheritdoc/>
-		public override void Write(IJsonSchemaWriter Writer)
+		public override void Write(IJsonSchemaWriter writer)
 		{
-			Writer.WriteString("type", "integer");
+			writer.WriteString("type", "integer");
 		}
 	}
 
@@ -64,9 +59,9 @@ namespace EpicGames.Core
 	public class JsonSchemaNumber : JsonSchemaPrimitiveType
 	{
 		/// <inheritdoc/>
-		public override void Write(IJsonSchemaWriter Writer)
+		public override void Write(IJsonSchemaWriter writer)
 		{
-			Writer.WriteString("type", "integer");
+			writer.WriteString("type", "integer");
 		}
 	}
 
@@ -124,24 +119,24 @@ namespace EpicGames.Core
 		/// <summary>
 		/// Constructor
 		/// </summary>
-		/// <param name="Format"></param>
-		public JsonSchemaString(JsonSchemaStringFormat? Format = null, string? Pattern = null)
+		/// <param name="format"></param>
+		public JsonSchemaString(JsonSchemaStringFormat? format = null, string? pattern = null)
 		{
-			this.Format = Format;
-			this.Pattern = Pattern;
+			Format = format;
+			Pattern = pattern;
 		}
 
 		/// <inheritdoc/>
-		public override void Write(IJsonSchemaWriter Writer)
+		public override void Write(IJsonSchemaWriter writer)
 		{
-			Writer.WriteString("type", "string");
+			writer.WriteString("type", "string");
 			if (Format != null)
 			{
-				Writer.WriteString("format", Format.Value.ToString().ToLowerInvariant());
+				writer.WriteString("format", Format.Value.ToString().ToLowerInvariant());
 			}
 			if (Pattern != null)
 			{
-				Writer.WriteString("pattern", Pattern);
+				writer.WriteString("pattern", Pattern);
 			}
 		}
 	}
@@ -159,20 +154,20 @@ namespace EpicGames.Core
 		/// <summary>
 		/// Constructor
 		/// </summary>
-		public JsonSchemaEnum(IEnumerable<string> Values)
+		public JsonSchemaEnum(IEnumerable<string> values)
 		{
-			this.Values.AddRange(Values);
+			Values.AddRange(values);
 		}
 
 		/// <inheritdoc/>
-		public override void Write(IJsonSchemaWriter Writer)
+		public override void Write(IJsonSchemaWriter writer)
 		{
-			Writer.WriteStartArray("enum");
-			foreach (string Value in Values)
+			writer.WriteStartArray("enum");
+			foreach (string value in Values)
 			{
-				Writer.WriteStringValue(Value);
+				writer.WriteStringValue(value);
 			}
-			Writer.WriteEndArray();
+			writer.WriteEndArray();
 		}
 	}
 
@@ -189,18 +184,18 @@ namespace EpicGames.Core
 		/// <summary>
 		/// Constructor
 		/// </summary>
-		public JsonSchemaArray(JsonSchemaType ItemType)
+		public JsonSchemaArray(JsonSchemaType itemType)
 		{
-			this.ItemType = ItemType;
+			ItemType = itemType;
 		}
 
 		/// <inheritdoc/>
-		public override void Write(IJsonSchemaWriter Writer)
+		public override void Write(IJsonSchemaWriter writer)
 		{
-			Writer.WriteString("type", "array");
-			Writer.WriteStartObject("items");
-			Writer.WriteType(ItemType);
-			Writer.WriteEndObject();
+			writer.WriteString("type", "array");
+			writer.WriteStartObject("items");
+			writer.WriteType(ItemType);
+			writer.WriteEndObject();
 		}
 	}
 
@@ -227,27 +222,27 @@ namespace EpicGames.Core
 		/// <summary>
 		/// Constructor
 		/// </summary>
-		public JsonSchemaProperty(string Name, string? Description, JsonSchemaType Type)
+		public JsonSchemaProperty(string name, string? description, JsonSchemaType type)
 		{
-			this.Name = Name;
-			this.Description = Description;
-			this.Type = Type;
+			Name = name;
+			Description = description;
+			Type = type;
 		}
 
 		/// <summary>
 		/// Writes a property to a schema
 		/// </summary>
-		/// <param name="Writer"></param>
-		public void Write(IJsonSchemaWriter Writer)
+		/// <param name="writer"></param>
+		public void Write(IJsonSchemaWriter writer)
 		{
-			string CamelCaseName = Char.ToLowerInvariant(Name[0]) + Name.Substring(1);
-			Writer.WriteStartObject(CamelCaseName);
-			Writer.WriteType(Type);
+			string camelCaseName = Char.ToLowerInvariant(Name[0]) + Name.Substring(1);
+			writer.WriteStartObject(camelCaseName);
+			writer.WriteType(Type);
 			if (Description != null)
 			{
-				Writer.WriteString("description", Description);
+				writer.WriteString("description", Description);
 			}
-			Writer.WriteEndObject();
+			writer.WriteEndObject();
 		}
 	}
 
@@ -267,28 +262,28 @@ namespace EpicGames.Core
 		public JsonSchemaType? AdditionalProperties { get; set; }
 
 		/// <inheritdoc/>
-		public override void Write(IJsonSchemaWriter Writer)
+		public override void Write(IJsonSchemaWriter writer)
 		{
-			Writer.WriteString("type", "object");
+			writer.WriteString("type", "object");
 			if (Properties.Count > 0)
 			{
-				Writer.WriteStartObject("properties");
-				foreach (JsonSchemaProperty Property in Properties)
+				writer.WriteStartObject("properties");
+				foreach (JsonSchemaProperty property in Properties)
 				{
-					Property.Write(Writer);
+					property.Write(writer);
 				}
-				Writer.WriteEndObject();
+				writer.WriteEndObject();
 			}
 
 			if (AdditionalProperties == null)
 			{
-				Writer.WriteBoolean("additionalProperties", false);
+				writer.WriteBoolean("additionalProperties", false);
 			}
 			else
 			{
-				Writer.WriteStartObject("additionalProperties");
-				AdditionalProperties.Write(Writer);
-				Writer.WriteEndObject();
+				writer.WriteStartObject("additionalProperties");
+				AdditionalProperties.Write(writer);
+				writer.WriteEndObject();
 			}
 		}
 	}
@@ -304,16 +299,16 @@ namespace EpicGames.Core
 		public List<JsonSchemaType> Types { get; } = new List<JsonSchemaType>();
 
 		/// <inheritdoc/>
-		public override void Write(IJsonSchemaWriter Writer)
+		public override void Write(IJsonSchemaWriter writer)
 		{
-			Writer.WriteStartArray("oneOf");
-			foreach (JsonSchemaType Type in Types)
+			writer.WriteStartArray("oneOf");
+			foreach (JsonSchemaType type in Types)
 			{
-				Writer.WriteStartObject();
-				Type.Write(Writer);
-				Writer.WriteEndObject();
+				writer.WriteStartObject();
+				type.Write(writer);
+				writer.WriteEndObject();
 			}
-			Writer.WriteEndArray();
+			writer.WriteEndArray();
 		}
 	}
 }
