Reformat EpicGames.Core according to standard coding conventions.

diff --git a/Engine/Source/Programs/Shared/EpicGames.Core/JsonTracer.cs b/Engine/Source/Programs/Shared/EpicGames.Core/JsonTracer.cs
--- a/Engine/Source/Programs/Shared/EpicGames.Core/JsonTracer.cs
+++ b/Engine/Source/Programs/Shared/EpicGames.Core/JsonTracer.cs
@@ -25,10 +25,10 @@ namespace EpicGames.Core
 		public string TraceId { get; }
 		public string SpanId { get; }
 
-		public JsonTracerSpanContext(string TraceId, string SpanId)
+		public JsonTracerSpanContext(string traceId, string spanId)
 		{
-			this.TraceId = TraceId;
-			this.SpanId = SpanId;
+			TraceId = traceId;
+			SpanId = spanId;
 		}
 	}
 
@@ -37,38 +37,33 @@ namespace EpicGames.Core
 		/// <summary>
 		/// Used to monotonically update ids
 		/// </summary>
-		private static long _nextIdCounter = 0;
+		private static long s_nextIdCounter = 0;
 
 		/// <summary>
 		/// A simple-as-possible (consecutive for repeatability) id generator.
 		/// </summary>
 		private static string NextId()
 		{
-			return Interlocked.Increment(ref _nextIdCounter).ToString(CultureInfo.InvariantCulture);
-		}
-		
-		public JsonTracerSpanContext Context
-		{
-			// C# doesn't have "return type covariance" so we use the trick with the explicit interface implementation
-			// and this separate property.
-			get
-			{
-				return JsonTracerContext;
-			}
+			return Interlocked.Increment(ref s_nextIdCounter).ToString(CultureInfo.InvariantCulture);
 		}
+
+		// C# doesn't have "return type covariance" so we use the trick with the explicit interface implementation
+		// and this separate property.
+		public JsonTracerSpanContext Context => _jsonTracerContext;
+
 		ISpanContext ISpan.Context => Context;
 		
 		public string OperationName { get; private set; }
 		
-		private readonly JsonTracer Tracer;
-		private JsonTracerSpanContext JsonTracerContext;
-		private DateTimeOffset _FinishTimestamp;
-		private bool Finished;
-		private readonly Dictionary<string, object> _Tags;
-		private readonly List<JsonTracerSpan.Reference> References;
+		private readonly JsonTracer _tracer;
+		private readonly JsonTracerSpanContext _jsonTracerContext;
+		private DateTimeOffset _finishTimestamp;
+		private bool _finished;
+		private readonly Dictionary<string, object> _tags;
+		private readonly List<JsonTracerSpan.Reference> _references;
 		
 		public DateTimeOffset StartTimestamp { get; }
-		public Dictionary<string, object> Tags => new Dictionary<string, object>(_Tags);
+		public Dictionary<string, object> Tags => new Dictionary<string, object>(_tags);
 		
 		/// <summary>
 		/// The finish time of the span; only valid after a call to <see cref="Finish()"/>.
@@ -77,10 +72,12 @@ namespace EpicGames.Core
 		{
 			get
 			{
-				if (_FinishTimestamp == DateTimeOffset.MinValue)
+				if (_finishTimestamp == DateTimeOffset.MinValue)
+				{
 					throw new InvalidOperationException("Must call Finish() before FinishTimestamp");
+				}
 
-				return _FinishTimestamp;
+				return _finishTimestamp;
 			}
 		}
 
@@ -92,68 +89,88 @@ namespace EpicGames.Core
 		/// <seealso cref="MockSpan.References"/>
 		public string? ParentId { get; }
 
-		private readonly object Lock = new object();
+		private readonly object _lock = new object();
 
-		public JsonTracerSpan(JsonTracer JsonTracer, string OperationName, DateTimeOffset StartTimestamp, Dictionary<string, object>? InitialTags, List<Reference>? References)
+		public JsonTracerSpan(JsonTracer jsonTracer, string operationName, DateTimeOffset startTimestamp, Dictionary<string, object>? initialTags, List<Reference>? references)
 		{
-			this.Tracer = JsonTracer;
-			this.OperationName = OperationName;
-			this.StartTimestamp = StartTimestamp;
+			_tracer = jsonTracer;
+			OperationName = operationName;
+			StartTimestamp = startTimestamp;
 			
-			_Tags = InitialTags == null
+			_tags = initialTags == null
 				? new Dictionary<string, object>()
-				: new Dictionary<string, object>(InitialTags);
+				: new Dictionary<string, object>(initialTags);
 
-			this.References = References == null
+			_references = references == null
 				? new List<Reference>()
-				: References.ToList();
+				: references.ToList();
 			
-			JsonTracerSpanContext? ParentContext = FindPreferredParentRef(this.References);
-			if (ParentContext == null)
+			JsonTracerSpanContext? parentContext = FindPreferredParentRef(_references);
+			if (parentContext == null)
 			{
 				// we are a root span
-				JsonTracerContext = new JsonTracerSpanContext(NextId(), NextId());
+				_jsonTracerContext = new JsonTracerSpanContext(NextId(), NextId());
 				ParentId = null;
 			}
 			else
 			{
 				// we are a child span
-				JsonTracerContext = new JsonTracerSpanContext(ParentContext.TraceId, NextId());
-				ParentId = ParentContext.SpanId;
+				_jsonTracerContext = new JsonTracerSpanContext(parentContext.TraceId, NextId());
+				ParentId = parentContext.SpanId;
 			}
 		}
 
-		public ISpan SetTag(string Key, string Value) { return SetObjectTag(Key, Value); }
-		public ISpan SetTag(string Key, bool Value) { return SetObjectTag(Key, Value); }
-		public ISpan SetTag(string Key, int Value) { return SetObjectTag(Key, Value); }
-		public ISpan SetTag(string Key, double Value) { return SetObjectTag(Key, Value); }
-		public ISpan SetTag(BooleanTag Tag, bool Value) { SetObjectTag(Tag.Key, Value); return this; }
-		public ISpan SetTag(IntOrStringTag Tag, string Value) { SetObjectTag(Tag.Key, Value); return this; }
-		public ISpan SetTag(IntTag Tag, int Value) { SetObjectTag(Tag.Key, Value); return this; }
-		public ISpan SetTag(StringTag Tag, string Value) { SetObjectTag(Tag.Key, Value); return this; }
+		public ISpan SetTag(string key, string value) => SetObjectTag(key, value); 
+		public ISpan SetTag(string key, bool value) => SetObjectTag(key, value);
+		public ISpan SetTag(string key, int value) => SetObjectTag(key, value);
+		public ISpan SetTag(string key, double value) => SetObjectTag(key, value);
+
+		public ISpan SetTag(BooleanTag tag, bool value) 
+		{ 
+			SetObjectTag(tag.Key, value); 
+			return this; 
+		}
+
+		public ISpan SetTag(IntOrStringTag tag, string value) 
+		{ 
+			SetObjectTag(tag.Key, value); 
+			return this; 
+		}
+
+		public ISpan SetTag(IntTag tag, int value) 
+		{ 
+			SetObjectTag(tag.Key, value); 
+			return this; 
+		}
 
-		private ISpan SetObjectTag(string Key, object Value)
+		public ISpan SetTag(StringTag tag, string value) 
+		{ 
+			SetObjectTag(tag.Key, value); 
+			return this; 
+		}
+
+		private ISpan SetObjectTag(string key, object value)
 		{
-			lock (Lock)
+			lock (_lock)
 			{
-				CheckForFinished("Setting tag [{0}:{1}] on already finished span", Key, Value);
-				_Tags[Key] = Value;
+				CheckForFinished("Setting tag [{0}:{1}] on already finished span", key, value);
+				_tags[key] = value;
 				return this;
 			}
 		}
 
 		public ISpan Log(IEnumerable<KeyValuePair<string, object>> fields) { throw new NotSupportedException("Log() calls are not supported in JsonTracerSpans"); }
 		public ISpan Log(DateTimeOffset timestamp, IEnumerable<KeyValuePair<string, object>> fields) { throw new NotSupportedException("Log() calls are not supported in JsonTracerSpans"); }
-		public ISpan Log(string @event) { throw new NotSupportedException("Log() calls are not supported in JsonTracerSpans"); }
-		public ISpan Log(DateTimeOffset timestamp, string @event) { throw new NotSupportedException("Log() calls are not supported in JsonTracerSpans"); }
+		public ISpan Log(string message) { throw new NotSupportedException("Log() calls are not supported in JsonTracerSpans"); }
+		public ISpan Log(DateTimeOffset timestamp, string message) { throw new NotSupportedException("Log() calls are not supported in JsonTracerSpans"); }
 
 		public ISpan SetBaggageItem(string key, string value) { throw new NotImplementedException(); }
 		public string GetBaggageItem(string key) { throw new NotImplementedException(); }
 
-		public ISpan SetOperationName(string OperationName)
+		public ISpan SetOperationName(string operationName)
 		{
-			CheckForFinished("Setting operationName [{0}] on already finished span", OperationName);
-			this.OperationName = OperationName;
+			CheckForFinished("Setting operationName [{0}] on already finished span", operationName);
+			OperationName = operationName;
 			return this;
 		}
 
@@ -162,38 +179,42 @@ namespace EpicGames.Core
 			Finish(DateTimeOffset.UtcNow);
 		}
 
-		public void Finish(DateTimeOffset FinishTimestamp)
+		public void Finish(DateTimeOffset finishTimestamp)
 		{
-			lock (Lock)
+			lock (_lock)
 			{
 				CheckForFinished("Tried to finish already finished span");
-				_FinishTimestamp = FinishTimestamp;
-				Tracer.AppendFinishedSpan(this);
-				Finished = true;
+				_finishTimestamp = finishTimestamp;
+				_tracer.AppendFinishedSpan(this);
+				_finished = true;
 			}
 		}
 		
-		private static JsonTracerSpanContext? FindPreferredParentRef(IList<Reference> References)
+		private static JsonTracerSpanContext? FindPreferredParentRef(IList<Reference> references)
 		{
-			if (!References.Any())
+			if (!references.Any())
+			{
 				return null;
+			}
 
 			// return the context of the parent, if applicable
-			foreach (var Reference in References)
+			foreach (Reference reference in references)
 			{
-				if (OpenTracing.References.ChildOf.Equals(Reference.ReferenceType))
-					return Reference.Context;
+				if (OpenTracing.References.ChildOf.Equals(reference.ReferenceType))
+				{
+					return reference.Context;
+				}
 			}
 
 			// otherwise, return the context of the first reference
-			return References.First().Context;
+			return references.First().Context;
 		}
 		
-		private void CheckForFinished(string Format, params object[] Args)
+		private void CheckForFinished(string format, params object[] args)
 		{
-			if (Finished)
+			if (_finished)
 			{
-				throw new InvalidOperationException(string.Format(Format, Args));
+				throw new InvalidOperationException(String.Format(format, args));
 			}
 		}
 		
@@ -206,70 +227,76 @@ namespace EpicGames.Core
 			/// </summary>
 			public string ReferenceType { get; }
 
-			public Reference(JsonTracerSpanContext Context, string ReferenceType)
+			public Reference(JsonTracerSpanContext context, string referenceType)
 			{
-				this.Context = Context ?? throw new ArgumentNullException(nameof(Context));
-				this.ReferenceType = ReferenceType ?? throw new ArgumentNullException(nameof(ReferenceType));
+				Context = context ?? throw new ArgumentNullException(nameof(context));
+				ReferenceType = referenceType ?? throw new ArgumentNullException(nameof(referenceType));
 			}
 
-			public override bool Equals(object? Obj)
+			public override bool Equals(object? obj)
 			{
-				return Equals(Obj as Reference);
+				return Equals(obj as Reference);
 			}
 			
-			public bool Equals(Reference? Other)
+			public bool Equals(Reference? other)
 			{
-				return Other != null &&
-				       EqualityComparer<JsonTracerSpanContext>.Default.Equals(Context, Other.Context) &&
-				       ReferenceType == Other.ReferenceType;
+				return other != null &&
+				       EqualityComparer<JsonTracerSpanContext>.Default.Equals(Context, other.Context) &&
+				       ReferenceType == other.ReferenceType;
 			}
-			
-			public override int GetHashCode()
+
+#pragma warning disable IDE0070 // Use 'System.HashCode'
+            public override int GetHashCode()
 			{
-				int HashCode = 2083322454;
-				HashCode = HashCode * -1521134295 + EqualityComparer<JsonTracerSpanContext>.Default.GetHashCode(Context);
-				HashCode = HashCode * -1521134295 + EqualityComparer<string>.Default.GetHashCode(ReferenceType);
-				return HashCode;
+				int hashCode = 2083322454;
+				hashCode = hashCode * -1521134295 + EqualityComparer<JsonTracerSpanContext>.Default.GetHashCode(Context);
+				hashCode = hashCode * -1521134295 + EqualityComparer<string>.Default.GetHashCode(ReferenceType);
+				return hashCode;
 			}
+#pragma warning restore IDE0070 // Use 'System.HashCode'
 		}
 	}
 
 	public class JsonTracerSpanBuilder : ISpanBuilder
 	{
-		private readonly JsonTracer Tracer;
-		private readonly string OperationName;
-		private DateTimeOffset StartTimestamp = DateTimeOffset.MinValue;
-		private readonly List<JsonTracerSpan.Reference> References = new List<JsonTracerSpan.Reference>();
-		private readonly Dictionary<string, object> Tags = new Dictionary<string, object>();
-		private bool _IgnoreActiveSpan;
+		private readonly JsonTracer _tracer;
+		private readonly string _operationName;
+		private DateTimeOffset _startTimestamp = DateTimeOffset.MinValue;
+		private readonly List<JsonTracerSpan.Reference> _references = new List<JsonTracerSpan.Reference>();
+		private readonly Dictionary<string, object> _tags = new Dictionary<string, object>();
+		private bool _ignoreActiveSpan;
 		
-		public JsonTracerSpanBuilder(JsonTracer Tracer, string OperationName)
+		public JsonTracerSpanBuilder(JsonTracer tracer, string operationName)
 		{
-			this.Tracer = Tracer;
-			this.OperationName = OperationName;
+			_tracer = tracer;
+			_operationName = operationName;
 		}
 		
-		public ISpanBuilder AsChildOf(ISpanContext? Parent)
+		public ISpanBuilder AsChildOf(ISpanContext? parent)
 		{
-			if (Parent == null)
+			if (parent == null)
+			{
 				return this;
+			}
 
-			return AddReference(OpenTracing.References.ChildOf, Parent);
+			return AddReference(OpenTracing.References.ChildOf, parent);
 		}
 
-		public ISpanBuilder AsChildOf(ISpan? Parent)
+		public ISpanBuilder AsChildOf(ISpan? parent)
 		{
-			if (Parent == null)
+			if (parent == null)
+			{
 				return this;
+			}
 
-			return AddReference(OpenTracing.References.ChildOf, Parent.Context);
+			return AddReference(OpenTracing.References.ChildOf, parent.Context);
 		}
 
-		public ISpanBuilder AddReference(string ReferenceType, ISpanContext? ReferencedContext)
+		public ISpanBuilder AddReference(string referenceType, ISpanContext? referencedContext)
 		{
-			if (ReferencedContext != null)
+			if (referencedContext != null)
 			{
-				References.Add(new JsonTracerSpan.Reference((JsonTracerSpanContext)ReferencedContext, ReferenceType));
+				_references.Add(new JsonTracerSpan.Reference((JsonTracerSpanContext)referencedContext, referenceType));
 			}
 
 			return this;
@@ -277,22 +304,61 @@ namespace EpicGames.Core
 
 		public ISpanBuilder IgnoreActiveSpan()
 		{
-			_IgnoreActiveSpan = true;
+			_ignoreActiveSpan = true;
 			return this;
 		}
 
-		public ISpanBuilder WithTag(string Key, string Value) { Tags[Key] = Value; return this; }
-		public ISpanBuilder WithTag(string Key, bool Value) { Tags[Key] = Value; return this; }
-		public ISpanBuilder WithTag(string Key, int Value) { Tags[Key] = Value; return this; }
-		public ISpanBuilder WithTag(string Key, double Value) { Tags[Key] = Value; return this; }
-		public ISpanBuilder WithTag(BooleanTag Tag, bool Value) { Tags[Tag.Key] = Value; return this; }
-		public ISpanBuilder WithTag(IntOrStringTag Tag, string Value) { Tags[Tag.Key] = Value; return this; }
-		public ISpanBuilder WithTag(IntTag Tag, int Value) { Tags[Tag.Key] = Value; return this; }
-		public ISpanBuilder WithTag(StringTag Tag, string Value) { Tags[Tag.Key] = Value; return this; }
+		public ISpanBuilder WithTag(string key, string value)
+		{ 
+			_tags[key] = value; 
+			return this; 
+		}
+		
+		public ISpanBuilder WithTag(string key, bool value) 
+		{ 
+			_tags[key] = value; 
+			return this; 
+		}
+
+		public ISpanBuilder WithTag(string key, int value)
+		{ 
+			_tags[key] = value; 
+			return this; 
+		}
+
+		public ISpanBuilder WithTag(string key, double value) 
+		{ 
+			_tags[key] = value; 
+			return this; 
+		}
+
+		public ISpanBuilder WithTag(BooleanTag tag, bool value) 
+		{ 
+			_tags[tag.Key] = value; 
+			return this; 
+		}
+
+		public ISpanBuilder WithTag(IntOrStringTag tag, string value) 
+		{ 
+			_tags[tag.Key] = value; 
+			return this; 
+		}
+		
+		public ISpanBuilder WithTag(IntTag tag, int value) 
+		{ 
+			_tags[tag.Key] = value; 
+			return this; 
+		}
+
+		public ISpanBuilder WithTag(StringTag tag, string value)
+		{ 
+			_tags[tag.Key] = value; 
+			return this; 
+		}
 
-		public ISpanBuilder WithStartTimestamp(DateTimeOffset Timestamp)
+		public ISpanBuilder WithStartTimestamp(DateTimeOffset timestamp)
 		{
-			StartTimestamp = Timestamp;
+			_startTimestamp = timestamp;
 			return this;
 		}
 
@@ -301,26 +367,26 @@ namespace EpicGames.Core
 			return StartActive(true);
 		}
 
-		public IScope StartActive(bool FinishSpanOnDispose)
+		public IScope StartActive(bool finishSpanOnDispose)
 		{
-			ISpan Span = Start();
-			return Tracer.ScopeManager.Activate(Span, FinishSpanOnDispose);
+			ISpan span = Start();
+			return _tracer.ScopeManager.Activate(span, finishSpanOnDispose);
 		}
 
 		public ISpan Start()
 		{
-			if (StartTimestamp == DateTimeOffset.MinValue) // value was not set by builder
+			if (_startTimestamp == DateTimeOffset.MinValue) // value was not set by builder
 			{
-				StartTimestamp = DateTimeOffset.UtcNow;
+				_startTimestamp = DateTimeOffset.UtcNow;
 			}
 
-			ISpanContext? ActiveSpanContext = Tracer.ActiveSpan?.Context;
-			if (!References.Any() && !_IgnoreActiveSpan && ActiveSpanContext != null)
+			ISpanContext? activeSpanContext = _tracer.ActiveSpan?.Context;
+			if (!_references.Any() && !_ignoreActiveSpan && activeSpanContext != null)
 			{
-				References.Add(new JsonTracerSpan.Reference((JsonTracerSpanContext)ActiveSpanContext, OpenTracing.References.ChildOf));
+				_references.Add(new JsonTracerSpan.Reference((JsonTracerSpanContext)activeSpanContext, OpenTracing.References.ChildOf));
 			}
 
-			return new JsonTracerSpan(Tracer, OperationName, StartTimestamp, Tags, References);
+			return new JsonTracerSpan(_tracer, _operationName, _startTimestamp, _tags, _references);
 		}
 	}
 	
@@ -329,107 +395,107 @@ namespace EpicGames.Core
 		public IScopeManager ScopeManager { get; }
 		public ISpan? ActiveSpan => ScopeManager.Active?.Span;
 
-		private readonly object Lock = new object();
-		private readonly List<JsonTracerSpan> FinishedSpans = new List<JsonTracerSpan>();
-		private readonly DirectoryReference? TelemetryDir;
+		private readonly object _lock = new object();
+		private readonly List<JsonTracerSpan> _finishedSpans = new List<JsonTracerSpan>();
+		private readonly DirectoryReference? _telemetryDir;
 		
-		public JsonTracer(DirectoryReference? TelemetryDir = null)
+		public JsonTracer(DirectoryReference? telemetryDir = null)
 		{
 			ScopeManager = new AsyncLocalScopeManager();
-			this.TelemetryDir = TelemetryDir;
+			_telemetryDir = telemetryDir;
 		}
 
 		public static JsonTracer? TryRegisterAsGlobalTracer()
 		{
-			string? TelemetryDir = Environment.GetEnvironmentVariable("UE_TELEMETRY_DIR");
-			if (TelemetryDir != null)
+			string? telemetryDir = Environment.GetEnvironmentVariable("UE_TELEMETRY_DIR");
+			if (telemetryDir != null)
 			{
-				JsonTracer Tracer = new JsonTracer(new DirectoryReference(TelemetryDir));
-				return GlobalTracer.RegisterIfAbsent(Tracer) ? Tracer : null;
+				JsonTracer tracer = new JsonTracer(new DirectoryReference(telemetryDir));
+				return GlobalTracer.RegisterIfAbsent(tracer) ? tracer : null;
 			}
 
 			return null;
 		}
 
-		public ISpanBuilder BuildSpan(string OperationName)
+		public ISpanBuilder BuildSpan(string operationName)
 		{
-			return new JsonTracerSpanBuilder(this, OperationName);
+			return new JsonTracerSpanBuilder(this, operationName);
 		}
 
-		public void Inject<TCarrier>(ISpanContext SpanContext, IFormat<TCarrier> Format, TCarrier Carrier)
+		public void Inject<TCarrier>(ISpanContext spanContext, IFormat<TCarrier> format, TCarrier carrier)
 		{
-			throw new NotSupportedException(string.Format("Tracer.Inject is not implemented for {0} by JsonTracer", Format));
+			throw new NotSupportedException(String.Format("Tracer.Inject is not implemented for {0} by JsonTracer", format));
 		}
 
-		public ISpanContext Extract<TCarrier>(IFormat<TCarrier> Format, TCarrier Carrier)
+		public ISpanContext Extract<TCarrier>(IFormat<TCarrier> format, TCarrier carrier)
 		{
-			throw new NotSupportedException(string.Format("Tracer.Extract is not implemented for {0} by JsonTracer", Format));
+			throw new NotSupportedException(String.Format("Tracer.Extract is not implemented for {0} by JsonTracer", format));
 		}
 
 		public void Flush()
 		{
-			if (TelemetryDir != null)
+			if (_telemetryDir != null)
 			{
-				string TelemetryScopeId = Environment.GetEnvironmentVariable("UE_TELEMETRY_SCOPE_ID") ?? "noscope";
+				string telemetryScopeId = Environment.GetEnvironmentVariable("UE_TELEMETRY_SCOPE_ID") ?? "noscope";
 				
-				FileReference File;
-				using (System.Diagnostics.Process Process = System.Diagnostics.Process.GetCurrentProcess())
+				FileReference file;
+				using (System.Diagnostics.Process process = System.Diagnostics.Process.GetCurrentProcess())
 				{
-					DirectoryReference.CreateDirectory(TelemetryDir);
+					DirectoryReference.CreateDirectory(_telemetryDir);
 
-					string FileName = String.Format("{0}.{1}.{2}.opentracing.json", Path.GetFileName(Assembly.GetEntryAssembly()!.Location), TelemetryScopeId, Process.Id);
-					File = FileReference.Combine(TelemetryDir, FileName);
+					string fileName = String.Format("{0}.{1}.{2}.opentracing.json", Path.GetFileName(Assembly.GetEntryAssembly()!.Location), telemetryScopeId, process.Id);
+					file = FileReference.Combine(_telemetryDir, fileName);
 				}
 
-				using (JsonWriter Writer = new JsonWriter(File))
+				using (JsonWriter writer = new JsonWriter(file))
 				{
-					GetFinishedSpansAsJson(Writer);
+					GetFinishedSpansAsJson(writer);
 				}
 			}
 		}
 		
 		public List<JsonTracerSpan> GetFinishedSpans()
 		{
-			lock (Lock)
+			lock (_lock)
 			{
-				return new List<JsonTracerSpan>(FinishedSpans);
+				return new List<JsonTracerSpan>(_finishedSpans);
 			}
 		}
 
-		public void GetFinishedSpansAsJson(JsonWriter Writer)
+		public void GetFinishedSpansAsJson(JsonWriter writer)
 		{
-			Writer.WriteObjectStart();
-			Writer.WriteArrayStart("Spans");
-			foreach (JsonTracerSpan Span in GetFinishedSpans())
+			writer.WriteObjectStart();
+			writer.WriteArrayStart("Spans");
+			foreach (JsonTracerSpan span in GetFinishedSpans())
 			{
-				Writer.WriteObjectStart();
-				Writer.WriteValue("Name", Span.OperationName);
-				Dictionary<string, object> Tags = Span.Tags;
-				if (Tags.TryGetValue("Resource", out object? Resource) && Resource is string ResourceString)
+				writer.WriteObjectStart();
+				writer.WriteValue("Name", span.OperationName);
+				Dictionary<string, object> tags = span.Tags;
+				if (tags.TryGetValue("Resource", out object? resource) && resource is string resourceString)
 				{
-					Writer.WriteValue("Resource", ResourceString);
+					writer.WriteValue("Resource", resourceString);
 				}
-				if (Tags.TryGetValue("Service", out object? Service) && Service is string ServiceString)
+				if (tags.TryGetValue("Service", out object? service) && service is string serviceString)
 				{
-					Writer.WriteValue("Service", ServiceString);
+					writer.WriteValue("Service", serviceString);
 				}
-				Writer.WriteValue("StartTime", Span.StartTimestamp.ToString("o", CultureInfo.InvariantCulture));
-				Writer.WriteValue("FinishTime", Span.FinishTimestamp.ToString("o", CultureInfo.InvariantCulture));
-				Writer.WriteObjectStart("Metadata");
+				writer.WriteValue("StartTime", span.StartTimestamp.ToString("o", CultureInfo.InvariantCulture));
+				writer.WriteValue("FinishTime", span.FinishTimestamp.ToString("o", CultureInfo.InvariantCulture));
+				writer.WriteObjectStart("Metadata");
 				// TODO: Write tags as metadata?
-				Writer.WriteObjectEnd();
-				Writer.WriteObjectEnd();
+				writer.WriteObjectEnd();
+				writer.WriteObjectEnd();
 			}
-			Writer.WriteArrayEnd();
-			Writer.WriteObjectEnd();
+			writer.WriteArrayEnd();
+			writer.WriteObjectEnd();
 		}
 		
-		internal void AppendFinishedSpan(JsonTracerSpan JsonTracerSpan)
+		internal void AppendFinishedSpan(JsonTracerSpan jsonTracerSpan)
 		{
-			lock (Lock)
+			lock (_lock)
 			{
-				FinishedSpans.Add(JsonTracerSpan);
+				_finishedSpans.Add(jsonTracerSpan);
 			}
 		}
 	}
-}
\ No newline at end of file
+}
diff --git a/Engine/Source/Programs/Shared/EpicGames.Core/JsonWriter.cs b/Engine/Source/Programs/Shared/EpicGames.Core/JsonWriter.cs
--- a/Engine/Source/Programs/Shared/EpicGames.Core/JsonWriter.cs
+++ b/Engine/Source/Programs/Shared/EpicGames.Core/JsonWriter.cs
@@ -5,7 +5,6 @@ using System.Collections.Generic;
 using System.IO;
 using System.Linq;
 using System.Text;
-using System.Threading.Tasks;
 
 namespace EpicGames.Core
 {
@@ -30,46 +29,46 @@ namespace EpicGames.Core
 	/// </summary>
 	public class JsonWriter : IDisposable
 	{
-		TextWriter Writer;
-		bool bLeaveOpen;
-		JsonWriterStyle Style;
-		bool bRequiresComma;
-		string Indent;
+		TextWriter _writer;
+		readonly bool _leaveOpen;
+		readonly JsonWriterStyle _style;
+		bool _bRequiresComma;
+		string _indent;
 
 		/// <summary>
 		/// Constructor
 		/// </summary>
-		/// <param name="FileName">File to write to</param>
-		/// <param name="Style">Should use packed JSON or not</param>
-		public JsonWriter(string FileName, JsonWriterStyle Style = JsonWriterStyle.Readable)
-			: this(new StreamWriter(FileName))
+		/// <param name="fileName">File to write to</param>
+		/// <param name="style">Should use packed JSON or not</param>
+		public JsonWriter(string fileName, JsonWriterStyle style = JsonWriterStyle.Readable)
+			: this(new StreamWriter(fileName))
 		{
-			this.Style = Style;
+			_style = style;
 		}
 
 		/// <summary>
 		/// Constructor
 		/// </summary>
-		/// <param name="FileName">File to write to</param>
-		/// <param name="Style">Should use packed JSON or not</param>
-		public JsonWriter(FileReference FileName, JsonWriterStyle Style = JsonWriterStyle.Readable)
-			: this(new StreamWriter(FileName.FullName))
+		/// <param name="fileName">File to write to</param>
+		/// <param name="style">Should use packed JSON or not</param>
+		public JsonWriter(FileReference fileName, JsonWriterStyle style = JsonWriterStyle.Readable)
+			: this(new StreamWriter(fileName.FullName))
 		{
-			this.Style = Style;
+			_style = style;
 		}
 
 		/// <summary>
 		/// Constructor
 		/// </summary>
-		/// <param name="Writer">The text writer to output to</param>
-		/// <param name="bLeaveOpen">Whether to leave the writer open when the object is disposed</param>
-		/// <param name="Style">The output style</param>
-		public JsonWriter(TextWriter Writer, bool bLeaveOpen = false, JsonWriterStyle Style = JsonWriterStyle.Readable)
+		/// <param name="writer">The text writer to output to</param>
+		/// <param name="leaveOpen">Whether to leave the writer open when the object is disposed</param>
+		/// <param name="style">The output style</param>
+		public JsonWriter(TextWriter writer, bool leaveOpen = false, JsonWriterStyle style = JsonWriterStyle.Readable)
 		{
-			this.Writer = Writer;
-			this.bLeaveOpen = bLeaveOpen;
-			this.Style = Style;
-			Indent = "";
+			_writer = writer;
+			_leaveOpen = leaveOpen;
+			_style = style;
+			_indent = "";
 		}
 
 		/// <summary>
@@ -77,26 +76,26 @@ namespace EpicGames.Core
 		/// </summary>
 		public void Dispose()
 		{
-			if(!bLeaveOpen && Writer != null)
+			if(!_leaveOpen && _writer != null)
 			{
-				Writer.Dispose();
-				Writer = null!;
+				_writer.Dispose();
+				_writer = null!;
 			}
 		}
 
 		private void IncreaseIndent()
 		{
-			if (Style == JsonWriterStyle.Readable)
+			if (_style == JsonWriterStyle.Readable)
 			{
-				Indent += "\t";
+				_indent += "\t";
 			}
 		}
 		
 		private void DecreaseIndent()
 		{
-			if (Style == JsonWriterStyle.Readable)
+			if (_style == JsonWriterStyle.Readable)
 			{
-				Indent = Indent.Substring(0, Indent.Length - 1);
+				_indent = _indent.Substring(0, _indent.Length - 1);
 			}
 		}
 
@@ -107,25 +106,25 @@ namespace EpicGames.Core
 		{
 			WriteCommaNewline();
 
-			Writer.Write(Indent);
-			Writer.Write("{");
+			_writer.Write(_indent);
+			_writer.Write("{");
 
 			IncreaseIndent();
-			bRequiresComma = false;
+			_bRequiresComma = false;
 		}
 
 		/// <summary>
 		/// Write the name and opening brace for an object
 		/// </summary>
-		/// <param name="ObjectName">Name of the field</param>
-		public void WriteObjectStart(string ObjectName)
+		/// <param name="objectName">Name of the field</param>
+		public void WriteObjectStart(string objectName)
 		{
 			WriteCommaNewline();
 
-			string Space = (Style == JsonWriterStyle.Readable) ? " " : "";
-			Writer.Write("{0}\"{1}\":{2}", Indent, ObjectName, Space);
+			string space = (_style == JsonWriterStyle.Readable) ? " " : "";
+			_writer.Write("{0}\"{1}\":{2}", _indent, objectName, space);
 
-			bRequiresComma = false;
+			_bRequiresComma = false;
 
 			WriteObjectStart();
 		}
@@ -138,10 +137,10 @@ namespace EpicGames.Core
 			DecreaseIndent();
 
 			WriteLine();
-			Writer.Write(Indent);
-			Writer.Write("}");
+			_writer.Write(_indent);
+			_writer.Write("}");
 
-			bRequiresComma = true;
+			_bRequiresComma = true;
 		}
 
 		/// <summary>
@@ -152,25 +151,25 @@ namespace EpicGames.Core
 		{
 			WriteCommaNewline();
 
-			Writer.Write("{0}[", Indent);
+			_writer.Write("{0}[", _indent);
 
 			IncreaseIndent();
-			bRequiresComma = false;
+			_bRequiresComma = false;
 		}
 
 		/// <summary>
 		/// Write the name and opening bracket for an array
 		/// </summary>
-		/// <param name="ArrayName">Name of the field</param>
-		public void WriteArrayStart(string ArrayName)
+		/// <param name="arrayName">Name of the field</param>
+		public void WriteArrayStart(string arrayName)
 		{
 			WriteCommaNewline();
 
-			string Space = (Style == JsonWriterStyle.Readable) ? " " : "";
-			Writer.Write("{0}\"{1}\":{2}[", Indent, ArrayName, Space);
+			string space = (_style == JsonWriterStyle.Readable) ? " " : "";
+			_writer.Write("{0}\"{1}\":{2}[", _indent, arrayName, space);
 
 			IncreaseIndent();
-			bRequiresComma = false;
+			_bRequiresComma = false;
 		}
 
 		/// <summary>
@@ -181,42 +180,42 @@ namespace EpicGames.Core
 			DecreaseIndent();
 
 			WriteLine();
-			Writer.Write("{0}]", Indent);
+			_writer.Write("{0}]", _indent);
 
-			bRequiresComma = true;
+			_bRequiresComma = true;
 		}
 
 		private void WriteLine()
 		{
-			if (Style == JsonWriterStyle.Readable)
+			if (_style == JsonWriterStyle.Readable)
 			{
-				Writer.WriteLine();
+				_writer.WriteLine();
 			}
 		}
 		
-		private void WriteLine(string Line)
+		private void WriteLine(string line)
 		{
-			if (Style == JsonWriterStyle.Readable)
+			if (_style == JsonWriterStyle.Readable)
 			{
-				Writer.WriteLine(Line);
+				_writer.WriteLine(line);
 			}
 			else
 			{
-				Writer.Write(Line);
+				_writer.Write(line);
 			}
 		}
 
 		/// <summary>
 		/// Write an array of strings
 		/// </summary>
-		/// <param name="Name">Name of the field</param>
-		/// <param name="Values">Values for the field</param>
-		public void WriteStringArrayField(string Name, IEnumerable<string> Values)
+		/// <param name="name">Name of the field</param>
+		/// <param name="values">Values for the field</param>
+		public void WriteStringArrayField(string name, IEnumerable<string> values)
 		{
-			WriteArrayStart(Name);
-			foreach(string Value in Values)
+			WriteArrayStart(name);
+			foreach(string value in values)
 			{
-				WriteValue(Value);
+				WriteValue(value);
 			}
 			WriteArrayEnd();
 		}
@@ -224,177 +223,177 @@ namespace EpicGames.Core
 		/// <summary>
 		/// Write an array of enum values
 		/// </summary>
-		/// <param name="Name">Name of the field</param>
-		/// <param name="Values">Values for the field</param>
-		public void WriteEnumArrayField<T>(string Name, IEnumerable<T> Values) where T : struct
+		/// <param name="name">Name of the field</param>
+		/// <param name="values">Values for the field</param>
+		public void WriteEnumArrayField<T>(string name, IEnumerable<T> values) where T : struct
 		{
-			WriteStringArrayField(Name, Values.Select(x => x.ToString()!));
+			WriteStringArrayField(name, values.Select(x => x.ToString()!));
 		}
 
 		/// <summary>
 		/// Write a value with no field name, for the contents of an array
 		/// </summary>
-		/// <param name="Value">Value to write</param>
-		public void WriteValue(int Value)
+		/// <param name="value">Value to write</param>
+		public void WriteValue(int value)
 		{
 			WriteCommaNewline();
 
-			Writer.Write(Indent);
-			Writer.Write(Value);
+			_writer.Write(_indent);
+			_writer.Write(value);
 
-			bRequiresComma = true;
+			_bRequiresComma = true;
 		}
 
 		/// <summary>
 		/// Write a value with no field name, for the contents of an array
 		/// </summary>
-		/// <param name="Value">Value to write</param>
-		public void WriteValue(string Value)
+		/// <param name="value">Value to write</param>
+		public void WriteValue(string value)
 		{
 			WriteCommaNewline();
 
-			Writer.Write(Indent);
-			WriteEscapedString(Value);
+			_writer.Write(_indent);
+			WriteEscapedString(value);
 
-			bRequiresComma = true;
+			_bRequiresComma = true;
 		}
 
 		/// <summary>
 		/// Write a field name and string value
 		/// </summary>
-		/// <param name="Name">Name of the field</param>
-		/// <param name="Value">Value for the field</param>
-		public void WriteValue(string Name, string? Value)
+		/// <param name="name">Name of the field</param>
+		/// <param name="value">Value for the field</param>
+		public void WriteValue(string name, string? value)
 		{
 			WriteCommaNewline();
 
-			string Space = (Style == JsonWriterStyle.Readable) ? " " : "";
-			Writer.Write("{0}\"{1}\":{2}", Indent, Name, Space);
-			WriteEscapedString(Value);
+			string space = (_style == JsonWriterStyle.Readable) ? " " : "";
+			_writer.Write("{0}\"{1}\":{2}", _indent, name, space);
+			WriteEscapedString(value);
 
-			bRequiresComma = true;
+			_bRequiresComma = true;
 		}
 
 		/// <summary>
 		/// Write a field name and integer value
 		/// </summary>
-		/// <param name="Name">Name of the field</param>
-		/// <param name="Value">Value for the field</param>
-		public void WriteValue(string Name, int Value)
+		/// <param name="name">Name of the field</param>
+		/// <param name="value">Value for the field</param>
+		public void WriteValue(string name, int value)
 		{
-			WriteValueInternal(Name, Value.ToString());
+			WriteValueInternal(name, value.ToString());
 		}
 
 		/// <summary>
 		/// Write a field name and double value
 		/// </summary>
-		/// <param name="Name">Name of the field</param>
-		/// <param name="Value">Value for the field</param>
-		public void WriteValue(string Name, double Value)
+		/// <param name="name">Name of the field</param>
+		/// <param name="value">Value for the field</param>
+		public void WriteValue(string name, double value)
 		{
-			WriteValueInternal(Name, Value.ToString());
+			WriteValueInternal(name, value.ToString());
 		}
 
 		/// <summary>
 		/// Write a field name and bool value
 		/// </summary>
-		/// <param name="Name">Name of the field</param>
-		/// <param name="Value">Value for the field</param>
-		public void WriteValue(string Name, bool Value)
+		/// <param name="name">Name of the field</param>
+		/// <param name="value">Value for the field</param>
+		public void WriteValue(string name, bool value)
 		{
-			WriteValueInternal(Name, Value ? "true" : "false");
+			WriteValueInternal(name, value ? "true" : "false");
 		}
 
 		/// <summary>
 		/// Write a field name and enum value
 		/// </summary>
 		/// <typeparam name="T">The enum type</typeparam>
-		/// <param name="Name">Name of the field</param>
-		/// <param name="Value">Value for the field</param>
-		public void WriteEnumValue<T>(string Name, T Value) where T : struct
+		/// <param name="name">Name of the field</param>
+		/// <param name="value">Value for the field</param>
+		public void WriteEnumValue<T>(string name, T value) where T : struct
 		{
-			WriteValue(Name, Value.ToString()!);
+			WriteValue(name, value.ToString()!);
 		}
 
 		void WriteCommaNewline()
 		{
-			if (bRequiresComma)
+			if (_bRequiresComma)
 			{
 				WriteLine(",");
 			}
-			else if (Indent.Length > 0)
+			else if (_indent.Length > 0)
 			{
 				WriteLine();
 			}
 		}
 
-		void WriteValueInternal(string Name, string Value)
+		void WriteValueInternal(string name, string value)
 		{
 			WriteCommaNewline();
 
-			string Space = (Style == JsonWriterStyle.Readable) ? " " : "";
-			Writer.Write("{0}\"{1}\":{2}{3}", Indent, Name, Space, Value);
+			string space = (_style == JsonWriterStyle.Readable) ? " " : "";
+			_writer.Write("{0}\"{1}\":{2}{3}", _indent, name, space, value);
 
-			bRequiresComma = true;
+			_bRequiresComma = true;
 		}
 
-		void WriteEscapedString(string? Value)
+		void WriteEscapedString(string? value)
 		{
 			// Escape any characters which may not appear in a JSON string (see http://www.json.org).
-			Writer.Write("\"");
-			if (Value != null)
+			_writer.Write("\"");
+			if (value != null)
 			{
-				Writer.Write(EscapeString(Value));
+				_writer.Write(EscapeString(value));
 			}
-			Writer.Write("\"");
+			_writer.Write("\"");
 		}
 
 		/// <summary>
 		/// Escapes a string for serializing to JSON
 		/// </summary>
-		/// <param name="Value">The string to escape</param>
+		/// <param name="value">The string to escape</param>
 		/// <returns>The escaped string</returns>
-		public static string EscapeString(string Value)
+		public static string EscapeString(string value)
 		{
-			StringBuilder Result = new StringBuilder();
-			for (int Idx = 0; Idx < Value.Length; Idx++)
+			StringBuilder result = new StringBuilder();
+			for (int idx = 0; idx < value.Length; idx++)
 			{
-				switch (Value[Idx])
+				switch (value[idx])
 				{
 					case '\"':
-						Result.Append("\\\"");
+						result.Append("\\\"");
 						break;
 					case '\\':
-						Result.Append("\\\\");
+						result.Append("\\\\");
 						break;
 					case '\b':
-						Result.Append("\\b");
+						result.Append("\\b");
 						break;
 					case '\f':
-						Result.Append("\\f");
+						result.Append("\\f");
 						break;
 					case '\n':
-						Result.Append("\\n");
+						result.Append("\\n");
 						break;
 					case '\r':
-						Result.Append("\\r");
+						result.Append("\\r");
 						break;
 					case '\t':
-						Result.Append("\\t");
+						result.Append("\\t");
 						break;
 					default:
-						if (Char.IsControl(Value[Idx]))
+						if (Char.IsControl(value[idx]))
 						{
-							Result.AppendFormat("\\u{0:X4}", (int)Value[Idx]);
+							result.AppendFormat("\\u{0:X4}", (int)value[idx]);
 						}
 						else
 						{
-							Result.Append(Value[Idx]);
+							result.Append(value[idx]);
 						}
 						break;
 				}
 			}
-			return Result.ToString();
+			return result.ToString();
 		}
 	}
 }
diff --git a/Engine/Source/Programs/Shared/EpicGames.Core/KnownLogEvents.cs b/Engine/Source/Programs/Shared/EpicGames.Core/KnownLogEvents.cs
--- a/Engine/Source/Programs/Shared/EpicGames.Core/KnownLogEvents.cs
+++ b/Engine/Source/Programs/Shared/EpicGames.Core/KnownLogEvents.cs
@@ -1,10 +1,6 @@
 // Copyright Epic Games, Inc. All Rights Reserved.
 
 using Microsoft.Extensions.Logging;
-using System;
-using System.Collections.Generic;
-using System.Diagnostics.CodeAnalysis;
-using System.Text;
 
 namespace EpicGames.Core
 {
diff --git a/Engine/Source/Programs/Shared/EpicGames.Core/ListExtensions.cs b/Engine/Source/Programs/Shared/EpicGames.Core/ListExtensions.cs
--- a/Engine/Source/Programs/Shared/EpicGames.Core/ListExtensions.cs
+++ b/Engine/Source/Programs/Shared/EpicGames.Core/ListExtensions.cs
@@ -1,11 +1,7 @@
 // Copyright Epic Games, Inc. All Rights Reserved.
 
 using System;
-using System.Collections;
 using System.Collections.Generic;
-using System.Diagnostics.CodeAnalysis;
-using System.Globalization;
-using System.Text;
 
 namespace EpicGames.Core
 {
@@ -19,12 +15,12 @@ namespace EpicGames.Core
 		/// </summary>
 		/// <typeparam name="TElement">List element</typeparam>
 		/// <typeparam name="TField">Field type to sort by</typeparam>
-		/// <param name="List">List to sort</param>
-		/// <param name="Selector">Selects a field from the element</param>
-		public static void SortBy<TElement, TField>(this List<TElement> List, Func<TElement, TField> Selector)
+		/// <param name="list">List to sort</param>
+		/// <param name="selector">Selects a field from the element</param>
+		public static void SortBy<TElement, TField>(this List<TElement> list, Func<TElement, TField> selector)
 		{
-			IComparer<TField> DefaultComparer = Comparer<TField>.Default;
-			SortBy(List, Selector, DefaultComparer);
+			IComparer<TField> defaultComparer = Comparer<TField>.Default;
+			SortBy(list, selector, defaultComparer);
 		}
 
 		/// <summary>
@@ -32,12 +28,12 @@ namespace EpicGames.Core
 		/// </summary>
 		/// <typeparam name="TElement">List element</typeparam>
 		/// <typeparam name="TField">Field type to sort by</typeparam>
-		/// <param name="List">List to sort</param>
-		/// <param name="Selector">Selects a field from the element</param>
-		/// <param name="Comparer">Comparer for fields</param>
-		public static void SortBy<TElement, TField>(this List<TElement> List, Func<TElement, TField> Selector, IComparer<TField> Comparer)
+		/// <param name="list">List to sort</param>
+		/// <param name="selector">Selects a field from the element</param>
+		/// <param name="comparer">Comparer for fields</param>
+		public static void SortBy<TElement, TField>(this List<TElement> list, Func<TElement, TField> selector, IComparer<TField> comparer)
 		{
-			List.Sort((X, Y) => Comparer.Compare(Selector(X), Selector(Y)));
+			list.Sort((x, y) => comparer.Compare(selector(x), selector(y)));
 		}
 	}
 }
diff --git a/Engine/Source/Programs/Shared/EpicGames.Core/ListSegment.cs b/Engine/Source/Programs/Shared/EpicGames.Core/ListSegment.cs
--- a/Engine/Source/Programs/Shared/EpicGames.Core/ListSegment.cs
+++ b/Engine/Source/Programs/Shared/EpicGames.Core/ListSegment.cs
@@ -3,7 +3,6 @@
 using System;
 using System.Collections;
 using System.Collections.Generic;
-using System.Text;
 
 namespace EpicGames.Core
 {
@@ -42,31 +41,31 @@ namespace EpicGames.Core
 		/// Initializes a new instance of the System.ArraySegment`1 structure that delimits
 		/// all the elements in the specified array.
 		/// </summary>
-		public ListSegment(IList<T> List)
-			: this(List, 0, List.Count)
+		public ListSegment(IList<T> list)
+			: this(list, 0, list.Count)
 		{
 		}
 
 		/// <summary>
 		/// 
 		/// </summary>
-		/// <param name="List"></param>
-		/// <param name="Offset"></param>
-		/// <param name="Count"></param>
-		public ListSegment(IList<T> List, int Offset, int Count)
+		/// <param name="list"></param>
+		/// <param name="offset"></param>
+		/// <param name="count"></param>
+		public ListSegment(IList<T> list, int offset, int count)
 		{
-			if (Offset < 0 || Offset > List.Count)
+			if (offset < 0 || offset > list.Count)
 			{
-				throw new ArgumentOutOfRangeException(nameof(Offset));
+				throw new ArgumentOutOfRangeException(nameof(offset));
 			}
-			if (Count < 0 || Offset + Count > List.Count)
+			if (count < 0 || offset + count > list.Count)
 			{
-				throw new ArgumentOutOfRangeException(nameof(Count));
+				throw new ArgumentOutOfRangeException(nameof(count));
 			}
 
-			this.List = List;
-			this.Offset = Offset;
-			this.Count = Count;
+			List = list;
+			Offset = offset;
+			Count = count;
 		}
 
 		/// <summary>
@@ -75,33 +74,33 @@ namespace EpicGames.Core
 		/// <param name="index"></param>
 		/// <returns></returns>
 		/// <exception cref="ArgumentOutOfRangeException">Index is not a valid index in the ListSegment</exception>
-		public T this[int Index]
+		public T this[int index]
 		{
 			get
 			{
-				if (Index < 0 || Index >= Count)
+				if (index < 0 || index >= Count)
 				{
-					throw new ArgumentOutOfRangeException(nameof(Index));
+					throw new ArgumentOutOfRangeException(nameof(index));
 				}
-				return List[Index + Offset];
+				return List[index + Offset];
 			}
 			set
 			{
-				if (Index < 0 || Index >= Count)
+				if (index < 0 || index >= Count)
 				{
-					throw new ArgumentOutOfRangeException(nameof(Index));
+					throw new ArgumentOutOfRangeException(nameof(index));
 				}
-				List[Index + Offset] = value;
+				List[index + Offset] = value;
 			}
 		}
 
 		/// <summary>
 		/// 
 		/// </summary>
-		/// <param name="List"></param>
-		public static implicit operator ListSegment<T>(List<T> List)
+		/// <param name="list"></param>
+		public static implicit operator ListSegment<T>(List<T> list)
 		{
-			return new ListSegment<T>(List);
+			return new ListSegment<T>(list);
 		}
 
 		/// <summary>
@@ -109,32 +108,32 @@ namespace EpicGames.Core
 		/// </summary>
 		struct Enumerator : IEnumerator<T>, IEnumerator
 		{
-			IList<T> List;
-			int Index;
-			int MaxIndex;
+			readonly IList<T> _list;
+			int _index;
+			readonly int _maxIndex;
 
 			/// <inheritdoc/>
 			public T Current
 			{
 				get
 				{
-					if (Index >= MaxIndex)
+					if (_index >= _maxIndex)
 					{
 						throw new IndexOutOfRangeException();
 					}
-					return List[Index];
+					return _list[_index];
 				}
 			}
 
 			/// <summary>
 			/// Constructor
 			/// </summary>
-			/// <param name="Segment"></param>
-			public Enumerator(ListSegment<T> Segment)
+			/// <param name="segment"></param>
+			public Enumerator(ListSegment<T> segment)
 			{
-				this.List = Segment.List;
-				this.Index = Segment.Offset - 1;
-				this.MaxIndex = Segment.Offset + Segment.Count;
+				_list = segment.List;
+				_index = segment.Offset - 1;
+				_maxIndex = segment.Offset + segment.Count;
 			}
 
 			/// <inheritdoc/>
@@ -145,7 +144,7 @@ namespace EpicGames.Core
 			/// <inheritdoc/>
 			public bool MoveNext()
 			{
-				return ++Index < MaxIndex;
+				return ++_index < _maxIndex;
 			}
 
 			/// <inheritdoc/>
@@ -171,25 +170,25 @@ namespace EpicGames.Core
 		/// Create a list segment
 		/// </summary>
 		/// <typeparam name="T"></typeparam>
-		/// <param name="List"></param>
-		/// <param name="Offset"></param>
+		/// <param name="list"></param>
+		/// <param name="offset"></param>
 		/// <returns></returns>
-		public static ListSegment<T> Slice<T>(this IList<T> List, int Offset)
+		public static ListSegment<T> Slice<T>(this IList<T> list, int offset)
 		{
-			return new ListSegment<T>(List, Offset, List.Count - Offset);
+			return new ListSegment<T>(list, offset, list.Count - offset);
 		}
 
 		/// <summary>
 		/// Create a list segment
 		/// </summary>
 		/// <typeparam name="T"></typeparam>
-		/// <param name="List"></param>
-		/// <param name="Offset"></param>
-		/// <param name="Count"></param>
+		/// <param name="list"></param>
+		/// <param name="offset"></param>
+		/// <param name="count"></param>
 		/// <returns></returns>
-		public static ListSegment<T> Slice<T>(this IList<T> List, int Offset, int Count)
+		public static ListSegment<T> Slice<T>(this IList<T> list, int offset, int count)
 		{
-			return new ListSegment<T>(List, Offset, Count);
+			return new ListSegment<T>(list, offset, count);
 		}
 	}
 }
