Reformat EpicGames.Core according to standard coding conventions.

diff --git a/Engine/Source/Programs/Shared/EpicGames.Core/Log.cs b/Engine/Source/Programs/Shared/EpicGames.Core/Log.cs
--- a/Engine/Source/Programs/Shared/EpicGames.Core/Log.cs
+++ b/Engine/Source/Programs/Shared/EpicGames.Core/Log.cs
@@ -1,20 +1,18 @@
 // Copyright Epic Games, Inc. All Rights Reserved.
 
-using Microsoft.Extensions.DependencyInjection;
-using Microsoft.Extensions.Logging;
 using System;
+using System.Buffers;
 using System.Collections.Generic;
 using System.Diagnostics;
 using System.IO;
 using System.Linq;
 using System.Reflection;
-using System.Runtime.CompilerServices;
 using System.Text;
-using JetBrains.Annotations;
-using System.Buffers;
-using System.ComponentModel.Design;
 using System.Text.Json;
 using System.Text.RegularExpressions;
+using JetBrains.Annotations;
+using Microsoft.Extensions.DependencyInjection;
+using Microsoft.Extensions.Logging;
 
 namespace EpicGames.Core
 {
@@ -93,7 +91,7 @@ namespace EpicGames.Core
 	///   2. We have a verbose filter we would like to apply to all logs without having to have each listener filter individually, which would require our string formatting code to run every time.
 	///   3. We possibly want to ensure severity prefixes are logged, but Trace.WriteXXX does not allow any severity info to be passed down.
 	/// </summary>
-	static public class Log
+	public static class Log
 	{
 		/// <summary>
 		/// Singleton instance of the default logger
@@ -169,11 +167,10 @@ namespace EpicGames.Core
 		/// </summary>
 		public static FileReference? OutputFile => DefaultLogger?.OutputFile;
 
-
 		/// <summary>
 		/// A collection of strings that have been already written once
 		/// </summary>
-		private static HashSet<string> WriteOnceSet = new HashSet<string>();
+		private static readonly HashSet<string> s_writeOnceSet = new HashSet<string>();
 
 		/// <summary>
 		/// Adds a trace listener that writes to a log file.
@@ -182,83 +179,83 @@ namespace EpicGames.Core
 		/// If a StartupTraceListener was in use, this function will copy its captured data to the log file(s)
 		/// and remove the startup listener from the list of registered listeners.
 		/// </summary>
-		/// <param name="OutputFile">The file to write to</param>
+		/// <param name="outputFile">The file to write to</param>
 		/// <returns>The created trace listener</returns>
-		public static void AddFileWriter(string Name, FileReference OutputFile)
+		public static void AddFileWriter(string name, FileReference outputFile)
 		{
-			Log.TraceInformation($"Log file: {OutputFile}");
+			Log.TraceInformation($"Log file: {outputFile}");
 
-			if (Log.BackupLogFiles && FileReference.Exists(OutputFile))
+			if (Log.BackupLogFiles && FileReference.Exists(outputFile))
 			{
 				// before creating a new backup, cap the number of existing files
-				string FilenameWithoutExtension = OutputFile.GetFileNameWithoutExtension();
-				string Extension = OutputFile.GetExtension();
+				string filenameWithoutExtension = outputFile.GetFileNameWithoutExtension();
+				string extension = outputFile.GetExtension();
 
-				Regex BackupForm =
-					new Regex(FilenameWithoutExtension + @"-backup-\d\d\d\d\.\d\d\.\d\d-\d\d\.\d\d\.\d\d" + Extension);
+				Regex backupForm =
+					new Regex(filenameWithoutExtension + @"-backup-\d\d\d\d\.\d\d\.\d\d-\d\d\.\d\d\.\d\d" + extension);
 
-				foreach (FileReference OldBackup in DirectoryReference
-					.EnumerateFiles(OutputFile.Directory)
+				foreach (FileReference oldBackup in DirectoryReference
+					.EnumerateFiles(outputFile.Directory)
 					// find files that match the way that we name backup files
-					.Where(x => BackupForm.IsMatch(x.GetFileName()))
+					.Where(x => backupForm.IsMatch(x.GetFileName()))
 					// sort them from newest to oldest
 					.OrderByDescending(x => x.GetFileName())
 					// skip the newest ones that are to be kept; -1 because we're about to create another backup.
 					.Skip(Log.LogFileBackupCount - 1))
 				{
-					Log.TraceLog($"Deleting old log file: {OldBackup}");
-					FileReference.Delete(OldBackup);
+					Log.TraceLog($"Deleting old log file: {oldBackup}");
+					FileReference.Delete(oldBackup);
 				}
 
 				// Ensure that the backup gets a unique name, in the extremely unlikely case that UBT was run twice during
 				// the same second.
-				DateTime FileTime = File.GetCreationTimeUtc(OutputFile.FullName);
+				DateTime fileTime = File.GetCreationTimeUtc(outputFile.FullName);
 					
-				FileReference BackupFile;
+				FileReference backupFile;
 				for (;;)
 				{
-					string Timestamp = $"{FileTime:yyyy.MM.dd-HH.mm.ss}";
-					BackupFile = FileReference.Combine(OutputFile.Directory,
-						$"{FilenameWithoutExtension}-backup-{Timestamp}{Extension}");
-					if (!FileReference.Exists(BackupFile))
+					string timestamp = $"{fileTime:yyyy.MM.dd-HH.mm.ss}";
+					backupFile = FileReference.Combine(outputFile.Directory,
+						$"{filenameWithoutExtension}-backup-{timestamp}{extension}");
+					if (!FileReference.Exists(backupFile))
 					{
 						break;
 					}
 
-					FileTime = FileTime.AddSeconds(1);
+					fileTime = fileTime.AddSeconds(1);
 				}
 
-				FileReference.Move(OutputFile, BackupFile);
+				FileReference.Move(outputFile, backupFile);
 			}
 			
-			TextWriterTraceListener FirstTextWriter = DefaultLogger.AddFileWriter(Name, OutputFile);
+			TextWriterTraceListener firstTextWriter = DefaultLogger.AddFileWriter(name, outputFile);
 			
 			// find the StartupTraceListener in the listeners that was added early on
-			IEnumerable<StartupTraceListener> StartupListeners = Trace.Listeners.OfType<StartupTraceListener>();
-			if (StartupListeners.Any())
+			IEnumerable<StartupTraceListener> startupListeners = Trace.Listeners.OfType<StartupTraceListener>();
+			if (startupListeners.Any())
 			{
-				StartupTraceListener StartupListener = StartupListeners.First();
-				StartupListener.CopyTo(FirstTextWriter);
-				Trace.Listeners.Remove(StartupListener);
+				StartupTraceListener startupListener = startupListeners.First();
+				startupListener.CopyTo(firstTextWriter);
+				Trace.Listeners.Remove(startupListener);
 			}
 		}
 
 		/// <summary>
 		/// Adds a <see cref="TraceListener"/> to the collection in a safe manner.
 		/// </summary>
-		/// <param name="TraceListener">The <see cref="TraceListener"/> to add.</param>
-		public static void AddTraceListener(TraceListener TraceListener)
+		/// <param name="traceListener">The <see cref="TraceListener"/> to add.</param>
+		public static void AddTraceListener(TraceListener traceListener)
 		{
-			DefaultLogger.AddTraceListener(TraceListener);
+			DefaultLogger.AddTraceListener(traceListener);
 		}
 
 		/// <summary>
 		/// Removes a <see cref="TraceListener"/> from the collection in a safe manner.
 		/// </summary>
-		/// <param name="TraceListener">The <see cref="TraceListener"/> to remove.</param>
-		public static void RemoveTraceListener(TraceListener TraceListener)
+		/// <param name="traceListener">The <see cref="TraceListener"/> to remove.</param>
+		public static void RemoveTraceListener(TraceListener traceListener)
 		{
-			DefaultLogger.RemoveTraceListener(TraceListener);
+			DefaultLogger.RemoveTraceListener(traceListener);
 		}
 
 		/// <summary>
@@ -273,11 +270,11 @@ namespace EpicGames.Core
 		/// <summary>
 		/// Converts a LogEventType into a log prefix. Only used when bLogSeverity is true.
 		/// </summary>
-		/// <param name="Severity"></param>
+		/// <param name="severity"></param>
 		/// <returns></returns>
-		private static string GetSeverityPrefix(LogEventType Severity)
+		private static string GetSeverityPrefix(LogEventType severity)
 		{
-			switch (Severity)
+			switch (severity)
 			{
 				case LogEventType.Fatal:
 					return "FATAL ERROR: ";
@@ -298,431 +295,430 @@ namespace EpicGames.Core
 		/// Writes a formatted message to the console. All other functions should boil down to calling this method.
 		/// </summary>
 		/// <param name="bWriteOnce">If true, this message will be written only once</param>
-		/// <param name="Verbosity">Message verbosity level. We only meaningfully use values up to Verbose</param>
-		/// <param name="FormatOptions">Options for formatting messages</param>
-		/// <param name="Format">Message format string.</param>
-		/// <param name="Args">Optional arguments</param>
+		/// <param name="verbosity">Message verbosity level. We only meaningfully use values up to Verbose</param>
+		/// <param name="formatOptions">Options for formatting messages</param>
+		/// <param name="format">Message format string.</param>
+		/// <param name="args">Optional arguments</param>
 		[StringFormatMethod("Format")]
-		private static void WriteLinePrivate(bool bWriteOnce, LogEventType Verbosity, LogFormatOptions FormatOptions, string Format, params object?[] Args)
+		private static void WriteLinePrivate(bool bWriteOnce, LogEventType verbosity, LogFormatOptions formatOptions, string format, params object?[] args)
 		{
-			if (Logger.IsEnabled((LogLevel)Verbosity))
+			if (Logger.IsEnabled((LogLevel)verbosity))
 			{
-				StringBuilder Message = new StringBuilder();
+				StringBuilder message = new StringBuilder();
 
 				// Get the severity prefix for this message
-				if (IncludeSeverityPrefix && ((FormatOptions & LogFormatOptions.NoSeverityPrefix) == 0))
+				if (IncludeSeverityPrefix && ((formatOptions & LogFormatOptions.NoSeverityPrefix) == 0))
 				{
-					Message.Append(GetSeverityPrefix(Verbosity));
-					if (Message.Length > 0 && IncludeProgramNameWithSeverityPrefix)
+					message.Append(GetSeverityPrefix(verbosity));
+					if (message.Length > 0 && IncludeProgramNameWithSeverityPrefix)
 					{
 						// Include the executable name when running inside MSBuild. If unspecified, MSBuild re-formats them with an "EXEC :" prefix.
-						Message.Insert(0, $"{Path.GetFileNameWithoutExtension(Assembly.GetEntryAssembly()!.Location)}: ");
+						message.Insert(0, $"{Path.GetFileNameWithoutExtension(Assembly.GetEntryAssembly()!.Location)}: ");
 					}
 				}
 
 				// Append the formatted string
-				int IndentLen = Message.Length;
-				if (Args.Length == 0)
+				int indentLen = message.Length;
+				if (args.Length == 0)
 				{
-					Message.Append(Format);
+					message.Append(format);
 				}
 				else
 				{
-					Message.AppendFormat(Format, Args);
+					message.AppendFormat(format, args);
 				}
 
 				// Replace any Windows \r\n sequences with \n
-				Message.Replace("\r\n", "\n");
+				message.Replace("\r\n", "\n");
 
 				// Remove any trailing whitespace
-				int TrimLen = Message.Length;
-				while (TrimLen > 0 && " \t\r\n".Contains(Message[TrimLen - 1]))
+				int trimLen = message.Length;
+				while (trimLen > 0 && " \t\r\n".Contains(message[trimLen - 1]))
 				{
-					TrimLen--;
+					trimLen--;
 				}
-				Message.Remove(TrimLen, Message.Length - TrimLen);
+				message.Remove(trimLen, message.Length - trimLen);
 
 				// Update the indent length to include any whitespace at the start of the message
-				while (IndentLen < Message.Length && Message[IndentLen] == ' ')
+				while (indentLen < message.Length && message[indentLen] == ' ')
 				{
-					IndentLen++;
+					indentLen++;
 				}
 
 				// If there are multiple lines, insert a prefix at the start of each one
-				for (int Idx = 0; Idx < Message.Length; Idx++)
+				for (int idx = 0; idx < message.Length; idx++)
 				{
-					if (Message[Idx] == '\n')
+					if (message[idx] == '\n')
 					{
-						Message.Insert(Idx + 1, " ", IndentLen);
-						Idx += IndentLen;
+						message.Insert(idx + 1, " ", indentLen);
+						idx += indentLen;
 					}
 				}
 
 				// if we want this message only written one time, check if it was already written out
-				if (bWriteOnce && !WriteOnceSet.Add(Message.ToString()))
+				if (bWriteOnce && !s_writeOnceSet.Add(message.ToString()))
 				{
 					return;
 				}
 
 				// Forward it on to the internal logger
-				Logger.Log((LogLevel)Verbosity, new EventId(), Message, null, (Message, Ex) => Message.ToString());
+				Logger.Log((LogLevel)verbosity, new EventId(), message, null, (message, ex) => message.ToString());
 			}
 		}
 
 		/// <summary>
 		/// Similar to Trace.WriteLineIf
 		/// </summary>
-		/// <param name="Condition"></param>
-		/// <param name="Verbosity"></param>
-		/// <param name="Format"></param>
-		/// <param name="Args"></param>
+		/// <param name="condition"></param>
+		/// <param name="verbosity"></param>
+		/// <param name="format"></param>
+		/// <param name="args"></param>
 		[StringFormatMethod("Format")]
-		public static void WriteLineIf(bool Condition, LogEventType Verbosity, string Format, params object?[] Args)
+		public static void WriteLineIf(bool condition, LogEventType verbosity, string format, params object?[] args)
 		{
-			if (Condition)
+			if (condition)
 			{
-				WriteLinePrivate(false, Verbosity, LogFormatOptions.None, Format, Args);
+				WriteLinePrivate(false, verbosity, LogFormatOptions.None, format, args);
 			}
 		}
 
 		/// <summary>
 		/// Mostly an internal function, but expose StackFramesToSkip to allow UAT to use existing wrapper functions and still get proper formatting.
 		/// </summary>
-		/// <param name="StackFramesToSkip"></param>
-		/// <param name="Verbosity"></param>
-		/// <param name="Format"></param>
-		/// <param name="Args"></param>
+		/// <param name="stackFramesToSkip"></param>
+		/// <param name="verbosity"></param>
+		/// <param name="format"></param>
+		/// <param name="args"></param>
 		[Obsolete("StackFramesToSkip has been deprecated since 5.0. Please call an override without this parameter")]
 		[StringFormatMethod("Format")]
-		public static void WriteLine(int StackFramesToSkip, LogEventType Verbosity, string Format, params object?[] Args)
+		public static void WriteLine(int stackFramesToSkip, LogEventType verbosity, string format, params object?[] args)
 		{
-			WriteLinePrivate(false, Verbosity, LogFormatOptions.None, Format, Args);
+			WriteLinePrivate(false, verbosity, LogFormatOptions.None, format, args);
 		}
 
 		/// <summary>
 		/// Mostly an internal function, but expose StackFramesToSkip to allow UAT to use existing wrapper functions and still get proper formatting.
 		/// </summary>
-		/// <param name="StackFramesToSkip"></param>
-		/// <param name="Verbosity"></param>
-		/// <param name="FormatOptions"></param>
-		/// <param name="Format"></param>
-		/// <param name="Args"></param>
+		/// <param name="stackFramesToSkip"></param>
+		/// <param name="verbosity"></param>
+		/// <param name="formatOptions"></param>
+		/// <param name="format"></param>
+		/// <param name="args"></param>
 		[Obsolete("StackFramesToSkip has been deprecated since 5.0. Please call an override without this parameter")]
 		[StringFormatMethod("Format")]
-		public static void WriteLine(int StackFramesToSkip, LogEventType Verbosity, LogFormatOptions FormatOptions, string Format, params object?[] Args)
+		public static void WriteLine(int stackFramesToSkip, LogEventType verbosity, LogFormatOptions formatOptions, string format, params object?[] args)
 		{
-			WriteLinePrivate(false, Verbosity, FormatOptions, Format, Args);
+			WriteLinePrivate(false, verbosity, formatOptions, format, args);
 		}
 
 		/// <summary>
 		/// Similar to Trace.WriteLine
 		/// </summary>
-		/// <param name="Verbosity"></param>
-		/// <param name="Format"></param>
-		/// <param name="Args"></param>
+		/// <param name="verbosity"></param>
+		/// <param name="format"></param>
+		/// <param name="args"></param>
 		[StringFormatMethod("Format")]
-		public static void WriteLine(LogEventType Verbosity, string Format, params object?[] Args)
+		public static void WriteLine(LogEventType verbosity, string format, params object?[] args)
 		{
-			WriteLinePrivate(false, Verbosity, LogFormatOptions.None, Format, Args);
+			WriteLinePrivate(false, verbosity, LogFormatOptions.None, format, args);
 		}
 
 		/// <summary>
 		/// Similar to Trace.WriteLine
 		/// </summary>
-		/// <param name="Verbosity"></param>
-		/// <param name="FormatOptions"></param>
-		/// <param name="Format"></param>
-		/// <param name="Args"></param>
+		/// <param name="verbosity"></param>
+		/// <param name="formatOptions"></param>
+		/// <param name="format"></param>
+		/// <param name="args"></param>
 		[StringFormatMethod("Format")]
-		public static void WriteLine(LogEventType Verbosity, LogFormatOptions FormatOptions, string Format, params object?[] Args)
+		public static void WriteLine(LogEventType verbosity, LogFormatOptions formatOptions, string format, params object?[] args)
 		{
-			WriteLinePrivate(false, Verbosity, FormatOptions, Format, Args);
+			WriteLinePrivate(false, verbosity, formatOptions, format, args);
 		}
 
 		/// <summary>
 		/// Formats an exception for display in the log. The exception message is shown as an error, and the stack trace is included in the log.
 		/// </summary>
-		/// <param name="Ex">The exception to display</param>
-		/// <param name="LogFileName">The log filename to display, if any</param>
-		public static void WriteException(Exception Ex, string? LogFileName)
+		/// <param name="ex">The exception to display</param>
+		/// <param name="logFileName">The log filename to display, if any</param>
+		public static void WriteException(Exception ex, string? logFileName)
 		{
-			string LogSuffix = (LogFileName == null) ? "" : String.Format("\n(see {0} for full exception trace)", LogFileName);
+			string logSuffix = (logFileName == null) ? "" : String.Format("\n(see {0} for full exception trace)", logFileName);
 			TraceLog("==============================================================================");
-			TraceError("{0}{1}", ExceptionUtils.FormatException(Ex), LogSuffix);
-			TraceLog("\n{0}", ExceptionUtils.FormatExceptionDetails(Ex));
+			TraceError("{0}{1}", ExceptionUtils.FormatException(ex), logSuffix);
+			TraceLog("\n{0}", ExceptionUtils.FormatExceptionDetails(ex));
 			TraceLog("==============================================================================");
 		}
 
 		/// <summary>
 		/// Writes an error message to the console.
 		/// </summary>
-		/// <param name="Format">Message format string</param>
-		/// <param name="Args">Optional arguments</param>
+		/// <param name="format">Message format string</param>
+		/// <param name="args">Optional arguments</param>
 		[StringFormatMethod("Format")]
-		public static void TraceError(string Format, params object?[] Args)
+		public static void TraceError(string format, params object?[] args)
 		{
-			WriteLinePrivate(false, LogEventType.Error, LogFormatOptions.None, Format, Args);
+			WriteLinePrivate(false, LogEventType.Error, LogFormatOptions.None, format, args);
 		}
 
 		/// <summary>
 		/// Writes an error message to the console, in a format suitable for Visual Studio to parse.
 		/// </summary>
-		/// <param name="File">The file containing the error</param>
-		/// <param name="Format">Message format string</param>
-		/// <param name="Args">Optional arguments</param>
+		/// <param name="file">The file containing the error</param>
+		/// <param name="format">Message format string</param>
+		/// <param name="args">Optional arguments</param>
 		[StringFormatMethod("Format")]
-		public static void TraceErrorTask(FileReference File, string Format, params object?[] Args)
+		public static void TraceErrorTask(FileReference file, string format, params object?[] args)
 		{
-			WriteLinePrivate(false, LogEventType.Error, LogFormatOptions.NoSeverityPrefix, "{0}: error: {1}", File, String.Format(Format, Args));
+			WriteLinePrivate(false, LogEventType.Error, LogFormatOptions.NoSeverityPrefix, "{0}: error: {1}", file, String.Format(format, args));
 		}
 
 		/// <summary>
 		/// Writes an error message to the console, in a format suitable for Visual Studio to parse.
 		/// </summary>
-		/// <param name="File">The file containing the error</param>
-		/// <param name="Line">Line number of the error</param>
-		/// <param name="Format">Message format string</param>
-		/// <param name="Args">Optional arguments</param>
+		/// <param name="file">The file containing the error</param>
+		/// <param name="line">Line number of the error</param>
+		/// <param name="format">Message format string</param>
+		/// <param name="args">Optional arguments</param>
 		[StringFormatMethod("Format")]
-		public static void TraceErrorTask(FileReference File, int Line, string Format, params object?[] Args)
+		public static void TraceErrorTask(FileReference file, int line, string format, params object?[] args)
 		{
-			WriteLinePrivate(false, LogEventType.Error, LogFormatOptions.NoSeverityPrefix, "{0}({1}): error: {2}", File, Line, String.Format(Format, Args));
+			WriteLinePrivate(false, LogEventType.Error, LogFormatOptions.NoSeverityPrefix, "{0}({1}): error: {2}", file, line, String.Format(format, args));
 		}
 
 		/// <summary>
 		/// Writes a verbose message to the console.
 		/// </summary>
-		/// <param name="Format">Message format string</param>
-		/// <param name="Args">Optional arguments</param>
+		/// <param name="format">Message format string</param>
+		/// <param name="args">Optional arguments</param>
 		[Conditional("TRACE")]
 		[StringFormatMethod("Format")]
-		public static void TraceVerbose(string Format, params object?[] Args)
+		public static void TraceVerbose(string format, params object?[] args)
 		{
-			WriteLinePrivate(false, LogEventType.Verbose, LogFormatOptions.None, Format, Args);
+			WriteLinePrivate(false, LogEventType.Verbose, LogFormatOptions.None, format, args);
 		}
 
 		/// <summary>
 		/// Writes a message to the console.
 		/// </summary>
-		/// <param name="Format">Message format string</param>
-		/// <param name="Args">Optional arguments</param>
+		/// <param name="format">Message format string</param>
+		/// <param name="args">Optional arguments</param>
 		[StringFormatMethod("Format")]
-		public static void TraceInformation(string Format, params object?[] Args)
+		public static void TraceInformation(string format, params object?[] args)
 		{
-			WriteLinePrivate(false, LogEventType.Console, LogFormatOptions.None, Format, Args);
+			WriteLinePrivate(false, LogEventType.Console, LogFormatOptions.None, format, args);
 		}
 
 		/// <summary>
 		/// Writes a warning message to the console.
 		/// </summary>
-		/// <param name="Format">Message format string</param>
-		/// <param name="Args">Optional arguments</param>
+		/// <param name="format">Message format string</param>
+		/// <param name="args">Optional arguments</param>
 		[StringFormatMethod("Format")]
-		public static void TraceWarning(string Format, params object?[] Args)
+		public static void TraceWarning(string format, params object?[] args)
 		{
-			WriteLinePrivate(false, LogEventType.Warning, LogFormatOptions.None, Format, Args);
+			WriteLinePrivate(false, LogEventType.Warning, LogFormatOptions.None, format, args);
 		}
 
 		/// <summary>
 		/// Writes a warning message to the console, in a format suitable for Visual Studio to parse.
 		/// </summary>
-		/// <param name="File">The file containing the warning</param>
-		/// <param name="Format">Message format string</param>
-		/// <param name="Args">Optional arguments</param>
+		/// <param name="file">The file containing the warning</param>
+		/// <param name="format">Message format string</param>
+		/// <param name="args">Optional arguments</param>
 		[StringFormatMethod("Format")]
-		public static void TraceWarningTask(FileReference File, string Format, params object?[] Args)
+		public static void TraceWarningTask(FileReference file, string format, params object?[] args)
 		{
-			WriteLinePrivate(false, LogEventType.Warning, LogFormatOptions.NoSeverityPrefix, "{0}: warning: {1}", File, String.Format(Format, Args));
+			WriteLinePrivate(false, LogEventType.Warning, LogFormatOptions.NoSeverityPrefix, "{0}: warning: {1}", file, String.Format(format, args));
 		}
 
 		/// <summary>
 		/// Writes a warning message to the console, in a format suitable for Visual Studio to parse.
 		/// </summary>
-		/// <param name="File">The file containing the warning</param>
-		/// <param name="Line">Line number of the warning</param>
-		/// <param name="Format">Message format string</param>
-		/// <param name="Args">Optional arguments</param>
+		/// <param name="file">The file containing the warning</param>
+		/// <param name="line">Line number of the warning</param>
+		/// <param name="format">Message format string</param>
+		/// <param name="args">Optional arguments</param>
 		[StringFormatMethod("Format")]
-		public static void TraceWarningTask(FileReference File, int Line, string Format, params object?[] Args)
+		public static void TraceWarningTask(FileReference file, int line, string format, params object?[] args)
 		{
-			WriteLinePrivate(false, LogEventType.Warning, LogFormatOptions.NoSeverityPrefix, "{0}({1}): warning: {2}", File, Line, String.Format(Format, Args));
+			WriteLinePrivate(false, LogEventType.Warning, LogFormatOptions.NoSeverityPrefix, "{0}({1}): warning: {2}", file, line, String.Format(format, args));
 		}
 
 		/// <summary>
 		/// Writes a message to the console.
 		/// </summary>
-		/// <param name="File">The file containing the message</param>
-		/// <param name="Line">Line number of the message</param>
-		/// <param name="Format">Message format string</param>
-		/// <param name="Args">Optional arguments</param>
-		public static void TraceConsoleTask(FileReference File, int Line, string Format, params object?[] Args)
+		/// <param name="file">The file containing the message</param>
+		/// <param name="line">Line number of the message</param>
+		/// <param name="format">Message format string</param>
+		/// <param name="args">Optional arguments</param>
+		public static void TraceConsoleTask(FileReference file, int line, string format, params object?[] args)
 		{
-			WriteLinePrivate(false, LogEventType.Console, LogFormatOptions.NoSeverityPrefix, "{0}({1}): {2}", File, Line, String.Format(Format, Args));
+			WriteLinePrivate(false, LogEventType.Console, LogFormatOptions.NoSeverityPrefix, "{0}({1}): {2}", file, line, String.Format(format, args));
 		}
 
-
 		/// <summary>
 		/// Writes a very verbose message to the console.
 		/// </summary>
-		/// <param name="Format">Message format string</param>
-		/// <param name="Args">Optional arguments</param>
+		/// <param name="format">Message format string</param>
+		/// <param name="args">Optional arguments</param>
 		[Conditional("TRACE")]
 		[StringFormatMethod("Format")]
-		public static void TraceVeryVerbose(string Format, params object?[] Args)
+		public static void TraceVeryVerbose(string format, params object?[] args)
 		{
-			WriteLinePrivate(false, LogEventType.VeryVerbose, LogFormatOptions.None, Format, Args);
+			WriteLinePrivate(false, LogEventType.VeryVerbose, LogFormatOptions.None, format, args);
 		}
 
 		/// <summary>
 		/// Writes a message to the log only.
 		/// </summary>
-		/// <param name="Format">Message format string</param>
-		/// <param name="Args">Optional arguments</param>
+		/// <param name="format">Message format string</param>
+		/// <param name="args">Optional arguments</param>
 		[Conditional("TRACE")]
 		[StringFormatMethod("Format")]
-		public static void TraceLog(string Format, params object?[] Args)
+		public static void TraceLog(string format, params object?[] args)
 		{
-			WriteLinePrivate(false, LogEventType.Log, LogFormatOptions.None, Format, Args);
+			WriteLinePrivate(false, LogEventType.Log, LogFormatOptions.None, format, args);
 		}
 
 		/// <summary>
 		/// Similar to Trace.WriteLine
 		/// </summary>
-		/// <param name="Verbosity"></param>
-		/// <param name="Format"></param>
-		/// <param name="Args"></param>
+		/// <param name="verbosity"></param>
+		/// <param name="format"></param>
+		/// <param name="args"></param>
 		[StringFormatMethod("Format")]
-		public static void WriteLineOnce(LogEventType Verbosity, string Format, params object?[] Args)
+		public static void WriteLineOnce(LogEventType verbosity, string format, params object?[] args)
 		{
-			WriteLinePrivate(true, Verbosity, LogFormatOptions.None, Format, Args);
+			WriteLinePrivate(true, verbosity, LogFormatOptions.None, format, args);
 		}
 
 		/// <summary>
 		/// Similar to Trace.WriteLine
 		/// </summary>
-		/// <param name="Verbosity"></param>
-		/// <param name="Options"></param>
-		/// <param name="Format"></param>
-		/// <param name="Args"></param>
+		/// <param name="verbosity"></param>
+		/// <param name="options"></param>
+		/// <param name="format"></param>
+		/// <param name="args"></param>
 		[StringFormatMethod("Format")]
-		public static void WriteLineOnce(LogEventType Verbosity, LogFormatOptions Options, string Format, params object?[] Args)
+		public static void WriteLineOnce(LogEventType verbosity, LogFormatOptions options, string format, params object?[] args)
 		{
-			WriteLinePrivate(true, Verbosity, Options, Format, Args);
+			WriteLinePrivate(true, verbosity, options, format, args);
 		}
 
 		/// <summary>
 		/// Writes an error message to the console.
 		/// </summary>
-		/// <param name="Format">Message format string</param>
-		/// <param name="Args">Optional arguments</param>
+		/// <param name="format">Message format string</param>
+		/// <param name="args">Optional arguments</param>
 		[StringFormatMethod("Format")]
-		public static void TraceErrorOnce(string Format, params object?[] Args)
+		public static void TraceErrorOnce(string format, params object?[] args)
 		{
-			WriteLinePrivate(true, LogEventType.Error, LogFormatOptions.None, Format, Args);
+			WriteLinePrivate(true, LogEventType.Error, LogFormatOptions.None, format, args);
 		}
 
 		/// <summary>
 		/// Writes a verbose message to the console.
 		/// </summary>
-		/// <param name="Format">Message format string</param>
-		/// <param name="Args">Optional arguments</param>
+		/// <param name="format">Message format string</param>
+		/// <param name="args">Optional arguments</param>
 		[Conditional("TRACE")]
 		[StringFormatMethod("Format")]
-		public static void TraceVerboseOnce(string Format, params object?[] Args)
+		public static void TraceVerboseOnce(string format, params object?[] args)
 		{
-			WriteLinePrivate(true, LogEventType.Verbose, LogFormatOptions.None, Format, Args);
+			WriteLinePrivate(true, LogEventType.Verbose, LogFormatOptions.None, format, args);
 		}
 
 		/// <summary>
 		/// Writes a message to the console.
 		/// </summary>
-		/// <param name="Format">Message format string</param>
-		/// <param name="Args">Optional arguments</param>
+		/// <param name="format">Message format string</param>
+		/// <param name="args">Optional arguments</param>
 		[StringFormatMethod("Format")]
-		public static void TraceInformationOnce(string Format, params object?[] Args)
+		public static void TraceInformationOnce(string format, params object?[] args)
 		{
-			WriteLinePrivate(true, LogEventType.Console, LogFormatOptions.None, Format, Args);
+			WriteLinePrivate(true, LogEventType.Console, LogFormatOptions.None, format, args);
 		}
 
 		/// <summary>
 		/// Writes a warning message to the console.
 		/// </summary>
-		/// <param name="Format">Message format string</param>
-		/// <param name="Args">Optional arguments</param>
+		/// <param name="format">Message format string</param>
+		/// <param name="args">Optional arguments</param>
 		[StringFormatMethod("Format")]
-		public static void TraceWarningOnce(string Format, params object?[] Args)
+		public static void TraceWarningOnce(string format, params object?[] args)
 		{
-			WriteLinePrivate(true, LogEventType.Warning, LogFormatOptions.None, Format, Args);
+			WriteLinePrivate(true, LogEventType.Warning, LogFormatOptions.None, format, args);
 		}
 
 		/// <summary>
 		/// Writes a warning message to the console.
 		/// </summary>
-		/// <param name="File">The file containing the error</param>
-		/// <param name="Format">Message format string</param>
-		/// <param name="Args">Optional arguments</param>
+		/// <param name="file">The file containing the error</param>
+		/// <param name="format">Message format string</param>
+		/// <param name="args">Optional arguments</param>
 		[StringFormatMethod("Format")]
-		public static void TraceWarningOnce(FileReference File, string Format, params object?[] Args)
+		public static void TraceWarningOnce(FileReference file, string format, params object?[] args)
 		{
-			WriteLinePrivate( true, LogEventType.Warning, LogFormatOptions.NoSeverityPrefix, "{0}: warning: {1}", File, String.Format(Format, Args));
+			WriteLinePrivate( true, LogEventType.Warning, LogFormatOptions.NoSeverityPrefix, "{0}: warning: {1}", file, String.Format(format, args));
 		}
 
 		/// <summary>
 		/// Writes a warning message to the console.
 		/// </summary>
-		/// <param name="File">The file containing the error</param>
-		/// <param name="Line">Line number of the error</param>
-		/// <param name="Format">Message format string</param>
-		/// <param name="Args">Optional arguments</param>
+		/// <param name="file">The file containing the error</param>
+		/// <param name="line">Line number of the error</param>
+		/// <param name="format">Message format string</param>
+		/// <param name="args">Optional arguments</param>
 		[StringFormatMethod("Format")]
-		public static void TraceWarningOnce(FileReference File, int Line, string Format, params object?[] Args)
+		public static void TraceWarningOnce(FileReference file, int line, string format, params object?[] args)
 		{
-			WriteLinePrivate(true, LogEventType.Warning, LogFormatOptions.NoSeverityPrefix, "{0}({1}): warning: {2}", File, Line, String.Format(Format, Args));
+			WriteLinePrivate(true, LogEventType.Warning, LogFormatOptions.NoSeverityPrefix, "{0}({1}): warning: {2}", file, line, String.Format(format, args));
 		}
 
 		/// <summary>
 		/// Writes a very verbose message to the console.
 		/// </summary>
-		/// <param name="Format">Message format string</param>
-		/// <param name="Args">Optional arguments</param>
+		/// <param name="format">Message format string</param>
+		/// <param name="args">Optional arguments</param>
 		[Conditional("TRACE")]
 		[StringFormatMethod("Format")]
-		public static void TraceVeryVerboseOnce(string Format, params object?[] Args)
+		public static void TraceVeryVerboseOnce(string format, params object?[] args)
 		{
-			WriteLinePrivate(true, LogEventType.VeryVerbose, LogFormatOptions.None, Format, Args);
+			WriteLinePrivate(true, LogEventType.VeryVerbose, LogFormatOptions.None, format, args);
 		}
 
 		/// <summary>
 		/// Writes a message to the log only.
 		/// </summary>
-		/// <param name="Format">Message format string</param>
-		/// <param name="Args">Optional arguments</param>
+		/// <param name="format">Message format string</param>
+		/// <param name="args">Optional arguments</param>
 		[Conditional("TRACE")]
 		[StringFormatMethod("Format")]
-		public static void TraceLogOnce(string Format, params object?[] Args)
+		public static void TraceLogOnce(string format, params object?[] args)
 		{
-			WriteLinePrivate(true, LogEventType.Log, LogFormatOptions.None, Format, Args);
+			WriteLinePrivate(true, LogEventType.Log, LogFormatOptions.None, format, args);
 		}
 
 		/// <summary>
 		/// Enter a scope with the given status message. The message will be written to the console without a newline, allowing it to be updated through subsequent calls to UpdateStatus().
 		/// The message will be written to the log immediately. If another line is written while in a status scope, the initial status message is flushed to the console first.
 		/// </summary>
-		/// <param name="Message">The status message</param>
+		/// <param name="message">The status message</param>
 		[Conditional("TRACE")]
-		public static void PushStatus(string Message)
+		public static void PushStatus(string message)
 		{
-			DefaultLogger.PushStatus(Message);
+			DefaultLogger.PushStatus(message);
 		}
 
 		/// <summary>
 		/// Updates the current status message. This will overwrite the previous status line.
 		/// </summary>
-		/// <param name="Message">The status message</param>
+		/// <param name="message">The status message</param>
 		[Conditional("TRACE")]
-		public static void UpdateStatus(string Message)
+		public static void UpdateStatus(string message)
 		{
-			DefaultLogger.UpdateStatus(Message);
+			DefaultLogger.UpdateStatus(message);
 		}
 
 		/// <summary>
@@ -761,7 +757,7 @@ namespace EpicGames.Core
 		/// Renders the captured events as a single string
 		/// </summary>
 		/// <returns>Rendered log text</returns>
-		public string Render(string NewLine) => String.Join(NewLine, RenderLines());
+		public string Render(string newLine) => String.Join(newLine, RenderLines());
 
 		/// <summary>
 		/// Renders all the captured events
@@ -770,15 +766,15 @@ namespace EpicGames.Core
 		public List<string> RenderLines() => Events.ConvertAll(x => x.ToString());
 
 		/// <inheritdoc/>
-		public IDisposable BeginScope<TState>(TState State) => new NullScope();
+		public IDisposable BeginScope<TState>(TState state) => new NullScope();
 
 		/// <inheritdoc/>
-		public bool IsEnabled(LogLevel LogLevel) => true;
+		public bool IsEnabled(LogLevel logLevel) => true;
 
 		/// <inheritdoc/>
-		public void Log<TState>(LogLevel LogLevel, EventId EventId, TState State, Exception? Exception, Func<TState, Exception?, string> Formatter)
+		public void Log<TState>(LogLevel logLevel, EventId eventId, TState state, Exception? exception, Func<TState, Exception?, string> formatter)
 		{
-			Events.Add(LogEvent.FromState(LogLevel, EventId, State, Exception, Formatter));
+			Events.Add(LogEvent.FromState(logLevel, eventId, state, exception, formatter));
 		}
 	}
 
@@ -796,32 +792,32 @@ namespace EpicGames.Core
 			/// <summary>
 			/// The heading for this status message.
 			/// </summary>
-			public string HeadingText;
+			public string _headingText;
 
 			/// <summary>
 			/// The current status text.
 			/// </summary>
-			public string CurrentText;
+			public string _currentText;
 
 			/// <summary>
 			/// Whether the heading has been written to the console. Before the first time that lines are output to the log in the midst of a status scope, the heading will be written on a line of its own first.
 			/// </summary>
-			public bool bHasFlushedHeadingText;
+			public bool _hasFlushedHeadingText;
 
 			/// <summary>
 			/// Constructor
 			/// </summary>
-			public StatusMessage(string HeadingText, string CurrentText)
+			public StatusMessage(string headingText, string currentText)
 			{
-				this.HeadingText = HeadingText;
-				this.CurrentText = CurrentText;
+				_headingText = headingText;
+				_currentText = currentText;
 			}
 		}
 
 		/// <summary>
 		/// Object used for synchronization
 		/// </summary>
-		private object SyncObject = new object();
+		private readonly object _syncObject = new object();
 
 		/// <summary>
 		/// Minimum level for outputting messages
@@ -862,7 +858,7 @@ namespace EpicGames.Core
 		{
 			get; set;
 		}
-		private bool IncludeStartingTimestampWritten = false;
+		private bool _includeStartingTimestampWritten = false;
 
 		/// <summary>
 		/// Path to the log file being written to. May be null.
@@ -875,33 +871,29 @@ namespace EpicGames.Core
 		/// <summary>
 		/// Whether console output is redirected. This prevents writing status updates that rely on moving the cursor.
 		/// </summary>
-		private bool AllowStatusUpdates
-		{
-			get { return !Console.IsOutputRedirected; }
-		}
+		private bool AllowStatusUpdates => !Console.IsOutputRedirected;
 
 		/// <summary>
 		/// When configured, this tracks time since initialization to prepend a timestamp to each log.
 		/// </summary>
-		private Stopwatch Timer = Stopwatch.StartNew();
+		private readonly Stopwatch _timer = Stopwatch.StartNew();
 
 		/// <summary>
 		/// Stack of status scope information.
 		/// </summary>
-		private Stack<StatusMessage> StatusMessageStack = new Stack<StatusMessage>();
+		private readonly Stack<StatusMessage> _statusMessageStack = new Stack<StatusMessage>();
 
 		/// <summary>
 		/// The currently visible status text
 		/// </summary>
-		private string StatusText = "";
+		private string _statusText = "";
 
 		/// <summary>
 		/// Last time a status message was pushed to the stack
 		/// </summary>
-		private Stopwatch StatusTimer = new Stopwatch();
-
-		ArrayBufferWriter<byte> JsonBufferWriter;
-		Utf8JsonWriter JsonWriter;
+		private readonly Stopwatch _statusTimer = new Stopwatch();
+		readonly ArrayBufferWriter<byte> _jsonBufferWriter;
+		readonly Utf8JsonWriter _jsonWriter;
 
 		/// <summary>
 		/// Constructor
@@ -912,53 +904,53 @@ namespace EpicGames.Core
 			ColorConsoleOutput = true;
 			IncludeStartingTimestamp = true;
 
-			string? EnvVar = Environment.GetEnvironmentVariable("UE_STDOUT_JSON");
-			if(EnvVar != null && int.TryParse(EnvVar, out int Value) && Value != 0)
+			string? envVar = Environment.GetEnvironmentVariable("UE_STDOUT_JSON");
+			if(envVar != null && Int32.TryParse(envVar, out int value) && value != 0)
 			{
 				WriteJsonToStdOut = true;
 			}
 
-			JsonBufferWriter = new ArrayBufferWriter<byte>();
-			JsonWriter = new Utf8JsonWriter(JsonBufferWriter);
+			_jsonBufferWriter = new ArrayBufferWriter<byte>();
+			_jsonWriter = new Utf8JsonWriter(_jsonBufferWriter);
 		}
 
 		/// <summary>
 		/// Adds a trace listener that writes to a log file
 		/// </summary>
-		/// <param name="Name">Listener name</param>
-		/// <param name="OutputFile">The file to write to</param>
+		/// <param name="name">Listener name</param>
+		/// <param name="outputFile">The file to write to</param>
 		/// <returns>The created trace listener</returns>
-		public TextWriterTraceListener AddFileWriter(string Name, FileReference OutputFile)
+		public TextWriterTraceListener AddFileWriter(string name, FileReference outputFile)
 		{
 			try
 			{
-				this.OutputFile = OutputFile;
-				DirectoryReference.CreateDirectory(OutputFile.Directory);
-				TextWriterTraceListener LogTraceListener = new TextWriterTraceListener(new StreamWriter(OutputFile.FullName), Name);
-				lock (SyncObject)
+				OutputFile = outputFile;
+				DirectoryReference.CreateDirectory(outputFile.Directory);
+				TextWriterTraceListener logTraceListener = new TextWriterTraceListener(new StreamWriter(outputFile.FullName), name);
+				lock (_syncObject)
 				{
-					Trace.Listeners.Add(LogTraceListener);
+					Trace.Listeners.Add(logTraceListener);
 					WriteInitialTimestamp();
 				}
-				return LogTraceListener;
+				return logTraceListener;
 			}
-			catch (Exception Ex)
+			catch (Exception ex)
 			{
-				throw new Exception($"Error while creating log file \"{OutputFile}\"", Ex);
+				throw new Exception($"Error while creating log file \"{outputFile}\"", ex);
 			}
 		}
 
 		/// <summary>
 		/// Adds a <see cref="TraceListener"/> to the collection in a safe manner.
 		/// </summary>
-		/// <param name="TraceListener">The <see cref="TraceListener"/> to add.</param>
-		public void AddTraceListener(TraceListener TraceListener)
+		/// <param name="traceListener">The <see cref="TraceListener"/> to add.</param>
+		public void AddTraceListener(TraceListener traceListener)
 		{
-			lock (SyncObject)
+			lock (_syncObject)
 			{
-				if (!Trace.Listeners.Contains(TraceListener))
+				if (!Trace.Listeners.Contains(traceListener))
 				{
-					Trace.Listeners.Add(TraceListener);
+					Trace.Listeners.Add(traceListener);
 					WriteInitialTimestamp();
 				}
 			}
@@ -969,25 +961,25 @@ namespace EpicGames.Core
 		/// </summary>
 		private void WriteInitialTimestamp()
 		{
-			if (IncludeStartingTimestamp && !IncludeStartingTimestampWritten)
+			if (IncludeStartingTimestamp && !_includeStartingTimestampWritten)
 			{
-				DateTime Now = DateTime.Now;
-				this.LogDebug("{Message}", $"Log started at {Now} ({Now.ToUniversalTime():yyyy-MM-ddTHH\\:mm\\:ssZ})");
-				IncludeStartingTimestampWritten = true;
+				DateTime now = DateTime.Now;
+				this.LogDebug("{Message}", $"Log started at {now} ({now.ToUniversalTime():yyyy-MM-ddTHH\\:mm\\:ssZ})");
+				_includeStartingTimestampWritten = true;
 			}
 		}
 
 		/// <summary>
 		/// Removes a <see cref="TraceListener"/> from the collection in a safe manner.
 		/// </summary>
-		/// <param name="TraceListener">The <see cref="TraceListener"/> to remove.</param>
-		public void RemoveTraceListener(TraceListener TraceListener)
+		/// <param name="traceListener">The <see cref="TraceListener"/> to remove.</param>
+		public void RemoveTraceListener(TraceListener traceListener)
 		{
-			lock (SyncObject)
+			lock (_syncObject)
 			{
-				if (Trace.Listeners.Contains(TraceListener))
+				if (Trace.Listeners.Contains(traceListener))
 				{
-					Trace.Listeners.Remove(TraceListener);
+					Trace.Listeners.Remove(traceListener);
 				}
 			}
 		}
@@ -998,9 +990,9 @@ namespace EpicGames.Core
 		/// <returns>True if a TextWriterTraceListener has been added</returns>
 		public static bool HasFileWriter()
 		{
-			foreach (TraceListener? Listener in Trace.Listeners)
+			foreach (TraceListener? listener in Trace.Listeners)
 			{
-				if (Listener is TextWriterTraceListener)
+				if (listener is TextWriterTraceListener)
 				{
 					return true;
 				}
@@ -1008,56 +1000,56 @@ namespace EpicGames.Core
 			return false;
 		}
 
-		public IDisposable BeginScope<TState>(TState State)
+		public IDisposable BeginScope<TState>(TState state)
 		{
 			throw new NotImplementedException();
 		}
 
-		public bool IsEnabled(LogLevel LogLevel)
+		public bool IsEnabled(LogLevel logLevel)
 		{
-			return LogLevel >= OutputLevel;
+			return logLevel >= OutputLevel;
 		}
 
-		public void Log<TState>(LogLevel LogLevel, EventId EventId, TState State, Exception Exception, Func<TState, Exception, string> Formatter)
+		public void Log<TState>(LogLevel logLevel, EventId eventId, TState state, Exception exception, Func<TState, Exception, string> formatter)
 		{
-			string[] Lines = Formatter(State, Exception).Split('\n');
-			lock (SyncObject)
+			string[] lines = formatter(state, exception).Split('\n');
+			lock (_syncObject)
 			{
 				// Output to all the other trace listeners
-				string TimePrefix = String.Format("[{0:hh\\:mm\\:ss\\.fff}] ", Timer.Elapsed);
-				foreach (TraceListener? Listener in Trace.Listeners)
+				string timePrefix = String.Format("[{0:hh\\:mm\\:ss\\.fff}] ", _timer.Elapsed);
+				foreach (TraceListener? listener in Trace.Listeners)
 				{
-					if (Listener != null)
+					if (listener != null)
 					{
-						string TimePrefixActual =
+						string timePrefixActual =
 							IncludeTimestamps &&
-							!(Listener is DefaultTraceListener) // no timestamps when writing to the Visual Studio debug window
-							? TimePrefix
+							!(listener is DefaultTraceListener) // no timestamps when writing to the Visual Studio debug window
+							? timePrefix
 							: String.Empty;
 
-						foreach (string Line in Lines)
+						foreach (string line in lines)
 						{
-							string LineWithTime = TimePrefixActual + Line;
-							Listener.WriteLine(Line);
-							Listener.Flush();
+							string lineWithTime = timePrefixActual + line;
+							listener.WriteLine(line);
+							listener.Flush();
 						}
 					}
 				}
 
 				// Handle the console output separately; we format things differently
-				if (LogLevel >= LogLevel.Information)
+				if (logLevel >= LogLevel.Information)
 				{
 					FlushStatusHeading();
 
 					bool bResetConsoleColor = false;
 					if (ColorConsoleOutput)
 					{
-						if (LogLevel == LogLevel.Warning)
+						if (logLevel == LogLevel.Warning)
 						{
 							Console.ForegroundColor = ConsoleColor.Yellow;
 							bResetConsoleColor = true;
 						}
-						if (LogLevel >= LogLevel.Error)
+						if (logLevel >= LogLevel.Error)
 						{
 							Console.ForegroundColor = ConsoleColor.Red;
 							bResetConsoleColor = true;
@@ -1067,20 +1059,20 @@ namespace EpicGames.Core
 					{
 						if (WriteJsonToStdOut)
 						{
-							JsonBufferWriter.Clear();
-							JsonWriter.Reset();
+							_jsonBufferWriter.Clear();
+							_jsonWriter.Reset();
 
-							LogEvent Event = new LogEvent(DateTime.UtcNow, LogLevel, EventId, Formatter(State, Exception), null, null, LogException.FromException(Exception));
-							Event.Write(JsonWriter);
+							LogEvent logEvent = new LogEvent(DateTime.UtcNow, logLevel, eventId, formatter(state, exception), null, null, LogException.FromException(exception));
+							logEvent.Write(_jsonWriter);
 
-							JsonWriter.Flush();
-							Console.WriteLine(Encoding.UTF8.GetString(JsonBufferWriter.WrittenSpan));
+							_jsonWriter.Flush();
+							Console.WriteLine(Encoding.UTF8.GetString(_jsonBufferWriter.WrittenSpan));
 						}
 						else
 						{
-							foreach (string Line in Lines)
+							foreach (string line in lines)
 							{
-								Console.WriteLine(Line);
+								Console.WriteLine(line);
 							}
 						}
 					}
@@ -1099,9 +1091,9 @@ namespace EpicGames.Core
 						}
 					}
 
-					if (StatusMessageStack.Count > 0 && AllowStatusUpdates)
+					if (_statusMessageStack.Count > 0 && AllowStatusUpdates)
 					{
-						SetStatusText(StatusMessageStack.Peek().CurrentText);
+						SetStatusText(_statusMessageStack.Peek()._currentText);
 					}
 				}
 			}
@@ -1112,15 +1104,15 @@ namespace EpicGames.Core
 		/// </summary>
 		void FlushStatusHeading()
 		{
-			if (StatusMessageStack.Count > 0)
+			if (_statusMessageStack.Count > 0)
 			{
-				StatusMessage CurrentStatus = StatusMessageStack.Peek();
-				if (CurrentStatus.HeadingText.Length > 0 && !CurrentStatus.bHasFlushedHeadingText && AllowStatusUpdates)
+				StatusMessage currentStatus = _statusMessageStack.Peek();
+				if (currentStatus._headingText.Length > 0 && !currentStatus._hasFlushedHeadingText && AllowStatusUpdates)
 				{
-					SetStatusText(CurrentStatus.HeadingText);
+					SetStatusText(currentStatus._headingText);
 					Console.WriteLine();
-					StatusText = "";
-					CurrentStatus.bHasFlushedHeadingText = true;
+					_statusText = "";
+					currentStatus._hasFlushedHeadingText = true;
 				}
 				else
 				{
@@ -1133,23 +1125,23 @@ namespace EpicGames.Core
 		/// Enter a scope with the given status message. The message will be written to the console without a newline, allowing it to be updated through subsequent calls to UpdateStatus().
 		/// The message will be written to the log immediately. If another line is written while in a status scope, the initial status message is flushed to the console first.
 		/// </summary>
-		/// <param name="Message">The status message</param>
+		/// <param name="message">The status message</param>
 		[Conditional("TRACE")]
-		public void PushStatus(string Message)
+		public void PushStatus(string message)
 		{
-			lock (SyncObject)
+			lock (_syncObject)
 			{
 				FlushStatusHeading();
 
-				StatusMessage NewStatusMessage = new StatusMessage(Message, Message);
-				StatusMessageStack.Push(NewStatusMessage);
+				StatusMessage newStatusMessage = new StatusMessage(message, message);
+				_statusMessageStack.Push(newStatusMessage);
 
-				StatusTimer.Restart();
+				_statusTimer.Restart();
 
-				if (Message.Length > 0)
+				if (message.Length > 0)
 				{
-					this.LogDebug("{Message}", Message);
-					SetStatusText(Message);
+					this.LogDebug("{Message}", message);
+					SetStatusText(message);
 				}
 			}
 		}
@@ -1157,19 +1149,19 @@ namespace EpicGames.Core
 		/// <summary>
 		/// Updates the current status message. This will overwrite the previous status line.
 		/// </summary>
-		/// <param name="Message">The status message</param>
+		/// <param name="message">The status message</param>
 		[Conditional("TRACE")]
-		public void UpdateStatus(string Message)
+		public void UpdateStatus(string message)
 		{
-			lock (SyncObject)
+			lock (_syncObject)
 			{
-				StatusMessage CurrentStatusMessage = StatusMessageStack.Peek();
-				CurrentStatusMessage.CurrentText = Message;
+				StatusMessage currentStatusMessage = _statusMessageStack.Peek();
+				currentStatusMessage._currentText = message;
 
-				if (AllowStatusUpdates || StatusTimer.Elapsed.TotalSeconds > 10.0)
+				if (AllowStatusUpdates || _statusTimer.Elapsed.TotalSeconds > 10.0)
 				{
-					SetStatusText(Message);
-					StatusTimer.Restart();
+					SetStatusText(message);
+					_statusTimer.Restart();
 				}
 			}
 		}
@@ -1181,61 +1173,61 @@ namespace EpicGames.Core
 		[Conditional("TRACE")]
 		public void PopStatus()
 		{
-			lock (SyncObject)
+			lock (_syncObject)
 			{
-				StatusMessage CurrentStatusMessage = StatusMessageStack.Peek();
-				SetStatusText(CurrentStatusMessage.CurrentText);
+				StatusMessage currentStatusMessage = _statusMessageStack.Peek();
+				SetStatusText(currentStatusMessage._currentText);
 
-				if (StatusText.Length > 0)
+				if (_statusText.Length > 0)
 				{
 					Console.WriteLine();
-					StatusText = "";
+					_statusText = "";
 				}
 
-				StatusMessageStack.Pop();
+				_statusMessageStack.Pop();
 			}
 		}
 
 		/// <summary>
 		/// Update the status text. For internal use only; does not modify the StatusMessageStack objects.
 		/// </summary>
-		/// <param name="NewStatusText">New status text to display</param>
-		private void SetStatusText(string NewStatusText)
+		/// <param name="newStatusText">New status text to display</param>
+		private void SetStatusText(string newStatusText)
 		{
-			if (NewStatusText.Length > 0)
+			if (newStatusText.Length > 0)
 			{
-				NewStatusText = LogIndent.Current + NewStatusText;
+				newStatusText = LogIndent.Current + newStatusText;
 			}
 
-			if (StatusText != NewStatusText)
+			if (_statusText != newStatusText)
 			{
-				int NumCommonChars = 0;
-				while (NumCommonChars < StatusText.Length && NumCommonChars < NewStatusText.Length && StatusText[NumCommonChars] == NewStatusText[NumCommonChars])
+				int numCommonChars = 0;
+				while (numCommonChars < _statusText.Length && numCommonChars < newStatusText.Length && _statusText[numCommonChars] == newStatusText[numCommonChars])
 				{
-					NumCommonChars++;
+					numCommonChars++;
 				}
 
-				if (!AllowStatusUpdates && NumCommonChars < StatusText.Length)
+				if (!AllowStatusUpdates && numCommonChars < _statusText.Length)
 				{
 					// Prevent writing backspace characters if the console doesn't support it
 					Console.WriteLine();
-					StatusText = "";
-					NumCommonChars = 0;
+					_statusText = "";
+					numCommonChars = 0;
 				}
 
-				StringBuilder Text = new StringBuilder();
-				Text.Append('\b', StatusText.Length - NumCommonChars);
-				Text.Append(NewStatusText, NumCommonChars, NewStatusText.Length - NumCommonChars);
-				if (NewStatusText.Length < StatusText.Length)
+				StringBuilder text = new StringBuilder();
+				text.Append('\b', _statusText.Length - numCommonChars);
+				text.Append(newStatusText, numCommonChars, newStatusText.Length - numCommonChars);
+				if (newStatusText.Length < _statusText.Length)
 				{
-					int NumChars = StatusText.Length - NewStatusText.Length;
-					Text.Append(' ', NumChars);
-					Text.Append('\b', NumChars);
+					int numChars = _statusText.Length - newStatusText.Length;
+					text.Append(' ', numChars);
+					text.Append('\b', numChars);
 				}
-				Console.Write(Text.ToString());
+				Console.Write(text.ToString());
 
-				StatusText = NewStatusText;
-				StatusTimer.Restart();
+				_statusText = newStatusText;
+				_statusTimer.Restart();
 			}
 		}
 	}
@@ -1246,7 +1238,7 @@ namespace EpicGames.Core
 	public class DefaultLoggerProvider : ILoggerProvider
 	{
 		/// <inheritdoc/>
-		public ILogger CreateLogger(string CategoryName)
+		public ILogger CreateLogger(string categoryName)
 		{
 			return new DefaultLogger();
 		}
@@ -1265,10 +1257,10 @@ namespace EpicGames.Core
 		/// <summary>
 		/// Adds a regular Epic logger to the builder
 		/// </summary>
-		/// <param name="Builder">Logging builder</param>
-		public static void AddEpicDefault(this ILoggingBuilder Builder)
+		/// <param name="builder">Logging builder</param>
+		public static void AddEpicDefault(this ILoggingBuilder builder)
 		{
-			Builder.Services.AddSingleton<ILoggerProvider, DefaultLoggerProvider>();
+			builder.Services.AddSingleton<ILoggerProvider, DefaultLoggerProvider>();
 		}
 	}
 }
diff --git a/Engine/Source/Programs/Shared/EpicGames.Core/LogBuffer.cs b/Engine/Source/Programs/Shared/EpicGames.Core/LogBuffer.cs
--- a/Engine/Source/Programs/Shared/EpicGames.Core/LogBuffer.cs
+++ b/Engine/Source/Programs/Shared/EpicGames.Core/LogBuffer.cs
@@ -1,13 +1,7 @@
 // Copyright Epic Games, Inc. All Rights Reserved.
 
 using System;
-using System.Collections.Concurrent;
 using System.Collections.Generic;
-using System.Diagnostics.CodeAnalysis;
-using System.Linq;
-using System.Text;
-using System.Text.RegularExpressions;
-using System.Threading.Tasks;
 
 namespace EpicGames.Core
 {
@@ -16,19 +10,17 @@ namespace EpicGames.Core
 	/// </summary>
 	class LogBuffer : ILogCursor
 	{
-		int LineNumber;
+		int _lineNumber;
+		readonly string?[] _history;
+		int _historyIdx;
+		int _historyCount;
+		readonly List<string?> _nextLines;
 
-		string?[] History;
-		int HistoryIdx;
-		int HistoryCount;
-
-		List<string?> NextLines;
-
-		public LogBuffer(int HistorySize)
+		public LogBuffer(int historySize)
 		{
-			LineNumber = 1;
-			History = new string?[HistorySize];
-			NextLines = new List<string?>();
+			_lineNumber = 1;
+			_history = new string?[historySize];
+			_nextLines = new List<string?>();
 		}
 
 		public bool NeedMoreData
@@ -37,30 +29,21 @@ namespace EpicGames.Core
 			private set;
 		}
 
-		public string? CurrentLine
-		{
-			get { return this[0]; }
-		}
+		public string? CurrentLine => this[0];
 
-		public int CurrentLineNumber
-		{
-			get { return LineNumber; }
-		}
+		public int CurrentLineNumber => _lineNumber;
 
-		public int Length
-		{
-			get { return NextLines.Count; }
-		}
+		public int Length => _nextLines.Count;
 
-		public void AddLine(string? Line)
+		public void AddLine(string? line)
 		{
-			NextLines.Add(Line);
+			_nextLines.Add(line);
 			NeedMoreData = false;
 		}
 
-		public void Advance(int Count)
+		public void Advance(int count)
 		{
-			for (int Idx = 0; Idx < Count; Idx++)
+			for (int idx = 0; idx < count; idx++)
 			{
 				MoveNext();
 			}
@@ -68,50 +51,50 @@ namespace EpicGames.Core
 
 		public void MoveNext()
 		{
-			if (NextLines.Count == 0)
+			if (_nextLines.Count == 0)
 			{
 				throw new InvalidOperationException("Attempt to move past end of line buffer");
 			}
 
-			HistoryIdx++;
-			if (HistoryIdx >= History.Length)
+			_historyIdx++;
+			if (_historyIdx >= _history.Length)
 			{
-				HistoryIdx = 0;
+				_historyIdx = 0;
 			}
-			if (HistoryCount < History.Length)
+			if (_historyCount < _history.Length)
 			{
-				HistoryCount++;
+				_historyCount++;
 			}
 
-			History[HistoryIdx] = NextLines[0];
-			NextLines.RemoveAt(0);
+			_history[_historyIdx] = _nextLines[0];
+			_nextLines.RemoveAt(0);
 
-			LineNumber++;
+			_lineNumber++;
 		}
 
-		public string? this[int Offset]
+		public string? this[int offset]
 		{
 			get
 			{
-				if (Offset >= 0)
+				if (offset >= 0)
 				{
 					// Add new lines to the buffer
-					if (Offset >= NextLines.Count)
+					if (offset >= _nextLines.Count)
 					{
 						NeedMoreData = true;
 						return null;
 					}
-					return NextLines[Offset];
+					return _nextLines[offset];
 				}
-				else if (Offset >= -HistoryCount)
+				else if (offset >= -_historyCount)
 				{
 					// Retrieve a line from the history buffer
-					int Idx = HistoryIdx + 1 + Offset;
-					if (Idx < 0)
+					int idx = _historyIdx + 1 + offset;
+					if (idx < 0)
 					{
-						Idx += History.Length;
+						idx += _history.Length;
 					}
-					return History[Idx];
+					return _history[idx];
 				}
 				else
 				{
diff --git a/Engine/Source/Programs/Shared/EpicGames.Core/LogEvent.cs b/Engine/Source/Programs/Shared/EpicGames.Core/LogEvent.cs
--- a/Engine/Source/Programs/Shared/EpicGames.Core/LogEvent.cs
+++ b/Engine/Source/Programs/Shared/EpicGames.Core/LogEvent.cs
@@ -1,6 +1,5 @@
 // Copyright Epic Games, Inc. All Rights Reserved.
 
-using Microsoft.Extensions.Logging;
 using System;
 using System.Buffers;
 using System.Collections.Generic;
@@ -10,7 +9,7 @@ using System.Linq;
 using System.Text;
 using System.Text.Json;
 using System.Text.Json.Serialization;
-using System.Text.RegularExpressions;
+using Microsoft.Extensions.Logging;
 
 namespace EpicGames.Core
 {
@@ -94,294 +93,294 @@ namespace EpicGames.Core
 		/// <summary>
 		/// Constructor
 		/// </summary>
-		public LogEvent(DateTime Time, LogLevel Level, EventId EventId, string Message, string? Format, Dictionary<string, object>? Properties, LogException? Exception)
-			: this(Time, Level, EventId, 0, 1, Message, Format, Properties, Exception)
+		public LogEvent(DateTime time, LogLevel level, EventId eventId, string message, string? format, Dictionary<string, object>? properties, LogException? exception)
+			: this(time, level, eventId, 0, 1, message, format, properties, exception)
 		{
 		}
 
 		/// <summary>
 		/// Constructor
 		/// </summary>
-		public LogEvent(DateTime Time, LogLevel Level, EventId EventId, int LineIndex, int LineCount, string Message, string? Format, Dictionary<string, object>? Properties, LogException? Exception)
+		public LogEvent(DateTime time, LogLevel level, EventId eventId, int lineIndex, int lineCount, string message, string? format, Dictionary<string, object>? properties, LogException? exception)
 		{
-			this.Time = Time;
-			this.Level = Level;
-			this.Id = EventId;
-			this.LineIndex = LineIndex;
-			this.LineCount = LineCount;
-			this.Message = Message;
-			this.Format = Format;
-			this.Properties = Properties;
-			this.Exception = Exception;
+			Time = time;
+			Level = level;
+			Id = eventId;
+			LineIndex = lineIndex;
+			LineCount = lineCount;
+			Message = message;
+			Format = format;
+			Properties = properties;
+			Exception = exception;
 		}
 
 		/// <summary>
 		/// Read a log event from a utf-8 encoded json byte array
 		/// </summary>
-		/// <param name="Data"></param>
+		/// <param name="data"></param>
 		/// <returns></returns>
-		public static LogEvent Read(ReadOnlySpan<byte> Data)
+		public static LogEvent Read(ReadOnlySpan<byte> data)
 		{
-			Utf8JsonReader Reader = new Utf8JsonReader(Data);
-			Reader.Read();
-			return Read(ref Reader);
+			Utf8JsonReader reader = new Utf8JsonReader(data);
+			reader.Read();
+			return Read(ref reader);
 		}
 
 		/// <summary>
 		/// Read a log event from Json
 		/// </summary>
-		/// <param name="Reader">The Json reader</param>
+		/// <param name="reader">The Json reader</param>
 		/// <returns>New log event</returns>
-		public static LogEvent Read(ref Utf8JsonReader Reader)
+		public static LogEvent Read(ref Utf8JsonReader reader)
 		{
-			DateTime Time = new DateTime(0);
-			LogLevel Level = LogLevel.None;
-			EventId EventId = new EventId(0);
-			int Line = 0;
-			int LineCount = 1;
-			string Message = String.Empty;
-			string Format = String.Empty;
-			Dictionary<string, object>? Properties = null;
-			LogException? Exception = null;
-
-			ReadOnlySpan<byte> PropertyName;
-			for (; JsonExtensions.TryReadNextPropertyName(ref Reader, out PropertyName); Reader.Skip())
+			DateTime time = new DateTime(0);
+			LogLevel level = LogLevel.None;
+			EventId eventId = new EventId(0);
+			int line = 0;
+			int lineCount = 1;
+			string message = String.Empty;
+			string format = String.Empty;
+			Dictionary<string, object>? properties = null;
+			LogException? exception = null;
+
+			ReadOnlySpan<byte> propertyName;
+			for (; JsonExtensions.TryReadNextPropertyName(ref reader, out propertyName); reader.Skip())
 			{
-				if (Utf8StringComparer.OrdinalIgnoreCase.Equals(PropertyName, JsonPropertyNames.Time.Span))
+				if (Utf8StringComparer.OrdinalIgnoreCase.Equals(propertyName, JsonPropertyNames.Time.Span))
 				{
-					Time = Reader.GetDateTime();
+					time = reader.GetDateTime();
 				}
-				else if (Utf8StringComparer.OrdinalIgnoreCase.Equals(PropertyName, JsonPropertyNames.Level.Span))
+				else if (Utf8StringComparer.OrdinalIgnoreCase.Equals(propertyName, JsonPropertyNames.Level.Span))
 				{
-					Level = Enum.Parse<LogLevel>(Reader.GetString());
+					level = Enum.Parse<LogLevel>(reader.GetString());
 				}
-				else if (Utf8StringComparer.OrdinalIgnoreCase.Equals(PropertyName, JsonPropertyNames.Id.Span))
+				else if (Utf8StringComparer.OrdinalIgnoreCase.Equals(propertyName, JsonPropertyNames.Id.Span))
 				{
-					EventId = Reader.GetInt32();
+					eventId = reader.GetInt32();
 				}
-				else if (Utf8StringComparer.OrdinalIgnoreCase.Equals(PropertyName, JsonPropertyNames.Line.Span))
+				else if (Utf8StringComparer.OrdinalIgnoreCase.Equals(propertyName, JsonPropertyNames.Line.Span))
 				{
-					Line = Reader.GetInt32();
+					line = reader.GetInt32();
 				}
-				else if (Utf8StringComparer.OrdinalIgnoreCase.Equals(PropertyName, JsonPropertyNames.LineCount.Span))
+				else if (Utf8StringComparer.OrdinalIgnoreCase.Equals(propertyName, JsonPropertyNames.LineCount.Span))
 				{
-					LineCount = Reader.GetInt32();
+					lineCount = reader.GetInt32();
 				}
-				else if (Utf8StringComparer.OrdinalIgnoreCase.Equals(PropertyName, JsonPropertyNames.Message.Span))
+				else if (Utf8StringComparer.OrdinalIgnoreCase.Equals(propertyName, JsonPropertyNames.Message.Span))
 				{
-					Message = Reader.GetString();
+					message = reader.GetString();
 				}
-				else if (Utf8StringComparer.OrdinalIgnoreCase.Equals(PropertyName, JsonPropertyNames.Format.Span))
+				else if (Utf8StringComparer.OrdinalIgnoreCase.Equals(propertyName, JsonPropertyNames.Format.Span))
 				{
-					Format = Reader.GetString();
+					format = reader.GetString();
 				}
-				else if (Utf8StringComparer.OrdinalIgnoreCase.Equals(PropertyName, JsonPropertyNames.Properties.Span))
+				else if (Utf8StringComparer.OrdinalIgnoreCase.Equals(propertyName, JsonPropertyNames.Properties.Span))
 				{
-					Properties = ReadProperties(ref Reader);
+					properties = ReadProperties(ref reader);
 				}
 			}
 
-			return new LogEvent(Time, Level, EventId, Line, LineCount, Message, Format, Properties, Exception);
+			return new LogEvent(time, level, eventId, line, lineCount, message, format, properties, exception);
 		}
 
-		static Dictionary<string, object> ReadProperties(ref Utf8JsonReader Reader)
+		static Dictionary<string, object> ReadProperties(ref Utf8JsonReader reader)
 		{
-			Dictionary<string, object> Properties = new Dictionary<string, object>();
+			Dictionary<string, object> properties = new Dictionary<string, object>();
 
-			ReadOnlySpan<byte> PropertyName;
-			for (; JsonExtensions.TryReadNextPropertyName(ref Reader, out PropertyName); Reader.Skip())
+			ReadOnlySpan<byte> propertyName;
+			for (; JsonExtensions.TryReadNextPropertyName(ref reader, out propertyName); reader.Skip())
 			{
-				string Name = Encoding.UTF8.GetString(PropertyName);
-				object Value = ReadPropertyValue(ref Reader);
-				Properties.Add(Name, Value);
+				string name = Encoding.UTF8.GetString(propertyName);
+				object value = ReadPropertyValue(ref reader);
+				properties.Add(name, value);
 			}
 
-			return Properties;
+			return properties;
 		}
 
-		static object ReadPropertyValue(ref Utf8JsonReader Reader)
+		static object ReadPropertyValue(ref Utf8JsonReader reader)
 		{
-			switch (Reader.TokenType)
+			switch (reader.TokenType)
 			{
 				case JsonTokenType.True:
 					return true;
 				case JsonTokenType.False:
 					return true;
 				case JsonTokenType.StartObject:
-					return ReadStructuredPropertyValue(ref Reader);
+					return ReadStructuredPropertyValue(ref reader);
 				case JsonTokenType.String:
-					return Reader.GetString();
+					return reader.GetString();
 				case JsonTokenType.Number:
-					return Reader.GetInt32();
+					return reader.GetInt32();
 				default:
 					throw new InvalidOperationException("Unhandled property type");
 			}
 		}
 
-		static LogValue ReadStructuredPropertyValue(ref Utf8JsonReader Reader)
+		static LogValue ReadStructuredPropertyValue(ref Utf8JsonReader reader)
 		{
-			string Type = String.Empty;
-			string Text = String.Empty;
-			Dictionary<string, object>? Properties = null;
+			string type = String.Empty;
+			string text = String.Empty;
+			Dictionary<string, object>? properties = null;
 
-			ReadOnlySpan<byte> PropertyName;
-			for (; JsonExtensions.TryReadNextPropertyName(ref Reader, out PropertyName); Reader.Skip())
+			ReadOnlySpan<byte> propertyName;
+			for (; JsonExtensions.TryReadNextPropertyName(ref reader, out propertyName); reader.Skip())
 			{
-				if (Utf8StringComparer.OrdinalIgnoreCase.Equals(PropertyName, JsonPropertyNames.Type.Span))
+				if (Utf8StringComparer.OrdinalIgnoreCase.Equals(propertyName, JsonPropertyNames.Type.Span))
 				{
-					Type = Reader.GetString();
+					type = reader.GetString();
 				}
-				else if (Utf8StringComparer.OrdinalIgnoreCase.Equals(PropertyName, JsonPropertyNames.Text.Span))
+				else if (Utf8StringComparer.OrdinalIgnoreCase.Equals(propertyName, JsonPropertyNames.Text.Span))
 				{
-					Text = Reader.GetString();
+					text = reader.GetString();
 				}
 				else
 				{
-					Properties ??= new Dictionary<string, object>();
-					Properties.Add(Encoding.UTF8.GetString(PropertyName), ReadPropertyValue(ref Reader));
+					properties ??= new Dictionary<string, object>();
+					properties.Add(Encoding.UTF8.GetString(propertyName), ReadPropertyValue(ref reader));
 				}
 			}
 
-			return new LogValue(Type, Text, Properties);
+			return new LogValue(type, text, properties);
 		}
 
 		/// <summary>
 		/// Writes a log event to Json
 		/// </summary>
-		/// <param name="Writer"></param>
-		public void Write(Utf8JsonWriter Writer)
+		/// <param name="writer"></param>
+		public void Write(Utf8JsonWriter writer)
 		{
-			Writer.WriteStartObject();
-			Writer.WriteString(JsonPropertyNames.Time.Span, Time.ToString("s", CultureInfo.InvariantCulture));
-			Writer.WriteString(JsonPropertyNames.Level.Span, Level.ToString());
-			Writer.WriteString(JsonPropertyNames.Message.Span, Message);
+			writer.WriteStartObject();
+			writer.WriteString(JsonPropertyNames.Time.Span, Time.ToString("s", CultureInfo.InvariantCulture));
+			writer.WriteString(JsonPropertyNames.Level.Span, Level.ToString());
+			writer.WriteString(JsonPropertyNames.Message.Span, Message);
 
 			if (Id.Id != 0)
 			{
-				Writer.WriteNumber(JsonPropertyNames.Id.Span, Id.Id);
+				writer.WriteNumber(JsonPropertyNames.Id.Span, Id.Id);
 			}
 
 			if (LineIndex > 0)
 			{
-				Writer.WriteNumber(JsonPropertyNames.Line.Span, LineIndex);
+				writer.WriteNumber(JsonPropertyNames.Line.Span, LineIndex);
 			}
 
 			if (LineCount > 1)
 			{
-				Writer.WriteNumber(JsonPropertyNames.LineCount.Span, LineCount);
+				writer.WriteNumber(JsonPropertyNames.LineCount.Span, LineCount);
 			}
 
 			if (Format != null)
 			{
-				Writer.WriteString(JsonPropertyNames.Format.Span, Format);
+				writer.WriteString(JsonPropertyNames.Format.Span, Format);
 			}
 
 			if (Properties != null && Properties.Any())
 			{
-				Writer.WriteStartObject(JsonPropertyNames.Properties.Span);
-				foreach ((string Name, object? Value) in Properties!)
+				writer.WriteStartObject(JsonPropertyNames.Properties.Span);
+				foreach ((string name, object? value) in Properties!)
 				{
-					Writer.WritePropertyName(Name);
-					WritePropertyValue(ref Writer, Value);
+					writer.WritePropertyName(name);
+					WritePropertyValue(ref writer, value);
 				}
-				Writer.WriteEndObject();
+				writer.WriteEndObject();
 			}
 
 			if (Exception != null)
 			{
-				Writer.WriteStartObject(JsonPropertyNames.Exception.Span);
-				WriteException(ref Writer, Exception);
-				Writer.WriteEndObject();
+				writer.WriteStartObject(JsonPropertyNames.Exception.Span);
+				WriteException(ref writer, Exception);
+				writer.WriteEndObject();
 			}
-			Writer.WriteEndObject();
+			writer.WriteEndObject();
 		}
 
 		/// <summary>
 		/// Write a property value to a json object
 		/// </summary>
-		/// <param name="Writer"></param>
-		/// <param name="Value"></param>
-		static void WritePropertyValue(ref Utf8JsonWriter Writer, object? Value)
+		/// <param name="writer"></param>
+		/// <param name="value"></param>
+		static void WritePropertyValue(ref Utf8JsonWriter writer, object? value)
 		{
-			if (Value == null)
+			if (value == null)
 			{
-				Writer.WriteNullValue();
+				writer.WriteNullValue();
 			}
 			else
 			{
-				Type ValueType = Value.GetType();
-				if (ValueType == typeof(LogValue))
+				Type valueType = value.GetType();
+				if (valueType == typeof(LogValue))
 				{
-					WriteStructuredPropertyValue(ref Writer, (LogValue)Value);
+					WriteStructuredPropertyValue(ref writer, (LogValue)value);
 				}
 				else
 				{
-					WriteSimplePropertyValue(ref Writer, Value, ValueType);
+					WriteSimplePropertyValue(ref writer, value, valueType);
 				}
 			}
 		}
 
-		static void WriteStructuredPropertyValue(ref Utf8JsonWriter Writer, LogValue Value)
+		static void WriteStructuredPropertyValue(ref Utf8JsonWriter writer, LogValue value)
 		{
-			Writer.WriteStartObject();
-			Writer.WriteString(JsonPropertyNames.Type.Span, Value.Type);
-			Writer.WriteString(JsonPropertyNames.Text.Span, Value.Text);
-			if (Value.Properties != null)
+			writer.WriteStartObject();
+			writer.WriteString(JsonPropertyNames.Type.Span, value.Type);
+			writer.WriteString(JsonPropertyNames.Text.Span, value.Text);
+			if (value.Properties != null)
 			{
-				foreach ((string PropertyName, object? PropertyValue) in Value.Properties)
+				foreach ((string propertyName, object? propertyValue) in value.Properties)
 				{
-					Writer.WritePropertyName(PropertyName);
-					WritePropertyValue(ref Writer, PropertyValue);
+					writer.WritePropertyName(propertyName);
+					WritePropertyValue(ref writer, propertyValue);
 				}
 			}
-			Writer.WriteEndObject();
+			writer.WriteEndObject();
 		}
 
-		static void WriteSimplePropertyValue(ref Utf8JsonWriter Writer, object Value, Type ValueType)
+		static void WriteSimplePropertyValue(ref Utf8JsonWriter writer, object value, Type valueType)
 		{
-			switch (Type.GetTypeCode(ValueType))
+			switch (Type.GetTypeCode(valueType))
 			{
 				case TypeCode.Boolean:
-					Writer.WriteBooleanValue((bool)Value);
+					writer.WriteBooleanValue((bool)value);
 					break;
 				case TypeCode.Byte:
-					Writer.WriteNumberValue((Byte)Value);
+					writer.WriteNumberValue((byte)value);
 					break;
 				case TypeCode.SByte:
-					Writer.WriteNumberValue((SByte)Value);
+					writer.WriteNumberValue((sbyte)value);
 					break;
 				case TypeCode.UInt16:
-					Writer.WriteNumberValue((UInt16)Value);
+					writer.WriteNumberValue((ushort)value);
 					break;
 				case TypeCode.UInt32:
-					Writer.WriteNumberValue((UInt32)Value);
+					writer.WriteNumberValue((uint)value);
 					break;
 				case TypeCode.UInt64:
-					Writer.WriteNumberValue((UInt64)Value);
+					writer.WriteNumberValue((long)value);
 					break;
 				case TypeCode.Int16:
-					Writer.WriteNumberValue((Int16)Value);
+					writer.WriteNumberValue((short)value);
 					break;
 				case TypeCode.Int32:
-					Writer.WriteNumberValue((Int32)Value);
+					writer.WriteNumberValue((int)value);
 					break;
 				case TypeCode.Int64:
-					Writer.WriteNumberValue((Int64)Value);
+					writer.WriteNumberValue((long)value);
 					break;
 				case TypeCode.Decimal:
-					Writer.WriteNumberValue((Decimal)Value);
+					writer.WriteNumberValue((decimal)value);
 					break;
 				case TypeCode.Double:
-					Writer.WriteNumberValue((Double)Value);
+					writer.WriteNumberValue((double)value);
 					break;
 				case TypeCode.Single:
-					Writer.WriteNumberValue((Single)Value);
+					writer.WriteNumberValue((float)value);
 					break;
 				case TypeCode.String:
-					Writer.WriteStringValue((String)Value);
+					writer.WriteStringValue((string)value);
 					break;
 				default:
-					Writer.WriteStringValue(Value.ToString() ?? String.Empty);
+					writer.WriteStringValue(value.ToString() ?? String.Empty);
 					break;
 			}
 		}
@@ -389,81 +388,81 @@ namespace EpicGames.Core
 		/// <summary>
 		/// Writes an exception to a json object
 		/// </summary>
-		/// <param name="Writer">Writer to receive the exception data</param>
-		/// <param name="Exception">The exception</param>
-		static void WriteException(ref Utf8JsonWriter Writer, LogException Exception)
+		/// <param name="writer">Writer to receive the exception data</param>
+		/// <param name="exception">The exception</param>
+		static void WriteException(ref Utf8JsonWriter writer, LogException exception)
 		{
-			Writer.WriteString("message", Exception.Message);
-			Writer.WriteString("trace", Exception.Trace);
+			writer.WriteString("message", exception.Message);
+			writer.WriteString("trace", exception.Trace);
 
-			if (Exception.InnerException != null)
+			if (exception.InnerException != null)
 			{
-				Writer.WriteStartObject("innerException");
-				WriteException(ref Writer, Exception.InnerException);
-				Writer.WriteEndObject();
+				writer.WriteStartObject("innerException");
+				WriteException(ref writer, exception.InnerException);
+				writer.WriteEndObject();
 			}
 
-			if (Exception.InnerExceptions != null)
+			if (exception.InnerExceptions != null)
 			{
-				Writer.WriteStartArray("innerExceptions");
-				for (int Idx = 0; Idx < 16 && Idx < Exception.InnerExceptions.Count; Idx++) // Cap number of exceptions returned to avoid huge messages
+				writer.WriteStartArray("innerExceptions");
+				for (int idx = 0; idx < 16 && idx < exception.InnerExceptions.Count; idx++) // Cap number of exceptions returned to avoid huge messages
 				{
-					LogException InnerException = Exception.InnerExceptions[Idx];
-					Writer.WriteStartObject();
-					WriteException(ref Writer, InnerException);
-					Writer.WriteEndObject();
+					LogException innerException = exception.InnerExceptions[idx];
+					writer.WriteStartObject();
+					WriteException(ref writer, innerException);
+					writer.WriteEndObject();
 				}
-				Writer.WriteEndArray();
+				writer.WriteEndArray();
 			}
 		}
 
-		public static LogEvent Create(LogLevel Level, string Format, params object[] Args)
-			=> Create(Level, KnownLogEvents.None, null, Format, Args);
+		public static LogEvent Create(LogLevel level, string format, params object[] args)
+			=> Create(level, KnownLogEvents.None, null, format, args);
 
-		public static LogEvent Create(LogLevel Level, EventId EventId, string Format, params object[] Args)
-			=> Create(Level, EventId, null, Format, Args);
+		public static LogEvent Create(LogLevel level, EventId eventId, string format, params object[] args)
+			=> Create(level, eventId, null, format, args);
 
-		public static LogEvent Create(LogLevel Level, EventId EventId, Exception? Exception, string Format, params object[] Args)
+		public static LogEvent Create(LogLevel level, EventId eventId, Exception? exception, string format, params object[] args)
 		{
-			Dictionary<string, object> Properties = new Dictionary<string, object>();
-			MessageTemplate.ParsePropertyValues(Format, Args, Properties);
+			Dictionary<string, object> properties = new Dictionary<string, object>();
+			MessageTemplate.ParsePropertyValues(format, args, properties);
 
-			string Message = MessageTemplate.Render(Format, Properties!);
-			return new LogEvent(DateTime.UtcNow, Level, EventId, Message, Format, Properties, LogException.FromException(Exception));
+			string message = MessageTemplate.Render(format, properties!);
+			return new LogEvent(DateTime.UtcNow, level, eventId, message, format, properties, LogException.FromException(exception));
 		}
 
 		/// <summary>
 		/// Creates a log event from an ILogger parameters
 		/// </summary>
-		public static LogEvent FromState<TState>(LogLevel Level, EventId EventId, TState State, Exception? Exception, Func<TState, Exception?, string> Formatter)
+		public static LogEvent FromState<TState>(LogLevel level, EventId eventId, TState state, Exception? exception, Func<TState, Exception?, string> formatter)
 		{
-			DateTime Time = DateTime.UtcNow;
+			DateTime time = DateTime.UtcNow;
 
 			// Render the message
-			string Message = Formatter(State, Exception);
+			string message = formatter(state, exception);
 
 			// Try to log the event
-			IEnumerable<KeyValuePair<string, object>>? Values = State as IEnumerable<KeyValuePair<string, object>>;
-			if (Values != null)
+			IEnumerable<KeyValuePair<string, object>>? values = state as IEnumerable<KeyValuePair<string, object>>;
+			if (values != null)
 			{
-				KeyValuePair<string, object>? Format = Values.FirstOrDefault(x => x.Key.Equals(MessageTemplate.FormatPropertyName, StringComparison.Ordinal));
-				if (Format != null)
+				KeyValuePair<string, object>? format = values.FirstOrDefault(x => x.Key.Equals(MessageTemplate.FormatPropertyName, StringComparison.Ordinal));
+				if (format != null)
 				{
 					// Format all the other values
-					string FormatString = Format.Value.Value?.ToString() ?? String.Empty;
-					Dictionary<string, object>? Properties = MessageTemplate.CreatePositionalProperties(FormatString, Values);
-					return new LogEvent(Time, Level, EventId, Message, FormatString, Properties, LogException.FromException(Exception));
+					string formatString = format.Value.Value?.ToString() ?? String.Empty;
+					Dictionary<string, object>? properties = MessageTemplate.CreatePositionalProperties(formatString, values);
+					return new LogEvent(time, level, eventId, message, formatString, properties, LogException.FromException(exception));
 				}
 				else
 				{
 					// Include all the data, but don't use the format string
-					return new LogEvent(Time, Level, EventId, Message, null, Values.ToDictionary(x => x.Key, x => x.Value), LogException.FromException(Exception));
+					return new LogEvent(time, level, eventId, message, null, values.ToDictionary(x => x.Key, x => x.Value), LogException.FromException(exception));
 				}
 			}
 			else
 			{
 				// Format as a string
-				return new LogEvent(Time, Level, EventId, Message, null, null, LogException.FromException(Exception));
+				return new LogEvent(time, level, eventId, message, null, null, LogException.FromException(exception));
 			}
 		}
 
@@ -490,9 +489,9 @@ namespace EpicGames.Core
 
 			if (Properties != null)
 			{
-				foreach ((string Name, object? Value) in Properties)
+				foreach ((string name, object? value) in Properties)
 				{
-					yield return new KeyValuePair<string, object?>(Name, Value?.ToString());
+					yield return new KeyValuePair<string, object?>(name, value?.ToString());
 				}
 			}
 		}
@@ -503,9 +502,9 @@ namespace EpicGames.Core
 		/// <returns>Property pairs</returns>
 		System.Collections.IEnumerator System.Collections.IEnumerable.GetEnumerator()
 		{
-			foreach (KeyValuePair<string, object?> Pair in this)
+			foreach (KeyValuePair<string, object?> pair in this)
 			{
-				yield return Pair;
+				yield return pair;
 			}
 		}
 
@@ -514,12 +513,12 @@ namespace EpicGames.Core
 		/// </summary>
 		public byte[] ToJsonBytes()
 		{
-			ArrayBufferWriter<byte> Buffer = new ArrayBufferWriter<byte>();
-			using (Utf8JsonWriter Writer = new Utf8JsonWriter(Buffer))
+			ArrayBufferWriter<byte> buffer = new ArrayBufferWriter<byte>();
+			using (Utf8JsonWriter writer = new Utf8JsonWriter(buffer))
 			{
-				Write(Writer);
+				Write(writer);
 			}
-			return Buffer.WrittenSpan.ToArray();
+			return buffer.WrittenSpan.ToArray();
 		}
 
 		/// <summary>
@@ -527,12 +526,12 @@ namespace EpicGames.Core
 		/// </summary>
 		public string ToJson()
 		{
-			ArrayBufferWriter<byte> Buffer = new ArrayBufferWriter<byte>();
-			using (Utf8JsonWriter Writer = new Utf8JsonWriter(Buffer))
+			ArrayBufferWriter<byte> buffer = new ArrayBufferWriter<byte>();
+			using (Utf8JsonWriter writer = new Utf8JsonWriter(buffer))
 			{
-				Write(Writer);
+				Write(writer);
 			}
-			return Encoding.UTF8.GetString(Buffer.WrittenSpan);
+			return Encoding.UTF8.GetString(buffer.WrittenSpan);
 		}
 
 		/// <inheritdoc/>
@@ -562,14 +561,14 @@ namespace EpicGames.Core
 		/// <summary>
 		/// Constructor
 		/// </summary>
-		/// <param name="Type">Type of the value</param>
-		/// <param name="Text">Rendering of the value as text</param>
-		/// <param name="Properties">Additional properties for this value</param>
-		public LogValue(string Type, string Text, Dictionary<string, object>? Properties = null)
+		/// <param name="type">Type of the value</param>
+		/// <param name="text">Rendering of the value as text</param>
+		/// <param name="properties">Additional properties for this value</param>
+		public LogValue(string type, string text, Dictionary<string, object>? properties = null)
 		{
-			this.Type = Type;
-			this.Text = Text;
-			this.Properties = Properties;
+			Type = type;
+			Text = text;
+			Properties = properties;
 		}
 
 		/// <inheritdoc/>
@@ -602,48 +601,48 @@ namespace EpicGames.Core
 		/// <summary>
 		/// Multiple inner exceptions, in the case of an <see cref="AggregateException"/>
 		/// </summary>
-		public List<LogException> InnerExceptions = new List<LogException>();
+		public List<LogException> InnerExceptions { get; set; } = new List<LogException>();
 
 		/// <summary>
 		/// Constructor
 		/// </summary>
-		/// <param name="Message"></param>
-		/// <param name="Trace"></param>
-		public LogException(string Message, string Trace)
+		/// <param name="message"></param>
+		/// <param name="trace"></param>
+		public LogException(string message, string trace)
 		{
-			this.Message = Message;
-			this.Trace = Trace;
+			Message = message;
+			Trace = trace;
 		}
 
 		/// <summary>
 		/// Constructor
 		/// </summary>
-		/// <param name="Exception"></param>
-		[return: NotNullIfNotNull("Exception")]
-		public static LogException? FromException(Exception? Exception)
+		/// <param name="exception"></param>
+		[return: NotNullIfNotNull("exception")]
+		public static LogException? FromException(Exception? exception)
 		{
-			LogException? Result = null;
-			if (Exception != null)
+			LogException? result = null;
+			if (exception != null)
 			{
-				Result = new LogException(Exception.Message, Exception.StackTrace ?? String.Empty);
+				result = new LogException(exception.Message, exception.StackTrace ?? String.Empty);
 
-				if (Exception.InnerException != null)
+				if (exception.InnerException != null)
 				{
-					Result.InnerException = FromException(Exception.InnerException);
+					result.InnerException = FromException(exception.InnerException);
 				}
 
-				AggregateException? AggregateException = Exception as AggregateException;
-				if (AggregateException != null && AggregateException.InnerExceptions.Count > 0)
+				AggregateException? aggregateException = exception as AggregateException;
+				if (aggregateException != null && aggregateException.InnerExceptions.Count > 0)
 				{
-					Result.InnerExceptions = new List<LogException>();
-					for (int Idx = 0; Idx < 16 && Idx < AggregateException.InnerExceptions.Count; Idx++) // Cap number of exceptions returned to avoid huge messages
+					result.InnerExceptions = new List<LogException>();
+					for (int idx = 0; idx < 16 && idx < aggregateException.InnerExceptions.Count; idx++) // Cap number of exceptions returned to avoid huge messages
 					{
-						LogException InnerException = FromException(AggregateException.InnerExceptions[Idx]);
-						Result.InnerExceptions.Add(InnerException);
+						LogException innerException = FromException(aggregateException.InnerExceptions[idx]);
+						result.InnerExceptions.Add(innerException);
 					}
 				}
 			}
-			return Result;
+			return result;
 		}
 	}
 
@@ -653,15 +652,15 @@ namespace EpicGames.Core
 	public class LogEventConverter : JsonConverter<LogEvent>
 	{
 		/// <inheritdoc/>
-		public override LogEvent Read(ref Utf8JsonReader Reader, Type TypeToConvert, JsonSerializerOptions Options)
+		public override LogEvent Read(ref Utf8JsonReader reader, Type typeToConvert, JsonSerializerOptions options)
 		{
-			return LogEvent.Read(ref Reader);
+			return LogEvent.Read(ref reader);
 		}
 
 		/// <inheritdoc/>
-		public override void Write(Utf8JsonWriter Writer, LogEvent Value, JsonSerializerOptions Options)
+		public override void Write(Utf8JsonWriter writer, LogEvent value, JsonSerializerOptions options)
 		{
-			Value.Write(Writer);
+			value.Write(writer);
 		}
 	}
 }
