Reformat EpicGames.Core according to standard coding conventions.

diff --git a/Engine/Source/Programs/Shared/EpicGames.Core/LogEventBuilder.cs b/Engine/Source/Programs/Shared/EpicGames.Core/LogEventBuilder.cs
--- a/Engine/Source/Programs/Shared/EpicGames.Core/LogEventBuilder.cs
+++ b/Engine/Source/Programs/Shared/EpicGames.Core/LogEventBuilder.cs
@@ -1,13 +1,10 @@
 // Copyright Epic Games, Inc. All Rights Reserved.
 
-using Microsoft.Extensions.Logging;
 using System;
 using System.Collections.Generic;
-using System.Linq;
-using System.Reflection;
 using System.Text;
-using System.Text.Json;
 using System.Text.RegularExpressions;
+using Microsoft.Extensions.Logging;
 
 namespace EpicGames.Core
 {
@@ -43,13 +40,13 @@ namespace EpicGames.Core
 		/// <summary>
 		/// Constructor
 		/// </summary>
-		/// <param name="Offset">Starting offset within the line</param>
-		/// <param name="Text">The text for this span</param>
-		public LogEventSpan(int Offset, string Text)
+		/// <param name="offset">Starting offset within the line</param>
+		/// <param name="text">The text for this span</param>
+		public LogEventSpan(int offset, string text)
 		{
-			this.Offset = Offset;
-			this.Text = Text;
-			this.Properties = new Dictionary<string, object>(StringComparer.OrdinalIgnoreCase);
+			Offset = offset;
+			Text = text;
+			Properties = new Dictionary<string, object>(StringComparer.OrdinalIgnoreCase);
 		}
 
 		/// <summary>
@@ -86,56 +83,56 @@ namespace EpicGames.Core
 		/// <summary>
 		/// Constructor
 		/// </summary>
-		/// <param name="Text">Text for the line</param>
-		public LogEventLine(string Text)
+		/// <param name="text">Text for the line</param>
+		public LogEventLine(string text)
 		{
-			this.Text = Text;
-			this.Spans = new Dictionary<string, LogEventSpan>();
+			Text = text;
+			Spans = new Dictionary<string, LogEventSpan>();
 		}
 
 		/// <summary>
 		/// Adds a span containing markup on the source text
 		/// </summary>
-		/// <param name="Offset">Offset within the line</param>
-		/// <param name="Length">Length of the span</param>
-		/// <param name="Name">Name to use to identify the item in the format string</param>
+		/// <param name="offset">Offset within the line</param>
+		/// <param name="length">Length of the span</param>
+		/// <param name="name">Name to use to identify the item in the format string</param>
 		/// <returns>New span for the given range</returns>
-		public LogEventSpan AddSpan(int Offset, int Length, string Name)
+		public LogEventSpan AddSpan(int offset, int length, string name)
 		{
-			LogEventSpan Span = new LogEventSpan(Offset, Text.Substring(Offset, Length));
-			Spans.Add(Name, Span);
-			return Span;
+			LogEventSpan span = new LogEventSpan(offset, Text.Substring(offset, length));
+			Spans.Add(name, span);
+			return span;
 		}
 
 		/// <summary>
 		/// Adds a span containing markup for a regex match group
 		/// </summary>
-		/// <param name="Group">The match group</param>
-		/// <param name="Name">Name to use to identify the item in the format string</param>
-		public LogEventSpan AddSpan(Group Group, string Name)
+		/// <param name="group">The match group</param>
+		/// <param name="name">Name to use to identify the item in the format string</param>
+		public LogEventSpan AddSpan(Group group, string name)
 		{
-			return AddSpan(Group.Index, Group.Length, Name);
+			return AddSpan(group.Index, group.Length, name);
 		}
 
 		/// <summary>
 		/// Adds a span naming a regex match group, using the name of the group
 		/// </summary>
-		/// <param name="Group">The match group</param>
-		public LogEventSpan AddSpan(Group Group)
+		/// <param name="group">The match group</param>
+		public LogEventSpan AddSpan(Group group)
 		{
-			return AddSpan(Group.Index, Group.Length, Group.Name);
+			return AddSpan(group.Index, group.Length, group.Name);
 		}
 
 		/// <summary>
 		/// Adds a span naming a regex match group, using the name of the group
 		/// </summary>
-		/// <param name="Group">The match group</param>
-		/// <param name="Name">Name to use to identify the item in the format string</param>
-		public LogEventSpan? TryAddSpan(Group Group, string Name)
+		/// <param name="group">The match group</param>
+		/// <param name="name">Name to use to identify the item in the format string</param>
+		public LogEventSpan? TryAddSpan(Group group, string name)
 		{
-			if (Group.Success)
+			if (group.Success)
 			{
-				return AddSpan(Group, Name);
+				return AddSpan(group, name);
 			}
 			else
 			{
@@ -146,12 +143,12 @@ namespace EpicGames.Core
 		/// <summary>
 		/// Adds a span naming a regex match group, using the name of the group
 		/// </summary>
-		/// <param name="Group">The match group</param>
-		public LogEventSpan? TryAddSpan(Group Group)
+		/// <param name="group">The match group</param>
+		public LogEventSpan? TryAddSpan(Group group)
 		{
-			if (Group.Success)
+			if (group.Success)
 			{
-				return AddSpan(Group);
+				return AddSpan(group);
 			}
 			else
 			{
@@ -173,31 +170,31 @@ namespace EpicGames.Core
 	{
 		class LogSpan
 		{
-			public string Name;
-			public int Offset;
-			public int Length;
-			public object? Value;
+			public string _name;
+			public int _offset;
+			public int _length;
+			public object? _value;
 
-			public LogSpan(string Name, int Offset, int Length, object? Value)
+			public LogSpan(string name, int offset, int length, object? value)
 			{
-				this.Name = Name;
-				this.Offset = Offset;
-				this.Length = Length;
-				this.Value = Value;
+				_name = name;
+				_offset = offset;
+				_length = length;
+				_value = value;
 			}
 		}
 
 		class LogLine
 		{
-			public string Message;
-			public string? Format;
-			public Dictionary<string, object>? Properties;
+			public string _message;
+			public string? _format;
+			public Dictionary<string, object>? _properties;
 
-			public LogLine(string Message, string? Format, Dictionary<string, object>? Properties)
+			public LogLine(string message, string? format, Dictionary<string, object>? properties)
 			{
-				this.Message = Message;
-				this.Format = Format;
-				this.Properties = Properties;
+				_message = message;
+				_format = format;
+				_properties = properties;
 			}
 		}
 
@@ -214,30 +211,30 @@ namespace EpicGames.Core
 		/// <summary>
 		/// Events which have been parsed so far
 		/// </summary>
-		List<LogLine>? Lines;
+		List<LogLine>? _lines;
 
 		/// <summary>
 		/// Spans for the current line
 		/// </summary>
-		List<LogSpan>? Spans;
+		List<LogSpan>? _spans;
 
 		/// <summary>
 		/// Additional properties for this line
 		/// </summary>
-		Dictionary<string, object>? Properties;
+		Dictionary<string, object>? _properties;
 
 		/// <summary>
 		/// Starts building a log event at the current cursor position
 		/// </summary>
-		/// <param name="Cursor">The current cursor position</param>
-		public LogEventBuilder(ILogCursor Cursor, int LineCount = 1)
+		/// <param name="cursor">The current cursor position</param>
+		public LogEventBuilder(ILogCursor cursor, int lineCount = 1)
 		{
-			this.Current = Cursor;
-			this.Next = Cursor.Rebase(1);
+			Current = cursor;
+			Next = cursor.Rebase(1);
 
-			if (LineCount > 1)
+			if (lineCount > 1)
 			{
-				MoveNext(LineCount - 1);
+				MoveNext(lineCount - 1);
 			}
 		}
 
@@ -247,50 +244,50 @@ namespace EpicGames.Core
 		/// <returns></returns>
 		LogLine CreateLine()
 		{
-			int Offset = 0;
-			string CurrentLine = Current.CurrentLine!;
+			int offset = 0;
+			string currentLine = Current.CurrentLine!;
 
-			string? Format = null;
-			if (Spans != null)
+			string? format = null;
+			if (_spans != null)
 			{
-				StringBuilder Builder = new StringBuilder();
-				foreach (LogSpan Span in Spans)
+				StringBuilder builder = new StringBuilder();
+				foreach (LogSpan span in _spans)
 				{
-					if (Span.Offset >= Offset)
+					if (span._offset >= offset)
 					{
-						Builder.Append(CurrentLine, Offset, Span.Offset - Offset);
-						Builder.Append($"{{{Span.Name}}}");
-						Offset = Span.Offset + Span.Length;
+						builder.Append(currentLine, offset, span._offset - offset);
+						builder.Append($"{{{span._name}}}");
+						offset = span._offset + span._length;
 					}
 				}
-				Builder.Append(CurrentLine, Offset, CurrentLine.Length - Offset);
-				Format = Builder.ToString();
+				builder.Append(currentLine, offset, currentLine.Length - offset);
+				format = builder.ToString();
 			}
 
-			return new LogLine(CurrentLine, Format, Properties);
+			return new LogLine(currentLine, format, _properties);
 		}
 
 		/// <summary>
 		/// Adds a span containing markup on the source text
 		/// </summary>
-		/// <param name="Name">Name to use to identify the item in the format string</param>
-		/// <param name="Offset">Offset within the line</param>
-		/// <param name="Length">Length of the span</param>
+		/// <param name="name">Name to use to identify the item in the format string</param>
+		/// <param name="offset">Offset within the line</param>
+		/// <param name="length">Length of the span</param>
 		/// <param name="Data">Data of the span</param>
 		/// <returns>New span for the given range</returns>
-		public void Annotate(string Name, int Offset, int Length, object? Value = null)
+		public void Annotate(string name, int offset, int length, object? value = null)
 		{
-			LogSpan Span = new LogSpan(Name, Offset, Length, Value);
+			LogSpan span = new LogSpan(name, offset, length, value);
 
-			this.Properties ??= new Dictionary<string, object>(StringComparer.OrdinalIgnoreCase);
-			this.Properties.Add(Name, Span);
+			_properties ??= new Dictionary<string, object>(StringComparer.OrdinalIgnoreCase);
+			_properties.Add(name, span);
 
-			Spans ??= new List<LogSpan>();
-			for(int InsertIdx = Spans.Count; ;InsertIdx--)
+			_spans ??= new List<LogSpan>();
+			for(int insertIdx = _spans.Count; ;insertIdx--)
 			{
-				if (InsertIdx == 0 || Spans[InsertIdx - 1].Offset < Offset)
+				if (insertIdx == 0 || _spans[insertIdx - 1]._offset < offset)
 				{
-					Spans.Insert(InsertIdx, Span);
+					_spans.Insert(insertIdx, span);
 					break;
 				}
 			}
@@ -299,32 +296,32 @@ namespace EpicGames.Core
 		/// <summary>
 		/// Adds a span containing markup for a regex match group
 		/// </summary>
-		/// <param name="Group">The match group</param>
-		/// <param name="Name">Name to use to identify the item in the format string</param>
-		public void Annotate(string Name, Group Group, object? Value = null)
+		/// <param name="group">The match group</param>
+		/// <param name="name">Name to use to identify the item in the format string</param>
+		public void Annotate(string name, Group group, object? value = null)
 		{
-			Annotate(Name, Group.Index, Group.Length, Value);
+			Annotate(name, group.Index, group.Length, value);
 		}
 
 		/// <summary>
 		/// Adds a span naming a regex match group, using the name of the group
 		/// </summary>
-		/// <param name="Group">The match group</param>
-		public void Annotate(Group Group, object? Value = null)
+		/// <param name="group">The match group</param>
+		public void Annotate(Group group, object? value = null)
 		{
-			Annotate(Group.Name, Group.Index, Group.Length, Value);
+			Annotate(group.Name, group.Index, group.Length, value);
 		}
 
 		/// <summary>
 		/// Adds a span naming a regex match group, using the name of the group
 		/// </summary>
-		/// <param name="Group">The match group</param>
-		/// <param name="Name">Name to use to identify the item in the format string</param>
-		public bool TryAnnotate(string Name, Group Group, object? Value = null)
+		/// <param name="group">The match group</param>
+		/// <param name="name">Name to use to identify the item in the format string</param>
+		public bool TryAnnotate(string name, Group group, object? value = null)
 		{
-			if (Group.Success)
+			if (group.Success)
 			{
-				Annotate(Name, Group);
+				Annotate(name, group, value);
 				return true;
 			}
 			return false;
@@ -333,12 +330,12 @@ namespace EpicGames.Core
 		/// <summary>
 		/// Adds a span naming a regex match group, using the name of the group
 		/// </summary>
-		/// <param name="Group">The match group</param>
-		public bool TryAnnotate(Group Group, object? Value = null)
+		/// <param name="group">The match group</param>
+		public bool TryAnnotate(Group group, object? value = null)
 		{
-			if (Group.Success)
+			if (group.Success)
 			{
-				Annotate(Group, Value);
+				Annotate(group, value);
 				return true;
 			}
 			return false;
@@ -347,12 +344,12 @@ namespace EpicGames.Core
 		/// <summary>
 		/// Adds an additional named property
 		/// </summary>
-		/// <param name="Name">Name of the argument</param>
-		/// <param name="Value">Value to associate with it</param>
-		public void AddProperty(string Name, object Value)
+		/// <param name="name">Name of the argument</param>
+		/// <param name="value">Value to associate with it</param>
+		public void AddProperty(string name, object value)
 		{
-			Properties ??= new Dictionary<string, object>(StringComparer.OrdinalIgnoreCase);
-			Properties.Add(Name, Value);
+			_properties ??= new Dictionary<string, object>(StringComparer.OrdinalIgnoreCase);
+			_properties.Add(name, value);
 		}
 
 		/// <summary>
@@ -360,11 +357,11 @@ namespace EpicGames.Core
 		/// </summary>
 		public void MoveNext()
 		{
-			Lines ??= new List<LogLine>();
-			Lines.Add(CreateLine());
+			_lines ??= new List<LogLine>();
+			_lines.Add(CreateLine());
 
-			Spans = null;
-			Properties = null;
+			_spans = null;
+			_properties = null;
 
 			Current = Next;
 			Next = Next.Rebase(1);
@@ -373,10 +370,10 @@ namespace EpicGames.Core
 		/// <summary>
 		/// Advance by the given number of lines
 		/// </summary>
-		/// <param name="Count"></param>
-		public void MoveNext(int Count)
+		/// <param name="count"></param>
+		public void MoveNext(int count)
 		{
-			for (int Idx = 0; Idx < Count; Idx++)
+			for (int idx = 0; idx < count; idx++)
 			{
 				MoveNext();
 			}
@@ -386,75 +383,75 @@ namespace EpicGames.Core
 		/// Returns an array of log events
 		/// </summary>
 		/// <returns></returns>
-		public LogEvent[] ToArray(LogLevel Level, EventId EventId)
+		public LogEvent[] ToArray(LogLevel level, EventId eventId)
 		{
-			DateTime Time = DateTime.UtcNow;
+			DateTime time = DateTime.UtcNow;
 
-			int NumLines = Lines?.Count ?? 0;
-			int NumEvents = NumLines;
+			int numLines = _lines?.Count ?? 0;
+			int numEvents = numLines;
 			if (Current.CurrentLine != null)
 			{
-				NumEvents++;
+				numEvents++;
 			}
 
-			LogEvent[] Events = new LogEvent[NumEvents];
-			for (int Idx = 0; Idx < NumLines; Idx++)
+			LogEvent[] events = new LogEvent[numEvents];
+			for (int idx = 0; idx < numLines; idx++)
 			{
-				Events[Idx] = CreateEvent(Time, Level, EventId, Idx, NumEvents, Lines![Idx]);
+				events[idx] = CreateEvent(time, level, eventId, idx, numEvents, _lines![idx]);
 			}
 			if (Current.CurrentLine != null)
 			{
-				Events[NumEvents - 1] = CreateEvent(Time, Level, EventId, NumEvents - 1, NumEvents, CreateLine());
+				events[numEvents - 1] = CreateEvent(time, level, eventId, numEvents - 1, numEvents, CreateLine());
 			}
 
-			return Events;
+			return events;
 		}
 
-		static LogEvent CreateEvent(DateTime Time, LogLevel Level, EventId EventId, int LineIndex, int LineCount, LogLine Line)
+		static LogEvent CreateEvent(DateTime time, LogLevel level, EventId eventId, int lineIndex, int lineCount, LogLine line)
 		{
-			Dictionary<string, object>? Properties = null;
-			if (Line.Properties != null)
+			Dictionary<string, object>? properties = null;
+			if (line._properties != null)
 			{
-				Properties = new Dictionary<string, object>();
-				foreach ((string Name, object Value) in Line.Properties)
+				properties = new Dictionary<string, object>();
+				foreach ((string name, object value) in line._properties)
 				{
-					object NewValue;
-					if (Value is LogSpan Span)
+					object newValue;
+					if (value is LogSpan span)
 					{
-						string Text = Line.Message.Substring(Span.Offset, Span.Length);
-						if (Span.Value == null)
+						string text = line._message.Substring(span._offset, span._length);
+						if (span._value == null)
 						{
-							NewValue = Text;
+							newValue = text;
 						}
 						else
 						{
-							NewValue = LogEventFormatter.Format(Span.Value!);
-							if (NewValue is LogValue Event)
+							newValue = LogEventFormatter.Format(span._value!);
+							if (newValue is LogValue newLogValue)
 							{
-								Event.Text = Text;
+								newLogValue.Text = text;
 							}
 						}
 					}
 					else
 					{
-						NewValue = LogEventFormatter.Format(Value);
+						newValue = LogEventFormatter.Format(value);
 					}
-					Properties[Name] = NewValue;
+					properties[name] = newValue;
 				}
 			}
-			return new LogEvent(Time, Level, EventId, LineIndex, LineCount, Line.Message, Line.Format, Properties, null);
+			return new LogEvent(time, level, eventId, lineIndex, lineCount, line._message, line._format, properties, null);
 		}
 
 		/// <summary>
 		/// Creates a match object at the given priority
 		/// </summary>
-		/// <param name="Level"></param>
-		/// <param name="EventId">The event id</param>
-		/// <param name="Priority"></param>
+		/// <param name="level"></param>
+		/// <param name="eventId">The event id</param>
+		/// <param name="priority"></param>
 		/// <returns></returns>
-		public LogEventMatch ToMatch(LogEventPriority Priority, LogLevel Level, EventId EventId)
+		public LogEventMatch ToMatch(LogEventPriority priority, LogLevel level, EventId eventId)
 		{
-			return new LogEventMatch(Priority, ToArray(Level, EventId));
+			return new LogEventMatch(priority, ToArray(level, eventId));
 		}
 	}
 }
diff --git a/Engine/Source/Programs/Shared/EpicGames.Core/LogEventFormatter.cs b/Engine/Source/Programs/Shared/EpicGames.Core/LogEventFormatter.cs
--- a/Engine/Source/Programs/Shared/EpicGames.Core/LogEventFormatter.cs
+++ b/Engine/Source/Programs/Shared/EpicGames.Core/LogEventFormatter.cs
@@ -2,10 +2,7 @@
 
 using System;
 using System.Collections.Concurrent;
-using System.Collections.Generic;
-using System.Diagnostics.CodeAnalysis;
 using System.Reflection;
-using System.Text;
 
 namespace EpicGames.Core
 {
@@ -17,9 +14,9 @@ namespace EpicGames.Core
 		/// <summary>
 		/// Format a value for a log file.
 		/// </summary>
-		/// <param name="Value">The value to format</param>
+		/// <param name="value">The value to format</param>
 		/// <returns>Value to be written to the log event. This may be a primitive type (int, string, etc...) or a LogValue object. If it is neither, Result.ToString() is used.</returns>
-		public object Format(object Value);
+		public object Format(object value);
 	}
 
 	/// <summary>
@@ -36,10 +33,10 @@ namespace EpicGames.Core
 		/// <summary>
 		/// Constructor
 		/// </summary>
-		/// <param name="Name">Type name to use in log events</param>
-		public LogEventTypeAttribute(string Name)
+		/// <param name="name">Type name to use in log events</param>
+		public LogEventTypeAttribute(string name)
 		{
-			this.Name = Name;
+			Name = name;
 		}
 	}
 
@@ -57,10 +54,10 @@ namespace EpicGames.Core
 		/// <summary>
 		/// Constructor
 		/// </summary>
-		/// <param name="Type"></param>
-		public LogEventFormatterAttribute(Type Type)
+		/// <param name="type"></param>
+		public LogEventFormatterAttribute(Type type)
 		{
-			this.Type = Type;
+			Type = type;
 		}
 	}
 
@@ -75,7 +72,7 @@ namespace EpicGames.Core
 		public class DefaultFormatter : ILogEventFormatter
 		{
 			/// <inheritdoc/>
-			public object Format(object Value) => Value.ToString() ?? String.Empty;
+			public object Format(object value) => value.ToString() ?? String.Empty;
 		}
 
 		/// <summary>
@@ -84,7 +81,7 @@ namespace EpicGames.Core
 		public class PassThroughFormatter : ILogEventFormatter
 		{
 			/// <inheritdoc/>
-			public object Format(object Value) => Value;
+			public object Format(object value) => value;
 		}
 
 		/// <summary>
@@ -92,90 +89,89 @@ namespace EpicGames.Core
 		/// </summary>
 		public class AnnotateTypeFormatter : ILogEventFormatter
 		{
-			readonly string Name;
+			readonly string _name;
 
-			public AnnotateTypeFormatter(string Name)
+			public AnnotateTypeFormatter(string name)
 			{
-				this.Name = Name;
+				_name = name;
 			}
 
 			/// <inheritdoc/>
-			public object Format(object Value) => new LogValue(Name, Value.ToString() ?? String.Empty);
+			public object Format(object value) => new LogValue(_name, value.ToString() ?? String.Empty);
 		}
 
-		static DefaultFormatter DefaultFormatterInstance = new DefaultFormatter();
+		static readonly DefaultFormatter s_defaultFormatterInstance = new DefaultFormatter();
 
-		static ConcurrentDictionary<Type, ILogEventFormatter> CachedFormatters = GetDefaultFormatters();
+		static readonly ConcurrentDictionary<Type, ILogEventFormatter> s_cachedFormatters = GetDefaultFormatters();
 
 		static ConcurrentDictionary<Type, ILogEventFormatter> GetDefaultFormatters()
 		{
-			PassThroughFormatter PassThroughFormatter = new PassThroughFormatter();
+			PassThroughFormatter passThroughFormatter = new PassThroughFormatter();
 
-			ConcurrentDictionary<Type, ILogEventFormatter> Formatters = new ConcurrentDictionary<Type, ILogEventFormatter>();
-			Formatters.TryAdd(typeof(bool), PassThroughFormatter);
-			Formatters.TryAdd(typeof(int), PassThroughFormatter);
-			Formatters.TryAdd(typeof(string), PassThroughFormatter);
-			Formatters.TryAdd(typeof(LogValue), PassThroughFormatter);
+			ConcurrentDictionary<Type, ILogEventFormatter> formatters = new ConcurrentDictionary<Type, ILogEventFormatter>();
+			formatters.TryAdd(typeof(bool), passThroughFormatter);
+			formatters.TryAdd(typeof(int), passThroughFormatter);
+			formatters.TryAdd(typeof(string), passThroughFormatter);
+			formatters.TryAdd(typeof(LogValue), passThroughFormatter);
 
-			return Formatters;
+			return formatters;
 		}
 
 		/// <summary>
 		/// Registers a formatter for a particular type
 		/// </summary>
-		/// <param name="Type"></param>
-		/// <param name="Formatter"></param>
-		public static void RegisterFormatter(Type Type, ILogEventFormatter Formatter)
+		/// <param name="type"></param>
+		/// <param name="formatter"></param>
+		public static void RegisterFormatter(Type type, ILogEventFormatter formatter)
 		{
-			CachedFormatters.TryAdd(Type, Formatter);
+			s_cachedFormatters.TryAdd(type, formatter);
 		}
 
 		/// <summary>
 		/// Gets a formatter for the specified type
 		/// </summary>
-		/// <param name="Type"></param>
+		/// <param name="type"></param>
 		/// <returns></returns>
-		public static ILogEventFormatter GetFormatter(Type Type)
+		public static ILogEventFormatter GetFormatter(Type type)
 		{
 			for (; ; )
 			{
-				ILogEventFormatter? Formatter;
-				if (CachedFormatters.TryGetValue(Type, out Formatter))
+				ILogEventFormatter? formatter;
+				if (s_cachedFormatters.TryGetValue(type, out formatter))
 				{
-					return Formatter;
+					return formatter;
 				}
 
-				Formatter = CreateFormatter(Type);
+				formatter = CreateFormatter(type);
 
-				if (CachedFormatters.TryAdd(Type, Formatter))
+				if (s_cachedFormatters.TryAdd(type, formatter))
 				{
-					return Formatter;
+					return formatter;
 				}
 			}
 		}
 
-		static ILogEventFormatter CreateFormatter(Type Type)
+		static ILogEventFormatter CreateFormatter(Type type)
 		{
-			LogEventTypeAttribute? TypeAttribute = Type.GetCustomAttribute<LogEventTypeAttribute>();
-			if (TypeAttribute != null)
+			LogEventTypeAttribute? typeAttribute = type.GetCustomAttribute<LogEventTypeAttribute>();
+			if (typeAttribute != null)
 			{
-				return new AnnotateTypeFormatter(Type.Name);
+				return new AnnotateTypeFormatter(type.Name);
 			}
 
-			LogEventFormatterAttribute? FormatterAttribute = Type.GetCustomAttribute<LogEventFormatterAttribute>();
-			if (FormatterAttribute != null)
+			LogEventFormatterAttribute? formatterAttribute = type.GetCustomAttribute<LogEventFormatterAttribute>();
+			if (formatterAttribute != null)
 			{
-				return (ILogEventFormatter)Activator.CreateInstance(FormatterAttribute.Type)!;
+				return (ILogEventFormatter)Activator.CreateInstance(formatterAttribute.Type)!;
 			}
 
-			return DefaultFormatterInstance;
+			return s_defaultFormatterInstance;
 		}
 
-
-		public static object Format(object Value)
+		public static object Format(object value)
 		{
-			ILogEventFormatter Formatter = GetFormatter(Value.GetType());
-			return Formatter.Format(Value);
+			ILogEventFormatter formatter = GetFormatter(value.GetType());
+			return formatter.Format(value);
 		}
 	}
 }
diff --git a/Engine/Source/Programs/Shared/EpicGames.Core/LogEventParser.cs b/Engine/Source/Programs/Shared/EpicGames.Core/LogEventParser.cs
--- a/Engine/Source/Programs/Shared/EpicGames.Core/LogEventParser.cs
+++ b/Engine/Source/Programs/Shared/EpicGames.Core/LogEventParser.cs
@@ -1,18 +1,12 @@
 // Copyright Epic Games, Inc. All Rights Reserved.
 
-using Microsoft.Extensions.Logging;
 using System;
-using System.Buffers.Text;
-using System.Collections.Concurrent;
 using System.Collections.Generic;
 using System.IO;
 using System.Linq;
-using System.Reflection;
 using System.Text;
-using System.Text.Json;
 using System.Text.RegularExpressions;
-using System.Threading;
-using System.Threading.Tasks;
+using Microsoft.Extensions.Logging;
 
 namespace EpicGames.Core
 {
@@ -36,16 +30,16 @@ namespace EpicGames.Core
 		public LogEventPriority Priority { get; }
 		public List<LogEvent> Events { get; }
 
-		public LogEventMatch(LogEventPriority Priority, LogEvent Event)
+		public LogEventMatch(LogEventPriority priority, LogEvent logEvent)
 		{
-			this.Priority = Priority;
-			this.Events = new List<LogEvent> { Event };
+			Priority = priority;
+			Events = new List<LogEvent> { logEvent };
 		}
 
-		public LogEventMatch(LogEventPriority Priority, IEnumerable<LogEvent> Events)
+		public LogEventMatch(LogEventPriority priority, IEnumerable<LogEvent> events)
 		{
-			this.Priority = Priority;
-			this.Events = Events.ToList();
+			Priority = priority;
+			Events = events.ToList();
 		}
 	}
 
@@ -57,9 +51,9 @@ namespace EpicGames.Core
 		/// <summary>
 		/// Attempt to match events from the given input buffer
 		/// </summary>
-		/// <param name="Cursor">The input buffer</param>
+		/// <param name="cursor">The input buffer</param>
 		/// <returns>Information about the error that was matched, or null if an error was not matched</returns>
-		LogEventMatch? Match(ILogCursor Cursor);
+		LogEventMatch? Match(ILogCursor cursor);
 	}
 
 	/// <summary>
@@ -80,31 +74,31 @@ namespace EpicGames.Core
 		/// <summary>
 		/// Buffer of input lines
 		/// </summary>
-		LogBuffer Buffer;
+		readonly LogBuffer _buffer;
 
 		/// <summary>
 		/// Buffer for holding partial line data
 		/// </summary>
-		MemoryStream PartialLine = new MemoryStream();
+		readonly MemoryStream _partialLine = new MemoryStream();
 
 		/// <summary>
 		/// Whether matching is currently enabled
 		/// </summary>
-		int MatchingEnabled;
+		int _matchingEnabled;
 
 		/// <summary>
 		/// The inner logger
 		/// </summary>
-		ILogger Logger;
+		readonly ILogger _logger;
 
 		/// <summary>
 		/// Constructor
 		/// </summary>
-		/// <param name="Logger">The logger to receive parsed output messages</param>
-		public LogEventParser(ILogger Logger)
+		/// <param name="logger">The logger to receive parsed output messages</param>
+		public LogEventParser(ILogger logger)
 		{
-			this.Logger = Logger;
-			this.Buffer = new LogBuffer(50);
+			_logger = logger;
+			_buffer = new LogBuffer(50);
 		}
 
 		/// <inheritdoc/>
@@ -113,50 +107,50 @@ namespace EpicGames.Core
 		/// <summary>
 		/// Writes a line to the event filter
 		/// </summary>
-		/// <param name="Line">The line to output</param>
-		public void WriteLine(string Line)
+		/// <param name="line">The line to output</param>
+		public void WriteLine(string line)
 		{
-			Buffer.AddLine(Line);
+			_buffer.AddLine(line);
 			ProcessData(false);
 		}
 
 		/// <summary>
 		/// Writes data to the log parser
 		/// </summary>
-		/// <param name="Data">Data to write</param>
-		public void WriteData(ReadOnlyMemory<byte> Data)
+		/// <param name="data">Data to write</param>
+		public void WriteData(ReadOnlyMemory<byte> data)
 		{
-			int BaseIdx = 0;
-			int ScanIdx = 0;
-			ReadOnlySpan<byte> Span = Data.Span;
+			int baseIdx = 0;
+			int scanIdx = 0;
+			ReadOnlySpan<byte> span = data.Span;
 
 			// Handle a partially existing line
-			if (PartialLine.Length > 0)
+			if (_partialLine.Length > 0)
 			{
-				for (; ScanIdx < Span.Length; ScanIdx++)
+				for (; scanIdx < span.Length; scanIdx++)
 				{
-					if (Span[ScanIdx] == '\n')
+					if (span[scanIdx] == '\n')
 					{
-						PartialLine.Write(Span.Slice(BaseIdx, ScanIdx - BaseIdx));
+						_partialLine.Write(span.Slice(baseIdx, scanIdx - baseIdx));
 						FlushPartialLine();
-						BaseIdx = ++ScanIdx;
+						baseIdx = ++scanIdx;
 						break;
 					}
 				}
 			}
 
 			// Handle any complete lines
-			for (; ScanIdx < Span.Length; ScanIdx++)
+			for (; scanIdx < span.Length; scanIdx++)
 			{
-				if(Span[ScanIdx] == '\n')
+				if(span[scanIdx] == '\n')
 				{
-					AddLine(Data.Slice(BaseIdx, ScanIdx - BaseIdx));
-					BaseIdx = ScanIdx + 1;
+					AddLine(data.Slice(baseIdx, scanIdx - baseIdx));
+					baseIdx = scanIdx + 1;
 				}
 			}
 
 			// Add the rest of the text to the partial line buffer
-			PartialLine.Write(Span.Slice(BaseIdx));
+			_partialLine.Write(span.Slice(baseIdx));
 
 			// Process the new data
 			ProcessData(false);
@@ -168,7 +162,7 @@ namespace EpicGames.Core
 		public void Flush()
 		{
 			// If there's a partially written line, write that out first
-			if (PartialLine.Length > 0)
+			if (_partialLine.Length > 0)
 			{
 				FlushPartialLine();
 			}
@@ -180,24 +174,24 @@ namespace EpicGames.Core
 		/// <summary>
 		/// Adds a raw utf-8 string to the buffer
 		/// </summary>
-		/// <param name="Data">The string data</param>
-		private void AddLine(ReadOnlyMemory<byte> Data)
+		/// <param name="data">The string data</param>
+		private void AddLine(ReadOnlyMemory<byte> data)
 		{
-			if (Data.Length > 0 && Data.Span[Data.Length - 1] == '\r')
+			if (data.Length > 0 && data.Span[data.Length - 1] == '\r')
 			{
-				Data = Data.Slice(0, Data.Length - 1);
+				data = data.Slice(0, data.Length - 1);
 			}
-			if (Data.Length > 0 && Data.Span[0] == '{')
+			if (data.Length > 0 && data.Span[0] == '{')
 			{
-				JsonLogEvent JsonEvent;
-				if (JsonLogEvent.TryParse(Data, out JsonEvent))
+				JsonLogEvent jsonEvent;
+				if (JsonLogEvent.TryParse(data, out jsonEvent))
 				{
 					ProcessData(true);
-					Logger.Log(JsonEvent.Level, JsonEvent.EventId, JsonEvent, null, JsonLogEvent.Format);
+					_logger.Log(jsonEvent.Level, jsonEvent.EventId, jsonEvent, null, JsonLogEvent.Format);
 					return;
 				}
 			}
-			Buffer.AddLine(StringUtils.ParseEscapeCodes(Encoding.UTF8.GetString(Data.Span)));
+			_buffer.AddLine(StringUtils.ParseEscapeCodes(Encoding.UTF8.GetString(data.Span)));
 		}
 
 		/// <summary>
@@ -205,9 +199,9 @@ namespace EpicGames.Core
 		/// </summary>
 		private void FlushPartialLine()
 		{
-			AddLine(PartialLine.ToArray());
-			PartialLine.Position = 0;
-			PartialLine.SetLength(0);
+			AddLine(_partialLine.ToArray());
+			_partialLine.Position = 0;
+			_partialLine.SetLength(0);
 		}
 
 		/// <summary>
@@ -216,53 +210,53 @@ namespace EpicGames.Core
 		/// <param name="bFlush">Whether we've reached the end of the stream</param>
 		void ProcessData(bool bFlush)
 		{
-			while (Buffer.Length > 0)
+			while (_buffer.Length > 0)
 			{
 				// Try to match an event
-				List<LogEvent>? Events = null;
-				if (Regex.IsMatch(Buffer[0], "<-- Suspend Log Parsing -->", RegexOptions.IgnoreCase))
+				List<LogEvent>? events = null;
+				if (Regex.IsMatch(_buffer[0], "<-- Suspend Log Parsing -->", RegexOptions.IgnoreCase))
 				{
-					MatchingEnabled--;
+					_matchingEnabled--;
 				}
-				else if (Regex.IsMatch(Buffer[0], "<-- Resume Log Parsing -->", RegexOptions.IgnoreCase))
+				else if (Regex.IsMatch(_buffer[0], "<-- Resume Log Parsing -->", RegexOptions.IgnoreCase))
 				{
-					MatchingEnabled++;
+					_matchingEnabled++;
 				}
-				else if (MatchingEnabled >= 0)
+				else if (_matchingEnabled >= 0)
 				{
-					Events = MatchEvent();
+					events = MatchEvent();
 				}
 
 				// Bail out if we need more data
-				if (Buffer.Length < 1024 && !bFlush && Buffer.NeedMoreData)
+				if (_buffer.Length < 1024 && !bFlush && _buffer.NeedMoreData)
 				{
 					break;
 				}
 
 				// If we did match something, check if it's not negated by an ignore pattern. We typically have relatively few errors and many more ignore patterns than matchers, so it's quicker 
 				// to check them in response to an identified error than to treat them as matchers of their own.
-				if (Events != null)
+				if (events != null)
 				{
-					foreach (Regex IgnorePattern in IgnorePatterns)
+					foreach (Regex ignorePattern in IgnorePatterns)
 					{
-						if (IgnorePattern.IsMatch(Buffer[0]))
+						if (ignorePattern.IsMatch(_buffer[0]))
 						{
-							Events = null;
+							events = null;
 							break;
 						}
 					}
 				}
 
 				// Report the error to the listeners
-				if (Events != null)
+				if (events != null)
 				{
-					WriteEvents(Events);
-					Buffer.Advance(Events.Count);
+					WriteEvents(events);
+					_buffer.Advance(events.Count);
 				}
 				else
 				{
-					Logger.Log(LogLevel.Information, KnownLogEvents.None, Buffer[0]!, null, (State, Exception) => State);
-					Buffer.MoveNext();
+					_logger.Log(LogLevel.Information, KnownLogEvents.None, _buffer[0]!, null, (state, exception) => state);
+					_buffer.MoveNext();
 				}
 			}
 		}
@@ -273,30 +267,30 @@ namespace EpicGames.Core
 		/// <returns>The matched event</returns>
 		private List<LogEvent>? MatchEvent()
 		{
-			LogEventMatch? CurrentMatch = null;
-			foreach (ILogEventMatcher Matcher in Matchers)
+			LogEventMatch? currentMatch = null;
+			foreach (ILogEventMatcher matcher in Matchers)
 			{
-				LogEventMatch? Match = Matcher.Match(Buffer);
-				if(Match != null)
+				LogEventMatch? match = matcher.Match(_buffer);
+				if(match != null)
 				{
-					if (CurrentMatch == null || Match.Priority > CurrentMatch.Priority)
+					if (currentMatch == null || match.Priority > currentMatch.Priority)
 					{
-						CurrentMatch = Match;
+						currentMatch = match;
 					}
 				}
 			}
-			return CurrentMatch?.Events;
+			return currentMatch?.Events;
 		}
 
 		/// <summary>
 		/// Writes an event to the log
 		/// </summary>
 		/// <param name="Event">The event to write</param>
-		protected virtual void WriteEvents(List<LogEvent> Events)
+		protected virtual void WriteEvents(List<LogEvent> logEvents)
 		{
-			foreach (LogEvent Event in Events)
+			foreach (LogEvent logEvent in logEvents)
 			{
-				Logger.Log(Event.Level, Event.Id, Event, null, (State, Exception) => State.ToString());
+				_logger.Log(logEvent.Level, logEvent.Id, logEvent, null, (state, exception) => state.ToString());
 			}
 		}
 	}
diff --git a/Engine/Source/Programs/Shared/EpicGames.Core/LogIndent.cs b/Engine/Source/Programs/Shared/EpicGames.Core/LogIndent.cs
--- a/Engine/Source/Programs/Shared/EpicGames.Core/LogIndent.cs
+++ b/Engine/Source/Programs/Shared/EpicGames.Core/LogIndent.cs
@@ -1,8 +1,6 @@
 // Copyright Epic Games, Inc. All Rights Reserved.
 
 using System;
-using System.Collections.Generic;
-using System.Text;
 using System.Threading;
 
 namespace EpicGames.Core
@@ -20,35 +18,35 @@ namespace EpicGames.Core
 			public State? PrevState { get; }
 			public string Indent { get; }
 
-			public State(State? PrevState, string Indent)
+			public State(State? prevState, string indent)
 			{
-				this.PrevState = PrevState;
-				this.Indent = Indent;
+				PrevState = prevState;
+				Indent = indent;
 			}
 		}
 
 		/// <summary>
 		/// The current state value
 		/// </summary>
-		static AsyncLocal<State?> CurrentState = new AsyncLocal<State?>();
+		static readonly AsyncLocal<State?> s_currentState = new AsyncLocal<State?>();
 
 		/// <summary>
 		/// Gets the current indent string
 		/// </summary>
-		public static string Current => CurrentState.Value?.Indent ?? String.Empty;
+		public static string Current => s_currentState.Value?.Indent ?? String.Empty;
 
 		/// <summary>
 		/// Push a new indent onto the stack
 		/// </summary>
-		/// <param name="Indent">The indent to add</param>
-		public static void Push(string Indent)
+		/// <param name="indent">The indent to add</param>
+		public static void Push(string indent)
 		{
-			State? PrevState = CurrentState.Value;
-			if (PrevState != null)
+			State? prevState = s_currentState.Value;
+			if (prevState != null)
 			{
-				Indent = PrevState.Indent + Indent;
+				indent = prevState.Indent + indent;
 			}
-			CurrentState.Value = new State(PrevState, Indent);
+			s_currentState.Value = new State(prevState, indent);
 		}
 
 		/// <summary>
@@ -56,7 +54,7 @@ namespace EpicGames.Core
 		/// </summary>
 		public static void Pop()
 		{
-			CurrentState.Value = CurrentState.Value!.PrevState;
+			s_currentState.Value = s_currentState.Value!.PrevState;
 		}
 	}
 }
diff --git a/Engine/Source/Programs/Shared/EpicGames.Core/LogIndentScope.cs b/Engine/Source/Programs/Shared/EpicGames.Core/LogIndentScope.cs
--- a/Engine/Source/Programs/Shared/EpicGames.Core/LogIndentScope.cs
+++ b/Engine/Source/Programs/Shared/EpicGames.Core/LogIndentScope.cs
@@ -1,11 +1,6 @@
 // Copyright Epic Games, Inc. All Rights Reserved.
 
 using System;
-using System.Collections.Generic;
-using System.Linq;
-using System.Text;
-using System.Threading;
-using System.Threading.Tasks;
 
 namespace EpicGames.Core
 {
@@ -17,15 +12,15 @@ namespace EpicGames.Core
 		/// <summary>
 		/// Whether the object has been disposed
 		/// </summary>
-		bool Disposed;
+		bool _disposed;
 
 		/// <summary>
 		/// Constructor
 		/// </summary>
-		/// <param name="Indent">Indent to append to the existing indent</param>
-		public LogIndentScope(string Indent)
+		/// <param name="indent">Indent to append to the existing indent</param>
+		public LogIndentScope(string indent)
 		{
-			LogIndent.Push(Indent);
+			LogIndent.Push(indent);
 		}
 
 		/// <summary>
@@ -33,10 +28,10 @@ namespace EpicGames.Core
 		/// </summary>
 		public void Dispose()
 		{
-			if (!Disposed)
+			if (!_disposed)
 			{
 				LogIndent.Pop();
-				Disposed = true;
+				_disposed = true;
 			}
 		}
 	}
diff --git a/Engine/Source/Programs/Shared/EpicGames.Core/LogStatusScope.cs b/Engine/Source/Programs/Shared/EpicGames.Core/LogStatusScope.cs
--- a/Engine/Source/Programs/Shared/EpicGames.Core/LogStatusScope.cs
+++ b/Engine/Source/Programs/Shared/EpicGames.Core/LogStatusScope.cs
@@ -1,10 +1,7 @@
 // Copyright Epic Games, Inc. All Rights Reserved.
 
 using System;
-using System.Collections.Generic;
-using System.Linq;
 using System.Text;
-using System.Threading.Tasks;
 
 namespace EpicGames.Core
 {
@@ -16,71 +13,71 @@ namespace EpicGames.Core
 		/// <summary>
 		/// The base status message
 		/// </summary>
-		string Message;
+		string _message;
 
 		/// <summary>
 		/// Constructor
 		/// </summary>
-		/// <param name="Message">The status message</param>
-		public LogStatusScope(string Message)
+		/// <param name="message">The status message</param>
+		public LogStatusScope(string message)
 		{
-			this.Message = Message;
-			Log.PushStatus(Message);
+			_message = message;
+			Log.PushStatus(message);
 		}
 
 		/// <summary>
 		/// Constructor
 		/// </summary>
-		/// <param name="Format">The format specifier for the message</param>
-		/// <param name="Args">Arguments for the status message</param>
-		public LogStatusScope(string Format, params object[] Args)
-			: this(String.Format(Format, Args))
+		/// <param name="format">The format specifier for the message</param>
+		/// <param name="args">Arguments for the status message</param>
+		public LogStatusScope(string format, params object[] args)
+			: this(String.Format(format, args))
 		{
 		}
 
 		/// <summary>
 		/// Updates the base status message passed into the constructor.
 		/// </summary>
-		/// <param name="Message">The status message</param>
-		public void SetMessage(string Message)
+		/// <param name="message">The status message</param>
+		public void SetMessage(string message)
 		{
-			this.Message = Message;
-			Log.UpdateStatus(Message);
+			_message = message;
+			Log.UpdateStatus(message);
 		}
 
 		/// <summary>
 		/// Updates the base status message passed into the constructor.
 		/// </summary>
-		/// <param name="Format">The format specifier for the message</param>
-		/// <param name="Args">Arguments for the status message</param>
-		public void SetMessage(string Format, params object[] Args)
+		/// <param name="format">The format specifier for the message</param>
+		/// <param name="args">Arguments for the status message</param>
+		public void SetMessage(string format, params object[] args)
 		{
-			SetMessage(String.Format(Format, Args));
+			SetMessage(String.Format(format, args));
 		}
 
 		/// <summary>
 		/// Appends a progress string to the status message. Overwrites any previous progress message.
 		/// </summary>
-		/// <param name="Progress">The progress message</param>
-		public void SetProgress(string Progress)
+		/// <param name="progress">The progress message</param>
+		public void SetProgress(string progress)
 		{
-			StringBuilder FullMessage = new StringBuilder(Message);
-			FullMessage.Append(' ');
-			FullMessage.Append(Progress);
-			Log.UpdateStatus(FullMessage.ToString());
+			StringBuilder fullMessage = new StringBuilder(_message);
+			fullMessage.Append(' ');
+			fullMessage.Append(progress);
+			Log.UpdateStatus(fullMessage.ToString());
 		}
 
 		/// <summary>
 		/// Appends a progress string to the status message. Overwrites any previous progress message.
 		/// </summary>
-		/// <param name="Format">The format specifier for the message</param>
-		/// <param name="Args">Arguments for the status message</param>
-		public void SetProgress(string Format, params object[] Args)
+		/// <param name="format">The format specifier for the message</param>
+		/// <param name="args">Arguments for the status message</param>
+		public void SetProgress(string format, params object[] args)
 		{
-			StringBuilder FullMessage = new StringBuilder(Message);
-			FullMessage.Append(' ');
-			FullMessage.AppendFormat(Format, Args);
-			Log.UpdateStatus(FullMessage.ToString());
+			StringBuilder fullMessage = new StringBuilder(_message);
+			fullMessage.Append(' ');
+			fullMessage.AppendFormat(format, args);
+			Log.UpdateStatus(fullMessage.ToString());
 		}
 
 		/// <summary>
@@ -88,10 +85,10 @@ namespace EpicGames.Core
 		/// </summary>
 		public void Dispose()
 		{
-			if(Message != null)
+			if(_message != null)
 			{
 				Log.PopStatus();
-				Message = null!;
+				_message = null!;
 			}
 		}
 	}
