Reformat EpicGames.Core according to standard coding conventions.

diff --git a/Engine/Source/Programs/Shared/EpicGames.Core/ManagedProcess.cs b/Engine/Source/Programs/Shared/EpicGames.Core/ManagedProcess.cs
--- a/Engine/Source/Programs/Shared/EpicGames.Core/ManagedProcess.cs
+++ b/Engine/Source/Programs/Shared/EpicGames.Core/ManagedProcess.cs
@@ -1,19 +1,18 @@
 // Copyright Epic Games, Inc. All Rights Reserved.
 
-using Microsoft.Extensions.ObjectPool;
-using Microsoft.Win32.SafeHandles;
 using System;
 using System.Collections.Generic;
 using System.ComponentModel;
 using System.Diagnostics;
 using System.Diagnostics.CodeAnalysis;
 using System.IO;
-using System.IO.Pipes;
 using System.Runtime.InteropServices;
 using System.Text;
 using System.Threading;
 using System.Threading.Channels;
 using System.Threading.Tasks;
+using Microsoft.Extensions.ObjectPool;
+using Microsoft.Win32.SafeHandles;
 
 namespace EpicGames.Core
 {
@@ -39,6 +38,7 @@ namespace EpicGames.Core
 	/// </summary>
 	public class ManagedProcessGroup : IDisposable
 	{
+#pragma warning disable IDE1006 // Naming Styles
 		const UInt32 JOB_OBJECT_LIMIT_KILL_ON_JOB_CLOSE = 0x00002000;
 		const UInt32 JOB_OBJECT_LIMIT_BREAKAWAY_OK = 0x00000800;
 
@@ -77,7 +77,7 @@ namespace EpicGames.Core
 			public UIntPtr PeakProcessMemoryUsed;
 			public UIntPtr PeakJobMemoryUsed;
 		}
-		
+
 		[StructLayout(LayoutKind.Sequential)]
 		struct JOBOBJECT_BASIC_ACCOUNTING_INFORMATION
 		{
@@ -91,15 +91,15 @@ namespace EpicGames.Core
 			public UInt32 TotalTerminatedProcesses;
 		}
 
-		[DllImport("kernel32.dll", SetLastError=true)]
-		static extern SafeFileHandle CreateJobObject(IntPtr SecurityAttributes, IntPtr Name);
+		[DllImport("kernel32.dll", SetLastError = true)]
+		static extern SafeFileHandle CreateJobObject(IntPtr securityAttributes, IntPtr name);
 
 		const int JobObjectBasicAccountingInformation = 1;
 		const int JobObjectExtendedLimitInformation = 9;
 
-        [DllImport("kernel32.dll", SetLastError=true)]
-        static extern int SetInformationJobObject(SafeFileHandle hJob, int JobObjectInfoClass, IntPtr lpJobObjectInfo, int cbJobObjectInfoLength);
-		
+		[DllImport("kernel32.dll", SetLastError = true)]
+		static extern int SetInformationJobObject(SafeFileHandle hJob, int jobObjectInfoClass, IntPtr lpJobObjectInfo, int cbJobObjectInfoLength);
+
 		[DllImport("kernel32.dll", SetLastError = true)]
 		static extern bool QueryInformationJobObject(SafeFileHandle hJob, int JobObjectInformationClass, ref JOBOBJECT_BASIC_ACCOUNTING_INFORMATION lpJobObjectInformation, int cbJobObjectInformationLength);
 
@@ -116,10 +116,7 @@ namespace EpicGames.Core
 		/// <summary>
 		/// Determines support for using job objects
 		/// </summary>
-		static internal bool SupportsJobObjects
-		{
-			get { return RuntimePlatform.IsWindows; }
-		}
+		internal static bool SupportsJobObjects => RuntimePlatform.IsWindows;
 
 		/// <summary>
 		/// Constructor
@@ -136,14 +133,14 @@ namespace EpicGames.Core
 				}
 
 				// Configure the job object to terminate the processes added to it when the handle is closed
-				JOBOBJECT_EXTENDED_LIMIT_INFORMATION LimitInformation = new JOBOBJECT_EXTENDED_LIMIT_INFORMATION();
-				LimitInformation.BasicLimitInformation.LimitFlags = JOB_OBJECT_LIMIT_KILL_ON_JOB_CLOSE | JOB_OBJECT_LIMIT_BREAKAWAY_OK;
+				JOBOBJECT_EXTENDED_LIMIT_INFORMATION limitInformation = new JOBOBJECT_EXTENDED_LIMIT_INFORMATION();
+				limitInformation.BasicLimitInformation.LimitFlags = JOB_OBJECT_LIMIT_KILL_ON_JOB_CLOSE | JOB_OBJECT_LIMIT_BREAKAWAY_OK;
 
-				int Length = Marshal.SizeOf(typeof(JOBOBJECT_EXTENDED_LIMIT_INFORMATION));
-				IntPtr LimitInformationPtr = Marshal.AllocHGlobal(Length);
-				Marshal.StructureToPtr(LimitInformation, LimitInformationPtr, false);
+				int length = Marshal.SizeOf(typeof(JOBOBJECT_EXTENDED_LIMIT_INFORMATION));
+				IntPtr limitInformationPtr = Marshal.AllocHGlobal(length);
+				Marshal.StructureToPtr(limitInformation, limitInformationPtr, false);
 
-				if(SetInformationJobObject(JobHandle, JobObjectExtendedLimitInformation, LimitInformationPtr, Length) == 0)
+				if(SetInformationJobObject(JobHandle, JobObjectExtendedLimitInformation, limitInformationPtr, length) == 0)
 				{
 					throw new Win32Exception();
 				}
@@ -159,13 +156,13 @@ namespace EpicGames.Core
 			{
 				if (SupportsJobObjects)
 				{
-					JOBOBJECT_BASIC_ACCOUNTING_INFORMATION AccountingInformation = new JOBOBJECT_BASIC_ACCOUNTING_INFORMATION();
-					if (QueryInformationJobObject(JobHandle!, JobObjectBasicAccountingInformation, ref AccountingInformation, Marshal.SizeOf(typeof(JOBOBJECT_BASIC_ACCOUNTING_INFORMATION))) == false)
+					JOBOBJECT_BASIC_ACCOUNTING_INFORMATION accountingInformation = new JOBOBJECT_BASIC_ACCOUNTING_INFORMATION();
+					if (QueryInformationJobObject(JobHandle!, JobObjectBasicAccountingInformation, ref accountingInformation, Marshal.SizeOf(typeof(JOBOBJECT_BASIC_ACCOUNTING_INFORMATION))) == false)
 					{
 						throw new Win32Exception();
 					}
 
-					return new TimeSpan((long)AccountingInformation.TotalUserTime + (long)AccountingInformation.TotalKernelTime);
+					return new TimeSpan((long)accountingInformation.TotalUserTime + (long)accountingInformation.TotalKernelTime);
 				}
 				else
 				{
@@ -174,10 +171,6 @@ namespace EpicGames.Core
 			}
 		}
 
-		public void Add(IntPtr ProcessHandle)
-		{
-		}
-
 		/// <summary>
 		/// Dispose of the process group
 		/// </summary>
@@ -198,6 +191,7 @@ namespace EpicGames.Core
 	/// </summary>
 	public sealed class ManagedProcess : IDisposable
 	{
+#pragma warning disable IDE1006 // Naming Styles
 		[StructLayout(LayoutKind.Sequential)]
 		class SECURITY_ATTRIBUTES
 		{
@@ -238,10 +232,10 @@ namespace EpicGames.Core
 			public SafeHandle? hStdError;
 		}
 
-		[DllImport("kernel32.dll", SetLastError=true)]
+		[DllImport("kernel32.dll", SetLastError = true)]
 		static extern int AssignProcessToJobObject(SafeFileHandle? hJob, IntPtr hProcess);
 
-		[DllImport("kernel32.dll", SetLastError=true)]
+		[DllImport("kernel32.dll", SetLastError = true)]
 		static extern int CloseHandle(IntPtr hObject);
 
 		[DllImport("kernel32.dll", SetLastError = true)]
@@ -267,40 +261,40 @@ namespace EpicGames.Core
 			ABOVE_NORMAL_PRIORITY_CLASS = 0x00008000,
 		}
 
-        [DllImport("kernel32.dll", SetLastError=true)]
+		[DllImport("kernel32.dll", SetLastError = true)]
         static extern int CreateProcess(/*[MarshalAs(UnmanagedType.LPTStr)]*/ string? lpApplicationName, StringBuilder lpCommandLine, IntPtr lpProcessAttributes, IntPtr lpThreadAttributes, bool bInheritHandles, ProcessCreationFlags dwCreationFlags, IntPtr lpEnvironment, /*[MarshalAs(UnmanagedType.LPTStr)]*/ string? lpCurrentDirectory, STARTUPINFO lpStartupInfo, PROCESS_INFORMATION lpProcessInformation);
 
-		[DllImport("kernel32.dll", SetLastError=true)]
+		[DllImport("kernel32.dll", SetLastError = true)]
 		static extern int ResumeThread(IntPtr hThread);
 
-		[DllImport("kernel32.dll", SetLastError=true)]
+		[DllImport("kernel32.dll", SetLastError = true)]
 		static extern IntPtr GetCurrentProcess();
 
         const int DUPLICATE_SAME_ACCESS = 2;
 
-        [DllImport("kernel32.dll", SetLastError=true)]    
-        static extern int DuplicateHandle(
-            IntPtr hSourceProcessHandle,
-            SafeHandle hSourceHandle,
-            IntPtr hTargetProcess,
-            out SafeFileHandle targetHandle,
-            int dwDesiredAccess,
-            [MarshalAs(UnmanagedType.Bool)] bool bInheritHandle,
-            int dwOptions
-        );
+		[DllImport("kernel32.dll", SetLastError = true)]
+		static extern int DuplicateHandle(
+			IntPtr hSourceProcessHandle,
+			SafeHandle hSourceHandle,
+			IntPtr hTargetProcess,
+			out SafeFileHandle targetHandle,
+			int dwDesiredAccess,
+			[MarshalAs(UnmanagedType.Bool)] bool bInheritHandle,
+			int dwOptions
+		);
 
 		[DllImport("kernel32.dll", SetLastError = true)]
 		static extern int GetExitCodeProcess(SafeFileHandle hProcess, out int lpExitCode);
 
-		[DllImport ("kernel32.dll")]
-		static extern bool IsProcessInJob(IntPtr hProcess, IntPtr hJob, out bool Result);
+		[DllImport("kernel32.dll")]
+		static extern bool IsProcessInJob(IntPtr hProcess, IntPtr hJob, out bool result);
 
-		[DllImport("kernel32.dll", SetLastError=true)]
+		[DllImport("kernel32.dll", SetLastError = true)]
 		static extern int TerminateProcess(SafeHandleZeroOrMinusOneIsInvalid hProcess, uint uExitCode);
 
 		const UInt32 INFINITE = 0xFFFFFFFF;
 
-		[DllImport("kernel32.dll", SetLastError=true)]
+		[DllImport("kernel32.dll", SetLastError = true)]
 		static extern UInt32 WaitForSingleObject(SafeHandleZeroOrMinusOneIsInvalid hHandle, UInt32 dwMilliseconds);
 
 		[DllImport("kernel32.dll", SetLastError = true)]
@@ -316,7 +310,7 @@ namespace EpicGames.Core
 
 		[DllImport("kernel32.dll")]
 		static extern UInt32 GetActiveProcessorCount(UInt16 GroupNumber);
-		
+
 		[StructLayout(LayoutKind.Sequential)]
 		class GROUP_AFFINITY
 		{
@@ -330,19 +324,22 @@ namespace EpicGames.Core
 		[DllImport("kernel32.dll", SetLastError = true)]
 		static extern int SetThreadGroupAffinity(IntPtr hThread, GROUP_AFFINITY GroupAffinity, GROUP_AFFINITY? PreviousGroupAffinity);
 
+		const int ERROR_ACCESS_DENIED = 5;
+#pragma warning restore IDE1006 // Naming Styles
+
 		/// <summary>
 		/// Converts FILETIME to DateTime.
 		/// </summary>
-		/// <param name="Time">Input FILETIME structure</param>
+		/// <param name="time">Input FILETIME structure</param>
 		/// <returns>Converted DateTime</returns>
-		static DateTime FileTimeToDateTime(System.Runtime.InteropServices.ComTypes.FILETIME Time)
+		static DateTime FileTimeToDateTime(System.Runtime.InteropServices.ComTypes.FILETIME time)
 		{
-			ulong High = (ulong)Time.dwHighDateTime;
-			uint Low = (uint)Time.dwLowDateTime;
-			long FileTime = (long)((High << 32) + Low);
+			ulong high = (ulong)time.dwHighDateTime;
+			uint low = (uint)time.dwLowDateTime;
+			long fileTime = (long)((high << 32) + low);
 			try
 			{
-				return DateTime.FromFileTimeUtc(FileTime);
+				return DateTime.FromFileTimeUtc(fileTime);
 			}
 			catch (ArgumentOutOfRangeException)
 			{
@@ -350,8 +347,6 @@ namespace EpicGames.Core
 			}
 		}
 
-		const int ERROR_ACCESS_DENIED = 5;
-
 		/// <summary>
 		/// The process id
 		/// </summary>
@@ -360,22 +355,22 @@ namespace EpicGames.Core
 		/// <summary>
 		/// Handle for the child process.
 		/// </summary>
-		SafeFileHandle? ProcessHandle;
+		SafeFileHandle? _processHandle;
 
 		/// <summary>
 		/// The write end of the child process' stdin pipe.
 		/// </summary>
-		SafeFileHandle? StdInWrite;
+		SafeFileHandle? _stdInWrite;
 
 		/// <summary>
 		/// The read end of the child process' stdout pipe.
 		/// </summary>
-		SafeFileHandle? StdOutRead;
+		SafeFileHandle? _stdOutRead;
 
 		/// <summary>
 		/// The read end of the child process' stderr pipe.
 		/// </summary>
-		SafeFileHandle? StdErrRead;
+		SafeFileHandle? _stdErrRead;
 
 		/// <summary>
 		/// Input stream for the child process.
@@ -405,62 +400,62 @@ namespace EpicGames.Core
 		/// <summary>
 		/// Standard process implementation for non-Windows platforms.
 		/// </summary>
-		Process? FrameworkProcess;
+		Process? _frameworkProcess;
 
 		/// <summary>
 		/// Task to read from stdout and stderr
 		/// </summary>
-		Task? FrameworkOutputTask;
+		Task? _frameworkOutputTask;
 
 		/// <summary>
 		/// Merged output & error write stream for the framework child process. We use a channel rather than a pipe here to support cancellation, which is not supported by async pipes on Windows.
 		/// </summary>
-		Channel<ChannelBuffer>? FrameworkChannel;
+		Channel<ChannelBuffer>? _frameworkChannel;
 
 		/// <summary>
 		/// Pool of buffer objects output from the framework
 		/// </summary>
-		ObjectPool<ChannelBuffer>? FrameworkBufferPool;
+		ObjectPool<ChannelBuffer>? _frameworkBufferPool;
 
 		/// <summary>
 		/// Cancellation token for background threads
 		/// </summary>
-		CancellationTokenSource? FrameworkCancellationSource;
+		CancellationTokenSource? _frameworkCancellationSource;
 
 		/// <summary>
 		/// Static lock object. This is used to synchronize the creation of child processes - in particular, the inheritance of stdout/stderr write pipes. If processes
 		/// inherit pipes meant for other processes, they won't be closed until both terminate.
 		/// </summary>
-		static object LockObject = new object();
+		static readonly object s_lockObject = new object();
 
 		/// <summary>
 		/// Used to perform CPU usage and resource accounting of all children process involved in a single compilation unit.
 		/// </summary>
-		ManagedProcessGroup? AccountingProcessGroup;
+		ManagedProcessGroup? _accountingProcessGroup;
 
-		byte[] Buffer = Array.Empty<byte>();
-		int BufferPos = 0;
-		int BufferLen = 0;
-		bool bCarriageReturn = false;
+		byte[] _buffer = Array.Empty<byte>();
+		int _bufferPos = 0;
+		int _bufferLen = 0;
+		bool _bCarriageReturn = false;
 
 		class ChannelBuffer
 		{
-			public byte[] Data = new byte[1024];
-			public int Length = 0;
+			public byte[] _data = new byte[1024];
+			public int _length = 0;
 		}
 
 		class ChannelReadStream : Stream
 		{
-			ChannelReader<ChannelBuffer> Reader;
-			ObjectPool<ChannelBuffer> BufferPool;
+			readonly ChannelReader<ChannelBuffer> _reader;
+			readonly ObjectPool<ChannelBuffer> _bufferPool;
 
-			ChannelBuffer? CurrentBuf = null;
-			int CurrentPos = 0;
+			ChannelBuffer? _currentBuf = null;
+			int _currentPos = 0;
 
-			public ChannelReadStream(ChannelReader<ChannelBuffer> Reader, ObjectPool<ChannelBuffer> BufferPool)
+			public ChannelReadStream(ChannelReader<ChannelBuffer> reader, ObjectPool<ChannelBuffer> bufferPool)
 			{
-				this.Reader = Reader;
-				this.BufferPool = BufferPool;
+				_reader = reader;
+				_bufferPool = bufferPool;
 			}
 
 			public override bool CanRead => true;
@@ -473,57 +468,57 @@ namespace EpicGames.Core
 			public override void SetLength(long value) => throw new NotSupportedException();
 			public override void Write(byte[] buffer, int offset, int count) => throw new NotSupportedException();
 
-			public override int Read(byte[] Buffer, int Offset, int Count)
+			public override int Read(byte[] buffer, int offset, int count)
 			{
 				// Perform the read on the thread pool to avoid deadlocking the continuation pathway on WinForms apps
-				Task<int> ReadTask = Task.Run(async () => await ReadAsync(Buffer.AsMemory(Offset, Count), CancellationToken.None));
-				return ReadTask.GetAwaiter().GetResult();
+				Task<int> readTask = Task.Run(async () => await ReadAsync(buffer.AsMemory(offset, count), CancellationToken.None));
+				return readTask.GetAwaiter().GetResult();
 			}
 
-			public override async ValueTask<int> ReadAsync(Memory<byte> Buffer, CancellationToken CancellationToken)
+			public override async ValueTask<int> ReadAsync(Memory<byte> buffer, CancellationToken cancellationToken)
 			{
-				if (CurrentBuf != null && CurrentPos == CurrentBuf.Length)
+				if (_currentBuf != null && _currentPos == _currentBuf._length)
 				{
-					BufferPool.Return(CurrentBuf);
-					CurrentBuf = null;
-					CurrentPos = 0;
+					_bufferPool.Return(_currentBuf);
+					_currentBuf = null;
+					_currentPos = 0;
 				}
 
-				while (CurrentBuf == null)
+				while (_currentBuf == null)
 				{
-					if (!await Reader.WaitToReadAsync(CancellationToken))
+					if (!await _reader.WaitToReadAsync(cancellationToken))
 					{
 						return 0;
 					}
-					if (Reader.TryRead(out CurrentBuf))
+					if (_reader.TryRead(out _currentBuf))
 					{
 						break;
 					}
 				}
 
-				int ReadLen = Math.Min(CurrentBuf.Length - CurrentPos, Buffer.Length);
-				CurrentBuf.Data.AsSpan(CurrentPos, ReadLen).CopyTo(Buffer.Span);
-				CurrentPos += ReadLen;
-				return ReadLen;
+				int readLen = Math.Min(_currentBuf._length - _currentPos, buffer.Length);
+				_currentBuf._data.AsSpan(_currentPos, readLen).CopyTo(buffer.Span);
+				_currentPos += readLen;
+				return readLen;
 			}
 		}
 
 		/// <summary>
 		/// Spawns a new managed process.
 		/// </summary>
-		/// <param name="Group">The managed process group to add to</param>
-		/// <param name="FileName">Path to the executable to be run</param>
-		/// <param name="CommandLine">Command line arguments for the process</param>
-		/// <param name="WorkingDirectory">Working directory for the new process. May be null to use the current working directory.</param>
-		/// <param name="Environment">Environment variables for the new process. May be null, in which case the current process' environment is inherited</param>
-		/// <param name="Input">Text to be passed via stdin to the new process. May be null.</param>
-		/// <param name="Priority">Priority for the child process</param>
-		public ManagedProcess(ManagedProcessGroup? Group, string FileName, string CommandLine, string? WorkingDirectory, IReadOnlyDictionary<string, string>? Environment, byte[]? Input, ProcessPriorityClass Priority, ManagedProcessFlags Flags = ManagedProcessFlags.MergeOutputPipes)
-			: this(Group, FileName, CommandLine, WorkingDirectory, Environment, Priority, Flags)
+		/// <param name="group">The managed process group to add to</param>
+		/// <param name="fileName">Path to the executable to be run</param>
+		/// <param name="commandLine">Command line arguments for the process</param>
+		/// <param name="workingDirectory">Working directory for the new process. May be null to use the current working directory.</param>
+		/// <param name="environment">Environment variables for the new process. May be null, in which case the current process' environment is inherited</param>
+		/// <param name="input">Text to be passed via stdin to the new process. May be null.</param>
+		/// <param name="priority">Priority for the child process</param>
+		public ManagedProcess(ManagedProcessGroup? group, string fileName, string commandLine, string? workingDirectory, IReadOnlyDictionary<string, string>? environment, byte[]? input, ProcessPriorityClass priority, ManagedProcessFlags flags = ManagedProcessFlags.MergeOutputPipes)
+			: this(group, fileName, commandLine, workingDirectory, environment, priority, flags)
 		{
-			if (Input != null)
+			if (input != null)
 			{
-				StdIn.Write(Input, 0, Input.Length);
+				StdIn.Write(input, 0, input.Length);
 			}
 			StdIn.Close();
 		}
@@ -531,13 +526,13 @@ namespace EpicGames.Core
 		/// <summary>
 		/// Spawns a new managed process.
 		/// </summary>
-		/// <param name="Group">The managed process group to add to</param>
-		/// <param name="FileName">Path to the executable to be run</param>
-		/// <param name="CommandLine">Command line arguments for the process</param>
-		/// <param name="WorkingDirectory">Working directory for the new process. May be null to use the current working directory.</param>
-		/// <param name="Environment">Environment variables for the new process. May be null, in which case the current process' environment is inherited</param>
-		/// <param name="Priority">Priority for the child process</param>
-		public ManagedProcess(ManagedProcessGroup? Group, string FileName, string CommandLine, string? WorkingDirectory, IReadOnlyDictionary<string, string>? Environment, ProcessPriorityClass Priority, ManagedProcessFlags Flags = ManagedProcessFlags.MergeOutputPipes)
+		/// <param name="group">The managed process group to add to</param>
+		/// <param name="fileName">Path to the executable to be run</param>
+		/// <param name="commandLine">Command line arguments for the process</param>
+		/// <param name="workingDirectory">Working directory for the new process. May be null to use the current working directory.</param>
+		/// <param name="environment">Environment variables for the new process. May be null, in which case the current process' environment is inherited</param>
+		/// <param name="priority">Priority for the child process</param>
+		public ManagedProcess(ManagedProcessGroup? group, string fileName, string commandLine, string? workingDirectory, IReadOnlyDictionary<string, string>? environment, ProcessPriorityClass priority, ManagedProcessFlags flags = ManagedProcessFlags.MergeOutputPipes)
 		{
 			// Create the child process
 			// NOTE: Child process must be created in a separate method to avoid stomping exception callstacks (https://stackoverflow.com/a/2494150)
@@ -545,196 +540,196 @@ namespace EpicGames.Core
 			{
 				if (ManagedProcessGroup.SupportsJobObjects)
 				{
-					CreateManagedProcessWin32(Group, FileName, CommandLine, WorkingDirectory, Environment, Priority, Flags);
+					CreateManagedProcessWin32(group, fileName, commandLine, workingDirectory, environment, priority, flags);
 				}
 				else
 				{
-					CreateManagedProcessPortable(Group, FileName, CommandLine, WorkingDirectory, Environment, Priority, Flags);
+					CreateManagedProcessPortable(fileName, commandLine, workingDirectory, environment, priority, flags);
 				}
 			}
-			catch (Exception Ex)
+			catch (Exception ex)
 			{
-				ExceptionUtils.AddContext(Ex, "while launching {0} {1}", FileName, CommandLine);
+				ExceptionUtils.AddContext(ex, "while launching {0} {1}", fileName, commandLine);
 				throw;
 			}
 		}
 
-		static readonly int ProcessorGroupCount = GetActiveProcessorGroupCount();
+		static readonly int s_processorGroupCount = GetActiveProcessorGroupCount();
 
 		/// <summary>
 		/// A counter used to decide which processor group the next process should be assigned to
 		/// </summary>
-		static int ProcessorGroupCounter = 0;
+		static int s_processorGroupCounter = 0;
 
 		/// <summary>
 		/// Spawns a new managed process using Win32 native functions. 
 		/// </summary>
-		/// <param name="Group">The managed process group to add to</param>
-		/// <param name="FileName">Path to the executable to be run</param>
-		/// <param name="CommandLine">Command line arguments for the process</param>
-		/// <param name="WorkingDirectory">Working directory for the new process. May be null to use the current working directory.</param>
-		/// <param name="Environment">Environment variables for the new process. May be null, in which case the current process' environment is inherited</param>
-		/// <param name="Priority">Priority for the child process</param>
-		/// <param name="ManagedFlags">Flags controlling how the new process is created</param>
-		private void CreateManagedProcessWin32(ManagedProcessGroup? Group, string FileName, string CommandLine, string? WorkingDirectory, IReadOnlyDictionary<string, string>? Environment, ProcessPriorityClass Priority, ManagedProcessFlags ManagedFlags)
+		/// <param name="group">The managed process group to add to</param>
+		/// <param name="fileName">Path to the executable to be run</param>
+		/// <param name="commandLine">Command line arguments for the process</param>
+		/// <param name="workingDirectory">Working directory for the new process. May be null to use the current working directory.</param>
+		/// <param name="environment">Environment variables for the new process. May be null, in which case the current process' environment is inherited</param>
+		/// <param name="priority">Priority for the child process</param>
+		/// <param name="managedFlags">Flags controlling how the new process is created</param>
+		private void CreateManagedProcessWin32(ManagedProcessGroup? group, string fileName, string commandLine, string? workingDirectory, IReadOnlyDictionary<string, string>? environment, ProcessPriorityClass priority, ManagedProcessFlags managedFlags)
 		{
-			IntPtr EnvironmentBlock = IntPtr.Zero;
+			IntPtr environmentBlock = IntPtr.Zero;
 			try
 			{
 				// Create the environment block for the child process, if necessary.
-				if (Environment != null)
+				if (environment != null)
 				{
 					// The native format for the environment block is a sequence of null terminated strings with a final null terminator.
-					List<byte> EnvironmentBytes = new List<byte>();
-					foreach (KeyValuePair<string, string> Pair in Environment)
+					List<byte> environmentBytes = new List<byte>();
+					foreach (KeyValuePair<string, string> pair in environment)
 					{
-						EnvironmentBytes.AddRange(Console.OutputEncoding.GetBytes(Pair.Key));
-						EnvironmentBytes.Add((byte)'=');
-						EnvironmentBytes.AddRange(Console.OutputEncoding.GetBytes(Pair.Value));
-						EnvironmentBytes.Add((byte)0);
+						environmentBytes.AddRange(Console.OutputEncoding.GetBytes(pair.Key));
+						environmentBytes.Add((byte)'=');
+						environmentBytes.AddRange(Console.OutputEncoding.GetBytes(pair.Value));
+						environmentBytes.Add((byte)0);
 					}
-					EnvironmentBytes.Add((byte)0);
+					environmentBytes.Add((byte)0);
 
 					// Allocate an unmanaged block of memory to store it.
-					EnvironmentBlock = Marshal.AllocHGlobal(EnvironmentBytes.Count);
-					Marshal.Copy(EnvironmentBytes.ToArray(), 0, EnvironmentBlock, EnvironmentBytes.Count);
+					environmentBlock = Marshal.AllocHGlobal(environmentBytes.Count);
+					Marshal.Copy(environmentBytes.ToArray(), 0, environmentBlock, environmentBytes.Count);
 				}
 
-				PROCESS_INFORMATION ProcessInfo = new PROCESS_INFORMATION();
+				PROCESS_INFORMATION processInfo = new PROCESS_INFORMATION();
 				try
 				{
 					// Get the flags to create the new process
-					ProcessCreationFlags Flags = ProcessCreationFlags.CREATE_NO_WINDOW | ProcessCreationFlags.CREATE_SUSPENDED;
-					switch (Priority)
+					ProcessCreationFlags flags = ProcessCreationFlags.CREATE_NO_WINDOW | ProcessCreationFlags.CREATE_SUSPENDED;
+					switch (priority)
 					{
 						case ProcessPriorityClass.Normal:
-							Flags |= ProcessCreationFlags.NORMAL_PRIORITY_CLASS;
+							flags |= ProcessCreationFlags.NORMAL_PRIORITY_CLASS;
 							break;
 						case ProcessPriorityClass.Idle:
-							Flags |= ProcessCreationFlags.IDLE_PRIORITY_CLASS;
+							flags |= ProcessCreationFlags.IDLE_PRIORITY_CLASS;
 							break;
 						case ProcessPriorityClass.High:
-							Flags |= ProcessCreationFlags.HIGH_PRIORITY_CLASS;
+							flags |= ProcessCreationFlags.HIGH_PRIORITY_CLASS;
 							break;
 						case ProcessPriorityClass.RealTime:
-							Flags |= ProcessCreationFlags.REALTIME_PRIORITY_CLASS;
+							flags |= ProcessCreationFlags.REALTIME_PRIORITY_CLASS;
 							break;
 						case ProcessPriorityClass.BelowNormal:
-							Flags |= ProcessCreationFlags.BELOW_NORMAL_PRIORITY_CLASS;
+							flags |= ProcessCreationFlags.BELOW_NORMAL_PRIORITY_CLASS;
 							break;
 						case ProcessPriorityClass.AboveNormal:
-							Flags |= ProcessCreationFlags.ABOVE_NORMAL_PRIORITY_CLASS;
+							flags |= ProcessCreationFlags.ABOVE_NORMAL_PRIORITY_CLASS;
 							break;
 					}
 
 					// Acquire a global lock before creating inheritable handles. If multiple threads create inheritable handles at the same time, and child processes will inherit them all.
 					// Since we need to wait for output pipes to be closed (in order to consume all output), this can result in output reads not returning until all processes with the same
 					// inherited handles are closed.
-					lock (LockObject)
+					lock (s_lockObject)
 					{
-						SafeFileHandle? StdInRead = null;
-						SafeFileHandle? StdOutWrite = null;
-						SafeFileHandle? StdErrWrite = null;
+						SafeFileHandle? stdInRead = null;
+						SafeFileHandle? stdOutWrite = null;
+						SafeFileHandle? stdErrWrite = null;
 						try
 						{
 							// Create stdin and stdout pipes for the child process. We'll close the handles for the child process' ends after it's been created.
-							SECURITY_ATTRIBUTES SecurityAttributes = new SECURITY_ATTRIBUTES();
-							SecurityAttributes.nLength = Marshal.SizeOf(SecurityAttributes);
-							SecurityAttributes.bInheritHandle = 1;
+							SECURITY_ATTRIBUTES securityAttributes = new SECURITY_ATTRIBUTES();
+							securityAttributes.nLength = Marshal.SizeOf(securityAttributes);
+							securityAttributes.bInheritHandle = 1;
 
-							if (CreatePipe(out StdInRead, out StdInWrite, SecurityAttributes, 4 * 1024) == 0 || SetHandleInformation(StdInWrite, HANDLE_FLAG_INHERIT, 0) == 0)
+							if (CreatePipe(out stdInRead, out _stdInWrite, securityAttributes, 4 * 1024) == 0 || SetHandleInformation(_stdInWrite, HANDLE_FLAG_INHERIT, 0) == 0)
 							{
 								throw new Win32ExceptionWithCode("Unable to create stdin pipe");
 							}
-							if (CreatePipe(out StdOutRead, out StdOutWrite, SecurityAttributes, 1024 * 1024) == 0 || SetHandleInformation(StdOutRead, HANDLE_FLAG_INHERIT, 0) == 0)
+							if (CreatePipe(out _stdOutRead, out stdOutWrite, securityAttributes, 1024 * 1024) == 0 || SetHandleInformation(_stdOutRead, HANDLE_FLAG_INHERIT, 0) == 0)
 							{
 								throw new Win32ExceptionWithCode("Unable to create stdout pipe");
 							}
 
-							if ((ManagedFlags & ManagedProcessFlags.MergeOutputPipes) != 0)
+							if ((managedFlags & ManagedProcessFlags.MergeOutputPipes) != 0)
 							{
-								if (DuplicateHandle(GetCurrentProcess(), StdOutWrite, GetCurrentProcess(), out StdErrWrite, 0, true, DUPLICATE_SAME_ACCESS) == 0)
+								if (DuplicateHandle(GetCurrentProcess(), stdOutWrite, GetCurrentProcess(), out stdErrWrite, 0, true, DUPLICATE_SAME_ACCESS) == 0)
 								{
 									throw new Win32ExceptionWithCode("Unable to duplicate stdout handle");
 								}
 							}
 							else
 							{
-								if (CreatePipe(out StdErrRead, out StdErrWrite, SecurityAttributes, 1024 * 1024) == 0 || SetHandleInformation(StdErrRead, HANDLE_FLAG_INHERIT, 0) == 0)
+								if (CreatePipe(out _stdErrRead, out stdErrWrite, securityAttributes, 1024 * 1024) == 0 || SetHandleInformation(_stdErrRead, HANDLE_FLAG_INHERIT, 0) == 0)
 								{
 									throw new Win32ExceptionWithCode("Unable to create stderr pipe");
 								}
 							}
 
 							// Create the new process as suspended, so we can modify it before it starts executing (and potentially preempting us)
-							STARTUPINFO StartupInfo = new STARTUPINFO();
-							StartupInfo.cb = Marshal.SizeOf(StartupInfo);
-							StartupInfo.hStdInput = StdInRead;
-							StartupInfo.hStdOutput = StdOutWrite;
-							StartupInfo.hStdError = StdErrWrite;
-							StartupInfo.dwFlags = STARTF_USESTDHANDLES;
+							STARTUPINFO startupInfo = new STARTUPINFO();
+							startupInfo.cb = Marshal.SizeOf(startupInfo);
+							startupInfo.hStdInput = stdInRead;
+							startupInfo.hStdOutput = stdOutWrite;
+							startupInfo.hStdError = stdErrWrite;
+							startupInfo.dwFlags = STARTF_USESTDHANDLES;
 
 							// Under heavy load (ie. spawning large number of processes, typically Clang) we see CreateProcess very occasionally failing with ERROR_ACCESS_DENIED.
-							int[] RetryDelay = { 100, 200, 1000, 5000 };
-							for(int AttemptIdx = 0; ; AttemptIdx++)
+							int[] retryDelay = { 100, 200, 1000, 5000 };
+							for(int attemptIdx = 0; ; attemptIdx++)
 							{
-								if (CreateProcess(null, new StringBuilder("\"" + FileName + "\" " + CommandLine), IntPtr.Zero, IntPtr.Zero, true, Flags, EnvironmentBlock, WorkingDirectory, StartupInfo, ProcessInfo) != 0)
+								if (CreateProcess(null, new StringBuilder("\"" + fileName + "\" " + commandLine), IntPtr.Zero, IntPtr.Zero, true, flags, environmentBlock, workingDirectory, startupInfo, processInfo) != 0)
 								{
 									break;
 								}
 
-								if (Marshal.GetLastWin32Error() != ERROR_ACCESS_DENIED || AttemptIdx >= RetryDelay.Length)
+								if (Marshal.GetLastWin32Error() != ERROR_ACCESS_DENIED || attemptIdx >= retryDelay.Length)
 								{
 									throw new Win32ExceptionWithCode("Unable to create process");
 								}
 
-								Thread.Sleep(RetryDelay[AttemptIdx]);
+								Thread.Sleep(retryDelay[attemptIdx]);
 							}
 
 							// Save the process id
-							Id = (int)ProcessInfo.dwProcessId;
+							Id = (int)processInfo.dwProcessId;
 						}
 						finally
 						{
 							// Close the write ends of the handle. We don't want any other process to be able to inherit these.
-							if (StdInRead != null)
+							if (stdInRead != null)
 							{
-								StdInRead.Dispose();
-								StdInRead = null;
+								stdInRead.Dispose();
+								stdInRead = null;
 							}
-							if (StdOutWrite != null)
+							if (stdOutWrite != null)
 							{
-								StdOutWrite.Dispose();
-								StdOutWrite = null;
+								stdOutWrite.Dispose();
+								stdOutWrite = null;
 							}
-							if (StdErrWrite != null)
+							if (stdErrWrite != null)
 							{
-								StdErrWrite.Dispose();
-								StdErrWrite = null;
+								stdErrWrite.Dispose();
+								stdErrWrite = null;
 							}
 						}
 					}
 
 					// Add it to our job object
-					if (Group != null && AssignProcessToJobObject(Group.JobHandle, ProcessInfo.hProcess) == 0)
+					if (group != null && AssignProcessToJobObject(group.JobHandle, processInfo.hProcess) == 0)
 					{
 						// Support for nested job objects was only added in Windows 8; prior to that, assigning processes to job objects would fail. Figure out if we're already in a job, and ignore the error if we are.
-						int OriginalError = Marshal.GetLastWin32Error();
+						int originalError = Marshal.GetLastWin32Error();
 
 						bool bProcessInJob;
 						IsProcessInJob(GetCurrentProcess(), IntPtr.Zero, out bProcessInJob);
 
 						if (!bProcessInJob)
 						{
-							throw new Win32ExceptionWithCode(OriginalError, "Unable to assign process to job object");
+							throw new Win32ExceptionWithCode(originalError, "Unable to assign process to job object");
 						}
 					}
 
 					// Create a JobObject for each spawned process to do CPU usage accounting of spawned process and all its children
-					AccountingProcessGroup = new ManagedProcessGroup();
-					if (AccountingProcessGroup != null)
+					_accountingProcessGroup = new ManagedProcessGroup();
+					if (_accountingProcessGroup != null)
 					{
-						if (AssignProcessToJobObject(AccountingProcessGroup.JobHandle, ProcessInfo.hProcess) == 0) 
+						if (AssignProcessToJobObject(_accountingProcessGroup.JobHandle, processInfo.hProcess) == 0) 
 						{
 							throw new Win32Exception();
 						}
@@ -753,36 +748,36 @@ namespace EpicGames.Core
 					//
 					// To reduce the chance of pathological process scheduling, we will explicitly distribute processes to each process
 					// group. Doing so has been observed to reduce overall compile times for large builds by as much as 10%.
-					if (ProcessorGroupCount > 1)
+					if (s_processorGroupCount > 1)
 					{
-						ushort ProcessorGroup = (ushort)(Interlocked.Increment(ref ProcessorGroupCounter) % ProcessorGroupCount);
+						ushort processorGroup = (ushort)(Interlocked.Increment(ref s_processorGroupCounter) % s_processorGroupCount);
 
-						uint GroupProcessorCount = GetActiveProcessorCount(ProcessorGroup);
+						uint groupProcessorCount = GetActiveProcessorCount(processorGroup);
 
-						GROUP_AFFINITY GroupAffinity = new GROUP_AFFINITY();
-						GroupAffinity.Mask = ~0ul >> (int)(64 - GroupProcessorCount);
-						GroupAffinity.Group = ProcessorGroup;
-						SetThreadGroupAffinity(ProcessInfo.hThread, GroupAffinity, null);
+						GROUP_AFFINITY groupAffinity = new GROUP_AFFINITY();
+						groupAffinity.Mask = ~0ul >> (int)(64 - groupProcessorCount);
+						groupAffinity.Group = processorGroup;
+						SetThreadGroupAffinity(processInfo.hThread, groupAffinity, null);
 					}
 					
 
 					// Allow the thread to start running
-					if (ResumeThread(ProcessInfo.hThread) == -1)
+					if (ResumeThread(processInfo.hThread) == -1)
 					{
 						throw new Win32ExceptionWithCode("Unable to resume thread in child process");
 					}
 
 					// If we have any input text, write it to stdin now
-					StdIn = new FileStream(StdInWrite, FileAccess.Write, 4096, false);
+					StdIn = new FileStream(_stdInWrite, FileAccess.Write, 4096, false);
 
 					// Create the stream objects for reading the process output
-					StdOut = new FileStream(StdOutRead, FileAccess.Read, 4096, false);
+					StdOut = new FileStream(_stdOutRead, FileAccess.Read, 4096, false);
 					StdOutText = new StreamReader(StdOut, Console.OutputEncoding);
 
 					// Do the same for the stderr output
-					if (StdErrRead != null)
+					if (_stdErrRead != null)
 					{
-						StdErr = new FileStream(StdErrRead, FileAccess.Read, 4096, false);
+						StdErr = new FileStream(_stdErrRead, FileAccess.Read, 4096, false);
 						StdErrText = new StreamReader(StdErr, Console.OutputEncoding);
 					}
 					else
@@ -792,26 +787,25 @@ namespace EpicGames.Core
 					}
 
 					// Wrap the process handle in a SafeFileHandle
-					ProcessHandle = new SafeFileHandle(ProcessInfo.hProcess, true);
+					_processHandle = new SafeFileHandle(processInfo.hProcess, true);
 				}
 				finally
 				{
-					if (ProcessInfo.hProcess != IntPtr.Zero && ProcessHandle == null)
+					if (processInfo.hProcess != IntPtr.Zero && _processHandle == null)
 					{
-						CloseHandle(ProcessInfo.hProcess);
+						CloseHandle(processInfo.hProcess);
 					}
-					if (ProcessInfo.hThread != IntPtr.Zero)
+					if (processInfo.hThread != IntPtr.Zero)
 					{
-						CloseHandle(ProcessInfo.hThread);
+						CloseHandle(processInfo.hThread);
 					}
 				}
 			}
 			finally
 			{
-				if (EnvironmentBlock != IntPtr.Zero)
+				if (environmentBlock != IntPtr.Zero)
 				{
-					Marshal.FreeHGlobal(EnvironmentBlock);
-					EnvironmentBlock = IntPtr.Zero;
+					Marshal.FreeHGlobal(environmentBlock);
 				}
 			}
 		}
@@ -819,105 +813,104 @@ namespace EpicGames.Core
 		/// <summary>
 		/// Spawns a new managed process using Win32 native functions. 
 		/// </summary>
-		/// <param name="Group">The managed process group to add to</param>
-		/// <param name="FileName">Path to the executable to be run</param>
-		/// <param name="CommandLine">Command line arguments for the process</param>
-		/// <param name="WorkingDirectory">Working directory for the new process. May be null to use the current working directory.</param>
-		/// <param name="Environment">Environment variables for the new process. May be null, in which case the current process' environment is inherited</param>
-		/// <param name="Priority">Priority for the child process</param>
-		/// <param name="Flags">Flags for the new process</param>
-		private void CreateManagedProcessPortable(ManagedProcessGroup? Group, string FileName, string CommandLine, string? WorkingDirectory, IReadOnlyDictionary<string, string>? Environment, ProcessPriorityClass Priority, ManagedProcessFlags ManagedFlags)
+		/// <param name="fileName">Path to the executable to be run</param>
+		/// <param name="commandLine">Command line arguments for the process</param>
+		/// <param name="workingDirectory">Working directory for the new process. May be null to use the current working directory.</param>
+		/// <param name="environment">Environment variables for the new process. May be null, in which case the current process' environment is inherited</param>
+		/// <param name="priority">Priority for the child process</param>
+		/// <param name="flags">Flags for the new process</param>
+		private void CreateManagedProcessPortable(string fileName, string commandLine, string? workingDirectory, IReadOnlyDictionary<string, string>? environment, ProcessPriorityClass priority, ManagedProcessFlags flags)
 		{
 			// TODO: Process Arguments follow windows conventions in .NET Core
 			// Which means single quotes ' are not considered quotes.
 			// see [REDACTED_PROFILE_URL]/runtime/issues/29857
 			// also see UE-102580
 			// for rules see https://docs.microsoft.com/en-us/cpp/cpp/main-function-command-line-args
-			CommandLine = CommandLine.Replace('\'', '\"');
+			commandLine = commandLine.Replace('\'', '\"');
 
 			// for non-Windows platforms
-			FrameworkProcess = new Process();
-			FrameworkProcess.StartInfo.FileName = FileName;
-			FrameworkProcess.StartInfo.Arguments = CommandLine;
-			FrameworkProcess.StartInfo.WorkingDirectory = WorkingDirectory;
-			FrameworkProcess.StartInfo.RedirectStandardInput = true;
-			FrameworkProcess.StartInfo.RedirectStandardOutput = true;
-			FrameworkProcess.StartInfo.RedirectStandardError = true;
-			FrameworkProcess.StartInfo.UseShellExecute = false;
-			FrameworkProcess.StartInfo.CreateNoWindow = true;
-			FrameworkProcess.Exited += (sender, e) =>
+			_frameworkProcess = new Process();
+			_frameworkProcess.StartInfo.FileName = fileName;
+			_frameworkProcess.StartInfo.Arguments = commandLine;
+			_frameworkProcess.StartInfo.WorkingDirectory = workingDirectory;
+			_frameworkProcess.StartInfo.RedirectStandardInput = true;
+			_frameworkProcess.StartInfo.RedirectStandardOutput = true;
+			_frameworkProcess.StartInfo.RedirectStandardError = true;
+			_frameworkProcess.StartInfo.UseShellExecute = false;
+			_frameworkProcess.StartInfo.CreateNoWindow = true;
+			_frameworkProcess.Exited += (sender, e) =>
 			{
-				FrameworkExitTime = DateTime.Now;
+				_frameworkExitTime = DateTime.Now;
 			};
 
-			if (Environment != null)
+			if (environment != null)
 			{
-				foreach (KeyValuePair<string, string> Pair in Environment)
+				foreach (KeyValuePair<string, string> pair in environment)
 				{
-					FrameworkProcess.StartInfo.EnvironmentVariables[Pair.Key] = Pair.Value;
+					_frameworkProcess.StartInfo.EnvironmentVariables[pair.Key] = pair.Value;
 				}
 			}
 
-			FrameworkStartTime = FrameworkExitTime = DateTime.Now;
-			FrameworkProcess.Start();
+			_frameworkStartTime = _frameworkExitTime = DateTime.Now;
+			_frameworkProcess.Start();
 
 			try
 			{
-				FrameworkProcess.PriorityClass = Priority;
+				_frameworkProcess.PriorityClass = priority;
 			}
 			catch
 			{
 			}
 
-			Id = FrameworkProcess.Id;
-			StdIn = FrameworkProcess.StandardInput.BaseStream;
+			Id = _frameworkProcess.Id;
+			StdIn = _frameworkProcess.StandardInput.BaseStream;
 
-			if ((ManagedFlags & ManagedProcessFlags.MergeOutputPipes) != 0)
+			if ((flags & ManagedProcessFlags.MergeOutputPipes) != 0)
 			{
 				// AnonymousPipes block reading even if the stream has been fully read, until the writer pipe handle is closed.
-				FrameworkBufferPool = ObjectPool.Create<ChannelBuffer>();
-				FrameworkChannel = Channel.CreateBounded<ChannelBuffer>(new BoundedChannelOptions(128) { FullMode = BoundedChannelFullMode.Wait, SingleReader = true, SingleWriter = false });
-				FrameworkCancellationSource = new CancellationTokenSource();
+				_frameworkBufferPool = ObjectPool.Create<ChannelBuffer>();
+				_frameworkChannel = Channel.CreateBounded<ChannelBuffer>(new BoundedChannelOptions(128) { FullMode = BoundedChannelFullMode.Wait, SingleReader = true, SingleWriter = false });
+				_frameworkCancellationSource = new CancellationTokenSource();
 
-				StdOut = new ChannelReadStream(FrameworkChannel.Reader, FrameworkBufferPool);
+				StdOut = new ChannelReadStream(_frameworkChannel.Reader, _frameworkBufferPool);
 				StdOutText = new StreamReader(StdOut, Console.OutputEncoding);
 
 				StdErr = StdOut;
 				StdErrText = StdOutText;
 
-				CancellationToken CancellationToken = FrameworkCancellationSource.Token;
-				Task StdOutTask = Task.Run(() => CopyPipeAsync(FrameworkProcess.StandardOutput.BaseStream, FrameworkChannel, CancellationToken));
-				Task StdErrTask = Task.Run(() => CopyPipeAsync(FrameworkProcess.StandardError.BaseStream, FrameworkChannel, CancellationToken));
-				FrameworkOutputTask = Task.WhenAll(StdOutTask, StdErrTask).ContinueWith(x => FrameworkChannel.Writer.Complete());
+				CancellationToken cancellationToken = _frameworkCancellationSource.Token;
+				Task stdOutTask = Task.Run(() => CopyPipeAsync(_frameworkProcess.StandardOutput.BaseStream, _frameworkChannel, cancellationToken));
+				Task stdErrTask = Task.Run(() => CopyPipeAsync(_frameworkProcess.StandardError.BaseStream, _frameworkChannel, cancellationToken));
+				_frameworkOutputTask = Task.WhenAll(stdOutTask, stdErrTask).ContinueWith(x => _frameworkChannel.Writer.Complete());
 			}
 			else
 			{
-				StdOut = FrameworkProcess.StandardOutput.BaseStream;
-				StdOutText = FrameworkProcess.StandardOutput;
-				StdErr = FrameworkProcess.StandardError.BaseStream;
-				StdErrText = FrameworkProcess.StandardError;
+				StdOut = _frameworkProcess.StandardOutput.BaseStream;
+				StdOutText = _frameworkProcess.StandardOutput;
+				StdErr = _frameworkProcess.StandardError.BaseStream;
+				StdErrText = _frameworkProcess.StandardError;
 			}
 		}
 
 		/// <summary>
 		/// Copy data from one pipe to another.
 		/// </summary>
-		/// <param name="Source"></param>
-		/// <param name="Target"></param>
-		async Task CopyPipeAsync(Stream Source, ChannelWriter<ChannelBuffer> Target, CancellationToken CancellationToken)
+		/// <param name="source"></param>
+		/// <param name="target"></param>
+		async Task CopyPipeAsync(Stream source, ChannelWriter<ChannelBuffer> target, CancellationToken cancellationToken)
 		{
 			for (; ; )
 			{
-				ChannelBuffer Buffer = FrameworkBufferPool!.Get();
-				int ReadLen = await Source.ReadAsync(Buffer.Data, 0, Buffer.Data.Length, CancellationToken);
+				ChannelBuffer buffer = _frameworkBufferPool!.Get();
+				int readLen = await source.ReadAsync(buffer._data, 0, buffer._data.Length, cancellationToken);
 
-				if (ReadLen == 0)
+				if (readLen == 0)
 				{
 					break;
 				}
 
-				Buffer.Length = ReadLen;
-				await Target.WriteAsync(Buffer, CancellationToken);
+				buffer._length = readLen;
+				await target.WriteAsync(buffer, cancellationToken);
 			}
 		}
 
@@ -926,23 +919,23 @@ namespace EpicGames.Core
 		/// </summary>
 		public void Dispose()
 		{
-			if(ProcessHandle != null)
+			if(_processHandle != null)
 			{
-				TerminateProcess(ProcessHandle, 0);
-				WaitForSingleObject(ProcessHandle, INFINITE);
+				TerminateProcess(_processHandle, 0);
+				WaitForSingleObject(_processHandle, INFINITE);
 
-				ProcessHandle.Dispose();
-				ProcessHandle = null;
+				_processHandle.Dispose();
+				_processHandle = null;
 			}
-			if(StdInWrite != null)
+			if(_stdInWrite != null)
 			{
-				StdInWrite.Dispose();
-				StdInWrite = null;
+				_stdInWrite.Dispose();
+				_stdInWrite = null;
 			}
-			if(StdOutRead != null)
+			if(_stdOutRead != null)
 			{
-				StdOutRead.Dispose();
-				StdOutRead = null;
+				_stdOutRead.Dispose();
+				_stdOutRead = null;
 			}
 
 			StdIn?.Dispose();
@@ -950,103 +943,117 @@ namespace EpicGames.Core
 			StdOutText?.Dispose();
 			StdErr?.Dispose();
 			StdErrText?.Dispose();
-			AccountingProcessGroup?.Dispose();
+			_accountingProcessGroup?.Dispose();
 
-			if(FrameworkProcess != null)
+			if(_frameworkProcess != null)
 			{
-				FrameworkCancellationSource?.Cancel();
+				_frameworkCancellationSource?.Cancel();
 
-				if (!FrameworkProcess.HasExited)
+				if (!_frameworkProcess.HasExited)
 				{
-					try { FrameworkProcess.Kill(true); } catch { } // Kill entire process tree
-					try { FrameworkProcess.WaitForExit(); } catch { }
+					// Kill entire process tree
+					try
+					{
+						_frameworkProcess.Kill(true);
+					}
+					catch { }
+
+					try
+					{
+						_frameworkProcess.WaitForExit();
+					}
+					catch { }
 				}
 
-				FrameworkProcess.Dispose();
-				FrameworkProcess = null;
+				_frameworkProcess.Dispose();
+				_frameworkProcess = null;
 
-				FrameworkOutputTask?.ContinueWith(x => FrameworkCancellationSource?.Dispose());
+				_frameworkOutputTask?.ContinueWith(x => _frameworkCancellationSource?.Dispose());
 			}
 		}
 
 		/// <summary>
 		/// Reads data from the process output
 		/// </summary>
-		/// <param name="Buffer">The buffer to receive the data</param>
-		/// <param name="Offset">Offset within the buffer to write to</param>
-		/// <param name="Count">Maximum number of bytes to read</param>
+		/// <param name="buffer">The buffer to receive the data</param>
+		/// <param name="offset">Offset within the buffer to write to</param>
+		/// <param name="count">Maximum number of bytes to read</param>
 		/// <returns>Number of bytes read</returns>
-		public int Read(byte[] Buffer, int Offset, int Count)
+		public int Read(byte[] buffer, int offset, int count)
 		{
 			// Fill the buffer, reentering managed code every 20ms to allow thread abort exceptions to be thrown
-			Task<int> ReadTask = StdOut!.ReadAsync(Buffer, Offset, Count);
-			while(!ReadTask.Wait(20))
+			Task<int> readTask = StdOut!.ReadAsync(buffer, offset, count);
+			while(!readTask.Wait(20))
 			{
 				// Spin through managed code to allow things like ThreadAbortExceptions to be thrown.
 			}
-			return ReadTask.Result;
+			return readTask.Result;
 		}
 
 		/// <summary>
 		/// Reads data from the process output
 		/// </summary>
-		/// <param name="Buffer">The buffer to receive the data</param>
-		/// <param name="Offset">Offset within the buffer to write to</param>
-		/// <param name="Count">Maximum number of bytes to read</param>
+		/// <param name="buffer">The buffer to receive the data</param>
+		/// <param name="offset">Offset within the buffer to write to</param>
+		/// <param name="count">Maximum number of bytes to read</param>
 		/// <returns>Number of bytes read</returns>
-		public async ValueTask<int> ReadAsync(byte[] Buffer, int Offset, int Count, CancellationToken CancellationToken)
+		public async ValueTask<int> ReadAsync(byte[] buffer, int offset, int count, CancellationToken cancellationToken)
 		{
-			return await StdOut!.ReadAsync(Buffer, Offset, Count, CancellationToken);
+			return await StdOut!.ReadAsync(buffer, offset, count, cancellationToken);
 		}
 
 		/// <summary>
 		/// Copy the process output to the given stream
 		/// </summary>
-		/// <param name="OutputStream">The output stream</param>
-		/// <param name="CancellationToken">Cancellation token</param>
+		/// <param name="outputStream">The output stream</param>
+		/// <param name="cancellationToken">Cancellation token</param>
 		/// <returns></returns>
-		public Task CopyToAsync(Stream OutputStream, CancellationToken CancellationToken)
+		public Task CopyToAsync(Stream outputStream, CancellationToken cancellationToken)
 		{
-			return StdOut!.CopyToAsync(OutputStream, CancellationToken);
+			return StdOut!.CopyToAsync(outputStream, cancellationToken);
 		}
 
 		/// <summary>
 		/// Copy the process output to the given stream
 		/// </summary>
-		/// <param name="WriteOutput">The output stream</param>
-		/// <param name="CancellationToken">Cancellation token</param>
+		/// <param name="writeOutput">The output stream</param>
+		/// <param name="cancellationToken">Cancellation token</param>
 		/// <returns></returns>
-		public Task CopyToAsync(Action<byte[], int, int> WriteOutput, int BufferSize, CancellationToken CancellationToken)
+		public Task CopyToAsync(Action<byte[], int, int> writeOutput, int bufferSize, CancellationToken cancellationToken)
 		{
-			Func<byte[], int, int, CancellationToken, Task> WriteOutputAsync = (Buffer, Offset, Length, CancellationToken) => { WriteOutput(Buffer, Offset, Length); return Task.CompletedTask; };
-			return CopyToAsync(WriteOutputAsync, BufferSize, CancellationToken);
+			Task WriteOutputAsync(byte[] buffer, int offset, int length, CancellationToken cancellationToken)
+			{ 
+				writeOutput(buffer, offset, length); 
+				return Task.CompletedTask;
+			}
+			return CopyToAsync(WriteOutputAsync, bufferSize, cancellationToken);
 		}
 
 		/// <summary>
 		/// Copy the process output to the given stream
 		/// </summary>
-		/// <param name="WriteOutputAsync">The output stream</param>
-		/// <param name="CancellationToken">Cancellation token</param>
+		/// <param name="writeOutputAsync">The output stream</param>
+		/// <param name="cancellationToken">Cancellation token</param>
 		/// <returns></returns>
-		public async Task CopyToAsync(Func<byte[], int, int, CancellationToken, Task> WriteOutputAsync, int BufferSize, CancellationToken CancellationToken)
+		public async Task CopyToAsync(Func<byte[], int, int, CancellationToken, Task> writeOutputAsync, int bufferSize, CancellationToken cancellationToken)
 		{
-			TaskCompletionSource<bool> TaskCompletionSource = new TaskCompletionSource<bool>();
-			using (CancellationTokenRegistration Registration = CancellationToken.Register(() => TaskCompletionSource.SetResult(false)))
+			TaskCompletionSource<bool> taskCompletionSource = new TaskCompletionSource<bool>();
+			using (CancellationTokenRegistration registration = cancellationToken.Register(() => taskCompletionSource.SetResult(false)))
 			{
-				byte[] Buffer = new byte[BufferSize];
+				byte[] buffer = new byte[bufferSize];
 				for (; ; )
 				{
-					Task<int> ReadTask = StdOut.ReadAsync(Buffer, 0, BufferSize, CancellationToken);
+					Task<int> readTask = StdOut.ReadAsync(buffer, 0, bufferSize, cancellationToken);
 
-					Task CompletedTask = await Task.WhenAny(ReadTask, TaskCompletionSource.Task);
-					CancellationToken.ThrowIfCancellationRequested();
+					Task completedTask = await Task.WhenAny(readTask, taskCompletionSource.Task);
+					cancellationToken.ThrowIfCancellationRequested();
 
-					int Bytes = await ReadTask;
-					if (Bytes == 0)
+					int bytes = await readTask;
+					if (bytes == 0)
 					{
 						break;
 					}
-					await WriteOutputAsync(Buffer, 0, Bytes, CancellationToken);
+					await writeOutputAsync(buffer, 0, bytes, cancellationToken);
 				}
 			}
 		}
@@ -1054,60 +1061,60 @@ namespace EpicGames.Core
 		/// <summary>
 		/// Reads an individual line from the process
 		/// </summary>
-		/// <param name="CancellationToken"></param>
+		/// <param name="cancellationToken"></param>
 		/// <returns></returns>
-		public async ValueTask<string?> ReadLineAsync(CancellationToken CancellationToken = default)
+		public async ValueTask<string?> ReadLineAsync(CancellationToken cancellationToken = default)
 		{
 			for (; ; )
 			{
 				// Skip any '\n' that's part of a '\r\n' sequence
-				if (bCarriageReturn && BufferPos < BufferLen)
+				if (_bCarriageReturn && _bufferPos < _bufferLen)
 				{
-					if (Buffer[BufferPos] == '\n')
+					if (_buffer[_bufferPos] == '\n')
 					{
-						BufferPos++;
+						_bufferPos++;
 					}
-					bCarriageReturn = false;
+					_bCarriageReturn = false;
 				}
 
 				// Pull out the first complete output line
-				for (int Idx = BufferPos; Idx < BufferLen; Idx++)
+				for (int idx = _bufferPos; idx < _bufferLen; idx++)
 				{
-					byte Character = Buffer[Idx];
-					if (Character == '\r' || Character == '\n')
+					byte character = _buffer[idx];
+					if (character == '\r' || character == '\n')
 					{
-						string Line = Console.OutputEncoding.GetString(Buffer, BufferPos, Idx - BufferPos);
-						bCarriageReturn = (Character == '\r');
-						BufferPos = Idx + 1;
-						return Line;
+						string line = Console.OutputEncoding.GetString(_buffer, _bufferPos, idx - _bufferPos);
+						_bCarriageReturn = (character == '\r');
+						_bufferPos = idx + 1;
+						return line;
 					}
 				}
 
 				// Create some space in the output buffer
-				if (BufferPos > 0)
+				if (_bufferPos > 0)
 				{
-					Array.Copy(Buffer, BufferPos, Buffer, 0, BufferLen - BufferPos);
-					BufferLen -= BufferPos;
-					BufferPos = 0;
+					Array.Copy(_buffer, _bufferPos, _buffer, 0, _bufferLen - _bufferPos);
+					_bufferLen -= _bufferPos;
+					_bufferPos = 0;
 				}
-				else if (BufferLen == Buffer.Length)
+				else if (_bufferLen == _buffer.Length)
 				{
-					Array.Resize(ref Buffer, Buffer.Length + (32 * 1024));
+					Array.Resize(ref _buffer, _buffer.Length + (32 * 1024));
 				}
 
 				// Read more data into the buffer
-				int NumBytesRead = await ReadAsync(Buffer, BufferPos, Buffer.Length - BufferPos, CancellationToken);
-				if (NumBytesRead == 0)
+				int numBytesRead = await ReadAsync(_buffer, _bufferPos, _buffer.Length - _bufferPos, cancellationToken);
+				if (numBytesRead == 0)
 				{
-					string? Line = null;
-					if (BufferPos < BufferLen)
+					string? line = null;
+					if (_bufferPos < _bufferLen)
 					{
-						Line = Console.OutputEncoding.GetString(Buffer, BufferPos, BufferLen - BufferPos);
-						BufferPos = BufferLen;
+						line = Console.OutputEncoding.GetString(_buffer, _bufferPos, _bufferLen - _bufferPos);
+						_bufferPos = _bufferLen;
 					}
-					return Line;
+					return line;
 				}
-				BufferLen += NumBytesRead;
+				_bufferLen += numBytesRead;
 			}
 		}
 
@@ -1119,53 +1126,53 @@ namespace EpicGames.Core
 		{
 			// Manually read all the output lines from the stream. Using ReadToEndAsync() or ReadAsync() on the StreamReader is abysmally slow, especially
 			// for high-volume processes. Manually picking out the lines via a buffered ReadAsync() call was found to be 6x faster on 'p4 -ztag have' calls.
-			List<string> OutputLines = new List<string>();
-			byte[] Buffer = new byte[32 * 1024];
-			byte LastCharacter = 0;
-			int NumBytesInBuffer = 0;
+			List<string> outputLines = new List<string>();
+			byte[] buffer = new byte[32 * 1024];
+			byte lastCharacter = 0;
+			int numBytesInBuffer = 0;
 			for(;;)
 			{
 				// If we're got a single line larger than 32kb (!), enlarge the buffer to ensure we can handle it
-				if(NumBytesInBuffer == Buffer.Length)
+				if(numBytesInBuffer == buffer.Length)
 				{
-					Array.Resize(ref Buffer, Buffer.Length + 32 * 1024);
+					Array.Resize(ref buffer, buffer.Length + 32 * 1024);
 				}
 
 				// Fill the buffer, reentering managed code every 20ms to allow thread abort exceptions to be thrown
-				int NumBytesRead = Read(Buffer, NumBytesInBuffer, Buffer.Length - NumBytesInBuffer);
-				if(NumBytesRead == 0)
+				int numBytesRead = Read(buffer, numBytesInBuffer, buffer.Length - numBytesInBuffer);
+				if(numBytesRead == 0)
 				{
-					if(NumBytesInBuffer > 0)
+					if(numBytesInBuffer > 0)
 					{
-						OutputLines.Add(Console.OutputEncoding.GetString(Buffer, 0, NumBytesInBuffer));
+						outputLines.Add(Console.OutputEncoding.GetString(buffer, 0, numBytesInBuffer));
 					}
 					break;
 				}
 
 				// Otherwise append to the existing buffer
-				NumBytesInBuffer += NumBytesRead;
+				numBytesInBuffer += numBytesRead;
 
 				// Pull out all the complete output lines
-				int LastStartIdx = 0;
-				for(int Idx = 0; Idx < NumBytesInBuffer; Idx++)
+				int lastStartIdx = 0;
+				for(int idx = 0; idx < numBytesInBuffer; idx++)
 				{
-					if(Buffer[Idx] == '\r' || Buffer[Idx] == '\n')
+					if(buffer[idx] == '\r' || buffer[idx] == '\n')
 					{
-						if(Buffer[Idx] != '\n' || LastCharacter != '\r')
+						if(buffer[idx] != '\n' || lastCharacter != '\r')
 						{
-							OutputLines.Add(Console.OutputEncoding.GetString(Buffer, LastStartIdx, Idx - LastStartIdx));
+							outputLines.Add(Console.OutputEncoding.GetString(buffer, lastStartIdx, idx - lastStartIdx));
 						}
-						LastStartIdx = Idx + 1;
+						lastStartIdx = idx + 1;
 					}
-					LastCharacter = Buffer[Idx];
+					lastCharacter = buffer[idx];
 				}
 
 				// Shuffle everything back to the start of the buffer
-				Array.Copy(Buffer, LastStartIdx, Buffer, 0, Buffer.Length - LastStartIdx);
-				NumBytesInBuffer -= LastStartIdx;
+				Array.Copy(buffer, lastStartIdx, buffer, 0, buffer.Length - lastStartIdx);
+				numBytesInBuffer -= lastStartIdx;
 			}
 			WaitForExit();
-			return OutputLines;
+			return outputLines;
 		}
 
 		/// <summary>
@@ -1181,9 +1188,9 @@ namespace EpicGames.Core
 		{
 			get
 			{
-				if (AccountingProcessGroup != null)
+				if (_accountingProcessGroup != null)
 				{
-					return AccountingProcessGroup.TotalProcessorTime;
+					return _accountingProcessGroup.TotalProcessorTime;
 				}
 				else
 				{
@@ -1195,31 +1202,31 @@ namespace EpicGames.Core
 		/// <summary>
 		/// Block until the process outputs a line of text, or terminates.
 		/// </summary>
-		/// <param name="Line">Variable to receive the output line</param>
-		/// <param name="Token">Cancellation token which can be used to abort the read operation.</param>
+		/// <param name="line">Variable to receive the output line</param>
+		/// <param name="token">Cancellation token which can be used to abort the read operation.</param>
 		/// <returns>True if a line was read, false if the process terminated without writing another line, or the cancellation token was signaled.</returns>
-		public bool TryReadLine([NotNullWhen(true)] out string? Line, CancellationToken? Token = null)
+		public bool TryReadLine([NotNullWhen(true)] out string? line, CancellationToken? token = null)
 		{
 			try
 			{
 				// Busy wait for the ReadLine call to finish, so we can get interrupted by thread abort exceptions
-				Task<string?> ReadLineTask = StdOutText!.ReadLineAsync();
+				Task<string?> readLineTask = StdOutText!.ReadLineAsync();
 				for (;;)
 				{
 					const int MillisecondsTimeout = 20;
-					if (Token.HasValue 
-							? ReadLineTask.Wait(MillisecondsTimeout, Token.Value) 
-							: ReadLineTask.Wait(MillisecondsTimeout))
+					if (token.HasValue 
+							? readLineTask.Wait(MillisecondsTimeout, token.Value) 
+							: readLineTask.Wait(MillisecondsTimeout))
 					{
-						Line = ReadLineTask.IsCompleted ? ReadLineTask.Result : null;
-						return Line != null;
+						line = readLineTask.IsCompleted ? readLineTask.Result : null;
+						return line != null;
 					}
 				}
 			}
 			catch (OperationCanceledException)
 			{
 				// If the cancel token is signalled, just return false.
-				Line = null;
+				line = null;
 				return false;
 			}
 		}
@@ -1230,13 +1237,13 @@ namespace EpicGames.Core
 		/// <returns>Exit code of the process</returns>
 		public void WaitForExit()
 		{
-			if (FrameworkProcess == null)
+			if (_frameworkProcess == null)
 			{
-				WaitForSingleObject(ProcessHandle!, INFINITE);
+				WaitForSingleObject(_processHandle!, INFINITE);
 			}
 			else
 			{
-				FrameworkProcess.WaitForExit();
+				_frameworkProcess.WaitForExit();
 			}
 		}
 
@@ -1247,24 +1254,24 @@ namespace EpicGames.Core
 		{
 			get
 			{
-				if(FrameworkProcess == null)
+				if(_frameworkProcess == null)
 				{
-					int Value;
-					if(GetExitCodeProcess(ProcessHandle!, out Value) == 0)
+					int value;
+					if(GetExitCodeProcess(_processHandle!, out value) == 0)
 					{
 						throw new Win32Exception();
 					}
-					return Value;
+					return value;
 				}
 				else
 				{
-					return FrameworkProcess.ExitCode;
+					return _frameworkProcess.ExitCode;
 				}
 			}
 		}
 
-		private DateTime FrameworkStartTime = DateTime.MinValue;
-		private DateTime FrameworkExitTime = DateTime.MinValue;
+		private DateTime _frameworkStartTime = DateTime.MinValue;
+		private DateTime _frameworkExitTime = DateTime.MinValue;
 
 		/// <summary>
 		/// The creation time of the process.
@@ -1273,18 +1280,18 @@ namespace EpicGames.Core
 		{
 			get
 			{
-				if (FrameworkProcess == null)
+				if (_frameworkProcess == null)
 				{
-					System.Runtime.InteropServices.ComTypes.FILETIME CreationTime;
-					if (!GetProcessTimes(ProcessHandle!, out CreationTime, out _, out _, out _))
+					System.Runtime.InteropServices.ComTypes.FILETIME creationTime;
+					if (!GetProcessTimes(_processHandle!, out creationTime, out _, out _, out _))
 					{
 						throw new Win32Exception();
 					}
-					return FileTimeToDateTime(CreationTime);
+					return FileTimeToDateTime(creationTime);
 				}
 				else
 				{
-					return FrameworkStartTime;
+					return _frameworkStartTime;
 				}
 			}
 		}
@@ -1296,18 +1303,18 @@ namespace EpicGames.Core
 		{
 			get
 			{
-				if (FrameworkProcess == null)
+				if (_frameworkProcess == null)
 				{
-					System.Runtime.InteropServices.ComTypes.FILETIME ExitTime;
-					if (!GetProcessTimes(ProcessHandle!, out _, out ExitTime, out _, out _))
+					System.Runtime.InteropServices.ComTypes.FILETIME exitTime;
+					if (!GetProcessTimes(_processHandle!, out _, out exitTime, out _, out _))
 					{
 						throw new Win32Exception();
 					}
-					return FileTimeToDateTime(ExitTime);
+					return FileTimeToDateTime(exitTime);
 				}
 				else
 				{
-					return FrameworkExitTime;
+					return _frameworkExitTime;
 				}
 			}
 		}
diff --git a/Engine/Source/Programs/Shared/EpicGames.Core/Md5Hash.cs b/Engine/Source/Programs/Shared/EpicGames.Core/Md5Hash.cs
--- a/Engine/Source/Programs/Shared/EpicGames.Core/Md5Hash.cs
+++ b/Engine/Source/Programs/Shared/EpicGames.Core/Md5Hash.cs
@@ -2,12 +2,10 @@
 
 using System;
 using System.Buffers.Binary;
-using System.Collections.Generic;
 using System.ComponentModel;
 using System.Globalization;
 using System.IO;
 using System.Security.Cryptography;
-using System.Text;
 
 namespace EpicGames.Core
 {
@@ -30,7 +28,7 @@ namespace EpicGames.Core
 		/// <summary>
 		/// Memory storing the digest data
 		/// </summary>
-		public ReadOnlyMemory<byte> Memory;
+		public ReadOnlyMemory<byte> Memory { get; }
 
 		/// <summary>
 		/// Span for the underlying memory
@@ -45,87 +43,87 @@ namespace EpicGames.Core
 		/// <summary>
 		/// Constructor
 		/// </summary>
-		/// <param name="Memory">Memory to construct from</param>
-		public Md5Hash(ReadOnlyMemory<byte> Memory)
+		/// <param name="memory">Memory to construct from</param>
+		public Md5Hash(ReadOnlyMemory<byte> memory)
 		{
-			if (Memory.Length != NumBytes)
+			if (memory.Length != NumBytes)
 			{
 				throw new ArgumentException($"Md5Hash must be {NumBytes} bytes long");
 			}
 
-			this.Memory = Memory;
+			Memory = memory;
 		}
 
 		/// <summary>
 		/// Creates a content hash for a block of data, using a given algorithm.
 		/// </summary>
-		/// <param name="Data">Data to compute the hash for</param>
+		/// <param name="data">Data to compute the hash for</param>
 		/// <returns>New content hash instance containing the hash of the data</returns>
-		public static Md5Hash Compute(ReadOnlySpan<byte> Data)
+		public static Md5Hash Compute(ReadOnlySpan<byte> data)
 		{
-			byte[] Output = new byte[NumBytes];
-			using (MD5 Hasher = MD5.Create())
+			byte[] output = new byte[NumBytes];
+			using (MD5 hasher = MD5.Create())
 			{
-				Hasher.TryComputeHash(Data, Output, out _);
+				hasher.TryComputeHash(data, output, out _);
 			}
-			return new Md5Hash(Output);
+			return new Md5Hash(output);
 		}
 
 		/// <summary>
 		/// Creates a content hash for the input Stream object
 		/// </summary>
-		/// <param name="Stream">The Stream object to compoute the has for</param>
+		/// <param name="stream">The Stream object to compoute the has for</param>
 		/// <returns>New content hash instance containing the hash of the data</returns>
-		public static Md5Hash Compute(Stream Stream)
+		public static Md5Hash Compute(Stream stream)
 		{
-			using (MD5 Hasher = MD5.Create())
+			using (MD5 hasher = MD5.Create())
 			{
-				return new Md5Hash(Hasher.ComputeHash(Stream));
+				return new Md5Hash(hasher.ComputeHash(stream));
 			}
 		}
 
 		/// <summary>
 		/// Parses a digest from the given hex string
 		/// </summary>
-		/// <param name="Text"></param>
+		/// <param name="text"></param>
 		/// <returns></returns>
-		public static Md5Hash Parse(string Text)
+		public static Md5Hash Parse(string text)
 		{
-			return new Md5Hash(StringUtils.ParseHexString(Text));
+			return new Md5Hash(StringUtils.ParseHexString(text));
 		}
 
 		/// <summary>
 		/// Parses a digest from the given hex string
 		/// </summary>
-		/// <param name="Text"></param>
+		/// <param name="text"></param>
 		/// <returns></returns>
-		public static Md5Hash Parse(Utf8String Text)
+		public static Md5Hash Parse(Utf8String text)
 		{
-			return new Md5Hash(StringUtils.ParseHexString(Text.Span));
+			return new Md5Hash(StringUtils.ParseHexString(text.Span));
 		}
 
 		/// <inheritdoc cref="IComparable{T}.CompareTo(T)"/>
-		public int CompareTo(Md5Hash Other)
+		public int CompareTo(Md5Hash other)
 		{
-			ReadOnlySpan<byte> A = Span;
-			ReadOnlySpan<byte> B = Other.Span;
+			ReadOnlySpan<byte> a = Span;
+			ReadOnlySpan<byte> b = other.Span;
 
-			for (int Idx = 0; Idx < A.Length && Idx < B.Length; Idx++)
+			for (int idx = 0; idx < a.Length && idx < b.Length; idx++)
 			{
-				int Compare = A[Idx] - B[Idx];
-				if (Compare != 0)
+				int compare = a[idx] - b[idx];
+				if (compare != 0)
 				{
-					return Compare;
+					return compare;
 				}
 			}
-			return A.Length - B.Length;
+			return a.Length - b.Length;
 		}
 
 		/// <inheritdoc/>
-		public bool Equals(Md5Hash Other) => Span.SequenceEqual(Other.Span);
+		public bool Equals(Md5Hash other) => Span.SequenceEqual(other.Span);
 
 		/// <inheritdoc/>
-		public override bool Equals(object? Obj) => (Obj is Md5Hash Hash) && Hash.Span.SequenceEqual(Span);
+		public override bool Equals(object? obj) => (obj is Md5Hash hash) && hash.Span.SequenceEqual(Span);
 
 		/// <inheritdoc/>
 		public override int GetHashCode() => BinaryPrimitives.ReadInt32LittleEndian(Span);
@@ -136,32 +134,32 @@ namespace EpicGames.Core
 		/// <summary>
 		/// Test two hash values for equality
 		/// </summary>
-		public static bool operator ==(Md5Hash A, Md5Hash B) => A.Span.SequenceEqual(B.Span);
+		public static bool operator ==(Md5Hash a, Md5Hash b) => a.Span.SequenceEqual(b.Span);
 
 		/// <summary>
 		/// Test two hash values for equality
 		/// </summary>
-		public static bool operator !=(Md5Hash A, Md5Hash B) => !(A == B);
+		public static bool operator !=(Md5Hash a, Md5Hash b) => !(a == b);
 
 		/// <summary>
 		/// Tests whether A > B
 		/// </summary>
-		public static bool operator >(Md5Hash A, Md5Hash B) => A.CompareTo(B) > 0;
+		public static bool operator >(Md5Hash a, Md5Hash b) => a.CompareTo(b) > 0;
 
 		/// <summary>
 		/// Tests whether A is less than B
 		/// </summary>
-		public static bool operator <(Md5Hash A, Md5Hash B) => A.CompareTo(B) < 0;
+		public static bool operator <(Md5Hash a, Md5Hash b) => a.CompareTo(b) < 0;
 
 		/// <summary>
 		/// Tests whether A is greater than or equal to B
 		/// </summary>
-		public static bool operator >=(Md5Hash A, Md5Hash B) => A.CompareTo(B) >= 0;
+		public static bool operator >=(Md5Hash a, Md5Hash b) => a.CompareTo(b) >= 0;
 
 		/// <summary>
 		/// Tests whether A is less than or equal to B
 		/// </summary>
-		public static bool operator <=(Md5Hash A, Md5Hash B) => A.CompareTo(B) <= 0;
+		public static bool operator <=(Md5Hash a, Md5Hash b) => a.CompareTo(b) <= 0;
 	}
 
 	/// <summary>
@@ -172,21 +170,21 @@ namespace EpicGames.Core
 		/// <summary>
 		/// Read an <see cref="Md5Hash"/> from a memory reader
 		/// </summary>
-		/// <param name="Reader"></param>
+		/// <param name="reader"></param>
 		/// <returns></returns>
-		public static Md5Hash ReadMd5Hash(this MemoryReader Reader)
+		public static Md5Hash ReadMd5Hash(this MemoryReader reader)
 		{
-			return new Md5Hash(Reader.ReadFixedLengthBytes(Md5Hash.NumBytes));
+			return new Md5Hash(reader.ReadFixedLengthBytes(Md5Hash.NumBytes));
 		}
 
 		/// <summary>
 		/// Write an <see cref="Md5Hash"/> to a memory writer
 		/// </summary>
-		/// <param name="Writer"></param>
-		/// <param name="Hash"></param>
-		public static void WriteMd5Hash(this MemoryWriter Writer, Md5Hash Hash)
+		/// <param name="writer"></param>
+		/// <param name="hash"></param>
+		public static void WriteMd5Hash(this MemoryWriter writer, Md5Hash hash)
 		{
-			Writer.WriteFixedLengthBytes(Hash.Span);
+			writer.WriteFixedLengthBytes(hash.Span);
 		}
 	}
 
@@ -196,15 +194,15 @@ namespace EpicGames.Core
 	sealed class Md5HashTypeConverter : TypeConverter
 	{
 		/// <inheritdoc/>
-		public override bool CanConvertFrom(ITypeDescriptorContext Context, Type SourceType)
+		public override bool CanConvertFrom(ITypeDescriptorContext context, Type sourceType)
 		{
-			return SourceType == typeof(string);
+			return sourceType == typeof(string);
 		}
 
 		/// <inheritdoc/>
-		public override object ConvertFrom(ITypeDescriptorContext Context, CultureInfo Culture, object Value)
+		public override object ConvertFrom(ITypeDescriptorContext context, CultureInfo culture, object value)
 		{
-			return Md5Hash.Parse((string)Value);
+			return Md5Hash.Parse((string)value);
 		}
 	}
 }
diff --git a/Engine/Source/Programs/Shared/EpicGames.Core/MemoryReader.cs b/Engine/Source/Programs/Shared/EpicGames.Core/MemoryReader.cs
--- a/Engine/Source/Programs/Shared/EpicGames.Core/MemoryReader.cs
+++ b/Engine/Source/Programs/Shared/EpicGames.Core/MemoryReader.cs
@@ -2,8 +2,6 @@
 
 using System;
 using System.Buffers.Binary;
-using System.Collections.Generic;
-using System.Text;
 
 namespace EpicGames.Core
 {
@@ -36,30 +34,27 @@ namespace EpicGames.Core
 		/// <summary>
 		/// Returns the memory at the current offset
 		/// </summary>
-		public ReadOnlySpan<byte> Span
-		{
-			get { return Memory.Span.Slice(Offset); }
-		}
+		public ReadOnlySpan<byte> Span => Memory.Span.Slice(Offset);
 
 		/// <summary>
 		/// Constructor
 		/// </summary>
-		/// <param name="Memory">The memory to read from</param>
-		public MemoryReader(ReadOnlyMemory<byte> Memory)
+		/// <param name="memory">The memory to read from</param>
+		public MemoryReader(ReadOnlyMemory<byte> memory)
 		{
-			this.Memory = Memory;
-			this.Offset = 0;
+			Memory = memory;
+			Offset = 0;
 		}
 
 		/// <summary>
 		/// Checks that we've used the exact buffer length
 		/// </summary>
-		/// <param name="ExpectedOffset">Expected offset within the output buffer</param>
-		public void CheckOffset(int ExpectedOffset)
+		/// <param name="expectedOffset">Expected offset within the output buffer</param>
+		public void CheckOffset(int expectedOffset)
 		{
-			if (Offset != ExpectedOffset)
+			if (Offset != expectedOffset)
 			{
-				throw new Exception($"Serialization is not at expected offset within the input buffer (expected {ExpectedOffset}, actual {Offset})");
+				throw new Exception($"Serialization is not at expected offset within the input buffer (expected {expectedOffset}, actual {Offset})");
 			}
 		}
 
@@ -78,9 +73,9 @@ namespace EpicGames.Core
 		/// <returns>The value read from the buffer</returns>
 		public sbyte ReadInt8()
 		{
-			sbyte Value = (sbyte)Memory.Span[Offset];
+			sbyte value = (sbyte)Memory.Span[Offset];
 			Offset++;
-			return Value;
+			return value;
 		}
 
 		/// <summary>
@@ -89,9 +84,9 @@ namespace EpicGames.Core
 		/// <returns>The value read from the buffer</returns>
 		public byte ReadUInt8()
 		{
-			byte Value = Memory.Span[Offset];
+			byte value = Memory.Span[Offset];
 			Offset++;
-			return Value;
+			return value;
 		}
 
 		/// <summary>
@@ -100,9 +95,9 @@ namespace EpicGames.Core
 		/// <returns>The value read from the buffer</returns>
 		public short ReadInt16()
 		{
-			short Value = BinaryPrimitives.ReadInt16LittleEndian(Memory.Span.Slice(Offset));
+			short value = BinaryPrimitives.ReadInt16LittleEndian(Memory.Span.Slice(Offset));
 			Offset += sizeof(short);
-			return Value;
+			return value;
 		}
 
 		/// <summary>
@@ -111,9 +106,9 @@ namespace EpicGames.Core
 		/// <returns>The value read from the buffer</returns>
 		public ushort ReadUInt16()
 		{
-			ushort Value = BinaryPrimitives.ReadUInt16LittleEndian(Memory.Span.Slice(Offset));
+			ushort value = BinaryPrimitives.ReadUInt16LittleEndian(Memory.Span.Slice(Offset));
 			Offset += sizeof(ushort);
-			return Value;
+			return value;
 		}
 
 		/// <summary>
@@ -122,9 +117,9 @@ namespace EpicGames.Core
 		/// <returns>The value that was read from the buffer</returns>
 		public int ReadInt32()
 		{
-			int Value = BinaryPrimitives.ReadInt32LittleEndian(Memory.Span.Slice(Offset));
+			int value = BinaryPrimitives.ReadInt32LittleEndian(Memory.Span.Slice(Offset));
 			Offset += sizeof(int);
-			return Value;
+			return value;
 		}
 
 		/// <summary>
@@ -133,9 +128,9 @@ namespace EpicGames.Core
 		/// <returns>The value that was read from the buffer</returns>
 		public uint ReadUInt32()
 		{
-			uint Value = BinaryPrimitives.ReadUInt32LittleEndian(Memory.Span.Slice(Offset));
+			uint value = BinaryPrimitives.ReadUInt32LittleEndian(Memory.Span.Slice(Offset));
 			Offset += sizeof(uint);
-			return Value;
+			return value;
 		}
 
 		/// <summary>
@@ -144,9 +139,9 @@ namespace EpicGames.Core
 		/// <returns>The value that was read from the buffer</returns>
 		public long ReadInt64()
 		{
-			long Value = BinaryPrimitives.ReadInt64LittleEndian(Memory.Span.Slice(Offset));
+			long value = BinaryPrimitives.ReadInt64LittleEndian(Memory.Span.Slice(Offset));
 			Offset += sizeof(long);
-			return Value;
+			return value;
 		}
 
 		/// <summary>
@@ -155,9 +150,9 @@ namespace EpicGames.Core
 		/// <returns>The value that was read from the buffer</returns>
 		public ulong ReadUInt64()
 		{
-			ulong Value = BinaryPrimitives.ReadUInt64LittleEndian(Memory.Span.Slice(Offset));
+			ulong value = BinaryPrimitives.ReadUInt64LittleEndian(Memory.Span.Slice(Offset));
 			Offset += sizeof(ulong);
-			return Value;
+			return value;
 		}
 
 		/// <summary>
@@ -166,45 +161,45 @@ namespace EpicGames.Core
 		/// <returns>Sequence of bytes</returns>
 		public ReadOnlyMemory<byte> ReadVariableLengthBytes()
 		{
-			int Length = ReadInt32();
-			return ReadFixedLengthBytes(Length);
+			int length = ReadInt32();
+			return ReadFixedLengthBytes(length);
 		}
 
 		/// <summary>
 		/// Reads a sequence of bytes from the buffer
 		/// </summary>
-		/// <param name="Length">Number of bytes to read</param>
+		/// <param name="length">Number of bytes to read</param>
 		/// <returns>Sequence of bytes</returns>
-		public ReadOnlyMemory<byte> ReadFixedLengthBytes(int Length)
+		public ReadOnlyMemory<byte> ReadFixedLengthBytes(int length)
 		{
-			ReadOnlyMemory<byte> Bytes = Memory.Slice(Offset, Length);
-			Offset += Length;
-			return Bytes;
+			ReadOnlyMemory<byte> bytes = Memory.Slice(Offset, length);
+			Offset += length;
+			return bytes;
 		}
 
 		/// <summary>
 		/// Reads a variable length array
 		/// </summary>
-		/// <param name="ReadItem">Delegate to write an individual item</param>
-		public T[] ReadVariableLengthArray<T>(Func<T> ReadItem)
+		/// <param name="readItem">Delegate to write an individual item</param>
+		public T[] ReadVariableLengthArray<T>(Func<T> readItem)
 		{
-			int Length = ReadInt32();
-			return ReadFixedLengthArray(Length, ReadItem);
+			int length = ReadInt32();
+			return ReadFixedLengthArray(length, readItem);
 		}
 
 		/// <summary>
 		/// Writes a fixed length array
 		/// </summary>
-		/// <param name="Length">Length of the array to read</param>
-		/// <param name="ReadItem">Delegate to read an individual item</param>
-		public T[] ReadFixedLengthArray<T>(int Length, Func<T> ReadItem)
+		/// <param name="length">Length of the array to read</param>
+		/// <param name="readItem">Delegate to read an individual item</param>
+		public T[] ReadFixedLengthArray<T>(int length, Func<T> readItem)
 		{
-			T[] Array = new T[Length];
-			for (int Idx = 0; Idx < Length; Idx++)
+			T[] array = new T[length];
+			for (int idx = 0; idx < length; idx++)
 			{
-				Array[Idx] = ReadItem();
+				array[idx] = readItem();
 			}
-			return Array;
+			return array;
 		}
 	}
 }
diff --git a/Engine/Source/Programs/Shared/EpicGames.Core/MemoryWriter.cs b/Engine/Source/Programs/Shared/EpicGames.Core/MemoryWriter.cs
--- a/Engine/Source/Programs/Shared/EpicGames.Core/MemoryWriter.cs
+++ b/Engine/Source/Programs/Shared/EpicGames.Core/MemoryWriter.cs
@@ -2,8 +2,6 @@
 
 using System;
 using System.Buffers.Binary;
-using System.Collections.Generic;
-using System.Text;
 
 namespace EpicGames.Core
 {
@@ -15,30 +13,30 @@ namespace EpicGames.Core
 		/// <summary>
 		/// The memory block to write to
 		/// </summary>
-		Memory<byte> Memory;
+		readonly Memory<byte> _memory;
 
 		/// <summary>
 		/// Current offset within the memory
 		/// </summary>
-		int Offset;
+		int _offset;
 
 		/// <summary>
 		/// Constructor
 		/// </summary>
-		/// <param name="Memory">Memory to write to</param>
-		public MemoryWriter(Memory<byte> Memory)
+		/// <param name="memory">Memory to write to</param>
+		public MemoryWriter(Memory<byte> memory)
 		{
-			this.Memory = Memory;
-			this.Offset = 0;
+			_memory = memory;
+			_offset = 0;
 		}
 
 		/// <summary>
 		/// Checks that we've used the exact buffer length
 		/// </summary>
-		/// <param name="ExpectedOffset">Expected offset within the output buffer</param>
-		public void CheckOffset(int ExpectedOffset)
+		/// <param name="expectedOffset">Expected offset within the output buffer</param>
+		public void CheckOffset(int expectedOffset)
 		{
-			if (Offset != ExpectedOffset)
+			if (_offset != expectedOffset)
 			{
 				throw new Exception("Serialization is not at expected offset within the output buffer");
 			}
@@ -47,146 +45,146 @@ namespace EpicGames.Core
 		/// <summary>
 		/// Writes a boolean to memory
 		/// </summary>
-		/// <param name="Value">Value to write</param>
-		public void WriteBoolean(bool Value)
+		/// <param name="value">Value to write</param>
+		public void WriteBoolean(bool value)
 		{
-			WriteUInt8(Value ? (byte)1 : (byte)0);
+			WriteUInt8(value ? (byte)1 : (byte)0);
 		}
 
 		/// <summary>
 		/// Writes a byte to memory
 		/// </summary>
-		/// <param name="Value">Value to write</param>
-		public void WriteInt8(sbyte Value)
+		/// <param name="value">Value to write</param>
+		public void WriteInt8(sbyte value)
 		{
-			Memory.Span[Offset] = (byte)Value;
-			Offset++;
+			_memory.Span[_offset] = (byte)value;
+			_offset++;
 		}
 
 		/// <summary>
 		/// Writes a byte to memory
 		/// </summary>
-		/// <param name="Value">Value to write</param>
-		public void WriteUInt8(byte Value)
+		/// <param name="value">Value to write</param>
+		public void WriteUInt8(byte value)
 		{
-			Memory.Span[Offset] = Value;
-			Offset++;
+			_memory.Span[_offset] = value;
+			_offset++;
 		}
 
 		/// <summary>
 		/// Writes an int16 to the memory
 		/// </summary>
-		/// <param name="Value">Value to write</param>
-		public void WriteInt16(short Value)
+		/// <param name="value">Value to write</param>
+		public void WriteInt16(short value)
 		{
-			BinaryPrimitives.WriteInt16LittleEndian(Memory.Span.Slice(Offset), Value);
-			Offset += sizeof(short);
+			BinaryPrimitives.WriteInt16LittleEndian(_memory.Span.Slice(_offset), value);
+			_offset += sizeof(short);
 		}
 
 		/// <summary>
 		/// Writes a uint16 to the memory
 		/// </summary>
-		/// <param name="Value">Value to write</param>
-		public void WriteUInt16(ushort Value)
+		/// <param name="value">Value to write</param>
+		public void WriteUInt16(ushort value)
 		{
-			BinaryPrimitives.WriteUInt16LittleEndian(Memory.Span.Slice(Offset), Value);
-			Offset += sizeof(ushort);
+			BinaryPrimitives.WriteUInt16LittleEndian(_memory.Span.Slice(_offset), value);
+			_offset += sizeof(ushort);
 		}
 
 		/// <summary>
 		/// Writes an int32 to the memory
 		/// </summary>
-		/// <param name="Value">Value to write</param>
-		public void WriteInt32(int Value)
+		/// <param name="value">Value to write</param>
+		public void WriteInt32(int value)
 		{
-			BinaryPrimitives.WriteInt32LittleEndian(Memory.Span.Slice(Offset), Value);
-			Offset += sizeof(int);
+			BinaryPrimitives.WriteInt32LittleEndian(_memory.Span.Slice(_offset), value);
+			_offset += sizeof(int);
 		}
 
 		/// <summary>
 		/// Writes a uint32 to the memory
 		/// </summary>
-		/// <param name="Value">Value to write</param>
-		public void WriteUInt32(uint Value)
+		/// <param name="value">Value to write</param>
+		public void WriteUInt32(uint value)
 		{
-			BinaryPrimitives.WriteUInt32LittleEndian(Memory.Span.Slice(Offset), Value);
-			Offset += sizeof(uint);
+			BinaryPrimitives.WriteUInt32LittleEndian(_memory.Span.Slice(_offset), value);
+			_offset += sizeof(uint);
 		}
 
 		/// <summary>
 		/// Writes an int64 to the memory
 		/// </summary>
-		/// <param name="Value">Value to write</param>
-		public void WriteInt64(long Value)
+		/// <param name="value">Value to write</param>
+		public void WriteInt64(long value)
 		{
-			BinaryPrimitives.WriteInt64LittleEndian(Memory.Span.Slice(Offset), Value);
-			Offset += sizeof(long);
+			BinaryPrimitives.WriteInt64LittleEndian(_memory.Span.Slice(_offset), value);
+			_offset += sizeof(long);
 		}
 
 		/// <summary>
 		/// Writes a uint64 to the memory
 		/// </summary>
-		/// <param name="Value">Value to write</param>
-		public void WriteUInt64(ulong Value)
+		/// <param name="value">Value to write</param>
+		public void WriteUInt64(ulong value)
 		{
-			BinaryPrimitives.WriteUInt64LittleEndian(Memory.Span.Slice(Offset), Value);
-			Offset += sizeof(ulong);
+			BinaryPrimitives.WriteUInt64LittleEndian(_memory.Span.Slice(_offset), value);
+			_offset += sizeof(ulong);
 		}
 
 		/// <summary>
 		/// Writes a variable length span of bytes
 		/// </summary>
-		/// <param name="Bytes">The bytes to write</param>
-		public void WriteVariableLengthBytes(ReadOnlySpan<byte> Bytes)
+		/// <param name="bytes">The bytes to write</param>
+		public void WriteVariableLengthBytes(ReadOnlySpan<byte> bytes)
 		{
-			WriteInt32(Bytes.Length);
-			WriteFixedLengthBytes(Bytes);
+			WriteInt32(bytes.Length);
+			WriteFixedLengthBytes(bytes);
 		}
 
 		/// <summary>
 		/// Write a fixed-length sequence of bytes to the buffer
 		/// </summary>
-		/// <param name="Bytes">The bytes to write</param>
-		public void WriteFixedLengthBytes(ReadOnlySpan<byte> Bytes)
+		/// <param name="bytes">The bytes to write</param>
+		public void WriteFixedLengthBytes(ReadOnlySpan<byte> bytes)
 		{
-			Bytes.CopyTo(Memory.Span.Slice(Offset));
-			Offset += Bytes.Length;
+			bytes.CopyTo(_memory.Span.Slice(_offset));
+			_offset += bytes.Length;
 		}
 
 		/// <summary>
 		/// Writes a variable length array
 		/// </summary>
-		/// <param name="Array">The array to write</param>
-		/// <param name="WriteItem">Delegate to write an individual item</param>
-		public void WriteVariableLengthArray<T>(T[] Array, Action<T> WriteItem)
+		/// <param name="array">The array to write</param>
+		/// <param name="writeItem">Delegate to write an individual item</param>
+		public void WriteVariableLengthArray<T>(T[] array, Action<T> writeItem)
 		{
-			WriteInt32(Array.Length);
-			WriteFixedLengthArray(Array, WriteItem);
+			WriteInt32(array.Length);
+			WriteFixedLengthArray(array, writeItem);
 		}
 
 		/// <summary>
 		/// Writes a fixed length array
 		/// </summary>
-		/// <param name="Array">The array to write</param>
-		/// <param name="WriteItem">Delegate to write an individual item</param>
-		public void WriteFixedLengthArray<T>(T[] Array, Action<T> WriteItem)
+		/// <param name="array">The array to write</param>
+		/// <param name="writeItem">Delegate to write an individual item</param>
+		public void WriteFixedLengthArray<T>(T[] array, Action<T> writeItem)
 		{
-			for (int Idx = 0; Idx < Array.Length; Idx++)
+			for (int idx = 0; idx < array.Length; idx++)
 			{
-				WriteItem(Array[Idx]);
+				writeItem(array[idx]);
 			}
 		}
 
 		/// <summary>
 		/// Allocate a writable span from the buffer
 		/// </summary>
-		/// <param name="Length">Length of the span to allocate</param>
+		/// <param name="length">Length of the span to allocate</param>
 		/// <returns>Span that can be written to</returns>
-		public Span<byte> AllocateSpan(int Length)
+		public Span<byte> AllocateSpan(int length)
 		{
-			Span<byte> Span = Memory.Span.Slice(Offset, Length);
-			Offset += Length;
-			return Span;
+			Span<byte> span = _memory.Span.Slice(_offset, length);
+			_offset += length;
+			return span;
 		}
 	}
 }
