Reformat EpicGames.Core according to standard coding conventions.

diff --git a/Engine/Source/Programs/Shared/EpicGames.Core/MessageTemplate.cs b/Engine/Source/Programs/Shared/EpicGames.Core/MessageTemplate.cs
--- a/Engine/Source/Programs/Shared/EpicGames.Core/MessageTemplate.cs
+++ b/Engine/Source/Programs/Shared/EpicGames.Core/MessageTemplate.cs
@@ -1,15 +1,9 @@
 // Copyright Epic Games, Inc. All Rights Reserved.
 
-using Microsoft.Extensions.Logging;
 using System;
-using System.Buffers;
-using System.Collections.Concurrent;
 using System.Collections.Generic;
-using System.Globalization;
 using System.Linq;
-using System.Reflection;
 using System.Text;
-using System.Text.Json;
 
 namespace EpicGames.Core
 {
@@ -27,101 +21,101 @@ namespace EpicGames.Core
 		/// Converts any positional properties to named properties
 		/// </summary>
 		/// <returns></returns>
-		public static Dictionary<string, object>? CreatePositionalProperties(string Format, IEnumerable<KeyValuePair<string, object>> Properties)
+		public static Dictionary<string, object>? CreatePositionalProperties(string format, IEnumerable<KeyValuePair<string, object>> properties)
 		{
-			Dictionary<string, object>? OutputProperties = null;
-			if (Properties != null)
+			Dictionary<string, object>? outputProperties = null;
+			if (properties != null)
 			{
-				OutputProperties = new Dictionary<string, object>(Properties);
+				outputProperties = new Dictionary<string, object>(properties);
 
-				List<(int, int)>? Offsets = ParsePropertyNames(Format);
-				if (Offsets != null)
+				List<(int, int)>? offsets = ParsePropertyNames(format);
+				if (offsets != null)
 				{
-					foreach ((int Offset, int Length) in Offsets)
+					foreach ((int offset, int length) in offsets)
 					{
-						ReadOnlySpan<char> Name = Format.AsSpan(Offset, Length);
-						if (int.TryParse(Name, out int Number))
+						ReadOnlySpan<char> name = format.AsSpan(offset, length);
+						if (int.TryParse(name, out int number))
 						{
-							OutputProperties[Name.ToString()] = Properties.ElementAtOrDefault(Number).Value;
+							outputProperties[name.ToString()] = properties.ElementAtOrDefault(number).Value;
 						}
 					}
 				}
 			}
-			return OutputProperties;
+			return outputProperties;
 		}
 
 		/// <summary>
 		/// Renders a format string
 		/// </summary>
-		/// <param name="Format">The format string</param>
-		/// <param name="Properties">Property values to embed</param>
+		/// <param name="format">The format string</param>
+		/// <param name="properties">Property values to embed</param>
 		/// <returns>The rendered string</returns>
-		public static string Render(string Format, IEnumerable<KeyValuePair<string, object?>>? Properties)
+		public static string Render(string format, IEnumerable<KeyValuePair<string, object?>>? properties)
 		{
-			StringBuilder Result = new StringBuilder();
-			Render(Format, Properties, Result);
-			return Result.ToString();
+			StringBuilder result = new StringBuilder();
+			Render(format, properties, result);
+			return result.ToString();
 		}
 
 		/// <summary>
 		/// Renders a format string to the end of a string builder
 		/// </summary>
-		/// <param name="Format">The format string to render</param>
-		/// <param name="Properties">Sequence of key/value properties</param>
-		/// <param name="Result">Buffer to append the rendered string to</param>
-		public static void Render(string Format, IEnumerable<KeyValuePair<string, object?>>? Properties, StringBuilder Result)
+		/// <param name="format">The format string to render</param>
+		/// <param name="properties">Sequence of key/value properties</param>
+		/// <param name="result">Buffer to append the rendered string to</param>
+		public static void Render(string format, IEnumerable<KeyValuePair<string, object?>>? properties, StringBuilder result)
 		{
-			int NextOffset = 0;
+			int nextOffset = 0;
 
-			List<(int, int)>? Names = ParsePropertyNames(Format);
-			if (Names != null)
+			List<(int, int)>? names = ParsePropertyNames(format);
+			if (names != null)
 			{
-				foreach((int Offset, int Length) in Names)
+				foreach((int offset, int length) in names)
 				{
-					object? Value;
-					if (Properties != null && TryGetPropertyValue(Format.AsSpan(Offset, Length), Properties, out Value))
+					object? value;
+					if (properties != null && TryGetPropertyValue(format.AsSpan(offset, length), properties, out value))
 					{
-						int StartOffset = Offset - 1;
-						if (Format[StartOffset] == '@' || Format[StartOffset] == '$')
+						int startOffset = offset - 1;
+						if (format[startOffset] == '@' || format[startOffset] == '$')
 						{
-							StartOffset--;
+							startOffset--;
 						}
 
-						Unescape(Format.AsSpan(NextOffset, StartOffset - NextOffset), Result);
-						Result.Append(Value?.ToString() ?? "null");
-						NextOffset = Offset + Length + 1;
+						Unescape(format.AsSpan(nextOffset, startOffset - nextOffset), result);
+						result.Append(value?.ToString() ?? "null");
+						nextOffset = offset + length + 1;
 					}
 				}
 			}
 
-			Unescape(Format.AsSpan(NextOffset, Format.Length - NextOffset), Result);
+			Unescape(format.AsSpan(nextOffset, format.Length - nextOffset), result);
 		}
 
 		/// <summary>
 		/// Escapes a string for use in a message template
 		/// </summary>
-		/// <param name="Text">Text to escape</param>
+		/// <param name="text">Text to escape</param>
 		/// <returns>The escaped string</returns>
-		public static string Escape(string Text)
+		public static string Escape(string text)
 		{
-			StringBuilder Result = new StringBuilder();
-			Escape(Text, Result);
-			return Result.ToString();
+			StringBuilder result = new StringBuilder();
+			Escape(text, result);
+			return result.ToString();
 		}
 
 		/// <summary>
 		/// Escapes a span of characters and appends the result to a string
 		/// </summary>
-		/// <param name="Text">Span of characters to escape</param>
-		/// <param name="Result">Buffer to receive the escaped string</param>
-		public static void Escape(ReadOnlySpan<char> Text, StringBuilder Result)
+		/// <param name="text">Span of characters to escape</param>
+		/// <param name="result">Buffer to receive the escaped string</param>
+		public static void Escape(ReadOnlySpan<char> text, StringBuilder result)
 		{
-			foreach(char Char in Text)
+			foreach(char character in text)
 			{
-				Result.Append(Char);
-				if (Char == '{' || Char == '}')
+				result.Append(character);
+				if (character == '{' || character == '}')
 				{
-					Result.Append(Char);
+					result.Append(character);
 				}
 			}
 		}
@@ -129,98 +123,98 @@ namespace EpicGames.Core
 		/// <summary>
 		/// Unescapes a string from a message template
 		/// </summary>
-		/// <param name="Text">The text to unescape</param>
+		/// <param name="text">The text to unescape</param>
 		/// <returns>The unescaped text</returns>
-		public static string Unescape(string Text)
+		public static string Unescape(string text)
 		{
-			StringBuilder Result = new StringBuilder();
-			Unescape(Text.AsSpan(), Result);
-			return Result.ToString();
+			StringBuilder result = new StringBuilder();
+			Unescape(text.AsSpan(), result);
+			return result.ToString();
 		}
 
 		/// <summary>
 		/// Unescape a string and append the result to a string builder
 		/// </summary>
-		/// <param name="Text">Text to unescape</param>
-		/// <param name="Result">Receives the unescaped text</param>
-		public static void Unescape(ReadOnlySpan<char> Text, StringBuilder Result)
+		/// <param name="text">Text to unescape</param>
+		/// <param name="result">Receives the unescaped text</param>
+		public static void Unescape(ReadOnlySpan<char> text, StringBuilder result)
 		{
-			char LastChar = '\0';
-			foreach (char Char in Text)
+			char lastChar = '\0';
+			foreach (char character in text)
 			{
-				if ((Char != '{' || Char != '}') || Char != LastChar)
+				if ((character != '{' || character != '}') || character != lastChar)
 				{
-					Result.Append(Char);
+					result.Append(character);
 				}
-				LastChar = Char;
+				lastChar = character;
 			}
 		}
 
 		/// <summary>
 		/// Finds locations of property names from the given format string
 		/// </summary>
-		/// <param name="Format">The format string to parse</param>
+		/// <param name="format">The format string to parse</param>
 		/// <returns>List of offset, length pairs for property names. Null if the string does not contain any property references.</returns>
-		public static List<(int, int)>? ParsePropertyNames(string Format)
+		public static List<(int, int)>? ParsePropertyNames(string format)
 		{
-			List<(int, int)>? Names = null;
-			for (int Idx = 0; Idx < Format.Length - 1; Idx++)
+			List<(int, int)>? names = null;
+			for (int idx = 0; idx < format.Length - 1; idx++)
 			{
-				if (Format[Idx] == '{')
+				if (format[idx] == '{')
 				{
-					if (Format[Idx + 1] == '{')
+					if (format[idx + 1] == '{')
 					{
-						Idx++;
+						idx++;
 					}
 					else
 					{
-						int StartIdx = Idx + 1;
+						int startIdx = idx + 1;
 
-						Idx = Format.IndexOf('}', StartIdx);
-						if (Idx == -1)
+						idx = format.IndexOf('}', startIdx);
+						if (idx == -1)
 						{
 							break;
 						}
-						if (Names == null)
+						if (names == null)
 						{
-							Names = new List<(int, int)>();
+							names = new List<(int, int)>();
 						}
 
-						Names.Add((StartIdx, Idx - StartIdx));
+						names.Add((startIdx, idx - startIdx));
 					}
 				}
 			}
-			return Names;
+			return names;
 		}
 
 		/// <summary>
 		/// Parse the ordered arguments into a dictionary of named properties
 		/// </summary>
-		/// <param name="Format">Format string</param>
-		/// <param name="Args">Argument list to parse</param>
+		/// <param name="format">Format string</param>
+		/// <param name="args">Argument list to parse</param>
 		/// <returns></returns>
-		public static void ParsePropertyValues(string Format, object[] Args, Dictionary<string, object> Properties)
+		public static void ParsePropertyValues(string format, object[] args, Dictionary<string, object> properties)
 		{
-			List<(int, int)>? Offsets = ParsePropertyNames(Format);
-			if (Offsets != null)
+			List<(int, int)>? offsets = ParsePropertyNames(format);
+			if (offsets != null)
 			{
-				for (int Idx = 0; Idx < Offsets.Count; Idx++)
+				for (int idx = 0; idx < offsets.Count; idx++)
 				{
-					string Name = Format.Substring(Offsets[Idx].Item1, Offsets[Idx].Item2);
+					string name = format.Substring(offsets[idx].Item1, offsets[idx].Item2);
 
-					int Number;
-					if (int.TryParse(Name, out Number))
+					int number;
+					if (int.TryParse(name, out number))
 					{
-						if (Number >= 0 && Number < Args.Length)
+						if (number >= 0 && number < args.Length)
 						{
-							Properties[Name] = Args[Number];
+							properties[name] = args[number];
 						}
 					}
 					else
 					{
-						if (Idx < Args.Length)
+						if (idx < args.Length)
 						{
-							Properties[Name] = Args[Idx];
+							properties[name] = args[idx];
 						}
 					}
 				}
@@ -230,39 +224,39 @@ namespace EpicGames.Core
 		/// <summary>
 		/// Attempts to get a named property value from the given dictionary
 		/// </summary>
-		/// <param name="Name">Name of the property</param>
-		/// <param name="Properties">Sequence of property name/value pairs</param>
-		/// <param name="Value">On success, receives the property value</param>
+		/// <param name="name">Name of the property</param>
+		/// <param name="properties">Sequence of property name/value pairs</param>
+		/// <param name="value">On success, receives the property value</param>
 		/// <returns>True if the property was found, false otherwise</returns>
-		public static bool TryGetPropertyValue(ReadOnlySpan<char> Name, IEnumerable<KeyValuePair<string, object?>> Properties, out object? Value)
+		public static bool TryGetPropertyValue(ReadOnlySpan<char> name, IEnumerable<KeyValuePair<string, object?>> properties, out object? value)
 		{
-			int Number;
-			if (int.TryParse(Name, System.Globalization.NumberStyles.Integer, null, out Number))
+			int number;
+			if (int.TryParse(name, System.Globalization.NumberStyles.Integer, null, out number))
 			{
-				foreach (KeyValuePair<string, object?> Property in Properties)
+				foreach (KeyValuePair<string, object?> property in properties)
 				{
-					if (Number == 0)
+					if (number == 0)
 					{
-						Value = Property.Value;
+						value = property.Value;
 						return true;
 					}
-					Number--;
+					number--;
 				}
 			}
 			else
 			{
-				foreach (KeyValuePair<string, object?> Property in Properties)
+				foreach (KeyValuePair<string, object?> property in properties)
 				{
-					ReadOnlySpan<char> ParameterName = Property.Key.AsSpan();
-					if (Name.Equals(ParameterName, StringComparison.Ordinal))
+					ReadOnlySpan<char> parameterName = property.Key.AsSpan();
+					if (name.Equals(parameterName, StringComparison.Ordinal))
 					{
-						Value = Property.Value;
+						value = property.Value;
 						return true;
 					}
 				}
 			}
 
-			Value = null;
+			value = null;
 			return false;
 		}
 	}
diff --git a/Engine/Source/Programs/Shared/EpicGames.Core/ParallelTask.cs b/Engine/Source/Programs/Shared/EpicGames.Core/ParallelTask.cs
--- a/Engine/Source/Programs/Shared/EpicGames.Core/ParallelTask.cs
+++ b/Engine/Source/Programs/Shared/EpicGames.Core/ParallelTask.cs
@@ -3,7 +3,6 @@
 using System;
 using System.Collections.Generic;
 using System.Linq;
-using System.Text;
 using System.Threading.Tasks;
 using System.Threading.Tasks.Dataflow;
 
@@ -18,13 +17,13 @@ namespace EpicGames.Core
 		/// Execute a large number of tasks in parallel over a collection. Parallel.ForEach() method isn't generally compatible with asynchronous programming, because any
 		/// exceptions are thrown on the 
 		/// </summary>
-		/// <param name="FromInclusive">The starting index</param>
-		/// <param name="ToExclusive">The last index, exclusive</param>
-		/// <param name="Action">Action to perform for each item in the collection</param>
+		/// <param name="fromInclusive">The starting index</param>
+		/// <param name="toExclusive">The last index, exclusive</param>
+		/// <param name="action">Action to perform for each item in the collection</param>
 		/// <returns>Async task</returns>
-		public static Task ForAsync(int FromInclusive, int ToExclusive, Action<int> Action)
+		public static Task ForAsync(int fromInclusive, int toExclusive, Action<int> action)
 		{
-			return ForEachAsync(Enumerable.Range(FromInclusive, ToExclusive - FromInclusive), Action);
+			return ForEachAsync(Enumerable.Range(fromInclusive, toExclusive - fromInclusive), action);
 		}
 
 		/// <summary>
@@ -32,22 +31,22 @@ namespace EpicGames.Core
 		/// exceptions are thrown on the 
 		/// </summary>
 		/// <typeparam name="T">The collection type</typeparam>
-		/// <param name="Collection">The collection to iterate over</param>
-		/// <param name="Action">Action to perform for each item in the collection</param>
+		/// <param name="collection">The collection to iterate over</param>
+		/// <param name="action">Action to perform for each item in the collection</param>
 		/// <returns>Async task</returns>
-		public static Task ForEachAsync<T>(IEnumerable<T> Collection, Action<T> Action)
+		public static Task ForEachAsync<T>(IEnumerable<T> collection, Action<T> action)
 		{
-			ExecutionDataflowBlockOptions Options = new ExecutionDataflowBlockOptions();
-			Options.MaxDegreeOfParallelism = DataflowBlockOptions.Unbounded;
+			ExecutionDataflowBlockOptions options = new ExecutionDataflowBlockOptions();
+			options.MaxDegreeOfParallelism = DataflowBlockOptions.Unbounded;
 
-			ActionBlock<T> Actions = new ActionBlock<T>(Action, Options);
-			foreach (T Item in Collection)
+			ActionBlock<T> actions = new ActionBlock<T>(action, options);
+			foreach (T item in collection)
 			{
-				Actions.Post(Item);
+				actions.Post(item);
 			}
-			Actions.Complete();
+			actions.Complete();
 
-			return Actions.Completion;
+			return actions.Completion;
 		}
 
 		/// <summary>
@@ -55,23 +54,23 @@ namespace EpicGames.Core
 		/// exceptions are thrown on the 
 		/// </summary>
 		/// <typeparam name="T">The collection type</typeparam>
-		/// <param name="Collection">The collection to iterate over</param>
-		/// <param name="Action">Action to perform for each item in the collection</param>
-		/// <param name="MaxDegreeOfParallelism">Maximum degree of parallelism</param>
+		/// <param name="collection">The collection to iterate over</param>
+		/// <param name="action">Action to perform for each item in the collection</param>
+		/// <param name="maxDegreeOfParallelism">Maximum degree of parallelism</param>
 		/// <returns>Async task</returns>
-		public static Task ForEachAsync<T>(IEnumerable<T> Collection, Func<T, Task> Action, int MaxDegreeOfParallelism)
+		public static Task ForEachAsync<T>(IEnumerable<T> collection, Func<T, Task> action, int maxDegreeOfParallelism)
 		{
-			ExecutionDataflowBlockOptions Options = new ExecutionDataflowBlockOptions();
-			Options.MaxDegreeOfParallelism = MaxDegreeOfParallelism;
+			ExecutionDataflowBlockOptions options = new ExecutionDataflowBlockOptions();
+			options.MaxDegreeOfParallelism = maxDegreeOfParallelism;
 
-			ActionBlock<T> Actions = new ActionBlock<T>(Action, Options);
-			foreach (T Item in Collection)
+			ActionBlock<T> actions = new ActionBlock<T>(action, options);
+			foreach (T item in collection)
 			{
-				Actions.Post(Item);
+				actions.Post(item);
 			}
-			Actions.Complete();
+			actions.Complete();
 
-			return Actions.Completion;
+			return actions.Completion;
 		}
 	}
 }
diff --git a/Engine/Source/Programs/Shared/EpicGames.Core/ProcessUtils.cs b/Engine/Source/Programs/Shared/EpicGames.Core/ProcessUtils.cs
--- a/Engine/Source/Programs/Shared/EpicGames.Core/ProcessUtils.cs
+++ b/Engine/Source/Programs/Shared/EpicGames.Core/ProcessUtils.cs
@@ -1,17 +1,14 @@
 // Copyright Epic Games, Inc. All Rights Reserved.
 
-using Microsoft.Extensions.Logging;
-using Microsoft.Win32.SafeHandles;
 using System;
 using System.Collections.Generic;
 using System.Diagnostics;
-using System.Diagnostics.CodeAnalysis;
 using System.Linq;
-using System.Reflection;
 using System.Runtime.InteropServices;
 using System.Text;
 using System.Threading;
-using System.Threading.Tasks;
+using Microsoft.Extensions.Logging;
+using Microsoft.Win32.SafeHandles;
 
 namespace EpicGames.Core
 {
@@ -42,108 +39,108 @@ namespace EpicGames.Core
 		/// <summary>
 		/// Attempts to terminate all processes matching a predicate
 		/// </summary>
-		/// <param name="Predicate">The predicate for whether to terminate a process</param>
-		/// <param name="Logger">Logging device</param>
-		/// <param name="CancellationToken">Cancellation token to abort the search</param>
+		/// <param name="predicate">The predicate for whether to terminate a process</param>
+		/// <param name="logger">Logging device</param>
+		/// <param name="cancellationToken">Cancellation token to abort the search</param>
 		/// <returns>True if the process succeeded</returns>
-		public static bool TerminateProcesses(Predicate<FileReference> Predicate, ILogger Logger, CancellationToken CancellationToken)
+		public static bool TerminateProcesses(Predicate<FileReference> predicate, ILogger logger, CancellationToken cancellationToken)
 		{
 			if (RuntimeInformation.IsOSPlatform(OSPlatform.Windows))
 			{
-				return TerminateProcessesWin32(Predicate, Logger, CancellationToken);
+				return TerminateProcessesWin32(predicate, logger, cancellationToken);
 			}
 			else
 			{
-				return TerminateProcessesGenericPlatform(Predicate, Logger, CancellationToken);
+				return TerminateProcessesGenericPlatform(predicate, logger, cancellationToken);
 			}
 		}
 
 		/// <summary>
 		/// Attempts to terminate all processes matching a predicate
 		/// </summary>
-		/// <param name="Predicate">The predicate for whether to terminate a process</param>
-		/// <param name="Logger">Logging device</param>
-		/// <param name="CancellationToken">Cancellation token to abort the search</param>
+		/// <param name="predicate">The predicate for whether to terminate a process</param>
+		/// <param name="logger">Logging device</param>
+		/// <param name="cancellationToken">Cancellation token to abort the search</param>
 		/// <returns>True if the process succeeded</returns>
-		static bool TerminateProcessesWin32(Predicate<FileReference> Predicate, ILogger Logger, CancellationToken CancellationToken)
+		static bool TerminateProcessesWin32(Predicate<FileReference> predicate, ILogger logger, CancellationToken cancellationToken)
 		{
-			Dictionary<int, int> ProcessToCount = new Dictionary<int, int>();
+			Dictionary<int, int> processToCount = new Dictionary<int, int>();
 			for(; ;)
 			{
-				CancellationToken.ThrowIfCancellationRequested();
+				cancellationToken.ThrowIfCancellationRequested();
 
 				// Enumerate the processes
-				int NumProcessIds;
-				uint[] ProcessIds = new uint[512];
+				int numProcessIds;
+				uint[] processIds = new uint[512];
 				for (; ; )
 				{
-					int MaxBytes = ProcessIds.Length * sizeof(uint);
-					int NumBytes = 0;
-					if (!EnumProcesses(ProcessIds, MaxBytes, out NumBytes))
+					int maxBytes = processIds.Length * sizeof(uint);
+					int numBytes = 0;
+					if (!EnumProcesses(processIds, maxBytes, out numBytes))
 					{
 						throw new Win32ExceptionWithCode("Unable to enumerate processes");
 					}
-					if (NumBytes < MaxBytes)
+					if (numBytes < maxBytes)
 					{
-						NumProcessIds = NumBytes / sizeof(uint);
+						numProcessIds = numBytes / sizeof(uint);
 						break;
 					}
-					ProcessIds = new uint[ProcessIds.Length + 256];
+					processIds = new uint[processIds.Length + 256];
 				}
 
 				// Find the processes to terminate
-				List<SafeProcessHandle> WaitHandles = new List<SafeProcessHandle>();
+				List<SafeProcessHandle> waitHandles = new List<SafeProcessHandle>();
 				try
 				{
 					// Open each process in turn
-					foreach (uint ProcessId in ProcessIds)
+					foreach (uint processId in processIds)
 					{
-						SafeProcessHandle Handle = OpenProcess(PROCESS_QUERY_LIMITED_INFORMATION | PROCESS_TERMINATE, false, ProcessId);
+						SafeProcessHandle handle = OpenProcess(PROCESS_QUERY_LIMITED_INFORMATION | PROCESS_TERMINATE, false, processId);
 						try
 						{
-							if (!Handle.IsInvalid)
+							if (!handle.IsInvalid)
 							{
-								int Characters = 260;
-								StringBuilder Buffer = new StringBuilder(Characters);
-								if (QueryFullProcessImageName(Handle, 0, Buffer, ref Characters) > 0)
+								int characters = 260;
+								StringBuilder buffer = new StringBuilder(characters);
+								if (QueryFullProcessImageName(handle, 0, buffer, ref characters) > 0)
 								{
-									FileReference ImageFile = new FileReference(Buffer.ToString(0, Characters));
-									if (Predicate(ImageFile))
+									FileReference imageFile = new FileReference(buffer.ToString(0, characters));
+									if (predicate(imageFile))
 									{
-										Logger.LogInformation("Terminating {ImageName} ({ProcessId})", ImageFile, ProcessId);
-										if (TerminateProcess(Handle, 9))
+										logger.LogInformation("Terminating {ImageName} ({ProcessId})", imageFile, processId);
+										if (TerminateProcess(handle, 9))
 										{
-											WaitHandles.Add(Handle);
+											waitHandles.Add(handle);
 										}
 										else
 										{
-											Logger.LogInformation("Failed call to TerminateProcess ({Code})", Marshal.GetLastWin32Error());
+											logger.LogInformation("Failed call to TerminateProcess ({Code})", Marshal.GetLastWin32Error());
 										}
-										Handle.SetHandleAsInvalid();
+										handle.SetHandleAsInvalid();
 									}
 								}
 							}
 						}
 						finally
 						{
-							Handle.Dispose();
+							handle.Dispose();
 						}
 					}
 
 					// If there's nothing to do, exit immediately
-					if (WaitHandles.Count == 0)
+					if (waitHandles.Count == 0)
 					{
 						return true;
 					}
 
 					// Wait for them all to complete
-					WaitForMultipleObjects(WaitHandles.Count, WaitHandles.Select(x => x.DangerousGetHandle()).ToArray(), true, 10 * 1000);
+					WaitForMultipleObjects(waitHandles.Count, waitHandles.Select(x => x.DangerousGetHandle()).ToArray(), true, 10 * 1000);
 				}
 				finally
 				{
-					foreach (SafeProcessHandle WaitHandle in WaitHandles)
+					foreach (SafeProcessHandle waitHandle in waitHandles)
 					{
-						WaitHandle.Close();
+						waitHandle.Close();
 					}
 				}
 			}
@@ -152,76 +149,76 @@ namespace EpicGames.Core
 		/// <summary>
 		/// Attempts to terminate all processes matching a predicate
 		/// </summary>
-		/// <param name="Predicate">The predicate for whether to terminate a process</param>
-		/// <param name="Logger">Logging device</param>
-		/// <param name="CancellationToken">Cancellation token to abort the search</param>
+		/// <param name="predicate">The predicate for whether to terminate a process</param>
+		/// <param name="logger">Logging device</param>
+		/// <param name="cancellationToken">Cancellation token to abort the search</param>
 		/// <returns>True if the process succeeded</returns>
-		static bool TerminateProcessesGenericPlatform(Predicate<FileReference> Predicate, ILogger Logger, CancellationToken CancellationToken)
+		static bool TerminateProcessesGenericPlatform(Predicate<FileReference> predicate, ILogger logger, CancellationToken cancellationToken)
 		{
-			bool Result = true;
-			Dictionary<(int, DateTime), int> ProcessToCount = new Dictionary<(int, DateTime), int>();
+			bool result = true;
+			Dictionary<(int, DateTime), int> processToCount = new Dictionary<(int, DateTime), int>();
 			for (; ; )
 			{
-				CancellationToken.ThrowIfCancellationRequested();
+				cancellationToken.ThrowIfCancellationRequested();
 
 				bool bNoMatches = true;
 
 				// Enumerate all the processes
-				Process[] Processes = Process.GetProcesses();
-				foreach (Process Process in Processes)
+				Process[] processes = Process.GetProcesses();
+				foreach (Process process in processes)
 				{
 					// Attempt to get the image file. Ignore exceptions trying to fetch metadata for processes we don't have access to.
-					FileReference? ImageFile;
+					FileReference? imageFile;
 					try
 					{
-						ImageFile = new FileReference(Process.MainModule.FileName);
+						imageFile = new FileReference(process.MainModule.FileName);
 					}
 					catch
 					{
-						ImageFile = null;
+						imageFile = null;
 					}
 
 					// Test whether to terminate this process
-					if (ImageFile != null && Predicate(ImageFile))
+					if (imageFile != null && predicate(imageFile))
 					{
 						// Get a unique id for the process, given that process ids are recycled
-						(int, DateTime) UniqueId;
+						(int, DateTime) uniqueId;
 						try
 						{
-							UniqueId = (Process.Id, Process.StartTime);
+							uniqueId = (process.Id, process.StartTime);
 						}
 						catch
 						{
-							UniqueId = (Process.Id, DateTime.MinValue);
+							uniqueId = (process.Id, DateTime.MinValue);
 						}
 
 						// Figure out whether to try and terminate this process
 						const int MaxCount = 5;
-						if (!ProcessToCount.TryGetValue(UniqueId, out int Count) || Count < MaxCount)
+						if (!processToCount.TryGetValue(uniqueId, out int count) || count < MaxCount)
 						{
 							bNoMatches = false;
 							try
 							{
-								Logger.LogInformation("Terminating {ImageName} ({ProcessId})", ImageFile, Process.Id);
-								Process.Kill(true);
-								if (!Process.WaitForExit(5 * 1000))
+								logger.LogInformation("Terminating {ImageName} ({ProcessId})", imageFile, process.Id);
+								process.Kill(true);
+								if (!process.WaitForExit(5 * 1000))
 								{
-									Logger.LogInformation("Termination still pending; will retry...");
+									logger.LogInformation("Termination still pending; will retry...");
 								}
 							}
-							catch (Exception Ex)
+							catch (Exception ex)
 							{
-								Count++;
-								if (Count > 1)
+								count++;
+								if (count > 1)
 								{
-									Logger.LogInformation(Ex, "Exception while querying basic process info for pid {ProcessId}; will retry.", Process.Id);
+									logger.LogInformation(ex, "Exception while querying basic process info for pid {ProcessId}; will retry.", process.Id);
 								}
-								else if (Count == MaxCount)
+								else if (count == MaxCount)
 								{
-									Logger.LogWarning(Ex, "Unable to terminate process {ImageFile} ({ProcessId}): {Message}", ImageFile, Process.Id, Ex.Message);
-									Result = false;
+									logger.LogWarning(ex, "Unable to terminate process {ImageFile} ({ProcessId}): {Message}", imageFile, process.Id, ex.Message);
+									result = false;
 								}
-								ProcessToCount[UniqueId] = Count;
+								processToCount[uniqueId] = count;
 							}
 						}
 					}
@@ -230,7 +227,7 @@ namespace EpicGames.Core
 				// Return once we reach this point and haven't found anything else to terminate
 				if (bNoMatches)
 				{
-					return Result;
+					return result;
 				}
 			}
 		}
diff --git a/Engine/Source/Programs/Shared/EpicGames.Core/ReadOnlyHashSet.cs b/Engine/Source/Programs/Shared/EpicGames.Core/ReadOnlyHashSet.cs
--- a/Engine/Source/Programs/Shared/EpicGames.Core/ReadOnlyHashSet.cs
+++ b/Engine/Source/Programs/Shared/EpicGames.Core/ReadOnlyHashSet.cs
@@ -1,11 +1,7 @@
 // Copyright Epic Games, Inc. All Rights Reserved.
 
-using System;
 using System.Collections;
 using System.Collections.Generic;
-using System.Linq;
-using System.Text;
-using System.Threading.Tasks;
 
 namespace EpicGames.Core
 {
@@ -18,60 +14,54 @@ namespace EpicGames.Core
 		/// <summary>
 		/// The mutable hashset
 		/// </summary>
-		HashSet<T> Inner;
+		readonly HashSet<T> _inner;
 
 		/// <summary>
 		/// Constructor
 		/// </summary>
-		/// <param name="Inner">The mutable hashset</param>
-		public ReadOnlyHashSet(HashSet<T> Inner)
+		/// <param name="inner">The mutable hashset</param>
+		public ReadOnlyHashSet(HashSet<T> inner)
 		{
-			this.Inner = Inner;
+			_inner = inner;
 		}
 
 		/// <summary>
 		/// Constructor
 		/// </summary>
 		/// <param name="Collection">Elements for the hash set</param>
-		public ReadOnlyHashSet(IEnumerable<T> Elements)
+		public ReadOnlyHashSet(IEnumerable<T> elements)
 		{
-			this.Inner = new HashSet<T>(Elements);
+			_inner = new HashSet<T>(elements);
 		}
 
 		/// <summary>
 		/// Constructor
 		/// </summary>
 		/// <param name="Collection">Elements for the hash set</param>
-		/// <param name="Comparer">Comparer for elements in the set</param>
-		public ReadOnlyHashSet(IEnumerable<T> Elements, IEqualityComparer<T> Comparer)
+		/// <param name="comparer">Comparer for elements in the set</param>
+		public ReadOnlyHashSet(IEnumerable<T> elements, IEqualityComparer<T> comparer)
 		{
-			this.Inner = new HashSet<T>(Elements, Comparer);
+			_inner = new HashSet<T>(elements, comparer);
 		}
 
 		/// <summary>
 		/// Number of elements in the set
 		/// </summary>
-		public int Count
-		{
-			get { return Inner.Count; }
-		}
+		public int Count => _inner.Count;
 
 		/// <summary>
 		/// The comparer for elements in the set
 		/// </summary>
-		public IEqualityComparer<T> Comparer
-		{
-			get { return Inner.Comparer; }
-		}
+		public IEqualityComparer<T> Comparer => _inner.Comparer;
 
 		/// <summary>
 		/// Tests whether a given item is in the set
 		/// </summary>
-		/// <param name="Item">Item to check for</param>
+		/// <param name="item">Item to check for</param>
 		/// <returns>True if the item is in the set</returns>
-		public bool Contains(T Item)
+		public bool Contains(T item)
 		{
-			return Inner.Contains(Item);
+			return _inner.Contains(item);
 		}
 
 		/// <summary>
@@ -80,7 +70,7 @@ namespace EpicGames.Core
 		/// <returns>Enumerator instance</returns>
 		public IEnumerator<T> GetEnumerator()
 		{
-			return Inner.GetEnumerator();
+			return _inner.GetEnumerator();
 		}
 
 		/// <summary>
@@ -89,16 +79,16 @@ namespace EpicGames.Core
 		/// <returns>Enumerator instance</returns>
 		IEnumerator IEnumerable.GetEnumerator()
 		{
-			return ((IEnumerable)Inner).GetEnumerator();
+			return ((IEnumerable)_inner).GetEnumerator();
 		}
 
 		/// <summary>
 		/// Implicit conversion operator from hashsets
 		/// </summary>
-		/// <param name="HashSet"></param>
-		public static implicit operator ReadOnlyHashSet<T>(HashSet<T> HashSet)
+		/// <param name="hashSet"></param>
+		public static implicit operator ReadOnlyHashSet<T>(HashSet<T> hashSet)
 		{
-			return new ReadOnlyHashSet<T>(HashSet);
+			return new ReadOnlyHashSet<T>(hashSet);
 		}
 	}
 }
diff --git a/Engine/Source/Programs/Shared/EpicGames.Core/ReadOnlyListExtensions.cs b/Engine/Source/Programs/Shared/EpicGames.Core/ReadOnlyListExtensions.cs
--- a/Engine/Source/Programs/Shared/EpicGames.Core/ReadOnlyListExtensions.cs
+++ b/Engine/Source/Programs/Shared/EpicGames.Core/ReadOnlyListExtensions.cs
@@ -1,10 +1,7 @@
 // Copyright Epic Games, Inc. All Rights Reserved.
 
 using System;
-using System.Collections;
 using System.Collections.Generic;
-using System.Linq;
-using System.Threading.Tasks;
 
 namespace EpicGames.Core
 {
@@ -17,25 +14,25 @@ namespace EpicGames.Core
 		/// Performs a binary search on the given list
 		/// </summary>
 		/// <typeparam name="T">The element type for the list</typeparam>
-		/// <param name="List">The list type</param>
-		/// <param name="Item">Item to search for</param>
+		/// <param name="list">The list type</param>
+		/// <param name="item">Item to search for</param>
 		/// <returns>As List.BinarySearch</returns>
-		public static int BinarySearch<T>(this IReadOnlyList<T> List, T Item)
+		public static int BinarySearch<T>(this IReadOnlyList<T> list, T item)
 		{
-			return BinarySearch(List, x => x, Item, Comparer<T>.Default);
+			return BinarySearch(list, x => x, item, Comparer<T>.Default);
 		}
 
 		/// <summary>
 		/// Performs a binary search on the given list
 		/// </summary>
 		/// <typeparam name="T">The element type for the list</typeparam>
-		/// <param name="List">The list type</param>
-		/// <param name="Item">Item to search for</param>
-		/// <param name="Comparer">Comparer for elements in the list</param>
+		/// <param name="list">The list type</param>
+		/// <param name="item">Item to search for</param>
+		/// <param name="comparer">Comparer for elements in the list</param>
 		/// <returns>As List.BinarySearch</returns>
-		public static int BinarySearch<T>(this IReadOnlyList<T> List, T Item, IComparer<T> Comparer)
+		public static int BinarySearch<T>(this IReadOnlyList<T> list, T item, IComparer<T> comparer)
 		{
-			return BinarySearch(List, x => x, Item, Comparer);
+			return BinarySearch(list, x => x, item, comparer);
 		}
 
 		/// <summary>
@@ -43,13 +40,13 @@ namespace EpicGames.Core
 		/// </summary>
 		/// <typeparam name="TItem">The item in the list</typeparam>
 		/// <typeparam name="TField">The field to search on</typeparam>
-		/// <param name="List">The list to search</param>
-		/// <param name="Projection">The projection to apply to each item in the list</param>
-		/// <param name="Item">The item to find</param>
+		/// <param name="list">The list to search</param>
+		/// <param name="projection">The projection to apply to each item in the list</param>
+		/// <param name="item">The item to find</param>
 		/// <returns>As <see cref="List{T}.BinarySearch(T)"/></returns>
-		public static int BinarySearch<TItem, TField>(this IReadOnlyList<TItem> List, Func<TItem, TField> Projection, TField Item)
+		public static int BinarySearch<TItem, TField>(this IReadOnlyList<TItem> list, Func<TItem, TField> projection, TField item)
 		{
-			return BinarySearch(List, Projection, Item, Comparer<TField>.Default);
+			return BinarySearch(list, projection, item, Comparer<TField>.Default);
 		}
 
 		/// <summary>
@@ -57,34 +54,34 @@ namespace EpicGames.Core
 		/// </summary>
 		/// <typeparam name="TItem">The item in the list</typeparam>
 		/// <typeparam name="TField">The field to search on</typeparam>
-		/// <param name="List">The list to search</param>
-		/// <param name="Projection">The projection to apply to each item in the list</param>
-		/// <param name="Item">The item to find</param>
-		/// <param name="Comparer">Comparer for field elements</param>
+		/// <param name="list">The list to search</param>
+		/// <param name="projection">The projection to apply to each item in the list</param>
+		/// <param name="item">The item to find</param>
+		/// <param name="comparer">Comparer for field elements</param>
 		/// <returns>As <see cref="List{T}.BinarySearch(T)"/></returns>
-		public static int BinarySearch<TItem, TField>(this IReadOnlyList<TItem> List, Func<TItem, TField> Projection, TField Item, IComparer<TField> Comparer)
+		public static int BinarySearch<TItem, TField>(this IReadOnlyList<TItem> list, Func<TItem, TField> projection, TField item, IComparer<TField> comparer)
 		{
-			int LowerBound = 0;
-			int UpperBound = List.Count - 1;
-			while (LowerBound <= UpperBound)
+			int lowerBound = 0;
+			int upperBound = list.Count - 1;
+			while (lowerBound <= upperBound)
 			{
-				int Idx = LowerBound + (UpperBound - LowerBound) / 2;
+				int idx = lowerBound + (upperBound - lowerBound) / 2;
 
-				int Comparison = Comparer.Compare(Projection(List[Idx]), Item);
-				if (Comparison == 0)
+				int comparison = comparer.Compare(projection(list[idx]), item);
+				if (comparison == 0)
 				{
-					return Idx;
+					return idx;
 				}
-				else if (Comparison < 0)
+				else if (comparison < 0)
 				{
-					LowerBound = Idx + 1;
+					lowerBound = idx + 1;
 				}
 				else
 				{
-					UpperBound = Idx - 1;
+					upperBound = idx - 1;
 				}
 			}
-			return ~LowerBound;
+			return ~lowerBound;
 		}
 
 		/// <summary>
@@ -92,38 +89,38 @@ namespace EpicGames.Core
 		/// </summary>
 		/// <typeparam name="TInput">The input element type</typeparam>
 		/// <typeparam name="TOutput">The output element type</typeparam>
-		/// <param name="Input">Input list</param>
-		/// <param name="Convert">Conversion function</param>
+		/// <param name="input">Input list</param>
+		/// <param name="convert">Conversion function</param>
 		/// <returns>New list of items</returns>
-		public static List<TOutput> ConvertAll<TInput, TOutput>(this IReadOnlyList<TInput> Input, Func<TInput, TOutput> Convert)
+		public static List<TOutput> ConvertAll<TInput, TOutput>(this IReadOnlyList<TInput> input, Func<TInput, TOutput> convert)
 		{
-			List<TOutput> Outputs = new List<TOutput>(Input.Count);
-			for (int Idx = 0; Idx < Input.Count; Idx++)
+			List<TOutput> outputs = new List<TOutput>(input.Count);
+			for (int idx = 0; idx < input.Count; idx++)
 			{
-				Outputs.Add(Convert(Input[Idx]));
+				outputs.Add(convert(input[idx]));
 			}
-			return Outputs;
+			return outputs;
 		}
 
 		/// <summary>
 		/// Finds the index of the first element matching a predicate
 		/// </summary>
 		/// <typeparam name="T"></typeparam>
-		/// <param name="List">List to search</param>
-		/// <param name="Predicate">Predicate for the list</param>
+		/// <param name="list">List to search</param>
+		/// <param name="predicate">Predicate for the list</param>
 		/// <returns>Index of the element</returns>
-		public static int FindIndex<T>(this IReadOnlyList<T> List, Predicate<T> Predicate)
+		public static int FindIndex<T>(this IReadOnlyList<T> list, Predicate<T> predicate)
 		{
-			int FoundIndex = -1;
-			for(int Idx = 0; Idx < List.Count; Idx++)
+			int foundIndex = -1;
+			for(int idx = 0; idx < list.Count; idx++)
 			{
-				if (Predicate(List[Idx]))
+				if (predicate(list[idx]))
 				{
-					FoundIndex = Idx;
+					foundIndex = idx;
 					break;
 				}
 			}
-			return FoundIndex;
+			return foundIndex;
 		}
 	}
 }
diff --git a/Engine/Source/Programs/Shared/EpicGames.Core/ReadOnlyMemoryStream.cs b/Engine/Source/Programs/Shared/EpicGames.Core/ReadOnlyMemoryStream.cs
--- a/Engine/Source/Programs/Shared/EpicGames.Core/ReadOnlyMemoryStream.cs
+++ b/Engine/Source/Programs/Shared/EpicGames.Core/ReadOnlyMemoryStream.cs
@@ -1,29 +1,27 @@
 // Copyright Epic Games, Inc. All Rights Reserved.
 
 using System;
-using System.Collections.Generic;
 using System.IO;
-using System.Text;
 
 namespace EpicGames.Core
 {
 	/// <summary>
-	/// Stream which reads from a <see cref="ReadOnlyMemory{byte}"/>
+	/// Stream which reads from a <see cref="ReadOnlyMemory{Byte}"/>
 	/// </summary>
 	public class ReadOnlyMemoryStream : Stream
 	{
 		/// <summary>
 		/// The buffer to read from
 		/// </summary>
-		ReadOnlyMemory<byte> Memory;
+		readonly ReadOnlyMemory<byte> _memory;
 
 		/// <summary>
 		/// Constructor
 		/// </summary>
-		/// <param name="Memory">The memory to read from</param>
-		public ReadOnlyMemoryStream(ReadOnlyMemory<byte> Memory)
+		/// <param name="memory">The memory to read from</param>
+		public ReadOnlyMemoryStream(ReadOnlyMemory<byte> memory)
 		{
-			this.Memory = Memory;
+			_memory = memory;
 		}
 
 		/// <inheritdoc/>
@@ -36,7 +34,7 @@ namespace EpicGames.Core
 		public override bool CanWrite => false;
 
 		/// <inheritdoc/>
-		public override long Length => Memory.Length;
+		public override long Length => _memory.Length;
 
 		/// <inheritdoc/>
 		public override long Position { get; set; }
@@ -47,30 +45,30 @@ namespace EpicGames.Core
 		}
 
 		/// <inheritdoc/>
-		public override int Read(byte[] Buffer, int Offset, int Count)
+		public override int Read(byte[] buffer, int offset, int count)
 		{
-			int CopyLength = Math.Min(Count, (int)(Memory.Length - Position));
-			Memory.Slice((int)Position, CopyLength).CopyTo(Buffer.AsMemory(Offset, CopyLength));
-			Position += CopyLength;
-			return CopyLength;
+			int copyLength = Math.Min(count, (int)(_memory.Length - Position));
+			_memory.Slice((int)Position, copyLength).CopyTo(buffer.AsMemory(offset, copyLength));
+			Position += copyLength;
+			return copyLength;
 		}
 
 		/// <inheritdoc/>
-		public override long Seek(long Offset, SeekOrigin Origin)
+		public override long Seek(long offset, SeekOrigin origin)
 		{
-			switch (Origin)
+			switch (origin)
 			{
 				case SeekOrigin.Begin:
-					Position = Offset;
+					Position = offset;
 					break;
 				case SeekOrigin.Current:
-					Position += Offset;
+					Position += offset;
 					break;
 				case SeekOrigin.End:
-					Position = Memory.Length + Offset;
+					Position = _memory.Length + offset;
 					break;
 				default:
-					throw new ArgumentException(null, nameof(Origin));
+					throw new ArgumentException(null, nameof(origin));
 			}
 			return Position;
 		}
diff --git a/Engine/Source/Programs/Shared/EpicGames.Core/RuntimePlatform.cs b/Engine/Source/Programs/Shared/EpicGames.Core/RuntimePlatform.cs
--- a/Engine/Source/Programs/Shared/EpicGames.Core/RuntimePlatform.cs
+++ b/Engine/Source/Programs/Shared/EpicGames.Core/RuntimePlatform.cs
@@ -1,9 +1,6 @@
 // Copyright Epic Games, Inc. All Rights Reserved.
 
-using System;
-using System.Collections.Generic;
 using System.Runtime.InteropServices;
-using System.Text;
 
 namespace EpicGames.Core
 {
diff --git a/Engine/Source/Programs/Shared/EpicGames.Core/ScopeProfiler.cs b/Engine/Source/Programs/Shared/EpicGames.Core/ScopeProfiler.cs
--- a/Engine/Source/Programs/Shared/EpicGames.Core/ScopeProfiler.cs
+++ b/Engine/Source/Programs/Shared/EpicGames.Core/ScopeProfiler.cs
@@ -1,12 +1,10 @@
 // Copyright Epic Games, Inc. All Rights Reserved.
 
-using System;
-using System.Reflection;
-
 namespace EpicGames.Core
 {
 #if !__SCOPEPROFILER_AVAILABLE__
-	
+#pragma warning disable IDE0060 // Remove unused parameter
+
 	/// <summary>
 	/// A stub/no-op scope-profiler API that can be replaced by another implementation to record execution of
 	/// instrumented scopes.
@@ -15,16 +13,16 @@ namespace EpicGames.Core
 	{
 		public static ScopeProfiler Instance = new ScopeProfiler();
 		
-		public void InitializeAndStart(string ProgramName, string HostAddress, int MaxThreadCount) { }
+		public void InitializeAndStart(string programName, string hostAddress, int maxThreadCount) { }
 
 		public void StopAndShutdown() { }
 
 		/// <summary>
 		/// Start a stacked scope. Must start and end on the same thread
 		/// </summary>
-		/// <param name="Name"></param>
+		/// <param name="name"></param>
 		/// <param name="bIsStall"></param>
-		public void StartScope(string Name, bool bIsStall) { }
+		public void StartScope(string name, bool bIsStall) { }
 		
 		/// <summary>
 		/// End a stacked scope. Must start and end on the same thread.
@@ -34,21 +32,21 @@ namespace EpicGames.Core
 		/// <summary>
 		/// Record an un-stacked time span at a specified time.
 		/// </summary>
-		/// <param name="Name"></param>
-		/// <param name="Id"></param>
+		/// <param name="name"></param>
+		/// <param name="id"></param>
 		/// <param name="bStall"></param>
-		/// <param name="BeginThreadId"></param>
-		/// <param name="EndThreadId"></param>
-		/// <param name="StartTime">start time for the span (use FastTime())</param>
-		/// <param name="EndTime">end time for the span (use FastTime())</param>
-		public void AddSpanAtTime(string Name, ulong Id, bool bStall, uint BeginThreadId, uint EndThreadId, ulong StartTime, ulong EndTime) { }
+		/// <param name="beginThreadId"></param>
+		/// <param name="endThreadId"></param>
+		/// <param name="startTime">start time for the span (use FastTime())</param>
+		/// <param name="endTime">end time for the span (use FastTime())</param>
+		public void AddSpanAtTime(string name, ulong id, bool bStall, uint beginThreadId, uint endThreadId, ulong startTime, ulong endTime) { }
 
 		/// <summary>
 		/// Record a value against a particular name, to be plotted over time
 		/// </summary>
-		/// <param name="Name"></param>
-		/// <param name="Value"></param>
-		public void Plot(string Name, double Value) { }
+		/// <param name="name"></param>
+		/// <param name="value"></param>
+		public void Plot(string name, double value) { }
 
 		/// <summary>
 		/// Generate a timestamp value for the current moment.
@@ -68,6 +66,7 @@ namespace EpicGames.Core
 			return 0;
 		}
 	}
+#pragma warning restore IDE0060 // Remove unused parameter
 #endif
 }
 
diff --git a/Engine/Source/Programs/Shared/EpicGames.Core/Sha1Hash.cs b/Engine/Source/Programs/Shared/EpicGames.Core/Sha1Hash.cs
--- a/Engine/Source/Programs/Shared/EpicGames.Core/Sha1Hash.cs
+++ b/Engine/Source/Programs/Shared/EpicGames.Core/Sha1Hash.cs
@@ -1,16 +1,12 @@
 // Copyright Epic Games, Inc. All Rights Reserved.
 
 using System;
-using System.Collections.Generic;
-using System.Buffers;
-using System.Security.Cryptography;
-using System.Text;
 using System.Buffers.Binary;
-using System.Diagnostics.CodeAnalysis;
 using System.ComponentModel;
 using System.Globalization;
-using System.Text.Json.Serialization;
+using System.Security.Cryptography;
 using System.Text.Json;
+using System.Text.Json.Serialization;
 
 namespace EpicGames.Core
 {
@@ -30,10 +26,9 @@ namespace EpicGames.Core
 		/// Length of the hash in bits
 		/// </summary>
 		public const int NumBits = NumBytes * 8;
-
-		ulong A;
-		ulong B;
-		uint C;
+		readonly ulong _a;
+		readonly ulong _b;
+		readonly uint _c;
 
 		/// <summary>
 		/// Hash consisting of zeroes
@@ -44,8 +39,8 @@ namespace EpicGames.Core
 		/// Constructor
 		/// </summary>
 		/// <param name="Memory">Memory to construct from</param>
-		public Sha1Hash(ReadOnlySpan<byte> Span)
-			: this(BinaryPrimitives.ReadUInt64BigEndian(Span), BinaryPrimitives.ReadUInt64BigEndian(Span.Slice(8)), BinaryPrimitives.ReadUInt32BigEndian(Span.Slice(16)))
+		public Sha1Hash(ReadOnlySpan<byte> span)
+			: this(BinaryPrimitives.ReadUInt64BigEndian(span), BinaryPrimitives.ReadUInt64BigEndian(span.Slice(8)), BinaryPrimitives.ReadUInt32BigEndian(span.Slice(16)))
 		{
 		}
 
@@ -53,77 +48,77 @@ namespace EpicGames.Core
 		/// Constructor
 		/// </summary>
 		/// <param name="Memory">Memory to construct from</param>
-		public Sha1Hash(ulong A, ulong B, uint C)
+		public Sha1Hash(ulong a, ulong b, uint c)
 		{
-			this.A = A;
-			this.B = B;
-			this.C = C;
+			_a = a;
+			_b = b;
+			_c = c;
 		}
 
 		/// <summary>
 		/// Creates a content hash for a block of data, using a given algorithm.
 		/// </summary>
-		/// <param name="Data">Data to compute the hash for</param>
+		/// <param name="data">Data to compute the hash for</param>
 		/// <returns>New content hash instance containing the hash of the data</returns>
-		public static Sha1Hash Compute(ReadOnlySpan<byte> Data)
+		public static Sha1Hash Compute(ReadOnlySpan<byte> data)
 		{
-			byte[] Output = new byte[20];
-			using (SHA1 Sha1 = SHA1.Create())
+			byte[] output = new byte[20];
+			using (SHA1 sha1 = SHA1.Create())
 			{
-				int BytesWritten;
-				if (!Sha1.TryComputeHash(Data, Output, out BytesWritten) || BytesWritten != NumBytes)
+				int bytesWritten;
+				if (!sha1.TryComputeHash(data, output, out bytesWritten) || bytesWritten != NumBytes)
 				{
 					throw new Exception($"Unable to hash data");
 				}
 			}
-			return new Sha1Hash(Output);
+			return new Sha1Hash(output);
 		}
 
 		/// <summary>
 		/// Parses a digest from the given hex string
 		/// </summary>
-		/// <param name="Text"></param>
+		/// <param name="text"></param>
 		/// <returns></returns>
-		public static Sha1Hash Parse(string Text)
+		public static Sha1Hash Parse(string text)
 		{
-			return new Sha1Hash(StringUtils.ParseHexString(Text));
+			return new Sha1Hash(StringUtils.ParseHexString(text));
 		}
 
 		/// <summary>
 		/// Parses a digest from the given hex string
 		/// </summary>
-		/// <param name="Text"></param>
+		/// <param name="text"></param>
 		/// <returns></returns>
-		public static Sha1Hash Parse(ReadOnlySpan<byte> Text)
+		public static Sha1Hash Parse(ReadOnlySpan<byte> text)
 		{
-			return new Sha1Hash(StringUtils.ParseHexString(Text));
+			return new Sha1Hash(StringUtils.ParseHexString(text));
 		}
 
 		/// <inheritdoc cref="IComparable{T}.CompareTo(T)"/>
-		public int CompareTo(Sha1Hash Other)
+		public int CompareTo(Sha1Hash other)
 		{
-			if (A != Other.A)
+			if (_a != other._a)
 			{
-				return (A < Other.A) ? -1 : +1;
+				return (_a < other._a) ? -1 : +1;
 			}
-			else if (B != Other.B)
+			else if (_b != other._b)
 			{
-				return (B < Other.B) ? -1 : +1;
+				return (_b < other._b) ? -1 : +1;
 			}
 			else
 			{
-				return (C < Other.C) ? -1 : +1;
+				return (_c < other._c) ? -1 : +1;
 			}
 		}
 
 		/// <inheritdoc/>
-		public bool Equals(Sha1Hash Other) => A == Other.A && B == Other.B && C == Other.C;
+		public bool Equals(Sha1Hash other) => _a == other._a && _b == other._b && _c == other._c;
 
 		/// <inheritdoc/>
-		public override bool Equals(object? Obj) => (Obj is Sha1Hash Hash) && Equals(Hash);
+		public override bool Equals(object? obj) => (obj is Sha1Hash hash) && Equals(hash);
 
 		/// <inheritdoc/>
-		public override int GetHashCode() => (int)A;
+		public override int GetHashCode() => (int)_a;
 
 		/// <inheritdoc/>
 		public Utf8String ToUtf8String() => StringUtils.FormatUtf8HexString(ToByteArray());
@@ -133,51 +128,51 @@ namespace EpicGames.Core
 
 		public byte[] ToByteArray()
 		{
-			byte[] Data = new byte[NumBytes];
-			CopyTo(Data);
-			return Data;
+			byte[] data = new byte[NumBytes];
+			CopyTo(data);
+			return data;
 		}
 
 		/// <summary>
 		/// Copies this hash into a span
 		/// </summary>
-		/// <param name="Span"></param>
-		public void CopyTo(Span<byte> Span)
+		/// <param name="span"></param>
+		public void CopyTo(Span<byte> span)
 		{
-			BinaryPrimitives.WriteUInt64BigEndian(Span, A);
-			BinaryPrimitives.WriteUInt64BigEndian(Span[8..], B);
-			BinaryPrimitives.WriteUInt32BigEndian(Span[16..], C);
+			BinaryPrimitives.WriteUInt64BigEndian(span, _a);
+			BinaryPrimitives.WriteUInt64BigEndian(span[8..], _b);
+			BinaryPrimitives.WriteUInt32BigEndian(span[16..], _c);
 		}
 
 		/// <summary>
 		/// Test two hash values for equality
 		/// </summary>
-		public static bool operator ==(Sha1Hash A, Sha1Hash B) => A.Equals(B);
+		public static bool operator ==(Sha1Hash a, Sha1Hash b) => a.Equals(b);
 
 		/// <summary>
 		/// Test two hash values for equality
 		/// </summary>
-		public static bool operator !=(Sha1Hash A, Sha1Hash B) => !(A == B);
+		public static bool operator !=(Sha1Hash a, Sha1Hash b) => !(a == b);
 
 		/// <summary>
 		/// Tests whether A > B
 		/// </summary>
-		public static bool operator >(Sha1Hash A, Sha1Hash B) => A.CompareTo(B) > 0;
+		public static bool operator >(Sha1Hash a, Sha1Hash b) => a.CompareTo(b) > 0;
 
 		/// <summary>
 		/// Tests whether A is less than B
 		/// </summary>
-		public static bool operator <(Sha1Hash A, Sha1Hash B) => A.CompareTo(B) < 0;
+		public static bool operator <(Sha1Hash a, Sha1Hash b) => a.CompareTo(b) < 0;
 
 		/// <summary>
 		/// Tests whether A is greater than or equal to B
 		/// </summary>
-		public static bool operator >=(Sha1Hash A, Sha1Hash B) => A.CompareTo(B) >= 0;
+		public static bool operator >=(Sha1Hash a, Sha1Hash b) => a.CompareTo(b) >= 0;
 
 		/// <summary>
 		/// Tests whether A is less than or equal to B
 		/// </summary>
-		public static bool operator <=(Sha1Hash A, Sha1Hash B) => A.CompareTo(B) <= 0;
+		public static bool operator <=(Sha1Hash a, Sha1Hash b) => a.CompareTo(b) <= 0;
 	}
 
 	/// <summary>
@@ -188,21 +183,21 @@ namespace EpicGames.Core
 		/// <summary>
 		/// Read an <see cref="Sha1Hash"/> from a memory reader
 		/// </summary>
-		/// <param name="Reader"></param>
+		/// <param name="reader"></param>
 		/// <returns></returns>
-		public static Sha1Hash ReadSha1Hash(this MemoryReader Reader)
+		public static Sha1Hash ReadSha1Hash(this MemoryReader reader)
 		{
-			return new Sha1Hash(Reader.ReadFixedLengthBytes(Sha1Hash.NumBytes).Span);
+			return new Sha1Hash(reader.ReadFixedLengthBytes(Sha1Hash.NumBytes).Span);
 		}
 
 		/// <summary>
 		/// Write an <see cref="Sha1Hash"/> to a memory writer
 		/// </summary>
-		/// <param name="Writer"></param>
-		/// <param name="Hash"></param>
-		public static void WriteSha1Hash(this MemoryWriter Writer, Sha1Hash Hash)
+		/// <param name="writer"></param>
+		/// <param name="hash"></param>
+		public static void WriteSha1Hash(this MemoryWriter writer, Sha1Hash hash)
 		{
-			Hash.CopyTo(Writer.AllocateSpan(Sha1Hash.NumBytes));
+			hash.CopyTo(writer.AllocateSpan(Sha1Hash.NumBytes));
 		}
 	}
 
@@ -212,10 +207,10 @@ namespace EpicGames.Core
 	sealed class Sha1HashJsonConverter : JsonConverter<Sha1Hash>
 	{
 		/// <inheritdoc/>
-		public override Sha1Hash Read(ref Utf8JsonReader Reader, Type TypeToConvert, JsonSerializerOptions Options) => Sha1Hash.Parse(Reader.ValueSpan);
+		public override Sha1Hash Read(ref Utf8JsonReader reader, Type typeToConvert, JsonSerializerOptions options) => Sha1Hash.Parse(reader.ValueSpan);
 
 		/// <inheritdoc/>
-		public override void Write(Utf8JsonWriter Writer, Sha1Hash Value, JsonSerializerOptions Options) => Writer.WriteStringValue(Value.ToUtf8String().Span);
+		public override void Write(Utf8JsonWriter writer, Sha1Hash value, JsonSerializerOptions options) => writer.WriteStringValue(value.ToUtf8String().Span);
 	}
 
 	/// <summary>
@@ -224,15 +219,15 @@ namespace EpicGames.Core
 	sealed class Sha1HashTypeConverter : TypeConverter
 	{
 		/// <inheritdoc/>
-		public override bool CanConvertFrom(ITypeDescriptorContext Context, Type SourceType)
+		public override bool CanConvertFrom(ITypeDescriptorContext context, Type sourceType)
 		{
-			return SourceType == typeof(string);
+			return sourceType == typeof(string);
 		}
 
 		/// <inheritdoc/>
-		public override object ConvertFrom(ITypeDescriptorContext Context, CultureInfo Culture, object Value)
+		public override object ConvertFrom(ITypeDescriptorContext context, CultureInfo culture, object value)
 		{
-			return Sha1Hash.Parse((string)Value);
+			return Sha1Hash.Parse((string)value);
 		}
 	}
 }
diff --git a/Engine/Source/Programs/Shared/EpicGames.Core/StartupTraceListener.cs b/Engine/Source/Programs/Shared/EpicGames.Core/StartupTraceListener.cs
--- a/Engine/Source/Programs/Shared/EpicGames.Core/StartupTraceListener.cs
+++ b/Engine/Source/Programs/Shared/EpicGames.Core/StartupTraceListener.cs
@@ -1,16 +1,7 @@
 // Copyright Epic Games, Inc. All Rights Reserved.
 
-using System;
-using System.Collections.Generic;
-using System.Linq;
-using System.Text;
-using System.IO;
-using System.CodeDom.Compiler;
-using Microsoft.CSharp;
-using Microsoft.Win32;
-using System.Reflection;
 using System.Diagnostics;
-using EpicGames.Core;
+using System.Text;
 
 namespace EpicGames.Core
 {
@@ -19,41 +10,41 @@ namespace EpicGames.Core
 	/// </summary>
 	public class StartupTraceListener : TraceListener
 	{
-		StringBuilder Buffer = new StringBuilder();
+		readonly StringBuilder _buffer = new StringBuilder();
 
 		/// <summary>
 		/// Copy the contents of the buffered output to another trace listener
 		/// </summary>
-		/// <param name="Other">The trace listener to receive the buffered output</param>
-		public void CopyTo(TraceListener Other)
+		/// <param name="other">The trace listener to receive the buffered output</param>
+		public void CopyTo(TraceListener other)
 		{
-			foreach(string Line in Buffer.ToString().Split("\n"))
+			foreach(string line in _buffer.ToString().Split("\n"))
 			{
-				Other.WriteLine(Line);
+				other.WriteLine(line);
 			}
 		}
 
 		/// <summary>
 		/// Write a message to the buffer
 		/// </summary>
-		/// <param name="Message">The message to write</param>
-		public override void Write(string Message)
+		/// <param name="message">The message to write</param>
+		public override void Write(string message)
 		{
 			if(NeedIndent)
 			{
 				WriteIndent();
 			}
-			Buffer.Append(Message);
+			_buffer.Append(message);
 		}
 
 		/// <summary>
 		/// Write a message to the buffer, followed by a newline
 		/// </summary>
-		/// <param name="Message">The message to write</param>
-		public override void WriteLine(string Message)
+		/// <param name="message">The message to write</param>
+		public override void WriteLine(string message)
 		{
-			Write(Message);
-			Buffer.Append("\n");
+			Write(message);
+			_buffer.Append("\n");
 		}
 	}
 }
diff --git a/Engine/Source/Programs/Shared/EpicGames.Core/StringComparerUE.cs b/Engine/Source/Programs/Shared/EpicGames.Core/StringComparerUE.cs
--- a/Engine/Source/Programs/Shared/EpicGames.Core/StringComparerUE.cs
+++ b/Engine/Source/Programs/Shared/EpicGames.Core/StringComparerUE.cs
@@ -5,7 +5,6 @@ using System.Collections.Generic;
 
 namespace EpicGames.Core
 {
-
 	/// <summary>
 	/// Comparer for String objects.  However, it implements UE style ignore case compare
 	/// </summary>
@@ -29,43 +28,43 @@ namespace EpicGames.Core
 		/// <summary>
 		/// Constructor
 		/// </summary>
-		/// <param name="ComparisonType">Type of comparison to perform</param>
-		public StringComparerUE(StringComparison ComparisonType)
+		/// <param name="comparisonType">Type of comparison to perform</param>
+		public StringComparerUE(StringComparison comparisonType)
 		{
-			this.ComparisonType = ComparisonType;
+			ComparisonType = comparisonType;
 		}
 
 		/// <inheritdoc/>
-		public bool Equals(string? X, string? Y)
+		public bool Equals(string? x, string? y)
 		{
-			return string.Equals(X, Y, ComparisonType);
+			return String.Equals(x, y, ComparisonType);
 		}
 
 		/// <inheritdoc/>
 		public int GetHashCode(string obj)
 		{
-			return string.GetHashCode(obj, ComparisonType);
+			return String.GetHashCode(obj, ComparisonType);
 		}
 
 		/// <inheritdoc/>
-		public int Compare(string? X, string? Y)
+		public int Compare(string? x, string? y)
 		{
 			if (ComparisonType == StringComparison.OrdinalIgnoreCase)
 			{
-				if (X == null)
+				if (x == null)
 				{
-					return Y == null ? 0 : -1;
+					return y == null ? 0 : -1;
 				}
-				else if (Y == null)
+				else if (y == null)
 				{
 					return 1;
 				}
 
-				return StringUtils.CompareIgnoreCaseUE(X.AsSpan(), Y.AsSpan());
+				return StringUtils.CompareIgnoreCaseUe(x.AsSpan(), y.AsSpan());
 			}
 			else
 			{
-				return string.Compare(X, Y, ComparisonType);
+				return String.Compare(x, y, ComparisonType);
 			}
 		}
 	}
diff --git a/Engine/Source/Programs/Shared/EpicGames.Core/StringUtils.cs b/Engine/Source/Programs/Shared/EpicGames.Core/StringUtils.cs
--- a/Engine/Source/Programs/Shared/EpicGames.Core/StringUtils.cs
+++ b/Engine/Source/Programs/Shared/EpicGames.Core/StringUtils.cs
@@ -5,7 +5,6 @@ using System.Collections.Generic;
 using System.Diagnostics.CodeAnalysis;
 using System.Linq;
 using System.Text;
-using System.Threading.Tasks;
 
 namespace EpicGames.Core
 {
@@ -14,195 +13,199 @@ namespace EpicGames.Core
 		/// <summary>
 		/// Array mapping from ascii index to hexadecimal digits.
 		/// </summary>
-		static sbyte[] HexDigits;
+		static readonly sbyte[] s_hexDigits;
 
 		/// <summary>
 		/// Hex digits to utf8 byte
 		/// </summary>
-		static byte[] HexDigitToUtf8Byte = Encoding.UTF8.GetBytes("0123456789abcdef");
+		static readonly byte[] s_hexDigitToUtf8Byte = Encoding.UTF8.GetBytes("0123456789abcdef");
 
 		/// <summary>
 		/// Array mapping human readable size of bytes, 1024^x. long max is within the range of Exabytes.
 		/// </summary>
-		static string[] ByteSizes = { "B", "KB", "MB", "GB", "TB", "PB", "EB" };
+		static readonly string[] s_byteSizes = { "B", "KB", "MB", "GB", "TB", "PB", "EB" };
 
 		/// <summary>
 		/// Static constructor. Initializes the HexDigits array.
 		/// </summary>
 		static StringUtils()
 		{
-			HexDigits = new sbyte[256];
-			for (int Idx = 0; Idx < 256; Idx++)
+			s_hexDigits = new sbyte[256];
+			for (int idx = 0; idx < 256; idx++)
 			{
-				HexDigits[Idx] = -1;
+				s_hexDigits[idx] = -1;
 			}
-			for (int Idx = '0'; Idx <= '9'; Idx++)
+			for (int idx = '0'; idx <= '9'; idx++)
 			{
-				HexDigits[Idx] = (sbyte)(Idx - '0');
+				s_hexDigits[idx] = (sbyte)(idx - '0');
 			}
-			for (int Idx = 'a'; Idx <= 'f'; Idx++)
+			for (int idx = 'a'; idx <= 'f'; idx++)
 			{
-				HexDigits[Idx] = (sbyte)(10 + Idx - 'a');
+				s_hexDigits[idx] = (sbyte)(10 + idx - 'a');
 			}
-			for (int Idx = 'A'; Idx <= 'F'; Idx++)
+			for (int idx = 'A'; idx <= 'F'; idx++)
 			{
-				HexDigits[Idx] = (sbyte)(10 + Idx - 'A');
+				s_hexDigits[idx] = (sbyte)(10 + idx - 'A');
 			}
 		}
 
 		/// <summary>
 		/// Indents a string by a given indent
 		/// </summary>
-		/// <param name="Text">The text to indent</param>
-		/// <param name="Indent">The indent to add to each line</param>
+		/// <param name="text">The text to indent</param>
+		/// <param name="indent">The indent to add to each line</param>
 		/// <returns>The indented string</returns>
-		public static string Indent(string Text, string Indent)
+		public static string Indent(string text, string indent)
 		{
-			string Result = "";
-			if(Text.Length > 0)
+			string result = "";
+			if(text.Length > 0)
 			{
-				Result = Indent + Text.Replace("\n", "\n" + Indent);
+				result = indent + text.Replace("\n", "\n" + indent);
 			}
-			return Result;
+			return result;
 		}
 
 		/// <summary>
 		/// Expand all the property references (of the form $(PropertyName)) in a string.
 		/// </summary>
-		/// <param name="Text">The input string to expand properties in</param>
-		/// <param name="Properties">Dictionary of properties to expand</param>
+		/// <param name="text">The input string to expand properties in</param>
+		/// <param name="properties">Dictionary of properties to expand</param>
 		/// <returns>The expanded string</returns>
-		public static string ExpandProperties(string Text, Dictionary<string, string> Properties)
+		public static string ExpandProperties(string text, Dictionary<string, string> properties)
 		{
-			return ExpandProperties(Text, Name => { Properties.TryGetValue(Name, out string? Value); return Value; });
+			return ExpandProperties(text, name => 
+			{ 
+				properties.TryGetValue(name, out string? value); 
+				return value; 
+			});
 		}
 
 		/// <summary>
 		/// Expand all the property references (of the form $(PropertyName)) in a string.
 		/// </summary>
-		/// <param name="Text">The input string to expand properties in</param>
-		/// <param name="GetPropertyValue">Delegate to retrieve a property value</param>
+		/// <param name="text">The input string to expand properties in</param>
+		/// <param name="getPropertyValue">Delegate to retrieve a property value</param>
 		/// <returns>The expanded string</returns>
-		public static string ExpandProperties(string Text, Func<string, string?> GetPropertyValue)
+		public static string ExpandProperties(string text, Func<string, string?> getPropertyValue)
 		{
-			string Result = Text;
-			for (int Idx = Result.IndexOf("$("); Idx != -1; Idx = Result.IndexOf("$(", Idx))
+			string result = text;
+			for (int idx = result.IndexOf("$("); idx != -1; idx = result.IndexOf("$(", idx))
 			{
 				// Find the end of the variable name
-				int EndIdx = Result.IndexOf(')', Idx + 2);
-				if (EndIdx == -1)
+				int endIdx = result.IndexOf(')', idx + 2);
+				if (endIdx == -1)
 				{
 					break;
 				}
 
 				// Extract the variable name from the string
-				string Name = Result.Substring(Idx + 2, EndIdx - (Idx + 2));
+				string name = result.Substring(idx + 2, endIdx - (idx + 2));
 
 				// Check if we've got a value for this variable
-				string? Value = GetPropertyValue(Name);
-				if (Value == null)
+				string? value = getPropertyValue(name);
+				if (value == null)
 				{
 					// Do not expand it; must be preprocessing the script.
-					Idx = EndIdx;
+					idx = endIdx;
 				}
 				else
 				{
 					// Replace the variable, or skip past it
-					Result = Result.Substring(0, Idx) + Value + Result.Substring(EndIdx + 1);
+					result = result.Substring(0, idx) + value + result.Substring(endIdx + 1);
 
 					// Make sure we skip over the expanded variable; we don't want to recurse on it.
-					Idx += Value.Length;
+					idx += value.Length;
 				}
 			}
-			return Result;
+			return result;
 		}
 
-		/// <inheritdoc cref="WordWrap(string, int, int, int)">
-		public static IEnumerable<string> WordWrap(string Text, int MaxWidth)
+		/// <inheritdoc cref="WordWrap(String, Int32, Int32, Int32)">
+		public static IEnumerable<string> WordWrap(string text, int maxWidth)
 		{
-			return WordWrap(Text, 0, 0, MaxWidth);
+			return WordWrap(text, 0, 0, maxWidth);
 		}
 
 		/// <summary>
 		/// Takes a given sentence and wraps it on a word by word basis so that no line exceeds the set maximum line length. Words longer than a line 
 		/// are broken up. Returns the sentence as a list of individual lines.
 		/// </summary>
-		/// <param name="Text">The text to be wrapped</param>
-		/// <param name="InitialIndent">Indent for the first line</param>
-		/// <param name="HangingIndent">Indent for subsequent lines</param>
-		/// <param name="MaxWidth">The maximum (non negative) length of the returned sentences</param>
-		public static IEnumerable<string> WordWrap(string Text, int InitialIndent, int HangingIndent, int MaxWidth)
+		/// <param name="text">The text to be wrapped</param>
+		/// <param name="initialIndent">Indent for the first line</param>
+		/// <param name="hangingIndent">Indent for subsequent lines</param>
+		/// <param name="maxWidth">The maximum (non negative) length of the returned sentences</param>
+		public static IEnumerable<string> WordWrap(string text, int initialIndent, int hangingIndent, int maxWidth)
 		{
-			StringBuilder Builder = new StringBuilder();
+			StringBuilder builder = new StringBuilder();
 
-			int MinIdx = 0;
-			for (int LineIdx = 0; MinIdx < Text.Length; LineIdx++)
+			int minIdx = 0;
+			for (int lineIdx = 0; minIdx < text.Length; lineIdx++)
 			{
-				int Indent = (LineIdx == 0) ? InitialIndent : HangingIndent;
-				int MaxWidthForLine = MaxWidth - Indent;
-				int MaxIdx = GetWordWrapLineEnd(Text, MinIdx, MaxWidthForLine);
+				int indent = (lineIdx == 0) ? initialIndent : hangingIndent;
+				int maxWidthForLine = maxWidth - indent;
+				int maxIdx = GetWordWrapLineEnd(text, minIdx, maxWidthForLine);
 
-				int PrintMaxIdx = MaxIdx;
-				while (PrintMaxIdx > MinIdx && Char.IsWhiteSpace(Text[PrintMaxIdx - 1]))
+				int printMaxIdx = maxIdx;
+				while (printMaxIdx > minIdx && Char.IsWhiteSpace(text[printMaxIdx - 1]))
 				{
-					PrintMaxIdx--;
+					printMaxIdx--;
 				}
 
-				Builder.Clear();
-				Builder.Append(' ', Indent);
-				Builder.Append(Text, MinIdx, PrintMaxIdx - MinIdx);
-				yield return Builder.ToString();
+				builder.Clear();
+				builder.Append(' ', indent);
+				builder.Append(text, minIdx, printMaxIdx - minIdx);
+				yield return builder.ToString();
 
-				MinIdx = MaxIdx;
+				minIdx = maxIdx;
 			}
 		}
 
 		/// <summary>
 		/// Gets the next character index to end a word-wrapped line on
 		/// </summary>
-		static int GetWordWrapLineEnd(string Text, int MinIdx, int MaxWidth)
+		static int GetWordWrapLineEnd(string text, int minIdx, int maxWidth)
 		{
-			MaxWidth = Math.Min(MaxWidth, Text.Length - MinIdx);
+			maxWidth = Math.Min(maxWidth, text.Length - minIdx);
 
-			int MaxIdx = Text.IndexOf('\n', MinIdx, MaxWidth);
-			if (MaxIdx == -1)
+			int maxIdx = text.IndexOf('\n', minIdx, maxWidth);
+			if (maxIdx == -1)
 			{
-				MaxIdx = MinIdx + MaxWidth;
+				maxIdx = minIdx + maxWidth;
 			}
 			else
 			{
-				return MaxIdx + 1;
+				return maxIdx + 1;
 			}
 
-			if (MaxIdx == Text.Length)
+			if (maxIdx == text.Length)
 			{
-				return MaxIdx;
+				return maxIdx;
 			}
-			else if (Char.IsWhiteSpace(Text[MaxIdx]))
+			else if (Char.IsWhiteSpace(text[maxIdx]))
 			{
-				for (; ; MaxIdx++)
+				for (; ; maxIdx++)
 				{
-					if (MaxIdx == Text.Length)
+					if (maxIdx == text.Length)
 					{
-						return MaxIdx;
+						return maxIdx;
 					}
-					if (Text[MaxIdx] != ' ')
+					if (text[maxIdx] != ' ')
 					{
-						return MaxIdx;
+						return maxIdx;
 					}
 				}
 			}
 			else
 			{
-				for(int TryMaxIdx = MaxIdx; ; TryMaxIdx--)
+				for(int tryMaxIdx = maxIdx; ; tryMaxIdx--)
 				{
-					if(TryMaxIdx == MinIdx)
+					if(tryMaxIdx == minIdx)
 					{
-						return MaxIdx;
+						return maxIdx;
 					}
-					if (Text[TryMaxIdx - 1] == ' ')
+					if (text[tryMaxIdx - 1] == ' ')
 					{
-						return TryMaxIdx;
+						return tryMaxIdx;
 					}
 				}
 			}
@@ -211,319 +214,319 @@ namespace EpicGames.Core
 		/// <summary>
 		/// Extension method to allow formatting a string to a stringbuilder and appending a newline
 		/// </summary>
-		/// <param name="Builder">The string builder</param>
-		/// <param name="Format">Format string, as used for StringBuilder.AppendFormat</param>
-		/// <param name="Args">Arguments for the format string</param>
-		public static void AppendLine(this StringBuilder Builder, string Format, params object[] Args)
+		/// <param name="builder">The string builder</param>
+		/// <param name="format">Format string, as used for StringBuilder.AppendFormat</param>
+		/// <param name="args">Arguments for the format string</param>
+		public static void AppendLine(this StringBuilder builder, string format, params object[] args)
 		{
-			Builder.AppendFormat(Format, Args);
-			Builder.AppendLine();
+			builder.AppendFormat(format, args);
+			builder.AppendLine();
 		}
 
 		/// <summary>
 		/// Formats a list of strings in the style "1, 2, 3 and 4"
 		/// </summary>
-		/// <param name="Arguments">List of strings to format</param>
-		/// <param name="Conjunction">Conjunction to use between the last two items in the list (eg. "and" or "or")</param>
+		/// <param name="arguments">List of strings to format</param>
+		/// <param name="conjunction">Conjunction to use between the last two items in the list (eg. "and" or "or")</param>
 		/// <returns>Formatted list of strings</returns>
-		public static string FormatList(string[] Arguments, string Conjunction = "and")
+		public static string FormatList(string[] arguments, string conjunction = "and")
 		{
-			StringBuilder Result = new StringBuilder();
-			if (Arguments.Length > 0)
+			StringBuilder result = new StringBuilder();
+			if (arguments.Length > 0)
 			{
-				Result.Append(Arguments[0]);
-				for (int Idx = 1; Idx < Arguments.Length; Idx++)
+				result.Append(arguments[0]);
+				for (int idx = 1; idx < arguments.Length; idx++)
 				{
-					if (Idx == Arguments.Length - 1)
+					if (idx == arguments.Length - 1)
 					{
-						Result.AppendFormat(" {0} ", Conjunction);
+						result.AppendFormat(" {0} ", conjunction);
 					}
 					else
 					{
-						Result.Append(", ");
+						result.Append(", ");
 					}
-					Result.Append(Arguments[Idx]);
+					result.Append(arguments[idx]);
 				}
 			}
-			return Result.ToString();
+			return result.ToString();
 		}
 
 		/// <summary>
 		/// Formats a list of strings in the style "1, 2, 3 and 4"
 		/// </summary>
-		/// <param name="Arguments">List of strings to format</param>
-		/// <param name="Conjunction">Conjunction to use between the last two items in the list (eg. "and" or "or")</param>
+		/// <param name="arguments">List of strings to format</param>
+		/// <param name="conjunction">Conjunction to use between the last two items in the list (eg. "and" or "or")</param>
 		/// <returns>Formatted list of strings</returns>
-		public static string FormatList(IEnumerable<string> Arguments, string Conjunction = "and")
+		public static string FormatList(IEnumerable<string> arguments, string conjunction = "and")
 		{
-			return FormatList(Arguments.ToArray(), Conjunction);
+			return FormatList(arguments.ToArray(), conjunction);
 		}
 
 		/// <summary>
 		/// Formats a list of items
 		/// </summary>
-		/// <param name="Items">Array of items</param>
-		/// <param name="MaxCount">Maximum number of items to include in the list</param>
+		/// <param name="items">Array of items</param>
+		/// <param name="maxCount">Maximum number of items to include in the list</param>
 		/// <returns>Formatted list of items</returns>
-		public static string FormatList(string[] Items, int MaxCount)
+		public static string FormatList(string[] items, int maxCount)
 		{
-			if (Items.Length == 0)
+			if (items.Length == 0)
 			{
 				return "unknown";
 			}
-			else if (Items.Length == 1)
+			else if (items.Length == 1)
 			{
-				return Items[0];
+				return items[0];
 			}
-			else if (Items.Length <= MaxCount)
+			else if (items.Length <= maxCount)
 			{
-				return $"{String.Join(", ", Items.Take(Items.Length - 1))} and {Items.Last()}";
+				return $"{String.Join(", ", items.Take(items.Length - 1))} and {items.Last()}";
 			}
 			else
 			{
-				return $"{String.Join(", ", Items.Take(MaxCount - 1))} and {Items.Length - (MaxCount - 1)} others";
+				return $"{String.Join(", ", items.Take(maxCount - 1))} and {items.Length - (maxCount - 1)} others";
 			}
 		}
 
 		/// <summary>
 		/// Parses a hexadecimal digit
 		/// </summary>
-		/// <param name="Character">Character to parse</param>
+		/// <param name="character">Character to parse</param>
 		/// <returns>Value of this digit, or -1 if invalid</returns>
-		public static int GetHexDigit(byte Character)
+		public static int GetHexDigit(byte character)
 		{
-			return HexDigits[Character];
+			return s_hexDigits[character];
 		}
 
 		/// <summary>
 		/// Parses a hexadecimal digit
 		/// </summary>
-		/// <param name="Character">Character to parse</param>
+		/// <param name="character">Character to parse</param>
 		/// <returns>Value of this digit, or -1 if invalid</returns>
-		public static int GetHexDigit(char Character)
+		public static int GetHexDigit(char character)
 		{
-			return HexDigits[Math.Min((uint)Character, 127)];
+			return s_hexDigits[Math.Min((uint)character, 127)];
 		}
 
 		/// <summary>
 		/// Parses a hexadecimal string into an array of bytes
 		/// </summary>
 		/// <returns>Array of bytes</returns>
-		public static byte[] ParseHexString(string Text)
+		public static byte[] ParseHexString(string text)
 		{
-			byte[]? Bytes;
-			if(!TryParseHexString(Text, out Bytes))
+			byte[]? bytes;
+			if(!TryParseHexString(text, out bytes))
 			{
-				throw new FormatException(String.Format("Invalid hex string: '{0}'", Text));
+				throw new FormatException(String.Format("Invalid hex string: '{0}'", text));
 			}
-			return Bytes;
+			return bytes;
 		}
 
 		/// <summary>
 		/// Parses a hexadecimal string into an array of bytes
 		/// </summary>
 		/// <returns>Array of bytes</returns>
-		public static byte[] ParseHexString(ReadOnlySpan<byte> Text)
+		public static byte[] ParseHexString(ReadOnlySpan<byte> text)
 		{
-			byte[]? Bytes;
-			if (!TryParseHexString(Text, out Bytes))
+			byte[]? bytes;
+			if (!TryParseHexString(text, out bytes))
 			{
-				throw new FormatException($"Invalid hex string: '{Encoding.UTF8.GetString(Text)}'");
+				throw new FormatException($"Invalid hex string: '{Encoding.UTF8.GetString(text)}'");
 			}
-			return Bytes;
+			return bytes;
 		}
 
 		/// <summary>
 		/// Parses a hexadecimal string into an array of bytes
 		/// </summary>
-		/// <param name="Text">Text to parse</param>
+		/// <param name="text">Text to parse</param>
 		/// <returns></returns>
-		public static bool TryParseHexString(string Text, [NotNullWhen(true)] out byte[]? OutBytes)
+		public static bool TryParseHexString(string text, [NotNullWhen(true)] out byte[]? outBytes)
 		{
-			if((Text.Length & 1) != 0)
+			if((text.Length & 1) != 0)
 			{
-				OutBytes = null;
+				outBytes = null;
 				return false;
 			}
 
-			byte[] Bytes = new byte[Text.Length / 2];
-			for(int Idx = 0; Idx < Text.Length; Idx += 2)
+			byte[] bytes = new byte[text.Length / 2];
+			for(int idx = 0; idx < text.Length; idx += 2)
 			{
-				int Value = (GetHexDigit(Text[Idx]) << 4) | GetHexDigit(Text[Idx + 1]);
-				if(Value < 0)
+				int value = (GetHexDigit(text[idx]) << 4) | GetHexDigit(text[idx + 1]);
+				if(value < 0)
 				{
-					OutBytes = null;
+					outBytes = null;
 					return false;
 				}
-				Bytes[Idx / 2] = (byte)Value;
+				bytes[idx / 2] = (byte)value;
 			}
-			OutBytes = Bytes;
+			outBytes = bytes;
 			return true;
 		}
 
 		/// <summary>
 		/// Parses a hexadecimal string into an array of bytes
 		/// </summary>
-		/// <param name="Text">Text to parse</param>
+		/// <param name="text">Text to parse</param>
 		/// <returns></returns>
-		public static bool TryParseHexString(ReadOnlySpan<byte> Text, [NotNullWhen(true)] out byte[]? OutBytes)
+		public static bool TryParseHexString(ReadOnlySpan<byte> text, [NotNullWhen(true)] out byte[]? outBytes)
 		{
-			if ((Text.Length & 1) != 0)
+			if ((text.Length & 1) != 0)
 			{
-				OutBytes = null;
+				outBytes = null;
 				return false;
 			}
 
-			byte[] Bytes = new byte[Text.Length / 2];
-			for (int Idx = 0; Idx < Text.Length; Idx += 2)
+			byte[] bytes = new byte[text.Length / 2];
+			for (int idx = 0; idx < text.Length; idx += 2)
 			{
-				int Value = ParseHexByte(Text, Idx);
-				if (Value < 0)
+				int value = ParseHexByte(text, idx);
+				if (value < 0)
 				{
-					OutBytes = null;
+					outBytes = null;
 					return false;
 				}
-				Bytes[Idx / 2] = (byte)Value;
+				bytes[idx / 2] = (byte)value;
 			}
-			OutBytes = Bytes;
+			outBytes = bytes;
 			return true;
 		}
 
 		/// <summary>
 		/// Parse a hex byte from the given offset into a span of utf8 characters
 		/// </summary>
-		/// <param name="Text">The text to parse</param>
-		/// <param name="Idx">Index within the text to parse</param>
+		/// <param name="text">The text to parse</param>
+		/// <param name="idx">Index within the text to parse</param>
 		/// <returns>The parsed value, or a negative value on error</returns>
-		public static int ParseHexByte(ReadOnlySpan<byte> Text, int Idx)
+		public static int ParseHexByte(ReadOnlySpan<byte> text, int idx)
 		{
-			return ((int)HexDigits[Text[Idx]] << 4) | ((int)HexDigits[Text[Idx + 1]]);
+			return ((int)s_hexDigits[text[idx]] << 4) | ((int)s_hexDigits[text[idx + 1]]);
 		}
 
 		/// <summary>
 		/// Formats an array of bytes as a hexadecimal string
 		/// </summary>
-		/// <param name="Bytes">An array of bytes</param>
+		/// <param name="bytes">An array of bytes</param>
 		/// <returns>String representation of the array</returns>
-		public static string FormatHexString(byte[] Bytes)
+		public static string FormatHexString(byte[] bytes)
 		{
-			return FormatHexString(Bytes.AsSpan());
+			return FormatHexString(bytes.AsSpan());
 		}
 
 		/// <summary>
 		/// Formats an array of bytes as a hexadecimal string
 		/// </summary>
-		/// <param name="Bytes">An array of bytes</param>
+		/// <param name="bytes">An array of bytes</param>
 		/// <returns>String representation of the array</returns>
-		public static string FormatHexString(ReadOnlySpan<byte> Bytes)
+		public static string FormatHexString(ReadOnlySpan<byte> bytes)
 		{
 			const string HexDigits = "0123456789abcdef";
 
-			char[] Characters = new char[Bytes.Length * 2];
-			for (int Idx = 0; Idx < Bytes.Length; Idx++)
+			char[] characters = new char[bytes.Length * 2];
+			for (int idx = 0; idx < bytes.Length; idx++)
 			{
-				Characters[Idx * 2 + 0] = HexDigits[Bytes[Idx] >> 4];
-				Characters[Idx * 2 + 1] = HexDigits[Bytes[Idx] & 15];
+				characters[idx * 2 + 0] = HexDigits[bytes[idx] >> 4];
+				characters[idx * 2 + 1] = HexDigits[bytes[idx] & 15];
 			}
-			return new string(Characters);
+			return new string(characters);
 		}
 
 		/// <summary>
 		/// Formats an array of bytes as a hexadecimal string
 		/// </summary>
-		/// <param name="Bytes">An array of bytes</param>
+		/// <param name="bytes">An array of bytes</param>
 		/// <returns>String representation of the array</returns>
-		public static Utf8String FormatUtf8HexString(ReadOnlySpan<byte> Bytes)
+		public static Utf8String FormatUtf8HexString(ReadOnlySpan<byte> bytes)
 		{
-			byte[] Characters = new byte[Bytes.Length * 2];
-			for (int Idx = 0; Idx < Bytes.Length; Idx++)
+			byte[] characters = new byte[bytes.Length * 2];
+			for (int idx = 0; idx < bytes.Length; idx++)
 			{
-				Characters[Idx * 2 + 0] = HexDigitToUtf8Byte[Bytes[Idx] >> 4];
-				Characters[Idx * 2 + 1] = HexDigitToUtf8Byte[Bytes[Idx] & 15];
+				characters[idx * 2 + 0] = s_hexDigitToUtf8Byte[bytes[idx] >> 4];
+				characters[idx * 2 + 1] = s_hexDigitToUtf8Byte[bytes[idx] & 15];
 			}
-			return new Utf8String(Characters);
+			return new Utf8String(characters);
 		}
 
 		/// <summary>
 		/// Quotes a string as a command line argument
 		/// </summary>
-		/// <param name="String">The string to quote</param>
+		/// <param name="str">The string to quote</param>
 		/// <returns>The quoted argument if it contains any spaces, otherwise the original string</returns>
-		public static string QuoteArgument(this string String)
+		public static string QuoteArgument(this string str)
 		{
-			if (String.Contains(' '))
+			if (str.Contains(' '))
 			{
-				return $"\"{String}\"";
+				return $"\"{str}\"";
 			}
 			else
 			{
-				return String;
+				return str;
 			}
 		}
 
 		/// <summary>
 		/// Formats bytes into a human readable string
 		/// </summary>
-		/// <param name="Bytes">The total number of bytes</param>
-		/// <param name="DecimalPlaces">The number of decimal places to round the resulting value</param>
+		/// <param name="bytes">The total number of bytes</param>
+		/// <param name="decimalPlaces">The number of decimal places to round the resulting value</param>
 		/// <returns>Human readable string based on the value of Bytes</returns>
-		public static string FormatBytesString(long Bytes, int DecimalPlaces = 2)
+		public static string FormatBytesString(long bytes, int decimalPlaces = 2)
 		{
-			if (Bytes == 0)
+			if (bytes == 0)
 			{
-				return $"0 {ByteSizes[0]}";
+				return $"0 {s_byteSizes[0]}";
 			}
-			long BytesAbs = Math.Abs(Bytes);
-			int Power = Convert.ToInt32(Math.Floor(Math.Log(BytesAbs, 1024)));
-			double Value = Math.Round(BytesAbs / Math.Pow(1024, Power), DecimalPlaces);
-			return $"{(Math.Sign(Bytes) * Value)} {ByteSizes[Power]}";
+			long bytesAbs = Math.Abs(bytes);
+			int power = Convert.ToInt32(Math.Floor(Math.Log(bytesAbs, 1024)));
+			double value = Math.Round(bytesAbs / Math.Pow(1024, power), decimalPlaces);
+			return $"{(Math.Sign(bytes) * value)} {s_byteSizes[power]}";
 		}
 
 		/// <summary>
 		/// Converts a bytes string into bytes. E.g 1.2KB -> 1229
 		/// </summary>
-		/// <param name="BytesString"></param>
+		/// <param name="bytesString"></param>
 		/// <returns></returns>
-		public static long ParseBytesString( string BytesString )
+		public static long ParseBytesString( string bytesString )
 		{
-			BytesString = BytesString.Trim();
+			bytesString = bytesString.Trim();
 
-			int Power = ByteSizes.FindIndex( s => (s != ByteSizes[0]) && BytesString.EndsWith(s, StringComparison.InvariantCultureIgnoreCase ) ); // need to handle 'B' suffix separately
-			if (Power == -1 && BytesString.EndsWith(ByteSizes[0]))
+			int power = s_byteSizes.FindIndex( s => (s != s_byteSizes[0]) && bytesString.EndsWith(s, StringComparison.InvariantCultureIgnoreCase ) ); // need to handle 'B' suffix separately
+			if (power == -1 && bytesString.EndsWith(s_byteSizes[0]))
 			{
-				Power = 0;
+				power = 0;
 			}
-			if (Power != -1)
+			if (power != -1)
 			{
-				BytesString = BytesString.Substring(0, BytesString.Length - ByteSizes[Power].Length );
+				bytesString = bytesString.Substring(0, bytesString.Length - s_byteSizes[power].Length );
 			}
 
-			double Value = double.Parse(BytesString);
-			if (Power > 0 )
+			double value = Double.Parse(bytesString);
+			if (power > 0 )
 			{
-				Value *= Math.Pow(1024, Power);
+				value *= Math.Pow(1024, power);
 			}
 
-			return (long)Math.Round(Value);
+			return (long)Math.Round(value);
 		}
 
 		/// <summary>
 		/// Converts a bytes string into bytes. E.g 1.5KB -> 1536
 		/// </summary>
-		/// <param name="BytesString"></param>
+		/// <param name="bytesString"></param>
 		/// <returns></returns>
-		public static bool TryParseBytesString( string BytesString, out long? Bytes )
+		public static bool TryParseBytesString( string bytesString, out long? bytes )
 		{
 			try
 			{
-				Bytes = ParseBytesString(BytesString);
+				bytes = ParseBytesString(bytesString);
 				return true;
 			}
 			catch(Exception)
 			{
 			}
 
-			Bytes = null;
+			bytes = null;
 			return false;
 		}
 
@@ -531,109 +534,109 @@ namespace EpicGames.Core
 		/// Parses a string to remove VT100 escape codes
 		/// </summary>
 		/// <returns></returns>
-		public static string ParseEscapeCodes(string Line)
+		public static string ParseEscapeCodes(string line)
 		{
-			char EscapeChar = '\u001b';
+			char escapeChar = '\u001b';
 
-			int Index = Line.IndexOf(EscapeChar);
-			if (Index != -1)
+			int index = line.IndexOf(escapeChar);
+			if (index != -1)
 			{
-				int LastIndex = 0;
+				int lastIndex = 0;
 
-				StringBuilder Result = new StringBuilder();
+				StringBuilder result = new StringBuilder();
 				for (; ; )
 				{
-					Result.Append(Line, LastIndex, Index - LastIndex);
+					result.Append(line, lastIndex, index - lastIndex);
 
-					while (Index < Line.Length)
+					while (index < line.Length)
 					{
-						char Char = Line[Index];
-						if ((Char >= 'a' && Char <= 'z') || (Char >= 'A' && Char <= 'Z'))
+						char character = line[index];
+						if ((character >= 'a' && character <= 'z') || (character >= 'A' && character <= 'Z'))
 						{
-							Index++;
+							index++;
 							break;
 						}
-						Index++;
+						index++;
 					}
 
-					LastIndex = Index;
+					lastIndex = index;
 
-					Index = Line.IndexOf(EscapeChar, Index);
-					if (Index == -1)
+					index = line.IndexOf(escapeChar, index);
+					if (index == -1)
 					{
 						break;
 					}
 				}
-				Result.Append(Line, LastIndex, Line.Length - LastIndex);
+				result.Append(line, lastIndex, line.Length - lastIndex);
 
-				Line = Result.ToString();
+				line = result.ToString();
 			}
 
-			return Line;
+			return line;
 		}
 
 		/// <summary>
 		/// Truncates the given string to the maximum length, appending an elipsis if it is longer than allowed.
 		/// </summary>
-		/// <param name="Text"></param>
-		/// <param name="MaxLength"></param>
+		/// <param name="text"></param>
+		/// <param name="maxLength"></param>
 		/// <returns></returns>
-		public static string Truncate(string Text, int MaxLength)
+		public static string Truncate(string text, int maxLength)
 		{
-			if (Text.Length > MaxLength)
+			if (text.Length > maxLength)
 			{
-				Text = Text.Substring(0, MaxLength - 3) + "...";
+				text = text.Substring(0, maxLength - 3) + "...";
 			}
-			return Text;
+			return text;
 		}
 
 		/// <summary>
 		/// Compare two strings using UnrealEngine's ignore case algorithm
 		/// </summary>
-		/// <param name="X">First string to compare</param>
-		/// <param name="Y">Second string to compare</param>
+		/// <param name="x">First string to compare</param>
+		/// <param name="y">Second string to compare</param>
 		/// <returns>Less than zero if X < Y, zero if X == Y, and greater than zero if X > y</returns>
-		public static int CompareIgnoreCaseUE(ReadOnlySpan<char> X, ReadOnlySpan<char> Y)
+		public static int CompareIgnoreCaseUe(ReadOnlySpan<char> x, ReadOnlySpan<char> y)
 		{
-			int Length = X.Length < Y.Length ? X.Length : Y.Length;
+			int length = x.Length < y.Length ? x.Length : y.Length;
 
-			for (int Index = 0; Index < Length; ++Index)
+			for (int index = 0; index < length; ++index)
 			{
-				char XC = X[Index];
-				char YC = Y[Index];
-				if (XC == YC)
+				char xc = x[index];
+				char yc = y[index];
+				if (xc == yc)
 				{
 					continue;
 				}
-				else if (((XC | YC) & 0xffffff80) == 0) // if (BothAscii)
+				else if (((xc | yc) & 0xffffff80) == 0) // if (BothAscii)
 				{
-					if (XC >= 'A' && XC <= 'Z')
+					if (xc >= 'A' && xc <= 'Z')
 					{
-						XC += (char)32;
+						xc += (char)32;
 					}
-					if (YC >= 'A' && YC <= 'Z')
+					if (yc >= 'A' && yc <= 'Z')
 					{
-						YC += (char)32;
+						yc += (char)32;
 					}
-					int Diff = XC - YC;
-					if (Diff != 0)
+					int diff = xc - yc;
+					if (diff != 0)
 					{
-						return Diff;
+						return diff;
 					}
 				}
 				else
 				{
-					return XC - YC;
+					return xc - yc;
 				}
 			}
 
-			if (X.Length == Length)
+			if (x.Length == length)
 			{
-				return Y.Length == Length ? 0 : /* X[Length] */ -Y[Length];
+				return y.Length == length ? 0 : /* X[Length] */ -y[length];
 			}
 			else
 			{
-				return X[Length] /* - Y[Length] */;
+				return x[length] /* - Y[Length] */;
 			}
 		}
 	}
