Reformat EpicGames.Core according to standard coding conventions.

diff --git a/Engine/Source/Programs/Shared/EpicGames.Core/StringView.cs b/Engine/Source/Programs/Shared/EpicGames.Core/StringView.cs
--- a/Engine/Source/Programs/Shared/EpicGames.Core/StringView.cs
+++ b/Engine/Source/Programs/Shared/EpicGames.Core/StringView.cs
@@ -18,135 +18,129 @@ namespace EpicGames.Core
 		/// <summary>
 		/// Span for the sequence of characters
 		/// </summary>
-		public ReadOnlySpan<char> Span 
-		{ 
-			get { return Memory.Span; } 
-		}
+		public ReadOnlySpan<char> Span => Memory.Span; 
 
 		/// <summary>
 		/// Length of the sequence of characters
 		/// </summary>
-		public int Length
-		{
-			get { return Memory.Length; }
-		}
+		public int Length => Memory.Length;
 
 		/// <summary>
 		/// Constructor
 		/// </summary>
-		/// <param name="Text">String view to construct from</param>
-		public StringView(StringView Text)
+		/// <param name="text">String view to construct from</param>
+		public StringView(StringView text)
 		{
-			this.Memory = Text.Memory;
+			Memory = text.Memory;
 		}
 
 		/// <summary>
 		/// Constructor
 		/// </summary>
-		/// <param name="Text">String view to construct from</param>
-		/// <param name="Index">Offset within the string for this view</param>
-		public StringView(StringView Text, int Index)
+		/// <param name="text">String view to construct from</param>
+		/// <param name="index">Offset within the string for this view</param>
+		public StringView(StringView text, int index)
 		{
-			this.Memory = Text.Memory.Slice(Index);
+			Memory = text.Memory.Slice(index);
 		}
 
 		/// <summary>
 		/// Constructor
 		/// </summary>
-		/// <param name="Text">String view to construct from</param>
-		/// <param name="Index">Offset within the string for this view</param>
-		/// <param name="Count">Number of characters to include</param>
-		public StringView(StringView Text, int Index, int Count)
+		/// <param name="text">String view to construct from</param>
+		/// <param name="index">Offset within the string for this view</param>
+		/// <param name="count">Number of characters to include</param>
+		public StringView(StringView text, int index, int count)
 		{
-			this.Memory = Text.Memory.Slice(Index, Count);
+			Memory = text.Memory.Slice(index, count);
 		}
 
 		/// <summary>
 		/// Constructor
 		/// </summary>
-		/// <param name="Memory">The memory containing the characters</param>
-		public StringView(ReadOnlyMemory<char> Memory)
+		/// <param name="memory">The memory containing the characters</param>
+		public StringView(ReadOnlyMemory<char> memory)
 		{
-			this.Memory = Memory;
+			Memory = memory;
 		}
 
 		/// <summary>
 		/// Constructor
 		/// </summary>
-		/// <param name="Memory">The memory containing the characters</param>
-		/// <param name="Index">Offset within the string for this view</param>
-		public StringView(ReadOnlyMemory<char> Memory, int Index)
+		/// <param name="memory">The memory containing the characters</param>
+		/// <param name="index">Offset within the string for this view</param>
+		public StringView(ReadOnlyMemory<char> memory, int index)
 		{
-			this.Memory = Memory.Slice(Index);
+			Memory = memory.Slice(index);
 		}
 
 		/// <summary>
 		/// Constructor
 		/// </summary>
-		/// <param name="Memory">The memory containing the characters</param>
-		/// <param name="Index">Offset within the string for this view</param>
-		/// <param name="Count">Number of characters to include</param>
-		public StringView(ReadOnlyMemory<char> Memory, int Index, int Count)
+		/// <param name="memory">The memory containing the characters</param>
+		/// <param name="index">Offset within the string for this view</param>
+		/// <param name="count">Number of characters to include</param>
+		public StringView(ReadOnlyMemory<char> memory, int index, int count)
 		{
-			this.Memory = Memory.Slice(Index, Count);
+			Memory = memory.Slice(index, count);
 		}
 
 		/// <summary>
 		/// Constructor
 		/// </summary>
-		/// <param name="Text">String to construct from</param>
-		public StringView(string Text)
+		/// <param name="text">String to construct from</param>
+		public StringView(string text)
 		{
-			Memory = Text.AsMemory();
+			Memory = text.AsMemory();
 		}
 
 		/// <summary>
 		/// Constructs a view onto a substring of the given string
 		/// </summary>
-		/// <param name="Text">String to construct from</param>
-		/// <param name="Index">Offset within the string for this view</param>
-		public StringView(string Text, int Index)
+		/// <param name="text">String to construct from</param>
+		/// <param name="index">Offset within the string for this view</param>
+		public StringView(string text, int index)
 		{
-			Memory = Text.AsMemory(Index);
+			Memory = text.AsMemory(index);
 		}
 
 		/// <summary>
 		/// Constructs a view onto a substring of the given string
 		/// </summary>
-		/// <param name="Text">String to construct from</param>
-		/// <param name="Index">Offset within the string for this view</param>
-		/// <param name="Count">Number of characters to include</param>
-		public StringView(string Text, int Index, int Count)
+		/// <param name="text">String to construct from</param>
+		/// <param name="index">Offset within the string for this view</param>
+		/// <param name="count">Number of characters to include</param>
+		public StringView(string text, int index, int count)
 		{
-			Memory = Text.AsMemory(Index, Count);
+			Memory = text.AsMemory(index, count);
 		}
 
 		/// <summary>
 		/// Equality comparer
 		/// </summary>
-		/// <param name="X">First string to compare</param>
-		/// <param name="Y">Second string to compare</param>
+		/// <param name="x">First string to compare</param>
+		/// <param name="y">Second string to compare</param>
 		/// <returns>True if the strings are equal</returns>
-		public static bool operator ==(StringView X, StringView Y)
+		public static bool operator ==(StringView x, StringView y)
 		{
-			return X.Equals(Y);
+			return x.Equals(y);
 		}
 
 		/// <summary>
 		/// Inequality comparer
 		/// </summary>
-		/// <param name="X">First string to compare</param>
-		/// <param name="Y">Second string to compare</param>
+		/// <param name="x">First string to compare</param>
+		/// <param name="y">Second string to compare</param>
 		/// <returns>True if the strings are equal</returns>
-		public static bool operator !=(StringView X, StringView Y)
+		public static bool operator !=(StringView x, StringView y)
 		{
-			return !X.Equals(Y);
+			return !x.Equals(y);
 		}
 
 		/// <inheritdoc/>
-		public override bool Equals(object? Obj)
+		public override bool Equals(object? obj)
 		{
-			return Obj is StringView && Equals((StringView)Obj);
+			return obj is StringView view && Equals(view);
 		}
 
 		/// <inheritdoc/>
@@ -168,22 +162,22 @@ namespace EpicGames.Core
 		}
 
 		/// <inheritdoc/>
-		public bool Equals(StringView Other)
+		public bool Equals(StringView other)
 		{
-			return Equals(Other, StringComparison.CurrentCulture);
+			return Equals(other, StringComparison.CurrentCulture);
 		}
 
 		/// <inheritdoc/>
-		public bool Equals(StringView? Other, StringComparison comparisonType)
+		public bool Equals(StringView? other, StringComparison comparisonType)
 		{
-			return Other.HasValue && Memory.Span.Equals(Other.Value.Memory.Span, comparisonType);
+			return other.HasValue && Memory.Span.Equals(other.Value.Memory.Span, comparisonType);
 		}
 
 		/// <summary>
 		/// Implicit conversion operator from a regular string
 		/// </summary>
-		/// <param name="Text">The string to construct from</param>
-		public static implicit operator StringView(string Text) => new StringView(Text);
+		/// <param name="text">The string to construct from</param>
+		public static implicit operator StringView(string text) => new StringView(text);
 	}
 
 	/// <summary>
@@ -209,16 +203,16 @@ namespace EpicGames.Core
 		/// <summary>
 		/// Constructor
 		/// </summary>
-		/// <param name="ComparisonType">Type of comparison to perform</param>
-		public StringViewComparer(StringComparison ComparisonType)
+		/// <param name="comparisonType">Type of comparison to perform</param>
+		public StringViewComparer(StringComparison comparisonType)
 		{
-			this.ComparisonType = ComparisonType;
+			ComparisonType = comparisonType;
 		}
 
 		/// <inheritdoc/>
-		public bool Equals(StringView X, StringView Y)
+		public bool Equals(StringView x, StringView y)
 		{
-			return X.Span.Equals(Y.Span, ComparisonType);
+			return x.Span.Equals(y.Span, ComparisonType);
 		}
 
 		/// <inheritdoc/>
@@ -228,26 +222,26 @@ namespace EpicGames.Core
 		}
 
 		/// <inheritdoc/>
-		public int Compare(StringView X, StringView Y)
+		public int Compare(StringView x, StringView y)
 		{
-			return X.Span.CompareTo(Y.Span, ComparisonType);
+			return x.Span.CompareTo(y.Span, ComparisonType);
 		}
 	}
 
 	/// <summary>
 	/// Comparer for StringView objects.  However, it implements UE style ignore case compare
 	/// </summary>
-	public class StringViewComparerUE : IComparer<StringView>, IEqualityComparer<StringView>
+	public class StringViewComparerUe : IComparer<StringView>, IEqualityComparer<StringView>
 	{
 		/// <summary>
 		/// Static instance of an ordinal StringView comparer
 		/// </summary>
-		public static StringViewComparerUE Ordinal = new StringViewComparerUE(StringComparison.Ordinal);
+		public static StringViewComparerUe Ordinal = new StringViewComparerUe(StringComparison.Ordinal);
 
 		/// <summary>
 		/// Static instance of an ordinal StringView comparer which ignores case
 		/// </summary>
-		public static StringViewComparerUE OrdinalIgnoreCase = new StringViewComparerUE(StringComparison.OrdinalIgnoreCase);
+		public static StringViewComparerUe OrdinalIgnoreCase = new StringViewComparerUe(StringComparison.OrdinalIgnoreCase);
 
 		/// <summary>
 		/// The comparison type
@@ -257,16 +251,16 @@ namespace EpicGames.Core
 		/// <summary>
 		/// Constructor
 		/// </summary>
-		/// <param name="ComparisonType">Type of comparison to perform</param>
-		public StringViewComparerUE(StringComparison ComparisonType)
+		/// <param name="comparisonType">Type of comparison to perform</param>
+		public StringViewComparerUe(StringComparison comparisonType)
 		{
-			this.ComparisonType = ComparisonType;
+			ComparisonType = comparisonType;
 		}
 
 		/// <inheritdoc/>
-		public bool Equals(StringView X, StringView Y)
+		public bool Equals(StringView x, StringView y)
 		{
-			return X.Span.Equals(Y.Span, ComparisonType);
+			return x.Span.Equals(y.Span, ComparisonType);
 		}
 
 		/// <inheritdoc/>
@@ -276,15 +270,15 @@ namespace EpicGames.Core
 		}
 
 		/// <inheritdoc/>
-		public int Compare(StringView X, StringView Y)
+		public int Compare(StringView x, StringView y)
 		{
 			if (ComparisonType == StringComparison.OrdinalIgnoreCase)
 			{
-				return StringUtils.CompareIgnoreCaseUE(X.Span, Y.Span);
+				return StringUtils.CompareIgnoreCaseUe(x.Span, y.Span);
 			}
 			else
 			{
-				return X.Span.CompareTo(Y.Span, ComparisonType);
+				return x.Span.CompareTo(y.Span, ComparisonType);
 			}
 		}
 	}
diff --git a/Engine/Source/Programs/Shared/EpicGames.Core/ThreadPoolWorkQueue.cs b/Engine/Source/Programs/Shared/EpicGames.Core/ThreadPoolWorkQueue.cs
--- a/Engine/Source/Programs/Shared/EpicGames.Core/ThreadPoolWorkQueue.cs
+++ b/Engine/Source/Programs/Shared/EpicGames.Core/ThreadPoolWorkQueue.cs
@@ -2,10 +2,7 @@
 
 using System;
 using System.Collections.Generic;
-using System.Linq;
-using System.Text;
 using System.Threading;
-using System.Threading.Tasks;
 
 namespace EpicGames.Core
 {
@@ -17,22 +14,22 @@ namespace EpicGames.Core
 		/// <summary>
 		/// Object used for controlling access to NumOutstandingJobs and updating EmptyEvent
 		/// </summary>
-		object LockObject = new object();
+		readonly object _lockObject = new object();
 
 		/// <summary>
 		/// Number of jobs remaining in the queue. This is updated in an atomic way.
 		/// </summary>
-		int NumOutstandingJobs;
+		int _numOutstandingJobs;
 
 		/// <summary>
 		/// Event which indicates whether the queue is empty.
 		/// </summary>
-		ManualResetEvent EmptyEvent = new ManualResetEvent(true);
+		ManualResetEvent _emptyEvent = new ManualResetEvent(true);
 
 		/// <summary>
 		/// Exceptions which occurred while executing tasks
 		/// </summary>
-		List<Exception> Exceptions = new List<Exception>();
+		readonly List<Exception> _exceptions = new List<Exception>();
 
 		/// <summary>
 		/// Default constructor
@@ -46,70 +43,67 @@ namespace EpicGames.Core
 		/// </summary>
 		public void Dispose()
 		{
-			if(EmptyEvent != null)
+			if(_emptyEvent != null)
 			{
 				Wait();
 
-				EmptyEvent.Dispose();
-				EmptyEvent = null!;
+				_emptyEvent.Dispose();
+				_emptyEvent = null!;
 			}
 		}
 
 		/// <summary>
 		/// Returns the number of items remaining in the queue
 		/// </summary>
-		public int NumRemaining
-		{
-			get { return NumOutstandingJobs; }
-		}
+		public int NumRemaining => _numOutstandingJobs;
 
 		/// <summary>
 		/// Adds an item to the queue
 		/// </summary>
-		/// <param name="ActionToExecute">The action to add</param>
-		public void Enqueue(Action ActionToExecute)
+		/// <param name="actionToExecute">The action to add</param>
+		public void Enqueue(Action actionToExecute)
 		{
-			lock(LockObject)
+			lock(_lockObject)
 			{
-				if(NumOutstandingJobs == 0)
+				if(_numOutstandingJobs == 0)
 				{
-					EmptyEvent.Reset();
+					_emptyEvent.Reset();
 				}
-				NumOutstandingJobs++;
+				_numOutstandingJobs++;
 			}
 
 #if SINGLE_THREAD
 			Execute(ActionToExecute);
 #else
-			ThreadPool.QueueUserWorkItem(Execute, ActionToExecute);
+			ThreadPool.QueueUserWorkItem(Execute, actionToExecute);
 #endif
 		}
 
 		/// <summary>
 		/// Internal method to execute an action
 		/// </summary>
-		/// <param name="ActionToExecute">The action to execute</param>
-		void Execute(object? ActionToExecute)
+		/// <param name="actionToExecute">The action to execute</param>
+		void Execute(object? actionToExecute)
 		{
 			try
 			{
-				((Action)ActionToExecute!)();
+				((Action)actionToExecute!)();
 			}
-			catch(Exception Ex)
+			catch(Exception ex)
 			{
-				lock(LockObject)
+				lock(_lockObject)
 				{
-					Exceptions.Add(Ex);
+					_exceptions.Add(ex);
 				}
 			}
 			finally
 			{
-				lock(LockObject)
+				lock(_lockObject)
 				{
-					NumOutstandingJobs--;
-					if(NumOutstandingJobs == 0)
+					_numOutstandingJobs--;
+					if(_numOutstandingJobs == 0)
 					{
-						EmptyEvent.Set();
+						_emptyEvent.Set();
 					}
 				}
 			}
@@ -120,28 +114,28 @@ namespace EpicGames.Core
 		/// </summary>
 		public void Wait()
 		{
-			EmptyEvent.WaitOne();
+			_emptyEvent.WaitOne();
 			RethrowExceptions();
 		}
 
 		/// <summary>
 		/// Waits for all queued tasks to finish, or the timeout to elapse
 		/// </summary>
-		/// <param name="MillisecondsTimeout">Maximum time to wait</param>
+		/// <param name="millisecondsTimeout">Maximum time to wait</param>
 		/// <returns>True if the queue completed, false if the timeout elapsed</returns>
-		public bool Wait(int MillisecondsTimeout)
+		public bool Wait(int millisecondsTimeout)
 		{
-			return Wait(TimeSpan.FromMilliseconds(MillisecondsTimeout));
+			return Wait(TimeSpan.FromMilliseconds(millisecondsTimeout));
 		}
 
 		/// <summary>
 		/// Waits for all queued tasks to finish, or the timeout to elapse
 		/// </summary>
-		/// <param name="Timeout">Maximum time to wait</param>
+		/// <param name="timeout">Maximum time to wait</param>
 		/// <returns>True if the queue completed, false if the timeout elapsed</returns>
-		public bool Wait(TimeSpan Timeout)
+		public bool Wait(TimeSpan timeout)
 		{
-			bool bResult = EmptyEvent.WaitOne(Timeout);
+			bool bResult = _emptyEvent.WaitOne(timeout);
 			if (bResult)
 			{
 				RethrowExceptions();
@@ -154,11 +148,11 @@ namespace EpicGames.Core
 		/// </summary>
 		public void RethrowExceptions()
 		{
-			lock(LockObject)
+			lock(_lockObject)
 			{
-				if(Exceptions.Count > 0)
+				if(_exceptions.Count > 0)
 				{
-					throw new AggregateException(Exceptions.ToArray());
+					throw new AggregateException(_exceptions.ToArray());
 				}
 			}
 		}
diff --git a/Engine/Source/Programs/Shared/EpicGames.Core/TraceSpan.cs b/Engine/Source/Programs/Shared/EpicGames.Core/TraceSpan.cs
--- a/Engine/Source/Programs/Shared/EpicGames.Core/TraceSpan.cs
+++ b/Engine/Source/Programs/Shared/EpicGames.Core/TraceSpan.cs
@@ -5,10 +5,7 @@ using System.Collections.Generic;
 using System.Diagnostics;
 using System.Globalization;
 using System.IO;
-using System.Linq;
 using System.Reflection;
-using System.Text;
-using System.Threading.Tasks;
 
 namespace EpicGames.Core
 {
@@ -20,9 +17,9 @@ namespace EpicGames.Core
 		/// <summary>
 		/// Adds additional metadata to this scope
 		/// </summary>
-		/// <param name="Name">Name of the key</param>
-		/// <param name="Value">Value for this metadata</param>
-		void AddMetadata(string Name, string Value);
+		/// <param name="name">Name of the key</param>
+		/// <param name="value">Value for this metadata</param>
+		void AddMetadata(string name, string value);
 	}
 
 	/// <summary>
@@ -33,11 +30,11 @@ namespace EpicGames.Core
 		/// <summary>
 		/// Create a new trace span
 		/// </summary>
-		/// <param name="Operation">Name of the operation</param>
-		/// <param name="Resource">Resource that the operation is being performed on</param>
-		/// <param name="Service">Name of the service</param>
+		/// <param name="operation">Name of the operation</param>
+		/// <param name="resource">Resource that the operation is being performed on</param>
+		/// <param name="service">Name of the service</param>
 		/// <returns>New span instance</returns>
-		ITraceSpan Create(string Operation, string? Resource, string? Service);
+		ITraceSpan Create(string operation, string? resource, string? service);
 
 		/// <summary>
 		/// Flush all the current trace data
@@ -52,24 +49,24 @@ namespace EpicGames.Core
 	{
 		class TraceSpanImpl : ITraceSpan
 		{
-			public string Name;
-			public string? Resource;
-			public string? Service;
-			public DateTimeOffset StartTime;
-			public DateTimeOffset? FinishTime;
-			public Dictionary<string, string> Metadata = new Dictionary<string, string>();
+			public string Name { get; }
+			public string? Resource { get; }
+			public string? Service { get; }
+			public DateTimeOffset StartTime { get; }
+			public DateTimeOffset? FinishTime { get; private set; }
+			public Dictionary<string, string> Metadata { get; } = new Dictionary<string, string>();
 
-			public TraceSpanImpl(string Name, string? Resource, string? Service)
+			public TraceSpanImpl(string name, string? resource, string? service)
 			{
-				this.Name = Name;
-				this.Resource = Resource;
-				this.Service = Service;
-				this.StartTime = DateTimeOffset.Now;
+				Name = name;
+				Resource = resource;
+				Service = service;
+				StartTime = DateTimeOffset.Now;
 			}
 
-			public void AddMetadata(string Name, string Value)
+			public void AddMetadata(string name, string value)
 			{
-				Metadata[Name] = Value;
+				Metadata[name] = value;
 			}
 
 			public void Dispose()
@@ -84,30 +81,30 @@ namespace EpicGames.Core
 		/// <summary>
 		/// Output directory for telemetry info
 		/// </summary>
-		DirectoryReference TelemetryDir;
+		readonly DirectoryReference _telemetryDir;
 
 		/// <summary>
 		/// The current scope provider
 		/// </summary>
-		List<TraceSpanImpl> Spans = new List<TraceSpanImpl>();
+		readonly List<TraceSpanImpl> _spans = new List<TraceSpanImpl>();
 
 		/// <summary>
 		/// Constructor
 		/// </summary>
-		/// <param name="TelemetryDir">Directory to store telemetry files</param>
-		public JsonTraceSink(DirectoryReference TelemetryDir)
+		/// <param name="telemetryDir">Directory to store telemetry files</param>
+		public JsonTraceSink(DirectoryReference telemetryDir)
 		{
-			this.TelemetryDir = TelemetryDir;
+			_telemetryDir = telemetryDir;
 		}
 
 		/// <summary>
 		/// Creates a scope using the current provider
 		/// </summary>
-		public ITraceSpan Create(string Name, string? Resource = null, string? Service = null)
+		public ITraceSpan Create(string name, string? resource = null, string? service = null)
 		{
-			TraceSpanImpl Span = new TraceSpanImpl(Name, Resource, Service);
-			Spans.Add(Span);
-			return Span;
+			TraceSpanImpl span = new TraceSpanImpl(name, resource, service);
+			_spans.Add(span);
+			return span;
 		}
 
 		/// <summary>
@@ -115,46 +112,46 @@ namespace EpicGames.Core
 		/// </summary>
 		public void Flush()
 		{
-			FileReference File;
-			using (Process Process = Process.GetCurrentProcess())
+			FileReference file;
+			using (Process process = Process.GetCurrentProcess())
 			{
-				DirectoryReference.CreateDirectory(TelemetryDir);
+				DirectoryReference.CreateDirectory(_telemetryDir);
 
-				string FileName = String.Format("{0}.{1}.json", Path.GetFileName(Assembly.GetEntryAssembly()!.Location), Process.Id, Process.StartTime.Ticks);
-				File = FileReference.Combine(TelemetryDir, FileName);
+				string fileName = String.Format("{0}.{1}.json", Path.GetFileName(Assembly.GetEntryAssembly()!.Location), process.Id, process.StartTime.Ticks);
+				file = FileReference.Combine(_telemetryDir, fileName);
 			}
 
-			using (JsonWriter Writer = new JsonWriter(File))
+			using (JsonWriter writer = new JsonWriter(file))
 			{
-				Writer.WriteObjectStart();
-				Writer.WriteArrayStart("Spans");
-				foreach (TraceSpanImpl Span in Spans)
+				writer.WriteObjectStart();
+				writer.WriteArrayStart("Spans");
+				foreach (TraceSpanImpl span in _spans)
 				{
-					if (Span.FinishTime != null)
+					if (span.FinishTime != null)
 					{
-						Writer.WriteObjectStart();
-						Writer.WriteValue("Name", Span.Name);
-						if (Span.Resource != null)
+						writer.WriteObjectStart();
+						writer.WriteValue("Name", span.Name);
+						if (span.Resource != null)
 						{
-							Writer.WriteValue("Resource", Span.Resource);
+							writer.WriteValue("Resource", span.Resource);
 						}
-						if (Span.Service != null)
+						if (span.Service != null)
 						{
-							Writer.WriteValue("Service", Span.Service);
+							writer.WriteValue("Service", span.Service);
 						}
-						Writer.WriteValue("StartTime", Span.StartTime.ToString("o", CultureInfo.InvariantCulture));
-						Writer.WriteValue("FinishTime", Span.FinishTime.Value.ToString("o", CultureInfo.InvariantCulture));
-						Writer.WriteObjectStart("Metadata");
-						foreach (KeyValuePair<string, string> Pair in Span.Metadata)
+						writer.WriteValue("StartTime", span.StartTime.ToString("o", CultureInfo.InvariantCulture));
+						writer.WriteValue("FinishTime", span.FinishTime.Value.ToString("o", CultureInfo.InvariantCulture));
+						writer.WriteObjectStart("Metadata");
+						foreach (KeyValuePair<string, string> pair in span.Metadata)
 						{
-							Writer.WriteValue(Pair.Key, Pair.Value);
+							writer.WriteValue(pair.Key, pair.Value);
 						}
-						Writer.WriteObjectEnd();
-						Writer.WriteObjectEnd();
+						writer.WriteObjectEnd();
+						writer.WriteObjectEnd();
 					}
 				}
-				Writer.WriteArrayEnd();
-				Writer.WriteObjectEnd();
+				writer.WriteArrayEnd();
+				writer.WriteObjectEnd();
 			}
 		}
 	}
@@ -166,26 +163,26 @@ namespace EpicGames.Core
 	{
 		class CombinedTraceSpan : ITraceSpan
 		{
-			ITraceSpan[] Spans;
+			readonly ITraceSpan[] _spans;
 
-			public CombinedTraceSpan(ITraceSpan[] Spans)
+			public CombinedTraceSpan(ITraceSpan[] spans)
 			{
-				this.Spans = Spans;
+				_spans = spans;
 			}
 
-			public void AddMetadata(string Name, string Value)
+			public void AddMetadata(string name, string value)
 			{
-				foreach(ITraceSpan Span in Spans)
+				foreach(ITraceSpan span in _spans)
 				{
-					Span.AddMetadata(Name, Value);
+					span.AddMetadata(name, value);
 				}
 			}
 
 			public void Dispose()
 			{
-				foreach (ITraceSpan Span in Spans)
+				foreach (ITraceSpan span in _spans)
 				{
-					Span.Dispose();
+					span.Dispose();
 				}
 			}
 		}
@@ -193,7 +190,7 @@ namespace EpicGames.Core
 		/// <summary>
 		/// The sinks to use
 		/// </summary>
-		static List<ITraceSink> Sinks = GetDefaultSinks();
+		static readonly List<ITraceSink> s_sinks = GetDefaultSinks();
 
 		/// <summary>
 		/// Build a list of default sinks
@@ -201,51 +198,51 @@ namespace EpicGames.Core
 		/// <returns></returns>
 		static List<ITraceSink> GetDefaultSinks()
 		{
-			List<ITraceSink> Sinks = new List<ITraceSink>();
+			List<ITraceSink> sinks = new List<ITraceSink>();
 
-			string? TelemetryDir = Environment.GetEnvironmentVariable("UE_TELEMETRY_DIR");
-			if (TelemetryDir != null)
+			string? telemetryDir = Environment.GetEnvironmentVariable("UE_TELEMETRY_DIR");
+			if (telemetryDir != null)
 			{
-				Sinks.Add(new JsonTraceSink(new DirectoryReference(TelemetryDir)));
+				sinks.Add(new JsonTraceSink(new DirectoryReference(telemetryDir)));
 			}
 
-			return Sinks;
+			return sinks;
 		}
 
 		/// <summary>
 		/// Adds a new sink
 		/// </summary>
-		/// <param name="Sink">The sink to add</param>
-		public static void AddSink(ITraceSink Sink)
+		/// <param name="sink">The sink to add</param>
+		public static void AddSink(ITraceSink sink)
 		{
-			Sinks.Add(Sink);
+			s_sinks.Add(sink);
 		}
 
 		/// <summary>
 		/// Remove a sink from the current list
 		/// </summary>
-		/// <param name="Sink">The sink to remove</param>
-		public static void RemoveSink(ITraceSink Sink)
+		/// <param name="sink">The sink to remove</param>
+		public static void RemoveSink(ITraceSink sink)
 		{
-			Sinks.Remove(Sink);
+			s_sinks.Remove(sink);
 		}
 
 		/// <summary>
 		/// Creates a scope using the current provider
 		/// </summary>
-		public static ITraceSpan Create(string Operation, string? Resource = null, string? Service = null)
+		public static ITraceSpan Create(string operation, string? resource = null, string? service = null)
 		{
-			if (Sinks.Count == 0)
+			if (s_sinks.Count == 0)
 			{
 				return new CombinedTraceSpan(Array.Empty<ITraceSpan>());
 			}
-			else if (Sinks.Count == 1)
+			else if (s_sinks.Count == 1)
 			{
-				return Sinks[0].Create(Operation, Resource, Service);
+				return s_sinks[0].Create(operation, resource, service);
 			}
 			else
 			{
-				return new CombinedTraceSpan(Sinks.ConvertAll(x => x.Create(Operation, Resource, Service)).ToArray());
+				return new CombinedTraceSpan(s_sinks.ConvertAll(x => x.Create(operation, resource, service)).ToArray());
 			}
 		}
 
@@ -254,9 +251,9 @@ namespace EpicGames.Core
 		/// </summary>
 		public static void Flush()
 		{
-			foreach (ITraceSink Sink in Sinks)
+			foreach (ITraceSink sink in s_sinks)
 			{
-				Sink.Flush();
+				sink.Flush();
 			}
 		}
 	}
diff --git a/Engine/Source/Programs/Shared/EpicGames.Core/UnrealEngineTypes.cs b/Engine/Source/Programs/Shared/EpicGames.Core/UnrealEngineTypes.cs
--- a/Engine/Source/Programs/Shared/EpicGames.Core/UnrealEngineTypes.cs
+++ b/Engine/Source/Programs/Shared/EpicGames.Core/UnrealEngineTypes.cs
@@ -4,6 +4,8 @@ using System;
 using System.Collections.Generic;
 using System.Reflection;
 
+#pragma warning disable IDE1006 // Naming Styles
+
 namespace EpicGames.Core
 {
 	/// <summary>
@@ -89,7 +91,7 @@ namespace EpicGames.Core
 	/// Engine\Source\Runtime\CoreUObject\Public\UObject\ObjectMacros.h
 	/// </summary>
 	[Flags]
-	public enum EClassFlags : UInt32
+	public enum EClassFlags : uint
 	{
 		/// <summary>
 		/// No Flags
@@ -315,7 +317,7 @@ namespace EpicGames.Core
 	/// Engine\Source\Runtime\CoreUObject\Public\UObject\ObjectMacros.h
 	/// </summary>
 	[Flags]
-	public enum EClassCastFlags : UInt64
+	public enum EClassCastFlags : ulong
 	{
 		None = 0x0000000000000000,
 
@@ -406,7 +408,7 @@ namespace EpicGames.Core
 	/// Engine\Source\Runtime\CoreUObject\Public\UObject\ObjectMacros.h
 	/// </summary>
 	[Flags]
-	public enum EEnumFlags : UInt32
+	public enum EEnumFlags : uint
 	{
 		None,
 		Flags = 1 << 0,
@@ -441,7 +443,7 @@ namespace EpicGames.Core
 	/// Engine\Source\Runtime\CoreUObject\Public\UObject\Script.h
 	/// </summary>
 	[Flags]
-	public enum EFunctionFlags : UInt32
+	public enum EFunctionFlags : uint
 	{
 		None = 0x00000000,
 
@@ -635,7 +637,7 @@ namespace EpicGames.Core
 	/// Engine\Source\Runtime\CoreUObject\Public\UObject\ObjectMacros.h
 	/// </summary>
 	[Flags]
-	public enum EInternalObjectFlags : Int32
+	public enum EInternalObjectFlags : int
 	{
 		None = 0,
 		//~ All the other bits are reserved, DO NOT ADD NEW FLAGS HERE!
@@ -1226,7 +1228,7 @@ namespace EpicGames.Core
 	{
 		private static readonly Lazy<List<string>> BitNames = new Lazy<List<string>>(() =>
 		{
-			List<string> Out = new List<String>();
+			List<string> Out = new List<string>();
 			FieldInfo[] Fields = typeof(EPropertyFlags).GetFields();
 			for (int Bit = 0; Bit < 64; ++Bit)
 			{
@@ -1251,7 +1253,7 @@ namespace EpicGames.Core
 				}
 				if (!bFound)
 				{
-					Out.Add($"0x{Mask.ToString("X16")}");
+					Out.Add($"0x{Mask:X16}");
 				}
 			}
 			return Out;
@@ -1275,7 +1277,7 @@ namespace EpicGames.Core
 		public static List<string> ToStringList(this EPropertyFlags InFlags, bool bIncludePrefix = true)
 		{
 			List<string> LocalBitNames = BitNames.Value;
-			ulong IntFlags = (UInt64)InFlags;
+			ulong IntFlags = (ulong)InFlags;
 			List<string> Out = new List<string>();
 			for (int Bit = 0; Bit < 64; ++Bit)
 			{
@@ -1450,3 +1452,5 @@ namespace EpicGames.Core
 		}
 	}
 }
+
+#pragma warning restore IDE1006 // Naming Styles
diff --git a/Engine/Source/Programs/Shared/EpicGames.Core/Utf8String.cs b/Engine/Source/Programs/Shared/EpicGames.Core/Utf8String.cs
--- a/Engine/Source/Programs/Shared/EpicGames.Core/Utf8String.cs
+++ b/Engine/Source/Programs/Shared/EpicGames.Core/Utf8String.cs
@@ -2,8 +2,6 @@
 
 using System;
 using System.Collections.Generic;
-using System.Diagnostics.CodeAnalysis;
-using System.Runtime.CompilerServices;
 using System.Text;
 
 namespace EpicGames.Core
@@ -26,64 +24,52 @@ namespace EpicGames.Core
 		/// <summary>
 		/// Returns read only span for this string
 		/// </summary>
-		public ReadOnlySpan<byte> Span
-		{
-			get { return Memory.Span; } 
-		}
+		public ReadOnlySpan<byte> Span => Memory.Span;
 
 		/// <summary>
 		/// Determines if this string is empty
 		/// </summary>
-		public bool IsEmpty
-		{
-			get { return Memory.IsEmpty; }
-		}
+		public bool IsEmpty => Memory.IsEmpty;
 
 		/// <summary>
 		/// Returns the length of this string
 		/// </summary>
-		public int Length
-		{
-			get { return Memory.Length; }
-		}
+		public int Length => Memory.Length;
 
 		/// <summary>
 		/// Allows indexing individual bytes of the data
 		/// </summary>
-		/// <param name="Index">Byte index</param>
+		/// <param name="index">Byte index</param>
 		/// <returns>Byte at the given index</returns>
-		public byte this[int Index]
-		{
-			get { return Span[Index]; }
-		}
+		public byte this[int index] => Span[index];
 
 		/// <summary>
 		/// Constructor
 		/// </summary>
-		/// <param name="Text">Text to construct from</param>
-		public Utf8String(string Text)
+		/// <param name="text">Text to construct from</param>
+		public Utf8String(string text)
 		{
-			this.Memory = Encoding.UTF8.GetBytes(Text);
+			Memory = Encoding.UTF8.GetBytes(text);
 		}
 
 		/// <summary>
 		/// Constructor
 		/// </summary>
-		/// <param name="Memory">The data to construct from</param>
-		public Utf8String(ReadOnlyMemory<byte> Memory)
+		/// <param name="memory">The data to construct from</param>
+		public Utf8String(ReadOnlyMemory<byte> memory)
 		{
-			this.Memory = Memory;
+			Memory = memory;
 		}
 
 		/// <summary>
 		/// Constructor
 		/// </summary>
-		/// <param name="Buffer">The buffer to construct from</param>
-		/// <param name="Offset">Offset within the buffer</param>
-		/// <param name="Length">Length of the string within the buffer</param>
-		public Utf8String(byte[] Buffer, int Offset, int Length)
+		/// <param name="buffer">The buffer to construct from</param>
+		/// <param name="offset">Offset within the buffer</param>
+		/// <param name="length">Length of the string within the buffer</param>
+		public Utf8String(byte[] buffer, int offset, int length)
 		{
-			this.Memory = new ReadOnlyMemory<byte>(Buffer, Offset, Length);
+			Memory = new ReadOnlyMemory<byte>(buffer, offset, length);
 		}
 
 		/// <summary>
@@ -92,178 +78,178 @@ namespace EpicGames.Core
 		/// <returns></returns>
 		public Utf8String Clone()
 		{
-			byte[] NewBuffer = new byte[Memory.Length];
-			Memory.CopyTo(NewBuffer);
-			return new Utf8String(NewBuffer);
+			byte[] newBuffer = new byte[Memory.Length];
+			Memory.CopyTo(newBuffer);
+			return new Utf8String(newBuffer);
 		}
 
 		/// <summary>
 		/// Tests two strings for equality
 		/// </summary>
-		/// <param name="A">The first string to compare</param>
-		/// <param name="B">The second string to compare</param>
+		/// <param name="a">The first string to compare</param>
+		/// <param name="b">The second string to compare</param>
 		/// <returns>True if the strings are equal</returns>
-		public static bool operator ==(Utf8String A, Utf8String B)
+		public static bool operator ==(Utf8String a, Utf8String b)
 		{
-			return A.Equals(B);
+			return a.Equals(b);
 		}
 
 		/// <summary>
 		/// Tests two strings for inequality
 		/// </summary>
-		/// <param name="A">The first string to compare</param>
-		/// <param name="B">The second string to compare</param>
+		/// <param name="a">The first string to compare</param>
+		/// <param name="b">The second string to compare</param>
 		/// <returns>True if the strings are not equal</returns>
-		public static bool operator !=(Utf8String A, Utf8String B)
+		public static bool operator !=(Utf8String a, Utf8String b)
 		{
-			return !A.Equals(B);
+			return !a.Equals(b);
 		}
 
 		/// <inheritdoc/>
-		public bool Equals(Utf8String Other) => Utf8StringComparer.Ordinal.Equals(Span, Other.Span);
+		public bool Equals(Utf8String other) => Utf8StringComparer.Ordinal.Equals(Span, other.Span);
 
 		/// <inheritdoc/>
-		public int CompareTo(Utf8String Other) => Utf8StringComparer.Ordinal.Compare(Span, Other.Span);
+		public int CompareTo(Utf8String other) => Utf8StringComparer.Ordinal.Compare(Span, other.Span);
 
-		/// <inheritdoc cref="String.Contains(string)"/>
-		public bool Contains(Utf8String String) => IndexOf(String) != -1;
+		/// <inheritdoc cref="String.Contains(String)"/>
+		public bool Contains(Utf8String str) => IndexOf(str) != -1;
 
-		/// <inheritdoc cref="String.Contains(string, StringComparison)"/>
-		public bool Contains(Utf8String String, Utf8StringComparer Comparer) => IndexOf(String, Comparer) != -1;
+		/// <inheritdoc cref="String.Contains(String, StringComparison)"/>
+		public bool Contains(Utf8String str, Utf8StringComparer comparer) => IndexOf(str, comparer) != -1;
 
-		/// <inheritdoc cref="String.IndexOf(char)"/>
-		public int IndexOf(byte Char)
+		/// <inheritdoc cref="String.IndexOf(Char)"/>
+		public int IndexOf(byte character)
 		{
-			return Span.IndexOf(Char);
+			return Span.IndexOf(character);
 		}
 
-		/// <inheritdoc cref="String.IndexOf(char)"/>
-		public int IndexOf(char Char)
+		/// <inheritdoc cref="String.IndexOf(Char)"/>
+		public int IndexOf(char character)
 		{
-			if (Char < 0x80)
+			if (character < 0x80)
 			{
-				return Span.IndexOf((byte)Char);
+				return Span.IndexOf((byte)character);
 			}
 			else
 			{
-				return Span.IndexOf(Encoding.UTF8.GetBytes(new[] { Char }));
+				return Span.IndexOf(Encoding.UTF8.GetBytes(new[] { character }));
 			}
 		}
 
-		/// <inheritdoc cref="String.IndexOf(char, int)"/>
-		public int IndexOf(char Char, int Index) => IndexOf(Char, Index, Length - Index);
+		/// <inheritdoc cref="String.IndexOf(Char, Int32)"/>
+		public int IndexOf(char character, int index) => IndexOf(character, index, Length - index);
 
-		/// <inheritdoc cref="String.IndexOf(char, int, int)"/>
-		public int IndexOf(char Char, int Index, int Count)
+		/// <inheritdoc cref="String.IndexOf(Char, Int32, Int32)"/>
+		public int IndexOf(char character, int index, int count)
 		{
-			int Result;
-			if (Char < 0x80)
+			int result;
+			if (character < 0x80)
 			{
-				Result = Span.Slice(Index, Count).IndexOf((byte)Char);
+				result = Span.Slice(index, count).IndexOf((byte)character);
 			}
 			else
 			{
-				Result = Span.Slice(Index, Count).IndexOf(Encoding.UTF8.GetBytes(new[] { Char }));
+				result = Span.Slice(index, count).IndexOf(Encoding.UTF8.GetBytes(new[] { character }));
 			}
-			return (Result == -1) ? -1 : Result + Index;
+			return (result == -1) ? -1 : result + index;
 		}
 
-		/// <inheritdoc cref="String.IndexOf(string)"/>
-		public int IndexOf(Utf8String String)
+		/// <inheritdoc cref="String.IndexOf(String)"/>
+		public int IndexOf(Utf8String str)
 		{
-			return Span.IndexOf(String.Span);
+			return Span.IndexOf(str.Span);
 		}
 
-		/// <inheritdoc cref="String.IndexOf(string, StringComparison)"/>
-		public int IndexOf(Utf8String String, Utf8StringComparer Comparer)
+		/// <inheritdoc cref="String.IndexOf(String, StringComparison)"/>
+		public int IndexOf(Utf8String str, Utf8StringComparer comparer)
 		{
-			for (int Idx = 0; Idx < Length - String.Length; Idx++)
+			for (int idx = 0; idx < Length - str.Length; idx++)
 			{
-				if (Comparer.Equals(String.Slice(Idx, String.Length), String))
+				if (comparer.Equals(str.Slice(idx, str.Length), str))
 				{
-					return Idx;
+					return idx;
 				}
 			}
 			return -1;
 		}
 
-		/// <inheritdoc cref="String.LastIndexOf(char)"/>
-		public int LastIndexOf(byte Char)
+		/// <inheritdoc cref="String.LastIndexOf(Char)"/>
+		public int LastIndexOf(byte character)
 		{
-			return Span.LastIndexOf(Char);
+			return Span.LastIndexOf(character);
 		}
 
-		/// <inheritdoc cref="String.LastIndexOf(char)"/>
-		public int LastIndexOf(char Char)
+		/// <inheritdoc cref="String.LastIndexOf(Char)"/>
+		public int LastIndexOf(char character)
 		{
-			if (Char < 0x80)
+			if (character < 0x80)
 			{
-				return Span.LastIndexOf((byte)Char);
+				return Span.LastIndexOf((byte)character);
 			}
 			else
 			{
-				return Span.LastIndexOf(Encoding.UTF8.GetBytes(new[] { Char }));
+				return Span.LastIndexOf(Encoding.UTF8.GetBytes(new[] { character }));
 			}
 		}
 
 		/// <summary>
 		/// Tests if this string starts with another string
 		/// </summary>
-		/// <param name="Other">The string to check against</param>
+		/// <param name="other">The string to check against</param>
 		/// <returns>True if this string starts with the other string</returns>
-		public bool StartsWith(Utf8String Other)
+		public bool StartsWith(Utf8String other)
 		{
-			return Span.StartsWith(Other.Span);
+			return Span.StartsWith(other.Span);
 		}
 
 		/// <summary>
 		/// Tests if this string ends with another string
 		/// </summary>
-		/// <param name="Other">The string to check against</param>
-		/// <param name="Comparer">The string comparer</param>
+		/// <param name="other">The string to check against</param>
+		/// <param name="comparer">The string comparer</param>
 		/// <returns>True if this string ends with the other string</returns>
-		public bool StartsWith(Utf8String Other, Utf8StringComparer Comparer)
+		public bool StartsWith(Utf8String other, Utf8StringComparer comparer)
 		{
-			return Length >= Other.Length && Comparer.Equals(Slice(0, Other.Length), Other);
+			return Length >= other.Length && comparer.Equals(Slice(0, other.Length), other);
 		}
 
 		/// <summary>
 		/// Tests if this string ends with another string
 		/// </summary>
-		/// <param name="Other">The string to check against</param>
+		/// <param name="other">The string to check against</param>
 		/// <returns>True if this string ends with the other string</returns>
-		public bool EndsWith(Utf8String Other)
+		public bool EndsWith(Utf8String other)
 		{
-			return Span.EndsWith(Other.Span);
+			return Span.EndsWith(other.Span);
 		}
 
 		/// <summary>
 		/// Tests if this string ends with another string
 		/// </summary>
-		/// <param name="Other">The string to check against</param>
-		/// <param name="Comparer">The string comparer</param>
+		/// <param name="other">The string to check against</param>
+		/// <param name="comparer">The string comparer</param>
 		/// <returns>True if this string ends with the other string</returns>
-		public bool EndsWith(Utf8String Other, Utf8StringComparer Comparer)
+		public bool EndsWith(Utf8String other, Utf8StringComparer comparer)
 		{
-			return Length >= Other.Length && Comparer.Equals(Slice(Length - Other.Length), Other);
+			return Length >= other.Length && comparer.Equals(Slice(Length - other.Length), other);
 		}
 
-		/// <inheritdoc cref="Substring(int)"/>
-		public Utf8String Slice(int Start) => Substring(Start);
+		/// <inheritdoc cref="Substring(Int32)"/>
+		public Utf8String Slice(int start) => Substring(start);
 
-		/// <inheritdoc cref="Substring(int, int)"/>
-		public Utf8String Slice(int Start, int Count) => Substring(Start, Count);
+		/// <inheritdoc cref="Substring(Int32, Int32)"/>
+		public Utf8String Slice(int start, int count) => Substring(start, count);
 
-		/// <inheritdoc cref="String.Substring(int)"/>
-		public Utf8String Substring(int Start)
+		/// <inheritdoc cref="String.Substring(Int32)"/>
+		public Utf8String Substring(int start)
 		{
-			return new Utf8String(Memory.Slice(Start));
+			return new Utf8String(Memory.Slice(start));
 		}
 
-		/// <inheritdoc cref="String.Substring(int, int)"/>
-		public Utf8String Substring(int Start, int Count)
+		/// <inheritdoc cref="String.Substring(Int32, Int32)"/>
+		public Utf8String Substring(int start, int count)
 		{
-			return new Utf8String(Memory.Slice(Start, Count));
+			return new Utf8String(Memory.Slice(start, count));
 		}
 		
 		/// <summary>
@@ -273,8 +259,8 @@ namespace EpicGames.Core
 		/// <returns>True if the objects are equivalent</returns>
 		public override bool Equals(object? obj)
 		{
-			Utf8String? Other = obj as Utf8String?;
-			return Other != null && Equals(Other.Value);
+			Utf8String? other = obj as Utf8String?;
+			return other != null && Equals(other.Value);
 		}
 
 		/// <summary>
@@ -295,59 +281,59 @@ namespace EpicGames.Core
 		/// <summary>
 		/// Parse a string as an unsigned integer
 		/// </summary>
-		/// <param name="Text"></param>
+		/// <param name="text"></param>
 		/// <returns></returns>
-		public static uint ParseUnsignedInt(Utf8String Text)
+		public static uint ParseUnsignedInt(Utf8String text)
 		{
-			ReadOnlySpan<byte> Bytes = Text.Span;
-			if (Bytes.Length == 0)
+			ReadOnlySpan<byte> bytes = text.Span;
+			if (bytes.Length == 0)
 			{
 				throw new Exception("Cannot parse empty string as an integer");
 			}
 
-			uint Value = 0;
-			for (int Idx = 0; Idx < Bytes.Length; Idx++)
+			uint value = 0;
+			for (int idx = 0; idx < bytes.Length; idx++)
 			{
-				uint Digit = (uint)(Bytes[Idx] - '0');
-				if (Digit > 9)
+				uint digit = (uint)(bytes[idx] - '0');
+				if (digit > 9)
 				{
-					throw new Exception($"Cannot parse '{Text}' as an integer");
+					throw new Exception($"Cannot parse '{text}' as an integer");
 				}
-				Value = (Value * 10) + Digit;
+				value = (value * 10) + digit;
 			}
-			return Value;
+			return value;
 		}
 
 		/// <summary>
 		/// Appends two strings
 		/// </summary>
-		/// <param name="A"></param>
-		/// <param name="B"></param>
+		/// <param name="a"></param>
+		/// <param name="b"></param>
 		/// <returns></returns>
-		public static Utf8String operator +(Utf8String A, Utf8String B)
+		public static Utf8String operator +(Utf8String a, Utf8String b)
 		{
-			if (A.Length == 0)
+			if (a.Length == 0)
 			{
-				return B;
+				return b;
 			}
-			if (B.Length == 0)
+			if (b.Length == 0)
 			{
-				return A;
+				return a;
 			}
 
-			byte[] Buffer = new byte[A.Length + B.Length];
-			A.Span.CopyTo(Buffer);
-			B.Span.CopyTo(Buffer.AsSpan(A.Length));
-			return new Utf8String(Buffer);
+			byte[] buffer = new byte[a.Length + b.Length];
+			a.Span.CopyTo(buffer);
+			b.Span.CopyTo(buffer.AsSpan(a.Length));
+			return new Utf8String(buffer);
 		}
 
 		/// <summary>
 		/// Converts a string to a utf-8 string
 		/// </summary>
-		/// <param name="Text">Text to convert</param>
-		public static implicit operator Utf8String(string Text)
+		/// <param name="text">Text to convert</param>
+		public static implicit operator Utf8String(string text)
 		{
-			return new Utf8String(new ReadOnlyMemory<byte>(Encoding.UTF8.GetBytes(Text)));
+			return new Utf8String(new ReadOnlyMemory<byte>(Encoding.UTF8.GetBytes(text)));
 		}
 	}
 
@@ -362,25 +348,25 @@ namespace EpicGames.Core
 		public sealed class OrdinalComparer : Utf8StringComparer
 		{
 			/// <inheritdoc/>
-			public override bool Equals(ReadOnlySpan<byte> StrA, ReadOnlySpan<byte> StrB)
+			public override bool Equals(ReadOnlySpan<byte> strA, ReadOnlySpan<byte> strB)
 			{
-				return StrA.SequenceEqual(StrB);
+				return strA.SequenceEqual(strB);
 			}
 
 			/// <inheritdoc/>
-			public override int GetHashCode(ReadOnlySpan<byte> String)
+			public override int GetHashCode(ReadOnlySpan<byte> str)
 			{
-				int Hash = 5381;
-				for (int Idx = 0; Idx < String.Length; Idx++)
+				int hash = 5381;
+				for (int idx = 0; idx < str.Length; idx++)
 				{
-					Hash += (Hash << 5) + String[Idx];
+					hash += (hash << 5) + str[idx];
 				}
-				return Hash;
+				return hash;
 			}
 
-			public override int Compare(ReadOnlySpan<byte> StrA, ReadOnlySpan<byte> StrB)
+			public override int Compare(ReadOnlySpan<byte> strA, ReadOnlySpan<byte> strB)
 			{
-				return StrA.SequenceCompareTo(StrB);
+				return strA.SequenceCompareTo(strB);
 			}
 		}
 
@@ -390,16 +376,16 @@ namespace EpicGames.Core
 		public sealed class OrdinalIgnoreCaseComparer : Utf8StringComparer 
 		{
 			/// <inheritdoc/>
-			public override bool Equals(ReadOnlySpan<byte> StrA, ReadOnlySpan<byte> StrB)
+			public override bool Equals(ReadOnlySpan<byte> strA, ReadOnlySpan<byte> strB)
 			{
-				if (StrA.Length != StrB.Length)
+				if (strA.Length != strB.Length)
 				{
 					return false;
 				}
 
-				for (int Idx = 0; Idx < StrA.Length; Idx++)
+				for (int idx = 0; idx < strA.Length; idx++)
 				{
-					if (StrA[Idx] != StrB[Idx] && ToUpper(StrA[Idx]) != ToUpper(StrB[Idx]))
+					if (strA[idx] != strB[idx] && ToUpper(strA[idx]) != ToUpper(strB[idx]))
 					{
 						return false;
 					}
@@ -409,43 +395,43 @@ namespace EpicGames.Core
 			}
 
 			/// <inheritdoc/>
-			public override int GetHashCode(ReadOnlySpan<byte> String)
+			public override int GetHashCode(ReadOnlySpan<byte> str)
 			{
-				HashCode HashCode = new HashCode();
-				for (int Idx = 0; Idx < String.Length; Idx++)
+				HashCode hashCode = new HashCode();
+				for (int idx = 0; idx < str.Length; idx++)
 				{
-					HashCode.Add(ToUpper(String[Idx]));
+					hashCode.Add(ToUpper(str[idx]));
 				}
-				return HashCode.ToHashCode();
+				return hashCode.ToHashCode();
 			}
 
 			/// <inheritdoc/>
-			public override int Compare(ReadOnlySpan<byte> SpanA, ReadOnlySpan<byte> SpanB)
+			public override int Compare(ReadOnlySpan<byte> spanA, ReadOnlySpan<byte> spanB)
 			{
-				int Length = Math.Min(SpanA.Length, SpanB.Length);
-				for (int Idx = 0; Idx < Length; Idx++)
+				int length = Math.Min(spanA.Length, spanB.Length);
+				for (int idx = 0; idx < length; idx++)
 				{
-					if (SpanA[Idx] != SpanB[Idx])
+					if (spanA[idx] != spanB[idx])
 					{
-						int UpperA = ToUpper(SpanA[Idx]);
-						int UpperB = ToUpper(SpanB[Idx]);
-						if (UpperA != UpperB)
+						int upperA = ToUpper(spanA[idx]);
+						int upperB = ToUpper(spanB[idx]);
+						if (upperA != upperB)
 						{
-							return UpperA - UpperB;
+							return upperA - upperB;
 						}
 					}
 				}
-				return SpanA.Length - SpanB.Length;
+				return spanA.Length - spanB.Length;
 			}
 
 			/// <summary>
 			/// Convert a character to uppercase
 			/// </summary>
-			/// <param name="Character">Character to convert</param>
+			/// <param name="character">Character to convert</param>
 			/// <returns>The uppercase version of the character</returns>
-			static byte ToUpper(byte Character)
+			static byte ToUpper(byte character)
 			{
-				return (Character >= 'a' && Character <= 'z') ? (byte)(Character - 'a' + 'A') : Character;
+				return (character >= 'a' && character <= 'z') ? (byte)(character - 'a' + 'A') : character;
 			}
 		}
 
@@ -460,22 +446,22 @@ namespace EpicGames.Core
 		public static Utf8StringComparer OrdinalIgnoreCase { get; } = new OrdinalIgnoreCaseComparer();
 
 		/// <inheritdoc/>
-		public bool Equals(Utf8String StrA, Utf8String StrB) => Equals(StrA.Span, StrB.Span);
+		public bool Equals(Utf8String strA, Utf8String strB) => Equals(strA.Span, strB.Span);
 
 		/// <inheritdoc/>
-		public abstract bool Equals(ReadOnlySpan<byte> StrA, ReadOnlySpan<byte> StrB);
+		public abstract bool Equals(ReadOnlySpan<byte> strA, ReadOnlySpan<byte> strB);
 
 		/// <inheritdoc/>
-		public int GetHashCode(Utf8String String) => GetHashCode(String.Span);
+		public int GetHashCode(Utf8String str) => GetHashCode(str.Span);
 
 		/// <inheritdoc/>
-		public abstract int GetHashCode(ReadOnlySpan<byte> String);
+		public abstract int GetHashCode(ReadOnlySpan<byte> str);
 
 		/// <inheritdoc/>
-		public int Compare(Utf8String StrA, Utf8String StrB) => Compare(StrA.Span, StrB.Span);
+		public int Compare(Utf8String strA, Utf8String strB) => Compare(strA.Span, strB.Span);
 
 		/// <inheritdoc/>
-		public abstract int Compare(ReadOnlySpan<byte> StrA, ReadOnlySpan<byte> StrB);
+		public abstract int Compare(ReadOnlySpan<byte> strA, ReadOnlySpan<byte> strB);
 	}
 
 	/// <summary>
@@ -487,34 +473,34 @@ namespace EpicGames.Core
 		/// Reads a null-terminated utf8 string from the buffer
 		/// </summary>
 		/// <returns>The string data</returns>
-		public static Utf8String ReadString(this MemoryReader Reader)
+		public static Utf8String ReadString(this MemoryReader reader)
 		{
-			ReadOnlySpan<byte> Span = Reader.Span;
-			int Length = Span.IndexOf((byte)0);
-			Utf8String Value = new Utf8String(Reader.ReadFixedLengthBytes(Length));
-			Reader.ReadInt8();
-			return Value;
+			ReadOnlySpan<byte> span = reader.Span;
+			int length = span.IndexOf((byte)0);
+			Utf8String value = new Utf8String(reader.ReadFixedLengthBytes(length));
+			reader.ReadInt8();
+			return value;
 		}
 
 		/// <summary>
 		/// Writes a UTF8 string into memory with a null terminator
 		/// </summary>
-		/// <param name="Writer">The memory writer to serialize to</param>
-		/// <param name="String">String to write</param>
-		public static void WriteString(this MemoryWriter Writer, Utf8String String)
+		/// <param name="writer">The memory writer to serialize to</param>
+		/// <param name="str">String to write</param>
+		public static void WriteString(this MemoryWriter writer, Utf8String str)
 		{
-			Writer.WriteFixedLengthBytes(String.Span);
-			Writer.WriteInt8(0);
+			writer.WriteFixedLengthBytes(str.Span);
+			writer.WriteInt8(0);
 		}
 
 		/// <summary>
 		/// Determines the size of a serialized utf-8 string
 		/// </summary>
-		/// <param name="String">The string to measure</param>
+		/// <param name="str">The string to measure</param>
 		/// <returns>Size of the serialized string</returns>
-		public static int GetSerializedSize(this Utf8String String)
+		public static int GetSerializedSize(this Utf8String str)
 		{
-			return String.Length + 1;
+			return str.Length + 1;
 		}
 	}
 }
diff --git a/Engine/Source/Programs/Shared/EpicGames.Core/Win32ExceptionWithCode.cs b/Engine/Source/Programs/Shared/EpicGames.Core/Win32ExceptionWithCode.cs
--- a/Engine/Source/Programs/Shared/EpicGames.Core/Win32ExceptionWithCode.cs
+++ b/Engine/Source/Programs/Shared/EpicGames.Core/Win32ExceptionWithCode.cs
@@ -1,11 +1,7 @@
 // Copyright Epic Games, Inc. All Rights Reserved.
 
 using System;
-using System.Collections.Generic;
 using System.ComponentModel;
-using System.Linq;
-using System.Text;
-using System.Threading.Tasks;
 
 namespace EpicGames.Core
 {
@@ -17,37 +13,34 @@ namespace EpicGames.Core
 		/// <summary>
 		/// Constructor
 		/// </summary>
-		/// <param name="Code">The Windows error code</param>
-		public Win32ExceptionWithCode(int Code)
-			: base(Code)
+		/// <param name="code">The Windows error code</param>
+		public Win32ExceptionWithCode(int code)
+			: base(code)
 		{
 		}
 
 		/// <summary>
 		/// Constructor
 		/// </summary>
-		/// <param name="Message">Message to display</param>
-		public Win32ExceptionWithCode(string Message)
-			: base(Message)
+		/// <param name="message">Message to display</param>
+		public Win32ExceptionWithCode(string message)
+			: base(message)
 		{
 		}
 
 		/// <summary>
 		/// Constructor
 		/// </summary>
-		/// <param name="Code">The Windows error code</param>
-		/// <param name="Message">Message to display</param>
-		public Win32ExceptionWithCode(int Code, string Message)
-			: base(Code, Message)
+		/// <param name="code">The Windows error code</param>
+		/// <param name="message">Message to display</param>
+		public Win32ExceptionWithCode(int code, string message)
+			: base(code, message)
 		{
 		}
 
 		/// <summary>
 		/// Returns the exception message. Overriden to include the error code in the message.
 		/// </summary>
-		public override string Message
-		{
-			get { return String.Format("{0} (code 0x{1:X8})", base.Message, base.NativeErrorCode); }
-		}
+		public override string Message => String.Format("{0} (code 0x{1:X8})", base.Message, base.NativeErrorCode);
 	}
 }
diff --git a/Engine/Source/Programs/Shared/EpicGames.Core/ZipArchiveExtensions.cs b/Engine/Source/Programs/Shared/EpicGames.Core/ZipArchiveExtensions.cs
--- a/Engine/Source/Programs/Shared/EpicGames.Core/ZipArchiveExtensions.cs
+++ b/Engine/Source/Programs/Shared/EpicGames.Core/ZipArchiveExtensions.cs
@@ -1,11 +1,9 @@
 // Copyright Epic Games, Inc. All Rights Reserved.
 
 using System;
-using System.Collections.Generic;
-using System.Text;
-using System.Runtime.InteropServices;
 using System.IO.Compression;
 using System.Reflection;
+using System.Runtime.InteropServices;
 
 namespace EpicGames.Core
 {
@@ -17,62 +15,62 @@ namespace EpicGames.Core
 		/// <summary>
 		/// Create a zip archive entry, preserving platform mode bits
 		/// </summary>
-		/// <param name="Destination"></param>
-		/// <param name="SourceFileName"></param>
-		/// <param name="EntryName"></param>
-		/// <param name="CompressionLevel"></param>
+		/// <param name="destination"></param>
+		/// <param name="sourceFileName"></param>
+		/// <param name="entryName"></param>
+		/// <param name="compressionLevel"></param>
 		/// <returns></returns>
-		public static ZipArchiveEntry CreateEntryFromFile_CrossPlatform(this ZipArchive Destination, string SourceFileName, string EntryName, CompressionLevel CompressionLevel)
+		public static ZipArchiveEntry CreateEntryFromFile_CrossPlatform(this ZipArchive destination, string sourceFileName, string entryName, CompressionLevel compressionLevel)
 		{
-			ZipArchiveEntry Entry = ZipFileExtensions.CreateEntryFromFile(Destination, SourceFileName, EntryName, CompressionLevel);
-			int Result = -1;
+			ZipArchiveEntry entry = ZipFileExtensions.CreateEntryFromFile(destination, sourceFileName, entryName, compressionLevel);
+			int result = -1;
 
 			if (RuntimeInformation.IsOSPlatform(OSPlatform.Linux))
 			{
-				Result = FileUtils.GetFileMode_Linux(SourceFileName);
+				result = FileUtils.GetFileMode_Linux(sourceFileName);
 			}
 			else if (RuntimeInformation.IsOSPlatform(OSPlatform.OSX))
 			{
-				Result = FileUtils.GetFileMode_Mac(SourceFileName);
+				result = FileUtils.GetFileMode_Mac(sourceFileName);
 			}
 
-			if(Result >= 0)
+			if(result >= 0)
 			{
-				Entry.ExternalAttributes = (int)Result << 16;
+				entry.ExternalAttributes = (int)result << 16;
 			}
 
-			return Entry;
+			return entry;
 		}
 
 		/// <summary>
 		/// Internal field storing information about the platform that created a ZipArchiveEntry. Cannot interpret how to treat the attribute bits without reading this.
 		/// </summary>
-		static FieldInfo VersionMadeByPlatformField = typeof(ZipArchiveEntry).GetField("_versionMadeByPlatform", BindingFlags.NonPublic | BindingFlags.Instance)!;
+		static readonly FieldInfo s_versionMadeByPlatformField = typeof(ZipArchiveEntry).GetField("_versionMadeByPlatform", BindingFlags.NonPublic | BindingFlags.Instance)!;
 
 		/// <summary>
 		/// Extract a zip archive entry, preserving platform mode bits
 		/// </summary>
-		/// <param name="Entry"></param>
-		/// <param name="TargetFileName"></param>
-		/// <param name="Overwrite"></param>
-		public static void ExtractToFile_CrossPlatform(this ZipArchiveEntry Entry, string TargetFileName, bool Overwrite)
+		/// <param name="entry"></param>
+		/// <param name="targetFileName"></param>
+		/// <param name="overwrite"></param>
+		public static void ExtractToFile_CrossPlatform(this ZipArchiveEntry entry, string targetFileName, bool overwrite)
 		{
-			ZipFileExtensions.ExtractToFile(Entry, TargetFileName, Overwrite);
+			ZipFileExtensions.ExtractToFile(entry, targetFileName, overwrite);
 
 			if (RuntimeInformation.IsOSPlatform(OSPlatform.Linux))
 			{
-				int MadeByPlatform = Convert.ToInt32(VersionMadeByPlatformField.GetValue(Entry)!);
-				if (MadeByPlatform == 3 || MadeByPlatform == 19) // Unix or OSX
+				int madeByPlatform = Convert.ToInt32(s_versionMadeByPlatformField.GetValue(entry)!);
+				if (madeByPlatform == 3 || madeByPlatform == 19) // Unix or OSX
 				{
-					FileUtils.SetFileMode_Linux(TargetFileName, (ushort)(Entry.ExternalAttributes >> 16));
+					FileUtils.SetFileMode_Linux(targetFileName, (ushort)(entry.ExternalAttributes >> 16));
 				}
 			}
 			else if (RuntimeInformation.IsOSPlatform(OSPlatform.OSX))
 			{
-				int MadeByPlatform = Convert.ToInt32(VersionMadeByPlatformField.GetValue(Entry)!);
-				if (MadeByPlatform == 3 || MadeByPlatform == 19) // Unix or OSX
+				int madeByPlatform = Convert.ToInt32(s_versionMadeByPlatformField.GetValue(entry)!);
+				if (madeByPlatform == 3 || madeByPlatform == 19) // Unix or OSX
 				{
-					FileUtils.SetFileMode_Mac(TargetFileName, (ushort)(Entry.ExternalAttributes >> 16));
+					FileUtils.SetFileMode_Mac(targetFileName, (ushort)(entry.ExternalAttributes >> 16));
 				}
 			}
 		}
diff --git a/Engine/Source/Programs/Shared/EpicGames.Perforce.Managed/ManagedWorkspace.cs b/Engine/Source/Programs/Shared/EpicGames.Perforce.Managed/ManagedWorkspace.cs
--- a/Engine/Source/Programs/Shared/EpicGames.Perforce.Managed/ManagedWorkspace.cs
+++ b/Engine/Source/Programs/Shared/EpicGames.Perforce.Managed/ManagedWorkspace.cs
@@ -2057,7 +2057,7 @@ namespace EpicGames.Perforce.Managed
 		/// <returns>Disposable object for the trace</returns>
 		private IDisposable Trace(string Operation)
 		{
-			return TraceSpan.Create(Operation, Service: "hordeagent_repository");
+			return TraceSpan.Create(Operation, service: "hordeagent_repository");
 		}
 
 		#endregion
diff --git a/Engine/Source/Programs/Shared/EpicGames.Perforce/PerforceChildProcess.cs b/Engine/Source/Programs/Shared/EpicGames.Perforce/PerforceChildProcess.cs
--- a/Engine/Source/Programs/Shared/EpicGames.Perforce/PerforceChildProcess.cs
+++ b/Engine/Source/Programs/Shared/EpicGames.Perforce/PerforceChildProcess.cs
@@ -86,7 +86,7 @@ namespace EpicGames.Perforce
 			string FullArgumentList = CommandLineArguments.Join(FullArguments);
 			Logger.LogDebug("Running {Executable} {Arguments}", PerforceFileName, FullArgumentList);
 
-			Scope = TraceSpan.Create(Command, Service: "perforce");
+			Scope = TraceSpan.Create(Command, service: "perforce");
 			Scope.AddMetadata("arguments", FullArgumentList);
 
 			ChildProcessGroup = new ManagedProcessGroup();
diff --git a/Engine/Source/Programs/Shared/EpicGames.Redis/RedisChannel.cs b/Engine/Source/Programs/Shared/EpicGames.Redis/RedisChannel.cs
--- a/Engine/Source/Programs/Shared/EpicGames.Redis/RedisChannel.cs
+++ b/Engine/Source/Programs/Shared/EpicGames.Redis/RedisChannel.cs
@@ -2,8 +2,6 @@
 
 using StackExchange.Redis;
 using System;
-using System.Collections.Generic;
-using System.Text;
 using System.Threading.Tasks;
 
 namespace EpicGames.Redis
@@ -24,23 +22,23 @@ namespace EpicGames.Redis
 		/// <summary>
 		/// Constructor
 		/// </summary>
-		/// <param name="Channel"></param>
-		public RedisChannel(RedisChannel Channel) => this.Channel = Channel;
+		/// <param name="channel"></param>
+		public RedisChannel(RedisChannel channel) => Channel = channel;
 
 		/// <inheritdoc/>
-		public override bool Equals(object? Obj) => Obj is RedisChannel<T> List && Channel == List.Channel;
+		public override bool Equals(object? obj) => obj is RedisChannel<T> list && Channel == list.Channel;
 
 		/// <inheritdoc/>
-		public bool Equals(RedisChannel<T> Other) => Channel == Other.Channel;
+		public bool Equals(RedisChannel<T> other) => Channel == other.Channel;
 
 		/// <inheritdoc/>
 		public override int GetHashCode() => Channel.GetHashCode();
 
 		/// <summary>Compares two instances for equality</summary>
-		public static bool operator ==(RedisChannel<T> Left, RedisChannel<T> Right) => Left.Channel == Right.Channel;
+		public static bool operator ==(RedisChannel<T> left, RedisChannel<T> right) => left.Channel == right.Channel;
 
 		/// <summary>Compares two instances for equality</summary>
-		public static bool operator !=(RedisChannel<T> Left, RedisChannel<T> Right) => Left.Channel != Right.Channel;
+		public static bool operator !=(RedisChannel<T> left, RedisChannel<T> right) => left.Channel != right.Channel;
 	}
 
 	/// <summary>
@@ -52,7 +50,7 @@ namespace EpicGames.Redis
 		/// <summary>
 		/// The subscriber to register with
 		/// </summary>
-		ISubscriber? Subscriber;
+		ISubscriber? _subscriber;
 
 		/// <summary>
 		/// The channel to post on
@@ -62,41 +60,41 @@ namespace EpicGames.Redis
 		/// <summary>
 		/// The handler to call
 		/// </summary>
-		Action<RedisChannel<T>, T> Handler;
+		readonly Action<RedisChannel<T>, T> _handler;
 
 		/// <summary>
 		/// Constructor
 		/// </summary>
-		/// <param name="Subscriber"></param>
-		/// <param name="Channel"></param>
-		/// <param name="Handler"></param>
-		RedisChannelSubscription(ISubscriber Subscriber, RedisChannel<T> Channel, Action<RedisChannel<T>, T> Handler)
+		/// <param name="subscriber"></param>
+		/// <param name="channel"></param>
+		/// <param name="handler"></param>
+		RedisChannelSubscription(ISubscriber subscriber, RedisChannel<T> channel, Action<RedisChannel<T>, T> handler)
 		{
-			this.Subscriber = Subscriber;
-			this.Channel = Channel;
-			this.Handler = Handler;
+			_subscriber = subscriber;
+			Channel = channel;
+			_handler = handler;
 		}
 		
 		/// <summary>
 		/// Start the subscription
 		/// </summary>
-		internal static async Task<RedisChannelSubscription<T>> CreateAsync(ISubscriber Subscriber, RedisChannel<T> Channel, Action<RedisChannel<T>, T> Handler, CommandFlags Flags = CommandFlags.None)
+		internal static async Task<RedisChannelSubscription<T>> CreateAsync(ISubscriber subscriber, RedisChannel<T> channel, Action<RedisChannel<T>, T> handler, CommandFlags flags = CommandFlags.None)
 		{
-			RedisChannelSubscription<T> Subscription = new RedisChannelSubscription<T>(Subscriber, Channel, Handler);
-			await Subscription.Subscriber!.SubscribeAsync(Channel.Channel, Subscription.UntypedHandler, Flags);
-			return Subscription;
+			RedisChannelSubscription<T> subscription = new RedisChannelSubscription<T>(subscriber, channel, handler);
+			await subscription._subscriber!.SubscribeAsync(channel.Channel, subscription.UntypedHandler, flags);
+			return subscription;
 		}
 
 		/// <summary>
 		/// Unsubscribe from the channel
 		/// </summary>
-		/// <param name="Flags">Flags for the operation</param>
-		public async Task UnsubscribeAsync(CommandFlags Flags = CommandFlags.None)
+		/// <param name="flags">Flags for the operation</param>
+		public async Task UnsubscribeAsync(CommandFlags flags = CommandFlags.None)
 		{
-			if (Subscriber != null)
+			if (_subscriber != null)
 			{
-				await Subscriber.UnsubscribeAsync(Channel.Channel, UntypedHandler, Flags);
-				Subscriber = null;
+				await _subscriber.UnsubscribeAsync(Channel.Channel, UntypedHandler, flags);
+				_subscriber = null;
 			}
 		}
 
@@ -104,10 +102,10 @@ namespace EpicGames.Redis
 		/// The callback for messages being received
 		/// </summary>
 		/// <param name="_"></param>
-		/// <param name="Message"></param>
-		void UntypedHandler(RedisChannel _, RedisValue Message)
+		/// <param name="message"></param>
+		void UntypedHandler(RedisChannel _, RedisValue message)
 		{
-			Handler(Channel, RedisSerializer.Deserialize<T>(Message));
+			_handler(Channel, RedisSerializer.Deserialize<T>(message));
 		}
 
 		/// <inheritdoc/>
@@ -129,16 +127,16 @@ namespace EpicGames.Redis
 	public static class RedisChannelExtensions
 	{
 		/// <inheritdoc cref="ISubscriber.SubscribeAsync(StackExchange.Redis.RedisChannel, Action{StackExchange.Redis.RedisChannel, RedisValue}, CommandFlags)"/>
-		public static Task<RedisChannelSubscription<T>> SubscribeAsync<T>(this ISubscriber Subscriber, RedisChannel<T> Channel, Action<RedisChannel<T>, T> Handler, CommandFlags Flags = CommandFlags.None)
+		public static Task<RedisChannelSubscription<T>> SubscribeAsync<T>(this ISubscriber subscriber, RedisChannel<T> channel, Action<RedisChannel<T>, T> handler, CommandFlags flags = CommandFlags.None)
 		{
-			return RedisChannelSubscription<T>.CreateAsync(Subscriber, Channel, Handler, Flags);
+			return RedisChannelSubscription<T>.CreateAsync(subscriber, channel, handler, flags);
 		}
 
 		/// <inheritdoc cref="IDatabaseAsync.PublishAsync(StackExchange.Redis.RedisChannel, RedisValue, CommandFlags)"/>
-		public static Task PublishAsync<T>(this IDatabaseAsync Database, RedisChannel<T> Channel, T Item, CommandFlags Flags = CommandFlags.None)
+		public static Task PublishAsync<T>(this IDatabaseAsync database, RedisChannel<T> channel, T item, CommandFlags flags = CommandFlags.None)
 		{
-			RedisValue Value = RedisSerializer.Serialize(Item);
-			return Database.PublishAsync(Channel.Channel, Value, Flags);
+			RedisValue value = RedisSerializer.Serialize(item);
+			return database.PublishAsync(channel.Channel, value, flags);
 		}
 	}
 }
diff --git a/Engine/Source/Programs/Shared/EpicGames.Redis/RedisExtensions.cs b/Engine/Source/Programs/Shared/EpicGames.Redis/RedisExtensions.cs
--- a/Engine/Source/Programs/Shared/EpicGames.Redis/RedisExtensions.cs
+++ b/Engine/Source/Programs/Shared/EpicGames.Redis/RedisExtensions.cs
@@ -1,10 +1,6 @@
 // Copyright Epic Games, Inc. All Rights Reserved.
 
 using StackExchange.Redis;
-using System;
-using System.Collections.Generic;
-using System.Linq;
-using System.Threading.Tasks;
 
 namespace EpicGames.Redis
 {
@@ -16,21 +12,21 @@ namespace EpicGames.Redis
 		/// <summary>
 		/// Convert a RedisValue to a RedisKey
 		/// </summary>
-		/// <param name="Value"></param>
+		/// <param name="value"></param>
 		/// <returns></returns>
-		public static RedisKey AsKey(this RedisValue Value)
+		public static RedisKey AsKey(this RedisValue value)
 		{
-			return (byte[])Value;
+			return (byte[])value;
 		}
 
 		/// <summary>
 		/// Convert a RedisKey to a RedisValue
 		/// </summary>
-		/// <param name="Key"></param>
+		/// <param name="key"></param>
 		/// <returns></returns>
-		public static RedisValue AsValue(this RedisKey Key)
+		public static RedisValue AsValue(this RedisKey key)
 		{
-			return (byte[])Key;
+			return (byte[])key;
 		}
 	}
 }
diff --git a/Engine/Source/Programs/Shared/EpicGames.Redis/RedisHash.cs b/Engine/Source/Programs/Shared/EpicGames.Redis/RedisHash.cs
--- a/Engine/Source/Programs/Shared/EpicGames.Redis/RedisHash.cs
+++ b/Engine/Source/Programs/Shared/EpicGames.Redis/RedisHash.cs
@@ -4,7 +4,6 @@ using StackExchange.Redis;
 using System;
 using System.Collections.Generic;
 using System.Linq;
-using System.Text;
 using System.Threading.Tasks;
 
 namespace EpicGames.Redis
@@ -26,83 +25,83 @@ namespace EpicGames.Redis
 		/// <summary>
 		/// Constructor
 		/// </summary>
-		public RedisHash(IDatabaseAsync Database, RedisKey Key)
+		public RedisHash(IDatabaseAsync database, RedisKey key)
 		{
-			this.Database = Database;
-			this.Key = Key;
+			Database = database;
+			Key = key;
 		}
 
 		/// <inheritdoc cref="IDatabaseAsync.HashDeleteAsync(RedisKey, RedisValue, CommandFlags)"/>
-		public Task<bool> DeleteAsync(TName Name, CommandFlags Flags = CommandFlags.None)
+		public Task<bool> DeleteAsync(TName name, CommandFlags flags = CommandFlags.None)
 		{
-			return Database.HashDeleteAsync(Key, RedisSerializer.Serialize(Name), Flags);
+			return Database.HashDeleteAsync(Key, RedisSerializer.Serialize(name), flags);
 		}
 
 		/// <inheritdoc cref="IDatabaseAsync.HashDeleteAsync(RedisKey, RedisValue[], CommandFlags)"/>
-		public Task<long> DeleteAsync(IEnumerable<TName> Names, CommandFlags Flags = CommandFlags.None)
+		public Task<long> DeleteAsync(IEnumerable<TName> names, CommandFlags flags = CommandFlags.None)
 		{
-			RedisValue[] NameArray = Names.Select(x => RedisSerializer.Serialize(x)).ToArray();
-			return Database.HashDeleteAsync(Key, NameArray, Flags);
+			RedisValue[] nameArray = names.Select(x => RedisSerializer.Serialize(x)).ToArray();
+			return Database.HashDeleteAsync(Key, nameArray, flags);
 		}
 
 		/// <inheritdoc cref="IDatabaseAsync.HashExistsAsync(RedisKey, RedisValue, CommandFlags)"/>
-		public Task<bool> ExistsAsync(TName Name, CommandFlags Flags = CommandFlags.None)
+		public Task<bool> ExistsAsync(TName name, CommandFlags flags = CommandFlags.None)
 		{
-			return Database.HashExistsAsync(Key, RedisSerializer.Serialize(Name), Flags);
+			return Database.HashExistsAsync(Key, RedisSerializer.Serialize(name), flags);
 		}
 
 		/// <inheritdoc cref="IDatabaseAsync.HashGetAsync(RedisKey, RedisValue, CommandFlags)"/>
-		public async Task<TValue> GetAsync(TName Name, CommandFlags Flags = CommandFlags.None)
+		public async Task<TValue> GetAsync(TName name, CommandFlags flags = CommandFlags.None)
 		{
-			RedisValue Value = await Database.HashGetAsync(Key, RedisSerializer.Serialize(Name), Flags);
-			return Value.IsNull ? default(TValue)! : RedisSerializer.Deserialize<TValue>(Value);
+			RedisValue value = await Database.HashGetAsync(Key, RedisSerializer.Serialize(name), flags);
+			return value.IsNull ? default! : RedisSerializer.Deserialize<TValue>(value);
 		}
 
 		/// <inheritdoc cref="IDatabaseAsync.HashGetAsync(RedisKey, RedisValue[], CommandFlags)"/>
-		public async Task<TValue[]> GetAsync(IEnumerable<TName> Names, CommandFlags Flags = CommandFlags.None)
+		public async Task<TValue[]> GetAsync(IEnumerable<TName> names, CommandFlags flags = CommandFlags.None)
 		{
-			RedisValue[] NameArray = Names.Select(x => RedisSerializer.Serialize(x)).ToArray();
-			RedisValue[] ValueArray = await Database.HashGetAsync(Key, NameArray, Flags);
-			return Array.ConvertAll(ValueArray, x => RedisSerializer.Deserialize<TValue>(x));
+			RedisValue[] nameArray = names.Select(x => RedisSerializer.Serialize(x)).ToArray();
+			RedisValue[] valueArray = await Database.HashGetAsync(Key, nameArray, flags);
+			return Array.ConvertAll(valueArray, x => RedisSerializer.Deserialize<TValue>(x));
 		}
 
 		/// <inheritdoc cref="IDatabaseAsync.HashGetAllAsync(RedisKey, CommandFlags)"/>
-		public async Task<HashEntry<TName, TValue>[]> GetAllAsync(CommandFlags Flags = CommandFlags.None)
+		public async Task<HashEntry<TName, TValue>[]> GetAllAsync(CommandFlags flags = CommandFlags.None)
 		{
-			HashEntry[] Entries = await Database.HashGetAllAsync(Key, Flags);
-			return Array.ConvertAll(Entries, x => new HashEntry<TName, TValue>(RedisSerializer.Deserialize<TName>(x.Name), RedisSerializer.Deserialize<TValue>(x.Value)));
+			HashEntry[] entries = await Database.HashGetAllAsync(Key, flags);
+			return Array.ConvertAll(entries, x => new HashEntry<TName, TValue>(RedisSerializer.Deserialize<TName>(x.Name), RedisSerializer.Deserialize<TValue>(x.Value)));
 		}
 
 		/// <inheritdoc cref="IDatabaseAsync.HashKeysAsync(RedisKey, CommandFlags)"/>
-		public async Task<TName[]> KeysAsync(CommandFlags Flags = CommandFlags.None)
+		public async Task<TName[]> KeysAsync(CommandFlags flags = CommandFlags.None)
 		{
-			RedisValue[] NameArray = await Database.HashKeysAsync(Key, Flags);
-			return Array.ConvertAll(NameArray, x => RedisSerializer.Deserialize<TName>(x));
+			RedisValue[] nameArray = await Database.HashKeysAsync(Key, flags);
+			return Array.ConvertAll(nameArray, x => RedisSerializer.Deserialize<TName>(x));
 		}
 
 		/// <inheritdoc cref="IDatabaseAsync.HashLengthAsync(RedisKey, CommandFlags)"/>
-		public Task<long> LengthAsync(CommandFlags Flags = CommandFlags.None)
+		public Task<long> LengthAsync(CommandFlags flags = CommandFlags.None)
 		{
-			return Database.HashLengthAsync(Key, Flags);
+			return Database.HashLengthAsync(Key, flags);
 		}
 
 		/// <inheritdoc cref="IDatabaseAsync.HashSetAsync(RedisKey, RedisValue, RedisValue, When, CommandFlags)"/>
-		public Task SetAsync(TName Name, TValue Value, When When = When.Always, CommandFlags Flags = CommandFlags.None)
+		public Task SetAsync(TName name, TValue value, When when = When.Always, CommandFlags flags = CommandFlags.None)
 		{
-			return Database.HashSetAsync(Key, RedisSerializer.Serialize(Name), RedisSerializer.Serialize(Value), When, Flags);
+			return Database.HashSetAsync(Key, RedisSerializer.Serialize(name), RedisSerializer.Serialize(value), when, flags);
 		}
 
 		/// <inheritdoc cref="IDatabaseAsync.HashSetAsync(RedisKey, HashEntry[], CommandFlags)"/>
-		public Task SetAsync(IEnumerable<HashEntry<TName, TValue>> Entries, CommandFlags Flags = CommandFlags.None)
+		public Task SetAsync(IEnumerable<HashEntry<TName, TValue>> entries, CommandFlags flags = CommandFlags.None)
 		{
-			return Database.HashSetAsync(Key, Entries.Select(x => (HashEntry)x).ToArray(), Flags);
+			return Database.HashSetAsync(Key, entries.Select(x => (HashEntry)x).ToArray(), flags);
 		}
 
 		/// <inheritdoc cref="IDatabaseAsync.HashValuesAsync(RedisKey, CommandFlags)"/>
-		public async Task<TValue[]> ValuesAsync(CommandFlags Flags = CommandFlags.None)
+		public async Task<TValue[]> ValuesAsync(CommandFlags flags = CommandFlags.None)
 		{
-			RedisValue[] Values = await Database.HashValuesAsync(Key);
-			return Array.ConvertAll(Values, x => RedisSerializer.Deserialize<TValue>(x));
+			RedisValue[] values = await Database.HashValuesAsync(Key, flags);
+			return Array.ConvertAll(values, x => RedisSerializer.Deserialize<TValue>(x));
 		}
 	}
 
@@ -118,30 +117,30 @@ namespace EpicGames.Redis
 		/// <summary>
 		/// Constructor
 		/// </summary>
-		/// <param name="Name"></param>
-		/// <param name="Value"></param>
-		public HashEntry(TName Name, TValue Value)
+		/// <param name="name"></param>
+		/// <param name="value"></param>
+		public HashEntry(TName name, TValue value)
 		{
-			this.Name = Name;
-			this.Value = Value;
+			Name = name;
+			Value = value;
 		}
 
 		/// <summary>
 		/// Implicit conversion to a <see cref="HashEntry"/>
 		/// </summary>
-		/// <param name="Entry"></param>
-		public static implicit operator HashEntry(HashEntry<TName, TValue> Entry)
+		/// <param name="entry"></param>
+		public static implicit operator HashEntry(HashEntry<TName, TValue> entry)
 		{
-			return new HashEntry(RedisSerializer.Serialize(Entry.Name), RedisSerializer.Serialize(Entry.Value));
+			return new HashEntry(RedisSerializer.Serialize(entry.Name), RedisSerializer.Serialize(entry.Value));
 		}
 
 		/// <summary>
 		/// Implicit conversion to a <see cref="KeyValuePair{TName, TValue}"/>
 		/// </summary>
-		/// <param name="Entry"></param>
-		public static implicit operator KeyValuePair<TName, TValue>(HashEntry<TName, TValue> Entry)
+		/// <param name="entry"></param>
+		public static implicit operator KeyValuePair<TName, TValue>(HashEntry<TName, TValue> entry)
 		{
-			return new KeyValuePair<TName, TValue>(Entry.Name, Entry.Value);
+			return new KeyValuePair<TName, TValue>(entry.Name, entry.Value);
 		}
 	}
 
@@ -150,36 +149,36 @@ namespace EpicGames.Redis
 	/// </summary>
 	public static class RedisHashExtensions
 	{
-		/// <inheritdoc cref="IDatabaseAsync.HashDecrementAsync(RedisKey, RedisValue, long, CommandFlags)"/>
-		public static Task<long> DecrementAsync<TName>(this RedisHash<TName, long> Hash, TName Name, long Value = 1L, CommandFlags Flags = CommandFlags.None)
+		/// <inheritdoc cref="IDatabaseAsync.HashDecrementAsync(RedisKey, RedisValue, Int64, CommandFlags)"/>
+		public static Task<long> DecrementAsync<TName>(this RedisHash<TName, long> hash, TName name, long value = 1L, CommandFlags flags = CommandFlags.None)
 		{
-			return Hash.Database.HashDecrementAsync(Hash.Key, RedisSerializer.Serialize<TName>(Name), Value, Flags);
+			return hash.Database.HashDecrementAsync(hash.Key, RedisSerializer.Serialize<TName>(name), value, flags);
 		}
 
-		/// <inheritdoc cref="IDatabaseAsync.HashDecrementAsync(RedisKey, RedisValue, double, CommandFlags)"/>
-		public static Task<double> DecrementAsync<TName>(this RedisHash<TName, long> Hash, TName Name, double Value = 1.0, CommandFlags Flags = CommandFlags.None)
+		/// <inheritdoc cref="IDatabaseAsync.HashDecrementAsync(RedisKey, RedisValue, Double, CommandFlags)"/>
+		public static Task<double> DecrementAsync<TName>(this RedisHash<TName, long> hash, TName name, double value = 1.0, CommandFlags flags = CommandFlags.None)
 		{
-			return Hash.Database.HashDecrementAsync(Hash.Key, RedisSerializer.Serialize<TName>(Name), Value, Flags);
+			return hash.Database.HashDecrementAsync(hash.Key, RedisSerializer.Serialize<TName>(name), value, flags);
 		}
 
-		/// <inheritdoc cref="IDatabaseAsync.HashIncrementAsync(RedisKey, RedisValue, long, CommandFlags)"/>
-		public static Task<long> IncrementAsync<TName>(this RedisHash<TName, long> Hash, TName Name, long Value = 1L, CommandFlags Flags = CommandFlags.None)
+		/// <inheritdoc cref="IDatabaseAsync.HashIncrementAsync(RedisKey, RedisValue, Int64, CommandFlags)"/>
+		public static Task<long> IncrementAsync<TName>(this RedisHash<TName, long> hash, TName name, long value = 1L, CommandFlags flags = CommandFlags.None)
 		{
-			return Hash.Database.HashIncrementAsync(Hash.Key, RedisSerializer.Serialize<TName>(Name), Value, Flags);
+			return hash.Database.HashIncrementAsync(hash.Key, RedisSerializer.Serialize<TName>(name), value, flags);
 		}
 
-		/// <inheritdoc cref="IDatabaseAsync.HashIncrementAsync(RedisKey, RedisValue, double, CommandFlags)"/>
-		public static Task<double> IncrementAsync<TName>(this RedisHash<TName, long> Hash, TName Name, double Value = 1.0, CommandFlags Flags = CommandFlags.None)
+		/// <inheritdoc cref="IDatabaseAsync.HashIncrementAsync(RedisKey, RedisValue, Double, CommandFlags)"/>
+		public static Task<double> IncrementAsync<TName>(this RedisHash<TName, long> hash, TName name, double value = 1.0, CommandFlags flags = CommandFlags.None)
 		{
-			return Hash.Database.HashIncrementAsync(Hash.Key, RedisSerializer.Serialize<TName>(Name), Value, Flags);
+			return hash.Database.HashIncrementAsync(hash.Key, RedisSerializer.Serialize<TName>(name), value, flags);
 		}
 
 		/// <summary>
 		/// Creates a version of this set which modifies a transaction rather than the direct DB
 		/// </summary>
-		public static RedisHash<TName, TValue> With<TName, TValue>(this ITransaction Transaction, RedisHash<TName, TValue> Set)
+		public static RedisHash<TName, TValue> With<TName, TValue>(this ITransaction transaction, RedisHash<TName, TValue> set)
 		{
-			return new RedisHash<TName, TValue>(Transaction, Set.Key);
+			return new RedisHash<TName, TValue>(transaction, set.Key);
 		}
 	}
 }
diff --git a/Engine/Source/Programs/Shared/EpicGames.Redis/RedisList.cs b/Engine/Source/Programs/Shared/EpicGames.Redis/RedisList.cs
--- a/Engine/Source/Programs/Shared/EpicGames.Redis/RedisList.cs
+++ b/Engine/Source/Programs/Shared/EpicGames.Redis/RedisList.cs
@@ -3,9 +3,7 @@
 using StackExchange.Redis;
 using System;
 using System.Collections.Generic;
-using System.Diagnostics.CodeAnalysis;
 using System.Linq;
-using System.Text;
 using System.Threading.Tasks;
 
 namespace EpicGames.Redis
@@ -29,53 +27,53 @@ namespace EpicGames.Redis
 		/// <summary>
 		/// Constructor
 		/// </summary>
-		public RedisList(IDatabaseAsync Database, RedisKey Key)
+		public RedisList(IDatabaseAsync database, RedisKey key)
 		{
-			this.Database = Database;
-			this.Key = Key;
+			Database = database;
+			Key = key;
 		}
 
-		/// <inheritdoc cref="IDatabaseAsync.ListGetByIndexAsync(RedisKey, long, CommandFlags)"/>
-		public async Task<TElement> GetByIndexAsync(long Index, CommandFlags Flags = CommandFlags.None)
+		/// <inheritdoc cref="IDatabaseAsync.ListGetByIndexAsync(RedisKey, Int64, CommandFlags)"/>
+		public async Task<TElement> GetByIndexAsync(long index, CommandFlags flags = CommandFlags.None)
 		{
-			RedisValue Value = await Database.ListGetByIndexAsync(Key, Index, Flags);
-			return Value.IsNull? default(TElement)! : RedisSerializer.Deserialize<TElement>(Value);
+			RedisValue value = await Database.ListGetByIndexAsync(Key, index, flags);
+			return value.IsNull? default! : RedisSerializer.Deserialize<TElement>(value);
 		}
 
 		/// <inheritdoc cref="IDatabaseAsync.ListInsertAfterAsync(RedisKey, RedisValue, RedisValue, CommandFlags)"/>
-		public Task<long> InsertAfterAsync(TElement Pivot, TElement Item, CommandFlags Flags = CommandFlags.None)
+		public Task<long> InsertAfterAsync(TElement pivot, TElement item, CommandFlags flags = CommandFlags.None)
 		{
-			RedisValue PivotValue = RedisSerializer.Serialize(Pivot);
-			RedisValue ItemValue = RedisSerializer.Serialize(Item);
-			return Database.ListInsertAfterAsync(Key, PivotValue, ItemValue, Flags);
+			RedisValue pivotValue = RedisSerializer.Serialize(pivot);
+			RedisValue itemValue = RedisSerializer.Serialize(item);
+			return Database.ListInsertAfterAsync(Key, pivotValue, itemValue, flags);
 		}
 
 		/// <inheritdoc cref="IDatabaseAsync.ListInsertBeforeAsync(RedisKey, RedisValue, RedisValue, CommandFlags)"/>
-		public Task<long> InsertBeforeAsync(TElement Pivot, TElement Item, CommandFlags Flags = CommandFlags.None)
+		public Task<long> InsertBeforeAsync(TElement pivot, TElement item, CommandFlags flags = CommandFlags.None)
 		{
-			RedisValue PivotValue = RedisSerializer.Serialize(Pivot);
-			RedisValue ItemValue = RedisSerializer.Serialize(Item);
-			return Database.ListInsertBeforeAsync(Key, PivotValue, ItemValue, Flags);
+			RedisValue pivotValue = RedisSerializer.Serialize(pivot);
+			RedisValue itemValue = RedisSerializer.Serialize(item);
+			return Database.ListInsertBeforeAsync(Key, pivotValue, itemValue, flags);
 		}
 
 		/// <inheritdoc cref="IDatabaseAsync.ListLeftPopAsync(RedisKey, CommandFlags)"/>
-		public async Task<TElement> LeftPopAsync(CommandFlags Flags = CommandFlags.None)
+		public async Task<TElement> LeftPopAsync(CommandFlags flags = CommandFlags.None)
 		{
-			RedisValue Value = await Database.ListLeftPopAsync(Key, Flags);
-			return Value.IsNull ? default(TElement)! : RedisSerializer.Deserialize<TElement>(Value);
+			RedisValue value = await Database.ListLeftPopAsync(Key, flags);
+			return value.IsNull ? default! : RedisSerializer.Deserialize<TElement>(value);
 		}
 
 		/// <inheritdoc cref="IDatabaseAsync.ListLeftPushAsync(RedisKey, RedisValue, When, CommandFlags)"/>
-		public Task<long> LeftPushAsync(TElement Item, When When = When.Always, CommandFlags Flags = CommandFlags.None)
+		public Task<long> LeftPushAsync(TElement item, When when = When.Always, CommandFlags flags = CommandFlags.None)
 		{
-			return Database.ListLeftPushAsync(Key, RedisSerializer.Serialize(Item), When, Flags);
+			return Database.ListLeftPushAsync(Key, RedisSerializer.Serialize(item), when, flags);
 		}
 
 		/// <inheritdoc cref="IDatabaseAsync.ListLeftPushAsync(RedisKey, RedisValue[], When, CommandFlags)"/>
-		public Task<long> LeftPushAsync(IEnumerable<TElement> Items, When When = When.Always, CommandFlags Flags = CommandFlags.None)
+		public Task<long> LeftPushAsync(IEnumerable<TElement> items, When when = When.Always, CommandFlags flags = CommandFlags.None)
 		{
-			RedisValue[] Values = Items.Select(x => RedisSerializer.Serialize(x)).ToArray();
-			return Database.ListLeftPushAsync(Key, Values, When, Flags);
+			RedisValue[] values = items.Select(x => RedisSerializer.Serialize(x)).ToArray();
+			return Database.ListLeftPushAsync(Key, values, when, flags);
 		}
 
 		/// <inheritdoc cref="IDatabaseAsync.ListLengthAsync(RedisKey, CommandFlags)"/>
@@ -84,51 +82,51 @@ namespace EpicGames.Redis
 			return Database.ListLengthAsync(Key);
 		}
 
-		/// <inheritdoc cref="IDatabaseAsync.ListRangeAsync(RedisKey, long, long, CommandFlags)"/>
-		public async Task<TElement[]> RangeAsync(long Start = 0, long Stop = -1, CommandFlags Flags = CommandFlags.None)
+		/// <inheritdoc cref="IDatabaseAsync.ListRangeAsync(RedisKey, Int64, Int64, CommandFlags)"/>
+		public async Task<TElement[]> RangeAsync(long start = 0, long stop = -1, CommandFlags flags = CommandFlags.None)
 		{
-			RedisValue[] Values = await Database.ListRangeAsync(Key, Start, Stop, Flags);
-			return Array.ConvertAll(Values, (Converter<RedisValue, TElement>)(x => (TElement)RedisSerializer.Deserialize<TElement>(x)));
+			RedisValue[] values = await Database.ListRangeAsync(Key, start, stop, flags);
+			return Array.ConvertAll(values, (Converter<RedisValue, TElement>)(x => (TElement)RedisSerializer.Deserialize<TElement>(x)));
 		}
 
-		/// <inheritdoc cref="IDatabaseAsync.ListRemoveAsync(RedisKey, RedisValue, long, CommandFlags)"/>
-		public Task<long> RemoveAsync(TElement Item, long Count = 0L, CommandFlags Flags = CommandFlags.None)
+		/// <inheritdoc cref="IDatabaseAsync.ListRemoveAsync(RedisKey, RedisValue, Int64, CommandFlags)"/>
+		public Task<long> RemoveAsync(TElement item, long count = 0L, CommandFlags flags = CommandFlags.None)
 		{
-			RedisValue Value = RedisSerializer.Serialize(Item);
-			return Database.ListRemoveAsync(Key, Value, Count, Flags);
+			RedisValue value = RedisSerializer.Serialize(item);
+			return Database.ListRemoveAsync(Key, value, count, flags);
 		}
 
 		/// <inheritdoc cref="IDatabaseAsync.ListRightPopAsync(RedisKey, CommandFlags)"/>
-		public async Task<TElement> RightPopAsync(CommandFlags Flags = CommandFlags.None)
+		public async Task<TElement> RightPopAsync(CommandFlags flags = CommandFlags.None)
 		{
-			RedisValue Value = await Database.ListRightPopAsync(Key, Flags);
-			return Value.IsNull? default(TElement)! : RedisSerializer.Deserialize<TElement>(Value);
+			RedisValue value = await Database.ListRightPopAsync(Key, flags);
+			return value.IsNull? default! : RedisSerializer.Deserialize<TElement>(value);
 		}
 
 		/// <inheritdoc cref="IDatabaseAsync.ListRightPushAsync(RedisKey, RedisValue, When, CommandFlags)"/>
-		public Task<long> RightPushAsync(TElement Item, When When = When.Always, CommandFlags Flags = CommandFlags.None)
+		public Task<long> RightPushAsync(TElement item, When when = When.Always, CommandFlags flags = CommandFlags.None)
 		{
-			return Database.ListRightPushAsync(Key, RedisSerializer.Serialize(Item), When, Flags);
+			return Database.ListRightPushAsync(Key, RedisSerializer.Serialize(item), when, flags);
 		}
 
 		/// <inheritdoc cref="IDatabaseAsync.ListRightPushAsync(RedisKey, RedisValue[], When, CommandFlags)"/>
-		public Task<long> RightPushAsync(IEnumerable<TElement> Items, When When = When.Always, CommandFlags Flags = CommandFlags.None)
+		public Task<long> RightPushAsync(IEnumerable<TElement> items, When when = When.Always, CommandFlags flags = CommandFlags.None)
 		{
-			RedisValue[] Values = Items.Select(x => RedisSerializer.Serialize(x)).ToArray();
-			return Database.ListRightPushAsync(Key, Values, When, Flags);
+			RedisValue[] values = items.Select(x => RedisSerializer.Serialize(x)).ToArray();
+			return Database.ListRightPushAsync(Key, values, when, flags);
 		}
 
-		/// <inheritdoc cref="IDatabaseAsync.ListSetByIndexAsync(RedisKey, long, RedisValue, CommandFlags)"/>
-		public Task SetByIndexAsync(long Index, TElement Item, CommandFlags Flags = CommandFlags.None)
+		/// <inheritdoc cref="IDatabaseAsync.ListSetByIndexAsync(RedisKey, Int64, RedisValue, CommandFlags)"/>
+		public Task SetByIndexAsync(long index, TElement item, CommandFlags flags = CommandFlags.None)
 		{
-			RedisValue Value = RedisSerializer.Serialize(Item);
-			return Database.ListSetByIndexAsync(Key, Index, Value, Flags);
+			RedisValue value = RedisSerializer.Serialize(item);
+			return Database.ListSetByIndexAsync(Key, index, value, flags);
 		}
 
-		/// <inheritdoc cref="IDatabaseAsync.ListTrimAsync(RedisKey, long, long, CommandFlags)"/>
-		public Task TrimAsync(long Start, long Stop, CommandFlags Flags = CommandFlags.None)
+		/// <inheritdoc cref="IDatabaseAsync.ListTrimAsync(RedisKey, Int64, Int64, CommandFlags)"/>
+		public Task TrimAsync(long start, long stop, CommandFlags flags = CommandFlags.None)
 		{
-			return Database.ListTrimAsync(Key, Start, Stop, Flags);
+			return Database.ListTrimAsync(Key, start, stop, flags);
 		}
 	}
 
@@ -140,9 +138,9 @@ namespace EpicGames.Redis
 		/// <summary>
 		/// Creates a version of this list which modifies a transaction rather than the direct DB
 		/// </summary>
-		public static RedisList<TElement> With<TElement>(this ITransaction Transaction, RedisList<TElement> Set)
+		public static RedisList<TElement> With<TElement>(this ITransaction transaction, RedisList<TElement> set)
 		{
-			return new RedisList<TElement>(Transaction, Set.Key);
+			return new RedisList<TElement>(transaction, set.Key);
 		}
 	}
 }
diff --git a/Engine/Source/Programs/Shared/EpicGames.Redis/RedisSerializer.cs b/Engine/Source/Programs/Shared/EpicGames.Redis/RedisSerializer.cs
--- a/Engine/Source/Programs/Shared/EpicGames.Redis/RedisSerializer.cs
+++ b/Engine/Source/Programs/Shared/EpicGames.Redis/RedisSerializer.cs
@@ -5,9 +5,7 @@ using StackExchange.Redis;
 using System;
 using System.Collections.Generic;
 using System.ComponentModel;
-using System.Linq;
 using System.Reflection;
-using System.Threading.Tasks;
 
 namespace EpicGames.Redis
 {
@@ -25,10 +23,10 @@ namespace EpicGames.Redis
 		/// <summary>
 		/// Constructor
 		/// </summary>
-		/// <param name="ConverterType">The converter type</param>
-		public RedisConverterAttribute(Type ConverterType)
+		/// <param name="converterType">The converter type</param>
+		public RedisConverterAttribute(Type converterType)
 		{
-			this.ConverterType = ConverterType;
+			ConverterType = converterType;
 		}
 	}
 
@@ -40,16 +38,16 @@ namespace EpicGames.Redis
 		/// <summary>
 		/// Serailize an object to a RedisValue
 		/// </summary>
-		/// <param name="Value"></param>
+		/// <param name="value"></param>
 		/// <returns></returns>
-		RedisValue ToRedisValue(T Value);
+		RedisValue ToRedisValue(T value);
 
 		/// <summary>
 		/// Deserialize an object from a RedisValue
 		/// </summary>
-		/// <param name="Value"></param>
+		/// <param name="value"></param>
 		/// <returns></returns>
-		T FromRedisValue(RedisValue Value);
+		T FromRedisValue(RedisValue value);
 	}
 
 	/// <summary>
@@ -59,15 +57,15 @@ namespace EpicGames.Redis
 	public sealed class RedisCbConverter<T> : IRedisConverter<T>
 	{
 		/// <inheritdoc/>
-		public RedisValue ToRedisValue(T Value)
+		public RedisValue ToRedisValue(T value)
 		{
-			return CbSerializer.Serialize(Value).GetView();
+			return CbSerializer.Serialize(value).GetView();
 		}
 
 		/// <inheritdoc/>
-		public T FromRedisValue(RedisValue Value)
+		public T FromRedisValue(RedisValue value)
 		{
-			return CbSerializer.Deserialize<T>(new CbField((byte[])Value));
+			return CbSerializer.Deserialize<T>(new CbField((byte[])value));
 		}
 	}
 
@@ -78,37 +76,37 @@ namespace EpicGames.Redis
 	{
 		class RedisStringConverter<T> : IRedisConverter<T>
 		{
-			TypeConverter TypeConverter;
+			readonly TypeConverter _typeConverter;
 
-			public RedisStringConverter(TypeConverter TypeConverter)
+			public RedisStringConverter(TypeConverter typeConverter)
 			{
-				this.TypeConverter = TypeConverter;
+				_typeConverter = typeConverter;
 			}
 
-			public RedisValue ToRedisValue(T Value) => (string)TypeConverter.ConvertTo(Value, typeof(string));
-			public T FromRedisValue(RedisValue Value) => (T)TypeConverter.ConvertFrom((string)Value);
+			public RedisValue ToRedisValue(T value) => (string)_typeConverter.ConvertTo(value, typeof(string));
+			public T FromRedisValue(RedisValue value) => (T)_typeConverter.ConvertFrom((string)value);
 		}
 
 		class RedisNativeConverter<T> : IRedisConverter<T>
 		{
-			Func<RedisValue, T> FromRedisValueFunc;
-			Func<T, RedisValue> ToRedisValueFunc;
+			readonly Func<RedisValue, T> _fromRedisValueFunc;
+			readonly Func<T, RedisValue> _toRedisValueFunc;
 
-			public RedisNativeConverter(Func<RedisValue, T> FromRedisValueFunc, Func<T, RedisValue> ToRedisValueFunc)
+			public RedisNativeConverter(Func<RedisValue, T> fromRedisValueFunc, Func<T, RedisValue> toRedisValueFunc)
 			{
-				this.FromRedisValueFunc = FromRedisValueFunc;
-				this.ToRedisValueFunc = ToRedisValueFunc;
+				_fromRedisValueFunc = fromRedisValueFunc;
+				_toRedisValueFunc = toRedisValueFunc;
 			}
 
-			public T FromRedisValue(RedisValue Value) => FromRedisValueFunc(Value);
-			public RedisValue ToRedisValue(T Value) => ToRedisValueFunc(Value);
+			public T FromRedisValue(RedisValue value) => _fromRedisValueFunc(value);
+			public RedisValue ToRedisValue(T value) => _toRedisValueFunc(value);
 		}
 
-		static Dictionary<Type, object> NativeConverters = CreateNativeConverterLookup();
+		static readonly Dictionary<Type, object> s_nativeConverters = CreateNativeConverterLookup();
 
 		static Dictionary<Type, object> CreateNativeConverterLookup()
 		{
-			KeyValuePair<Type, object>[] Converters =
+			KeyValuePair<Type, object>[] converters =
 			{
 				CreateNativeConverter(x => (bool)x, x => x),
 				CreateNativeConverter(x => (int)x, x => x),
@@ -125,12 +123,12 @@ namespace EpicGames.Redis
 				CreateNativeConverter(x => (byte[])x, x => x),
 				CreateNativeConverter(x => (string)x, x => x),
 			};
-			return new Dictionary<Type, object>(Converters);
+			return new Dictionary<Type, object>(converters);
 		}
 
-		static KeyValuePair<Type, object> CreateNativeConverter<T>(Func<RedisValue, T> FromRedisValueFunc, Func<T, RedisValue> ToRedisValueFunc)
+		static KeyValuePair<Type, object> CreateNativeConverter<T>(Func<RedisValue, T> fromRedisValueFunc, Func<T, RedisValue> toRedisValueFunc)
 		{
-			return new KeyValuePair<Type, object>(typeof(T), new RedisNativeConverter<T>(FromRedisValueFunc, ToRedisValueFunc));
+			return new KeyValuePair<Type, object>(typeof(T), new RedisNativeConverter<T>(fromRedisValueFunc, toRedisValueFunc));
 		}
 
 		/// <summary>
@@ -143,35 +141,35 @@ namespace EpicGames.Redis
 
 			static IRedisConverter<T> CreateConverter()
 			{
-				Type Type = typeof(T);
+				Type type = typeof(T);
 
 				// Check for a custom converter
-				RedisConverterAttribute? Attribute = Type.GetCustomAttribute<RedisConverterAttribute>();
-				if (Attribute != null)
+				RedisConverterAttribute? attribute = type.GetCustomAttribute<RedisConverterAttribute>();
+				if (attribute != null)
 				{
-					Type ConverterType = Attribute.ConverterType;
-					if (ConverterType.IsGenericTypeDefinition)
+					Type converterType = attribute.ConverterType;
+					if (converterType.IsGenericTypeDefinition)
 					{
-						ConverterType = ConverterType.MakeGenericType(Type);
+						converterType = converterType.MakeGenericType(type);
 					}
-					return (IRedisConverter<T>)Activator.CreateInstance(ConverterType)!;
+					return (IRedisConverter<T>)Activator.CreateInstance(converterType)!;
 				}
 
 				// Check for known basic types
-				object? NativeConverter;
-				if (NativeConverters.TryGetValue(typeof(T), out NativeConverter))
+				object? nativeConverter;
+				if (s_nativeConverters.TryGetValue(typeof(T), out nativeConverter))
                 {
-					return (IRedisConverter<T>)NativeConverter;
+					return (IRedisConverter<T>)nativeConverter;
                 }
 
 				// Check if there's a regular converter we can use to convert to/from a string
-				TypeConverter? Converter = TypeDescriptor.GetConverter(Type);
-				if (Converter != null && Converter.CanConvertFrom(typeof(string)) && Converter.CanConvertTo(typeof(string)))
+				TypeConverter? converter = TypeDescriptor.GetConverter(type);
+				if (converter != null && converter.CanConvertFrom(typeof(string)) && converter.CanConvertTo(typeof(string)))
 				{
-					return new RedisStringConverter<T>(Converter);
+					return new RedisStringConverter<T>(converter);
 				}
 
-				throw new Exception($"Unable to find Redis converter for {Type.Name}");
+				throw new Exception($"Unable to find Redis converter for {type.Name}");
 			}
 		}
 
@@ -189,22 +187,22 @@ namespace EpicGames.Redis
 		/// Serialize an object to a <see cref="RedisValue"/>
 		/// </summary>
 		/// <typeparam name="T"></typeparam>
-		/// <param name="Value"></param>
+		/// <param name="value"></param>
 		/// <returns></returns>
-		public static RedisValue Serialize<T>(T Value)
+		public static RedisValue Serialize<T>(T value)
 		{
-			return CachedConverter<T>.Converter.ToRedisValue(Value);
+			return CachedConverter<T>.Converter.ToRedisValue(value);
 		}
 
 		/// <summary>
 		/// Deserialize a <see cref="RedisValue"/>
 		/// </summary>
 		/// <typeparam name="T"></typeparam>
-		/// <param name="Value"></param>
+		/// <param name="value"></param>
 		/// <returns></returns>
-		public static T Deserialize<T>(RedisValue Value)
+		public static T Deserialize<T>(RedisValue value)
 		{
-			return CachedConverter<T>.Converter.FromRedisValue(Value);
+			return CachedConverter<T>.Converter.FromRedisValue(value);
 		}
 	}
 }
