Reformat EpicGames.Core according to standard coding conventions.

diff --git a/Engine/Source/Programs/Shared/EpicGames.Redis/RedisSet.cs b/Engine/Source/Programs/Shared/EpicGames.Redis/RedisSet.cs
--- a/Engine/Source/Programs/Shared/EpicGames.Redis/RedisSet.cs
+++ b/Engine/Source/Programs/Shared/EpicGames.Redis/RedisSet.cs
@@ -14,7 +14,7 @@ namespace EpicGames.Redis
 	/// <typeparam name="TElement">The type of element stored in the list</typeparam>
 	public readonly struct RedisSet<TElement>
 	{
-		readonly IDatabaseAsync Database;
+		readonly IDatabaseAsync _database;
 
 		/// <summary>
 		/// The key for the list
@@ -24,70 +24,70 @@ namespace EpicGames.Redis
 		/// <summary>
 		/// Constructor
 		/// </summary>
-		public RedisSet(IDatabaseAsync Database, RedisKey Key)
+		public RedisSet(IDatabaseAsync database, RedisKey key)
 		{
-			this.Database = Database;
-			this.Key = Key;
+			_database = database;
+			Key = key;
 		}
 
 		/// <inheritdoc cref="IDatabaseAsync.SetAddAsync(RedisKey, RedisValue, CommandFlags)"/>
-		public Task<bool> AddAsync(TElement Item, CommandFlags Flags = CommandFlags.None)
+		public Task<bool> AddAsync(TElement item, CommandFlags flags = CommandFlags.None)
 		{
-			return Database.SetAddAsync(Key, RedisSerializer.Serialize(Item), Flags);
+			return _database.SetAddAsync(Key, RedisSerializer.Serialize(item), flags);
 		}
 
 		/// <inheritdoc cref="IDatabaseAsync.SetAddAsync(RedisKey, RedisValue[], CommandFlags)"/>
-		public Task<long> AddAsync(IEnumerable<TElement> Items, CommandFlags Flags = CommandFlags.None)
+		public Task<long> AddAsync(IEnumerable<TElement> items, CommandFlags flags = CommandFlags.None)
 		{
-			RedisValue[] Values = Items.Select(x => RedisSerializer.Serialize(x)).ToArray();
-			return Database.SetAddAsync(Key, Values, Flags);
+			RedisValue[] values = items.Select(x => RedisSerializer.Serialize(x)).ToArray();
+			return _database.SetAddAsync(Key, values, flags);
 		}
 
 		/// <inheritdoc cref="IDatabaseAsync.SetContainsAsync(RedisKey, RedisValue, CommandFlags)"/>
-		public Task<bool> ContainsAsync(TElement Item, CommandFlags Flags = CommandFlags.None)
+		public Task<bool> ContainsAsync(TElement item, CommandFlags flags = CommandFlags.None)
 		{
-			RedisValue Value = RedisSerializer.Serialize(Item);
-			return Database.SetContainsAsync(Key, Value, Flags);
+			RedisValue value = RedisSerializer.Serialize(item);
+			return _database.SetContainsAsync(Key, value, flags);
 		}
 
 		/// <inheritdoc cref="IDatabaseAsync.SetLengthAsync(RedisKey, CommandFlags)"/>
-		public Task<long> LengthAsync(CommandFlags Flags = CommandFlags.None)
+		public Task<long> LengthAsync(CommandFlags flags = CommandFlags.None)
 		{
-			return Database.SetLengthAsync(Key, Flags);
+			return _database.SetLengthAsync(Key, flags);
 		}
 
 		/// <inheritdoc cref="IDatabaseAsync.SetRemoveAsync(RedisKey, RedisValue, CommandFlags)"/>
-		public async Task<TElement[]> MembersAsync(CommandFlags Flags = CommandFlags.None)
+		public async Task<TElement[]> MembersAsync(CommandFlags flags = CommandFlags.None)
 		{
-			RedisValue[] Values = await Database.SetMembersAsync(Key, Flags);
-			return Array.ConvertAll(Values, (Converter<RedisValue, TElement>)(x => RedisSerializer.Deserialize<TElement>(x)));
+			RedisValue[] values = await _database.SetMembersAsync(Key, flags);
+			return Array.ConvertAll(values, (Converter<RedisValue, TElement>)(x => RedisSerializer.Deserialize<TElement>(x)));
 		}
 
 		/// <inheritdoc cref="IDatabaseAsync.SetPopAsync(RedisKey, CommandFlags)"/>
-		public async Task<TElement> PopAsync(CommandFlags Flags = CommandFlags.None)
+		public async Task<TElement> PopAsync(CommandFlags flags = CommandFlags.None)
 		{
-			RedisValue Value = await Database.SetPopAsync(Key, Flags);
-			return Value.IsNull ? default(TElement)! : RedisSerializer.Deserialize<TElement>(Value);
+			RedisValue value = await _database.SetPopAsync(Key, flags);
+			return value.IsNull ? default! : RedisSerializer.Deserialize<TElement>(value);
 		}
 
-		/// <inheritdoc cref="IDatabaseAsync.SetPopAsync(RedisKey, long, CommandFlags)"/>
-		public async Task<TElement[]> PopAsync(long Count, CommandFlags Flags = CommandFlags.None)
+		/// <inheritdoc cref="IDatabaseAsync.SetPopAsync(RedisKey, Int64, CommandFlags)"/>
+		public async Task<TElement[]> PopAsync(long count, CommandFlags flags = CommandFlags.None)
 		{
-			RedisValue[] Values = await Database.SetPopAsync(Key, Count, Flags);
-			return Array.ConvertAll<RedisValue, TElement>(Values, x => RedisSerializer.Deserialize<TElement>(x));
+			RedisValue[] values = await _database.SetPopAsync(Key, count, flags);
+			return Array.ConvertAll<RedisValue, TElement>(values, x => RedisSerializer.Deserialize<TElement>(x));
 		}
 
 		/// <inheritdoc cref="IDatabaseAsync.SetRemoveAsync(RedisKey, RedisValue, CommandFlags)"/>
-		public Task<bool> RemoveAsync(TElement Item, CommandFlags Flags = CommandFlags.None)
+		public Task<bool> RemoveAsync(TElement item, CommandFlags flags = CommandFlags.None)
 		{
-			return Database.SetRemoveAsync(Key, RedisSerializer.Serialize(Item), Flags);
+			return _database.SetRemoveAsync(Key, RedisSerializer.Serialize(item), flags);
 		}
 
 		/// <inheritdoc cref="IDatabaseAsync.SetRemoveAsync(RedisKey, RedisValue[], CommandFlags)"/>
-		public Task<long> RemoveAsync(IEnumerable<TElement> Items, CommandFlags Flags = CommandFlags.None)
+		public Task<long> RemoveAsync(IEnumerable<TElement> items, CommandFlags flags = CommandFlags.None)
 		{
-			RedisValue[] Values = Items.Select(x => RedisSerializer.Serialize(x)).ToArray();
-			return Database.SetRemoveAsync(Key, Values, Flags);
+			RedisValue[] values = items.Select(x => RedisSerializer.Serialize(x)).ToArray();
+			return _database.SetRemoveAsync(Key, values, flags);
 		}
 	}
 
@@ -99,9 +99,9 @@ namespace EpicGames.Redis
 		/// <summary>
 		/// Creates a version of this set which modifies a transaction rather than the direct DB
 		/// </summary>
-		public static RedisSet<TElement> With<TElement>(this ITransaction Transaction, RedisSet<TElement> Set)
+		public static RedisSet<TElement> With<TElement>(this ITransaction transaction, RedisSet<TElement> set)
 		{
-			return new RedisSet<TElement>(Transaction, Set.Key);
+			return new RedisSet<TElement>(transaction, set.Key);
 		}
 	}
 }
diff --git a/Engine/Source/Programs/Shared/EpicGames.Redis/RedisSortedSet.cs b/Engine/Source/Programs/Shared/EpicGames.Redis/RedisSortedSet.cs
--- a/Engine/Source/Programs/Shared/EpicGames.Redis/RedisSortedSet.cs
+++ b/Engine/Source/Programs/Shared/EpicGames.Redis/RedisSortedSet.cs
@@ -3,7 +3,6 @@
 using StackExchange.Redis;
 using System;
 using System.Collections.Generic;
-using System.Diagnostics.CodeAnalysis;
 using System.Linq;
 using System.Threading.Tasks;
 
@@ -28,40 +27,40 @@ namespace EpicGames.Redis
 		/// <summary>
 		/// Score for the entry
 		/// </summary>
-		public readonly double Score;
+		public readonly double Score { get; }
 
 		/// <summary>
 		/// Constructor
 		/// </summary>
-		/// <param name="Entry"></param>
-		public SortedSetEntry(SortedSetEntry Entry)
+		/// <param name="entry"></param>
+		public SortedSetEntry(SortedSetEntry entry)
 		{
-			this.Element = RedisSerializer.Deserialize<T>(Entry.Element);
-			this.ElementValue = Entry.Element;
-			this.Score = Entry.Score;
+			Element = RedisSerializer.Deserialize<T>(entry.Element);
+			ElementValue = entry.Element;
+			Score = entry.Score;
 		}
 
 		/// <summary>
 		/// Constructor
 		/// </summary>
-		/// <param name="Element"></param>
-		/// <param name="Score"></param>
-		public SortedSetEntry(T Element, double Score)
+		/// <param name="element"></param>
+		/// <param name="score"></param>
+		public SortedSetEntry(T element, double score)
 		{
-			this.Element = Element;
-			this.ElementValue = RedisSerializer.Serialize<T>(Element);
-			this.Score = Score;
+			Element = element;
+			ElementValue = RedisSerializer.Serialize<T>(element);
+			Score = score;
 		}
 
 		/// <summary>
 		/// Deconstruct this item into a tuple
 		/// </summary>
-		/// <param name="OutElement"></param>
-		/// <param name="OutScore"></param>
-		public void Deconstruct(out T OutElement, out double OutScore)
+		/// <param name="outElement"></param>
+		/// <param name="outScore"></param>
+		public void Deconstruct(out T outElement, out double outScore)
 		{
-			OutElement = Element;
-			OutScore = Score;
+			outElement = Element;
+			outScore = Score;
 		}
 	}
 
@@ -71,7 +70,7 @@ namespace EpicGames.Redis
 	/// <typeparam name="TElement">The type of element stored in the list</typeparam>
 	public readonly struct RedisSortedSet<TElement>
 	{
-		readonly IDatabaseAsync Database;
+		readonly IDatabaseAsync _database;
 
 		/// <summary>
 		/// The key for the list
@@ -81,130 +80,130 @@ namespace EpicGames.Redis
 		/// <summary>
 		/// Constructor
 		/// </summary>
-		public RedisSortedSet(IDatabaseAsync Database, RedisKey Key)
+		public RedisSortedSet(IDatabaseAsync database, RedisKey key)
 		{
-			this.Database = Database;
-			this.Key = Key;
+			_database = database;
+			Key = key;
 		}
 
-		/// <inheritdoc cref="IDatabaseAsync.SortedSetAddAsync(RedisKey, RedisValue, double, CommandFlags)"/>
-		public Task<bool> AddAsync(TElement Item, double Score, When When = When.Always, CommandFlags Flags = CommandFlags.None)
+		/// <inheritdoc cref="IDatabaseAsync.SortedSetAddAsync(RedisKey, RedisValue, Double, CommandFlags)"/>
+		public Task<bool> AddAsync(TElement item, double score, When when = When.Always, CommandFlags flags = CommandFlags.None)
 		{
-			RedisValue Value = RedisSerializer.Serialize(Item);
-			return Database.SortedSetAddAsync(Key, Value, Score, When, Flags);
+			RedisValue value = RedisSerializer.Serialize(item);
+			return _database.SortedSetAddAsync(Key, value, score, when, flags);
 		}
 
-		/// <inheritdoc cref="IDatabaseAsync.SortedSetAddAsync(RedisKey, RedisValue, double, When, CommandFlags)"/>
-		public Task<long> AddAsync(SortedSetEntry<TElement>[] Values, When When = When.Always, CommandFlags Flags = CommandFlags.None)
+		/// <inheritdoc cref="IDatabaseAsync.SortedSetAddAsync(RedisKey, RedisValue, Double, When, CommandFlags)"/>
+		public Task<long> AddAsync(SortedSetEntry<TElement>[] values, When when = When.Always, CommandFlags flags = CommandFlags.None)
 		{
-			SortedSetEntry[] Untyped = Array.ConvertAll(Values, x => new SortedSetEntry(x.ElementValue, x.Score));
-			return Database.SortedSetAddAsync(Key, Untyped, When, Flags);
+			SortedSetEntry[] untyped = Array.ConvertAll(values, x => new SortedSetEntry(x.ElementValue, x.Score));
+			return _database.SortedSetAddAsync(Key, untyped, when, flags);
 		}
 
-		/// <inheritdoc cref="IDatabaseAsync.SortedSetLengthAsync(RedisKey, double, double, Exclude, CommandFlags)"/>
-		public Task<long> LengthAsync(double Min = double.NegativeInfinity, double Max = double.PositiveInfinity, Exclude Exclude = Exclude.None, CommandFlags Flags = CommandFlags.None)
+		/// <inheritdoc cref="IDatabaseAsync.SortedSetLengthAsync(RedisKey, Double, Double, Exclude, CommandFlags)"/>
+		public Task<long> LengthAsync(double min = Double.NegativeInfinity, double max = Double.PositiveInfinity, Exclude exclude = Exclude.None, CommandFlags flags = CommandFlags.None)
 		{
-			return Database.SortedSetLengthAsync(Key, Min, Max, Exclude, Flags);
+			return _database.SortedSetLengthAsync(Key, min, max, exclude, flags);
 		}
 
 		/// <inheritdoc cref="IDatabaseAsync.SortedSetLengthByValueAsync(RedisKey, RedisValue, RedisValue, Exclude, CommandFlags)"/>
-		public Task<long> LengthByValueAsync(TElement Min, TElement Max, Exclude Exclude = Exclude.None, CommandFlags Flags = CommandFlags.None)
+		public Task<long> LengthByValueAsync(TElement min, TElement max, Exclude exclude = Exclude.None, CommandFlags flags = CommandFlags.None)
 		{
-			RedisValue MinValue = RedisSerializer.Serialize(Min);
-			RedisValue MaxValue = RedisSerializer.Serialize(Max);
-			return Database.SortedSetLengthByValueAsync(Key, MinValue, MaxValue, Exclude, Flags);
+			RedisValue minValue = RedisSerializer.Serialize(min);
+			RedisValue maxValue = RedisSerializer.Serialize(max);
+			return _database.SortedSetLengthByValueAsync(Key, minValue, maxValue, exclude, flags);
 		}
 
-		/// <inheritdoc cref="IDatabaseAsync.SortedSetRangeByRankAsync(RedisKey, long, long, Order, CommandFlags)"/>
-		public async Task<TElement[]> RangeByRankAsync(long Start, long Stop = -1, Order Order = Order.Ascending, CommandFlags Flags = CommandFlags.None)
+		/// <inheritdoc cref="IDatabaseAsync.SortedSetRangeByRankAsync(RedisKey, Int64, Int64, Order, CommandFlags)"/>
+		public async Task<TElement[]> RangeByRankAsync(long start, long stop = -1, Order order = Order.Ascending, CommandFlags flags = CommandFlags.None)
 		{
-			RedisValue[] Values = await Database.SortedSetRangeByRankAsync(Key, Start, Stop, Order, Flags);
-			return Array.ConvertAll(Values, x => RedisSerializer.Deserialize<TElement>(x));
+			RedisValue[] values = await _database.SortedSetRangeByRankAsync(Key, start, stop, order, flags);
+			return Array.ConvertAll(values, x => RedisSerializer.Deserialize<TElement>(x));
 		}
 
-		/// <inheritdoc cref="IDatabaseAsync.SortedSetRangeByRankWithScoresAsync(RedisKey, long, long, Order, CommandFlags)"/>
-		public async Task<SortedSetEntry<TElement>[]> RangeByRankWithScoresAsync(long Start, long Stop = -1, Order Order = Order.Ascending, CommandFlags Flags = CommandFlags.None)
+		/// <inheritdoc cref="IDatabaseAsync.SortedSetRangeByRankWithScoresAsync(RedisKey, Int64, Int64, Order, CommandFlags)"/>
+		public async Task<SortedSetEntry<TElement>[]> RangeByRankWithScoresAsync(long start, long stop = -1, Order order = Order.Ascending, CommandFlags flags = CommandFlags.None)
 		{
-			SortedSetEntry[] Values = await Database.SortedSetRangeByRankWithScoresAsync(Key, Start, Stop, Order, Flags);
-			return Array.ConvertAll(Values, x => new SortedSetEntry<TElement>(x));
+			SortedSetEntry[] values = await _database.SortedSetRangeByRankWithScoresAsync(Key, start, stop, order, flags);
+			return Array.ConvertAll(values, x => new SortedSetEntry<TElement>(x));
 		}
 
-		/// <inheritdoc cref="IDatabaseAsync.SortedSetRangeByScoreAsync(RedisKey, double, double, Exclude, Order, long, long, CommandFlags)"/>
-		public async Task<TElement[]> RangeByScoreAsync(double Start = double.NegativeInfinity, double Stop = double.PositiveInfinity, Exclude Exclude = Exclude.None, Order Order = Order.Ascending, long Skip = 0L, long Take = -1L, CommandFlags Flags = CommandFlags.None)
+		/// <inheritdoc cref="IDatabaseAsync.SortedSetRangeByScoreAsync(RedisKey, Double, Double, Exclude, Order, Int64, Int64, CommandFlags)"/>
+		public async Task<TElement[]> RangeByScoreAsync(double start = Double.NegativeInfinity, double stop = Double.PositiveInfinity, Exclude exclude = Exclude.None, Order order = Order.Ascending, long skip = 0L, long take = -1L, CommandFlags flags = CommandFlags.None)
 		{
-			RedisValue[] Values = await Database.SortedSetRangeByScoreAsync(Key, Start, Stop, Exclude, Order, Skip, Take, Flags);
-			return Array.ConvertAll(Values, x => RedisSerializer.Deserialize<TElement>(x));
+			RedisValue[] values = await _database.SortedSetRangeByScoreAsync(Key, start, stop, exclude, order, skip, take, flags);
+			return Array.ConvertAll(values, x => RedisSerializer.Deserialize<TElement>(x));
 		}
 
-		/// <inheritdoc cref="IDatabaseAsync.SortedSetRangeByScoreWithScoresAsync(RedisKey, double, double, Exclude, Order, long, long, CommandFlags)"/>
-		public async Task<SortedSetEntry<TElement>[]> RangeByScoreWithScoresAsync(double Start = double.NegativeInfinity, double Stop = double.PositiveInfinity, Exclude Exclude = Exclude.None, Order Order = Order.Ascending, long Skip = 0L, long Take = -1L, CommandFlags Flags = CommandFlags.None)
+		/// <inheritdoc cref="IDatabaseAsync.SortedSetRangeByScoreWithScoresAsync(RedisKey, Double, Double, Exclude, Order, Int64, Int64, CommandFlags)"/>
+		public async Task<SortedSetEntry<TElement>[]> RangeByScoreWithScoresAsync(double start = Double.NegativeInfinity, double stop = Double.PositiveInfinity, Exclude exclude = Exclude.None, Order order = Order.Ascending, long skip = 0L, long take = -1L, CommandFlags flags = CommandFlags.None)
 		{
-			SortedSetEntry[] Values = await Database.SortedSetRangeByScoreWithScoresAsync(Key, Start, Stop, Exclude, Order, Skip, Take, Flags);
-			return Array.ConvertAll(Values, x => new SortedSetEntry<TElement>(x));
+			SortedSetEntry[] values = await _database.SortedSetRangeByScoreWithScoresAsync(Key, start, stop, exclude, order, skip, take, flags);
+			return Array.ConvertAll(values, x => new SortedSetEntry<TElement>(x));
 		}
 
-		/// <inheritdoc cref="IDatabaseAsync.SortedSetRangeByValueAsync(RedisKey, RedisValue, RedisValue, Exclude, Order, long, long, CommandFlags)"/>
-		public async Task<TElement[]> RangeByValueAsync(TElement Min, TElement Max, Exclude Exclude = Exclude.None, Order Order = Order.Ascending, long Skip = 0L, long Take = -1L, CommandFlags Flags = CommandFlags.None)
+		/// <inheritdoc cref="IDatabaseAsync.SortedSetRangeByValueAsync(RedisKey, RedisValue, RedisValue, Exclude, Order, Int64, Int64, CommandFlags)"/>
+		public async Task<TElement[]> RangeByValueAsync(TElement min, TElement max, Exclude exclude = Exclude.None, Order order = Order.Ascending, long skip = 0L, long take = -1L, CommandFlags flags = CommandFlags.None)
 		{
-			RedisValue MinValue = RedisSerializer.Serialize(Min);
-			RedisValue MaxValue = RedisSerializer.Serialize(Max);
-			RedisValue[] Values = await Database.SortedSetRangeByValueAsync(Key, MinValue, MaxValue, Exclude, Order, Skip, Take, Flags);
-			return Array.ConvertAll(Values, x => RedisSerializer.Deserialize<TElement>(x));
+			RedisValue minValue = RedisSerializer.Serialize(min);
+			RedisValue maxValue = RedisSerializer.Serialize(max);
+			RedisValue[] values = await _database.SortedSetRangeByValueAsync(Key, minValue, maxValue, exclude, order, skip, take, flags);
+			return Array.ConvertAll(values, x => RedisSerializer.Deserialize<TElement>(x));
 		}
 
 		/// <inheritdoc cref="IDatabaseAsync.SortedSetRankAsync(RedisKey, RedisValue, Order, CommandFlags)"/>
-		public Task<long?> RankAsync(TElement Item, Order Order = Order.Ascending, CommandFlags Flags = CommandFlags.None)
+		public Task<long?> RankAsync(TElement item, Order order = Order.Ascending, CommandFlags flags = CommandFlags.None)
 		{
-			return Database.SortedSetRankAsync(Key, RedisSerializer.Serialize(Item), Order, Flags);
+			return _database.SortedSetRankAsync(Key, RedisSerializer.Serialize(item), order, flags);
 		}
 
 		/// <inheritdoc cref="IDatabaseAsync.SortedSetRemoveAsync(RedisKey, RedisValue, CommandFlags)"/>
-		public Task<bool> RemoveAsync(TElement Value, CommandFlags Flags = CommandFlags.None)
+		public Task<bool> RemoveAsync(TElement value, CommandFlags flags = CommandFlags.None)
 		{
-			return Database.SortedSetRemoveAsync(Key, RedisSerializer.Serialize(Value), Flags);
+			return _database.SortedSetRemoveAsync(Key, RedisSerializer.Serialize(value), flags);
 		}
 
 		/// <inheritdoc cref="IDatabaseAsync.SortedSetRemoveAsync(RedisKey, RedisValue[], CommandFlags)"/>
-		public Task<long> RemoveAsync(IEnumerable<TElement> Items, CommandFlags Flags = CommandFlags.None)
+		public Task<long> RemoveAsync(IEnumerable<TElement> items, CommandFlags flags = CommandFlags.None)
 		{
-			RedisValue[] Values = Items.Select(x => RedisSerializer.Serialize(x)).ToArray();
-			return Database.SortedSetRemoveAsync(Key, Values, Flags);
+			RedisValue[] values = items.Select(x => RedisSerializer.Serialize(x)).ToArray();
+			return _database.SortedSetRemoveAsync(Key, values, flags);
 		}
 
-		/// <inheritdoc cref="IDatabaseAsync.SortedSetRemoveRangeByRankAsync(RedisKey, long, long, CommandFlags)"/>
-		public Task<long> RemoveRangeByRankAsync(long Start, long Stop, CommandFlags Flags = CommandFlags.None)
+		/// <inheritdoc cref="IDatabaseAsync.SortedSetRemoveRangeByRankAsync(RedisKey, Int64, Int64, CommandFlags)"/>
+		public Task<long> RemoveRangeByRankAsync(long start, long stop, CommandFlags flags = CommandFlags.None)
 		{
-			return Database.SortedSetRemoveRangeByRankAsync(Key, Start, Stop, Flags);
+			return _database.SortedSetRemoveRangeByRankAsync(Key, start, stop, flags);
 		}
 
-		/// <inheritdoc cref="IDatabaseAsync.SortedSetRemoveRangeByScoreAsync(RedisKey, double, double, Exclude, CommandFlags)"/>
-		public Task<long> RemoveRangeByScoreAsync(double Start, double Stop, Exclude Exclude = Exclude.None, CommandFlags Flags = CommandFlags.None)
+		/// <inheritdoc cref="IDatabaseAsync.SortedSetRemoveRangeByScoreAsync(RedisKey, Double, Double, Exclude, CommandFlags)"/>
+		public Task<long> RemoveRangeByScoreAsync(double start, double stop, Exclude exclude = Exclude.None, CommandFlags flags = CommandFlags.None)
 		{
-			return Database.SortedSetRemoveRangeByScoreAsync(Key, Start, Stop, Exclude, Flags);
+			return _database.SortedSetRemoveRangeByScoreAsync(Key, start, stop, exclude, flags);
 		}
 
 		/// <inheritdoc cref="IDatabaseAsync.SortedSetRemoveRangeByValueAsync(RedisKey, RedisValue, RedisValue, Exclude, CommandFlags)"/>
-		public async Task<long> RemoveRangeByValueAsync(TElement Min, TElement Max, Exclude Exclude = Exclude.None, CommandFlags Flags = CommandFlags.None)
+		public async Task<long> RemoveRangeByValueAsync(TElement min, TElement max, Exclude exclude = Exclude.None, CommandFlags flags = CommandFlags.None)
 		{
-			RedisValue MinValue = RedisSerializer.Serialize(Min);
-			RedisValue MaxValue = RedisSerializer.Serialize(Max);
-			return await Database.SortedSetRemoveRangeByValueAsync(Key, MinValue, MaxValue, Exclude, Flags);
+			RedisValue minValue = RedisSerializer.Serialize(min);
+			RedisValue maxValue = RedisSerializer.Serialize(max);
+			return await _database.SortedSetRemoveRangeByValueAsync(Key, minValue, maxValue, exclude, flags);
 		}
 
-		/// <inheritdoc cref="IDatabaseAsync.SortedSetAddAsync(RedisKey, RedisValue, double, CommandFlags)"/>
-		public async IAsyncEnumerable<SortedSetEntry<TElement>> ScanAsync(RedisValue Pattern = default, int PageSize = 250, long Cursor = 0, int PageOffset = 0, CommandFlags Flags = CommandFlags.None)
+		/// <inheritdoc cref="IDatabaseAsync.SortedSetAddAsync(RedisKey, RedisValue, Double, CommandFlags)"/>
+		public async IAsyncEnumerable<SortedSetEntry<TElement>> ScanAsync(RedisValue pattern = default, int pageSize = 250, long cursor = 0, int pageOffset = 0, CommandFlags flags = CommandFlags.None)
 		{
-			await foreach (SortedSetEntry Entry in Database.SortedSetScanAsync(Key, Pattern, PageSize, Cursor, PageOffset, Flags))
+			await foreach (SortedSetEntry entry in _database.SortedSetScanAsync(Key, pattern, pageSize, cursor, pageOffset, flags))
 			{
-				yield return new SortedSetEntry<TElement>(Entry);
+				yield return new SortedSetEntry<TElement>(entry);
 			}
 		}
 
 		/// <inheritdoc cref="IDatabaseAsync.SortedSetScoreAsync(RedisKey, RedisValue, CommandFlags)"/>
-		public Task<double?> ScoreAsync(TElement Member, CommandFlags Flags = CommandFlags.None)
+		public Task<double?> ScoreAsync(TElement member, CommandFlags flags = CommandFlags.None)
 		{
-			RedisValue Value = RedisSerializer.Serialize(Member);
-			return Database.SortedSetScoreAsync(Key, Value, Flags);
+			RedisValue value = RedisSerializer.Serialize(member);
+			return _database.SortedSetScoreAsync(Key, value, flags);
 		}
 	}
 
@@ -216,9 +215,9 @@ namespace EpicGames.Redis
 		/// <summary>
 		/// Creates a version of this set which modifies a transaction rather than the direct DB
 		/// </summary>
-		public static RedisSortedSet<TElement> With<TElement>(this ITransaction Transaction, RedisSortedSet<TElement> Set)
+		public static RedisSortedSet<TElement> With<TElement>(this ITransaction transaction, RedisSortedSet<TElement> set)
 		{
-			return new RedisSortedSet<TElement>(Transaction, Set.Key);
+			return new RedisSortedSet<TElement>(transaction, set.Key);
 		}
 	}
 }
diff --git a/Engine/Source/Programs/Shared/EpicGames.Redis/RedisString.cs b/Engine/Source/Programs/Shared/EpicGames.Redis/RedisString.cs
--- a/Engine/Source/Programs/Shared/EpicGames.Redis/RedisString.cs
+++ b/Engine/Source/Programs/Shared/EpicGames.Redis/RedisString.cs
@@ -2,8 +2,6 @@
 
 using StackExchange.Redis;
 using System;
-using System.Collections.Generic;
-using System.Linq;
 using System.Threading.Tasks;
 
 namespace EpicGames.Redis
@@ -24,24 +22,24 @@ namespace EpicGames.Redis
 		/// <summary>
 		/// Constructor
 		/// </summary>
-		/// <param name="Database"></param>
-		/// <param name="Key"></param>
-		public RedisString(IDatabaseAsync Database, RedisKey Key)
+		/// <param name="database"></param>
+		/// <param name="key"></param>
+		public RedisString(IDatabaseAsync database, RedisKey key)
 		{
-			this.Database = Database;
-			this.Key = Key;
+			Database = database;
+			Key = key;
 		}
 
 		/// <inheritdoc cref="IDatabaseAsync.StringLengthAsync(RedisKey, CommandFlags)"/>
-		public Task<long> LengthAsync(CommandFlags Flags = CommandFlags.None)
+		public Task<long> LengthAsync(CommandFlags flags = CommandFlags.None)
 		{
-			return Database.StringLengthAsync(Key, Flags);
+			return Database.StringLengthAsync(Key, flags);
 		}
 
 		/// <inheritdoc cref="IDatabaseAsync.StringSetAsync(RedisKey, RedisValue, TimeSpan?, When, CommandFlags)"/>
-		public Task SetAsync(TElement Value, TimeSpan? Expiry = null, When When = When.Always, CommandFlags Flags = CommandFlags.None)
+		public Task SetAsync(TElement value, TimeSpan? expiry = null, When when = When.Always, CommandFlags flags = CommandFlags.None)
 		{
-			return Database.StringSetAsync(Key, RedisSerializer.Serialize(Value), Expiry, When, Flags);
+			return Database.StringSetAsync(Key, RedisSerializer.Serialize(value), expiry, when, flags);
 		}
 	}
 
@@ -51,57 +49,57 @@ namespace EpicGames.Redis
 	public static class RedisStringExtensions
 	{
 		/// <inheritdoc cref="IDatabaseAsync.StringGetAsync(RedisKey, CommandFlags)"/>
-		public static async Task<TElement?> GetAsync<TElement>(this RedisString<TElement> String, CommandFlags Flags = CommandFlags.None) where TElement : class
+		public static async Task<TElement?> GetAsync<TElement>(this RedisString<TElement> @string, CommandFlags flags = CommandFlags.None) where TElement : class
 		{
-			RedisValue Value = await String.Database.StringGetAsync(String.Key, Flags);
-			if (Value.IsNullOrEmpty)
+			RedisValue value = await @string.Database.StringGetAsync(@string.Key, flags);
+			if (value.IsNullOrEmpty)
 			{
 				return null;
 			}
-			return RedisSerializer.Deserialize<TElement>(Value);
+			return RedisSerializer.Deserialize<TElement>(value);
 		}
 
 		/// <inheritdoc cref="IDatabaseAsync.StringGetAsync(RedisKey, CommandFlags)"/>
-		public static async Task<TElement?> GetValueAsync<TElement>(this RedisString<TElement> String, CommandFlags Flags = CommandFlags.None) where TElement : struct
+		public static async Task<TElement?> GetValueAsync<TElement>(this RedisString<TElement> @string, CommandFlags flags = CommandFlags.None) where TElement : struct
 		{
-			RedisValue Value = await String.Database.StringGetAsync(String.Key, Flags);
-			if (Value.IsNullOrEmpty)
+			RedisValue value = await @string.Database.StringGetAsync(@string.Key, flags);
+			if (value.IsNullOrEmpty)
 			{
 				return default(TElement);
 			}
-			return RedisSerializer.Deserialize<TElement>(Value);
+			return RedisSerializer.Deserialize<TElement>(value);
 		}
 
-		/// <inheritdoc cref="IDatabaseAsync.StringDecrementAsync(RedisKey, long, CommandFlags)"/>
-		public static Task<long> DecrementAsync(this RedisString<long> String, long Value = 1L, CommandFlags Flags = CommandFlags.None)
+		/// <inheritdoc cref="IDatabaseAsync.StringDecrementAsync(RedisKey, Int64, CommandFlags)"/>
+		public static Task<long> DecrementAsync(this RedisString<long> @string, long value = 1L, CommandFlags flags = CommandFlags.None)
 		{
-			return String.Database.StringDecrementAsync(String.Key, Value, Flags);
+			return @string.Database.StringDecrementAsync(@string.Key, value, flags);
 		}
 
-		/// <inheritdoc cref="IDatabaseAsync.StringDecrementAsync(RedisKey, double, CommandFlags)"/>
-		public static Task<double> DecrementAsync(this RedisString<double> String, double Value = 1.0, CommandFlags Flags = CommandFlags.None)
+		/// <inheritdoc cref="IDatabaseAsync.StringDecrementAsync(RedisKey, Double, CommandFlags)"/>
+		public static Task<double> DecrementAsync(this RedisString<double> @string, double value = 1.0, CommandFlags flags = CommandFlags.None)
 		{
-			return String.Database.StringDecrementAsync(String.Key, Value, Flags);
+			return @string.Database.StringDecrementAsync(@string.Key, value, flags);
 		}
 
-		/// <inheritdoc cref="IDatabaseAsync.StringDecrementAsync(RedisKey, double, CommandFlags)"/>
-		public static Task<long> IncrementAsync(this RedisString<long> String, long Value = 1L, CommandFlags Flags = CommandFlags.None)
+		/// <inheritdoc cref="IDatabaseAsync.StringDecrementAsync(RedisKey, Double, CommandFlags)"/>
+		public static Task<long> IncrementAsync(this RedisString<long> @string, long value = 1L, CommandFlags flags = CommandFlags.None)
 		{
-			return String.Database.StringIncrementAsync(String.Key, Value, Flags);
+			return @string.Database.StringIncrementAsync(@string.Key, value, flags);
 		}
 
-		/// <inheritdoc cref="IDatabaseAsync.StringDecrementAsync(RedisKey, double, CommandFlags)"/>
-		public static Task<double> IncrementAsync(this RedisString<double> String, double Value = 1.0, CommandFlags Flags = CommandFlags.None)
+		/// <inheritdoc cref="IDatabaseAsync.StringDecrementAsync(RedisKey, Double, CommandFlags)"/>
+		public static Task<double> IncrementAsync(this RedisString<double> @string, double value = 1.0, CommandFlags flags = CommandFlags.None)
 		{
-			return String.Database.StringIncrementAsync(String.Key, Value, Flags);
+			return @string.Database.StringIncrementAsync(@string.Key, value, flags);
 		}
 
 		/// <summary>
 		/// Creates a version of this string which modifies a transaction rather than the direct DB
 		/// </summary>
-		public static RedisString<TElement> With<TElement>(this ITransaction Transaction, RedisString<TElement> Set)
+		public static RedisString<TElement> With<TElement>(this ITransaction transaction, RedisString<TElement> set)
 		{
-			return new RedisString<TElement>(Transaction, Set.Key);
+			return new RedisString<TElement>(transaction, set.Key);
 		}
 	}
 }
diff --git a/Engine/Source/Programs/Shared/EpicGames.Redis/Utility/RedisLock.cs b/Engine/Source/Programs/Shared/EpicGames.Redis/Utility/RedisLock.cs
--- a/Engine/Source/Programs/Shared/EpicGames.Redis/Utility/RedisLock.cs
+++ b/Engine/Source/Programs/Shared/EpicGames.Redis/Utility/RedisLock.cs
@@ -2,9 +2,7 @@
 
 using StackExchange.Redis;
 using System;
-using System.Collections.Generic;
 using System.Diagnostics;
-using System.Text;
 using System.Threading;
 using System.Threading.Tasks;
 
@@ -15,20 +13,20 @@ namespace EpicGames.Redis.Utility
 	/// </summary>
 	public class RedisLock : IAsyncDisposable, IDisposable
 	{
-		IDatabase Database;
-		RedisKey Key;
-		CancellationTokenSource? CancellationSource;
-		Task? BackgroundTask;
+		readonly IDatabase _database;
+		readonly RedisKey _key;
+		CancellationTokenSource? _cancellationSource;
+		Task? _backgroundTask;
 
 		/// <summary>
 		/// Constructor
 		/// </summary>
-		/// <param name="Database"></param>
-		/// <param name="Key"></param>
-		public RedisLock(IDatabase Database, RedisKey Key)
+		/// <param name="database"></param>
+		/// <param name="key"></param>
+		public RedisLock(IDatabase database, RedisKey key)
 		{
-			this.Database = Database;
-			this.Key = Key;
+			_database = database;
+			_key = key;
 		}
 
 		/// <inheritdoc/>
@@ -40,27 +38,27 @@ namespace EpicGames.Redis.Utility
 		/// <inheritdoc/>
 		public async ValueTask DisposeAsync()
 		{
-			if (BackgroundTask != null)
+			if (_backgroundTask != null)
 			{
-				CancellationSource!.Cancel();
-				await BackgroundTask;
-				CancellationSource.Dispose();
-				BackgroundTask = null;
+				_cancellationSource!.Cancel();
+				await _backgroundTask;
+				_cancellationSource.Dispose();
+				_backgroundTask = null;
 			}
 		}
 
 		/// <summary>
 		/// Attempts to acquire the lock for the given period of time. The lock will be renewed once half of this interval has elapsed.
 		/// </summary>
-		/// <param name="Duration">Time after which the lock expires</param>
-		/// <param name="ReleaseOnDispose">Whether the lock should be released when disposed. If false, the lock will be held for the given time, but not renewed once disposed.</param>
+		/// <param name="duration">Time after which the lock expires</param>
+		/// <param name="releaseOnDispose">Whether the lock should be released when disposed. If false, the lock will be held for the given time, but not renewed once disposed.</param>
 		/// <returns>True if the lock was acquired, false if another service already has it</returns>
-		public async ValueTask<bool> AcquireAsync(TimeSpan Duration, bool ReleaseOnDispose = true)
+		public async ValueTask<bool> AcquireAsync(TimeSpan duration, bool releaseOnDispose = true)
 		{
-			if (await Database.StringSetAsync(Key, RedisValue.EmptyString, Duration, When.NotExists))
+			if (await _database.StringSetAsync(_key, RedisValue.EmptyString, duration, When.NotExists))
 			{
-				CancellationSource = new CancellationTokenSource();
-				BackgroundTask = Task.Run(() => RenewAsync(Duration, ReleaseOnDispose, CancellationSource.Token));
+				_cancellationSource = new CancellationTokenSource();
+				_backgroundTask = Task.Run(() => RenewAsync(duration, releaseOnDispose, _cancellationSource.Token));
 				return true;
 			}
 			return false;
@@ -69,30 +67,30 @@ namespace EpicGames.Redis.Utility
 		/// <summary>
 		/// Background task which renews the lock while the service is running
 		/// </summary>
-		/// <param name="Duration"></param>
-		/// <param name="ReleaseOnDispose">Whether the lock should be released when the cancellation token is fired</param>
-		/// <param name="CancellationToken"></param>
+		/// <param name="duration"></param>
+		/// <param name="releaseOnDispose">Whether the lock should be released when the cancellation token is fired</param>
+		/// <param name="cancellationToken"></param>
 		/// <returns></returns>
-		async Task RenewAsync(TimeSpan Duration, bool ReleaseOnDispose, CancellationToken CancellationToken)
+		async Task RenewAsync(TimeSpan duration, bool releaseOnDispose, CancellationToken cancellationToken)
 		{
-			Stopwatch Timer = Stopwatch.StartNew();
+			Stopwatch timer = Stopwatch.StartNew();
 			for (; ; )
 			{
-				await Task.Delay(Duration / 2, CancellationToken).ContinueWith(x => { }); // Do not throw
-				if (CancellationToken.IsCancellationRequested)
+				await Task.Delay(duration / 2, cancellationToken).ContinueWith(x => { }); // Do not throw
+				if (cancellationToken.IsCancellationRequested)
 				{
-					Timer.Stop();
-					if (ReleaseOnDispose || Timer.Elapsed > Duration)
+					timer.Stop();
+					if (releaseOnDispose || timer.Elapsed > duration)
 					{
-						await Database.StringSetAsync(Key, RedisValue.Null);
+						await _database.StringSetAsync(_key, RedisValue.Null);
 					}
 					else
 					{
-						await Database.StringSetAsync(Key, RedisValue.EmptyString, Duration - Timer.Elapsed, When.Exists);
+						await _database.StringSetAsync(_key, RedisValue.EmptyString, duration - timer.Elapsed, When.Exists);
 					}
 					break;
 				}
-				if (!await Database.StringSetAsync(Key, RedisValue.EmptyString, Duration, When.Exists))
+				if (!await _database.StringSetAsync(_key, RedisValue.EmptyString, duration, When.Exists))
 				{
 					break;
 				}
@@ -109,11 +107,11 @@ namespace EpicGames.Redis.Utility
 		/// <summary>
 		/// Constructor
 		/// </summary>
-		/// <param name="Database"></param>
-		/// <param name="BaseKey"></param>
-		/// <param name="Value"></param>
-		public RedisLock(IDatabase Database, RedisKey BaseKey, T Value)
-			: base(Database, BaseKey.Append(RedisSerializer.Serialize<T>(Value).AsKey()))
+		/// <param name="database"></param>
+		/// <param name="baseKey"></param>
+		/// <param name="value"></param>
+		public RedisLock(IDatabase database, RedisKey baseKey, T value)
+			: base(database, baseKey.Append(RedisSerializer.Serialize<T>(value).AsKey()))
 		{
 		}
 	}
