Reformat EpicGames.BuildGraph according to style guidelines.

diff --git a/Engine/Source/Programs/AutomationTool/BuildGraph/BuildGraph.cs b/Engine/Source/Programs/AutomationTool/BuildGraph/BuildGraph.cs
--- a/Engine/Source/Programs/AutomationTool/BuildGraph/BuildGraph.cs
+++ b/Engine/Source/Programs/AutomationTool/BuildGraph/BuildGraph.cs
@@ -772,7 +772,7 @@ namespace AutomationTool
 			bool bHasRequiredAttributes = true;
 			foreach (ScriptTaskParameterBinding Parameter in Task.NameToParameter.Values)
 			{
-				if (!Parameter.bOptional && !TaskInfo.Arguments.ContainsKey(Parameter.Name))
+				if (!Parameter.Optional && !TaskInfo.Arguments.ContainsKey(Parameter.Name))
 				{
 					OutputBindingError(TaskInfo, "Missing required attribute - {AttrName}", Parameter.Name);
 					bHasRequiredAttributes = false;
@@ -1499,7 +1499,7 @@ namespace AutomationTool
 								string[] Columns = new string[4];
 								Columns[0] = ParameterName;
 								Columns[1] = TypeName;
-								Columns[2] = Parameter.bOptional ? "Optional" : "Required";
+								Columns[2] = Parameter.Optional ? "Optional" : "Required";
 								Columns[3] = ConvertToMarkdown(ParameterElement.SelectSingleNode("summary"));
 								Rows.Add(Columns);
 							}
@@ -1605,7 +1605,7 @@ namespace AutomationTool
 								Writer.WriteLine("      <tr>");
 								Writer.WriteLine("         <td>{0}</td>", ParameterName);
 								Writer.WriteLine("         <td>{0}</td>", TypeName);
-								Writer.WriteLine("         <td>{0}</td>", Parameter.bOptional? "Optional" : "Required");
+								Writer.WriteLine("         <td>{0}</td>", Parameter.Optional? "Optional" : "Required");
 								Writer.WriteLine("         <td>{0}</td>", ParameterElement.SelectSingleNode("summary").InnerXml.Trim());
 								Writer.WriteLine("      </tr>");
 							}
diff --git a/Engine/Source/Programs/AutomationTool/BuildGraph/Scripts/InstalledBuild.cs b/Engine/Source/Programs/AutomationTool/BuildGraph/Scripts/InstalledBuild.cs
--- a/Engine/Source/Programs/AutomationTool/BuildGraph/Scripts/InstalledBuild.cs
+++ b/Engine/Source/Programs/AutomationTool/BuildGraph/Scripts/InstalledBuild.cs
@@ -146,7 +146,7 @@ namespace AutomationTool
 				.AddNode(x => CompileUnrealEditorWin64Async(x, CrashReporterCompileArgs, EmbedSrcSrvInfo, CompileDatasmithPlugins, WithFullDebugInfo, SignExecutables))
 				.Requires(WinUhtNode);
 			
-			Graph.AddAggregate("Win64 Editor", WinEditorNode, Label: "Editors/Win64");
+			Graph.AddAggregate("Win64 Editor", WinEditorNode, label: "Editors/Win64");
 
 
 			/////// TARGET PLATFORMS ////////////////////////////////////////////////
@@ -219,7 +219,7 @@ namespace AutomationTool
 				.AddNode(x => MakeInstalledBuildWin64Async(x, WinInstalledFiles, WinFinalizeArgs, WinOutputDir))
 				.Requires(WinInstalledFiles);
 
-			Graph.AddAggregate("HostPlatforms_Win64", WinInstalledNode, Label: "Builds/Win64");
+			Graph.AddAggregate("HostPlatforms_Win64", WinInstalledNode, label: "Builds/Win64");
 		}
 
 		/// <summary>
diff --git a/Engine/Source/Programs/Shared/EpicGames.BuildGraph/BgAgent.cs b/Engine/Source/Programs/Shared/EpicGames.BuildGraph/BgAgent.cs
--- a/Engine/Source/Programs/Shared/EpicGames.BuildGraph/BgAgent.cs
+++ b/Engine/Source/Programs/Shared/EpicGames.BuildGraph/BgAgent.cs
@@ -3,9 +3,6 @@
 using System;
 using System.Collections.Generic;
 using System.Diagnostics;
-using System.Linq;
-using System.Text;
-using System.Threading.Tasks;
 using System.Xml;
 
 namespace EpicGames.BuildGraph
@@ -35,28 +32,28 @@ namespace EpicGames.BuildGraph
 		/// <summary>
 		/// Constructor
 		/// </summary>
-		/// <param name="InName">Name of this agent group</param>
-		/// <param name="InPossibleTypes">Array of valid agent types. See comment for AgentTypes member.</param>
-		public BgAgent(string InName, string[] InPossibleTypes)
+		/// <param name="inName">Name of this agent group</param>
+		/// <param name="inPossibleTypes">Array of valid agent types. See comment for AgentTypes member.</param>
+		public BgAgent(string inName, string[] inPossibleTypes)
 		{
-			Name = InName;
-			PossibleTypes = InPossibleTypes;
+			Name = inName;
+			PossibleTypes = inPossibleTypes;
 		}
 
 		/// <summary>
 		/// Writes this agent group out to a file, filtering nodes by a controlling trigger
 		/// </summary>
-		/// <param name="Writer">The XML writer to output to</param>
-		public void Write(XmlWriter Writer)
+		/// <param name="writer">The XML writer to output to</param>
+		public void Write(XmlWriter writer)
 		{
-			Writer.WriteStartElement("Agent");
-			Writer.WriteAttributeString("Name", Name);
-			Writer.WriteAttributeString("Type", String.Join(";", PossibleTypes));
-			foreach (BgNode Node in Nodes)
+			writer.WriteStartElement("Agent");
+			writer.WriteAttributeString("Name", Name);
+			writer.WriteAttributeString("Type", String.Join(";", PossibleTypes));
+			foreach (BgNode node in Nodes)
 			{
-				Node.Write(Writer);
+				node.Write(writer);
 			}
-			Writer.WriteEndElement();
+			writer.WriteEndElement();
 		}
 	}
 }
diff --git a/Engine/Source/Programs/Shared/EpicGames.BuildGraph/BgAgentSpec.cs b/Engine/Source/Programs/Shared/EpicGames.BuildGraph/BgAgentSpec.cs
--- a/Engine/Source/Programs/Shared/EpicGames.BuildGraph/BgAgentSpec.cs
+++ b/Engine/Source/Programs/Shared/EpicGames.BuildGraph/BgAgentSpec.cs
@@ -1,14 +1,11 @@
 // Copyright Epic Games, Inc. All Rights Reserved.
 
-using EpicGames.BuildGraph.Expressions;
-using Microsoft.Extensions.Logging;
 using System;
 using System.Collections.Generic;
-using System.Linq;
 using System.Linq.Expressions;
-using System.Reflection;
-using System.Text;
 using System.Threading.Tasks;
+using EpicGames.BuildGraph.Expressions;
+using Microsoft.Extensions.Logging;
 
 namespace EpicGames.BuildGraph
 {
@@ -30,19 +27,19 @@ namespace EpicGames.BuildGraph
 		/// <summary>
 		/// Add agent types to the agent definition
 		/// </summary>
-		public BgAgentConfig AddTypes(params BgString[] TypeNames)
+		public BgAgentConfig AddTypes(params BgString[] typeNames)
 		{
-			Types = Types.Add(TypeNames);
+			Types = Types.Add(typeNames);
 			return this;
 		}
 
 		/// <summary>
 		/// Check that a certain condition is true before executing the node
 		/// </summary>
-		public BgAgentConfig Check(BgBool Condition, LogLevel Level, BgString Message)
+		public BgAgentConfig Check(BgBool condition, LogLevel level, BgString message)
 		{
-			BgDiagnosticSpec Diagnostic = new BgDiagnosticSpec(Level, Message);
-			Diagnostics = Diagnostics.AddIf(Condition, Diagnostic);
+			BgDiagnosticSpec diagnostic = new BgDiagnosticSpec(level, message);
+			Diagnostics = Diagnostics.AddIf(condition, diagnostic);
 			return this;
 		}
 	}
@@ -75,84 +72,84 @@ namespace EpicGames.BuildGraph
 		/// <summary>
 		/// Private constructor. Use <see cref="BgGraphSpec.AddAgent(BgString)"/> to construct an agent
 		/// </summary>
-		/// <param name="Name"></param>
-		internal BgAgentSpec(BgString Name)
+		/// <param name="name"></param>
+		internal BgAgentSpec(BgString name)
 		{
-			this.Name = Name;
+			Name = name;
 		}
 
 		/// <summary>
 		/// Adds a node to be executed on this agent
 		/// </summary>
-		/// <param name="Function">Lambda expression which calls a static method to be used for this node</param>
+		/// <param name="function">Lambda expression which calls a static method to be used for this node</param>
 		/// <returns>New node spec instance</returns>
-		public BgNodeSpec AddNode(Expression<Func<BgContext, Task>> Function)
+		public BgNodeSpec AddNode(Expression<Func<BgContext, Task>> function)
 		{
-			BgNodeSpec NodeSpec = BgNodeSpec.Create(Function);
-			NodeSpecs.Add(NodeSpec);
-			return NodeSpec;
+			BgNodeSpec nodeSpec = BgNodeSpec.Create(function);
+			NodeSpecs.Add(nodeSpec);
+			return nodeSpec;
 		}
 
 		/// <summary>
 		/// Adds a node to be executed on this agent
 		/// </summary>
-		/// <param name="Function">Lambda expression which calls a static method to be used for this node</param>
+		/// <param name="function">Lambda expression which calls a static method to be used for this node</param>
 		/// <returns>New node spec instance</returns>
-		public BgNodeSpec<T> AddNode<T>(Expression<Func<BgContext, Task<T>>> Function)
+		public BgNodeSpec<T> AddNode<T>(Expression<Func<BgContext, Task<T>>> function)
 		{
-			BgNodeSpec<T> NodeSpec = BgNodeSpec.Create<T>(Function);
-			NodeSpecs.Add(NodeSpec);
-			return NodeSpec;
+			BgNodeSpec<T> nodeSpec = BgNodeSpec.Create<T>(function);
+			NodeSpecs.Add(nodeSpec);
+			return nodeSpec;
 		}
 
 		/// <summary>
 		/// Add agent types to the agent definition
 		/// </summary>
-		public BgAgentSpec Type(params BgString[] NewTypeNames)
+		public BgAgentSpec Type(params BgString[] newTypeNames)
 		{
-			TypeNames = TypeNames.Add(NewTypeNames);
+			TypeNames = TypeNames.Add(newTypeNames);
 			return this;
 		}
 
 		/// <summary>
 		/// Check that a certain condition is true before executing the node
 		/// </summary>
-		public BgAgentSpec WarnIf(BgBool Condition, BgString Message)
+		public BgAgentSpec WarnIf(BgBool condition, BgString message)
 		{
-			BgDiagnosticSpec Diagnostic = new BgDiagnosticSpec(LogLevel.Warning, Message);
-			Diagnostics = Diagnostics.AddIf(Condition, Diagnostic);
+			BgDiagnosticSpec diagnostic = new BgDiagnosticSpec(LogLevel.Warning, message);
+			Diagnostics = Diagnostics.AddIf(condition, diagnostic);
 			return this;
 		}
 
 		/// <summary>
 		/// Check that a certain condition is true before executing the node
 		/// </summary>
-		public BgAgentSpec ErrorIf(BgBool Condition, BgString Message)
+		public BgAgentSpec ErrorIf(BgBool condition, BgString message)
 		{
-			BgDiagnosticSpec Diagnostic = new BgDiagnosticSpec(LogLevel.Error, Message);
-			Diagnostics = Diagnostics.AddIf(Condition, Diagnostic);
+			BgDiagnosticSpec diagnostic = new BgDiagnosticSpec(LogLevel.Error, message);
+			Diagnostics = Diagnostics.AddIf(condition, diagnostic);
 			return this;
 		}
 
 		/// <summary>
 		/// Creates a concrete <see cref="BgAgent"/> object from the specification
 		/// </summary>
-		public void AddToGraph(BgExprContext Context, BgGraph Graph)
+		public void AddToGraph(BgExprContext context, BgGraph graph)
 		{
-			string[] Types = TypeNames.Compute(Context).ToArray();
+			string[] types = TypeNames.Compute(context).ToArray();
 
-			BgAgent Agent = new BgAgent(Name.Compute(Context), Types);
-			Graph.NameToAgent.Add(Agent.Name, Agent);
-			Graph.Agents.Add(Agent);
+			BgAgent agent = new BgAgent(Name.Compute(context), types);
+			graph.NameToAgent.Add(agent.Name, agent);
+			graph.Agents.Add(agent);
 
-			foreach (BgDiagnosticSpec Precondition in Diagnostics.GetEnumerable(Context))
+			foreach (BgDiagnosticSpec precondition in Diagnostics.GetEnumerable(context))
 			{
-				Precondition.AddToGraph(Context, Graph, Agent, null);
+				precondition.AddToGraph(context, graph, agent, null);
 			}
 
-			foreach (BgNodeSpec NodeSpec in NodeSpecs)
+			foreach (BgNodeSpec nodeSpec in NodeSpecs)
 			{
-				NodeSpec.AddToGraph(Context, Graph, Agent);
+				nodeSpec.AddToGraph(context, graph, agent);
 			}
 		}
 	}
diff --git a/Engine/Source/Programs/Shared/EpicGames.BuildGraph/BgAggregate.cs b/Engine/Source/Programs/Shared/EpicGames.BuildGraph/BgAggregate.cs
--- a/Engine/Source/Programs/Shared/EpicGames.BuildGraph/BgAggregate.cs
+++ b/Engine/Source/Programs/Shared/EpicGames.BuildGraph/BgAggregate.cs
@@ -1,10 +1,6 @@
 // Copyright Epic Games, Inc. All Rights Reserved.
 
-using System;
 using System.Collections.Generic;
-using System.Linq;
-using System.Text;
-using System.Threading.Tasks;
 
 namespace EpicGames.BuildGraph
 {
@@ -26,10 +22,10 @@ namespace EpicGames.BuildGraph
 		/// <summary>
 		/// Constructor
 		/// </summary>
-		/// <param name="InName">Name of this aggregate</param>
-		public BgAggregate(string InName)
+		/// <param name="inName">Name of this aggregate</param>
+		public BgAggregate(string inName)
 		{
-			Name = InName;
+			Name = inName;
 		}
 
 		/// <summary>
diff --git a/Engine/Source/Programs/Shared/EpicGames.BuildGraph/BgAggregateSpec.cs b/Engine/Source/Programs/Shared/EpicGames.BuildGraph/BgAggregateSpec.cs
--- a/Engine/Source/Programs/Shared/EpicGames.BuildGraph/BgAggregateSpec.cs
+++ b/Engine/Source/Programs/Shared/EpicGames.BuildGraph/BgAggregateSpec.cs
@@ -1,10 +1,7 @@
 // Copyright Epic Games, Inc. All Rights Reserved.
 
-using EpicGames.BuildGraph.Expressions;
-using System;
-using System.Collections.Generic;
 using System.Linq;
-using System.Text;
+using EpicGames.BuildGraph.Expressions;
 
 namespace EpicGames.BuildGraph
 {
@@ -26,29 +23,29 @@ namespace EpicGames.BuildGraph
 		/// <summary>
 		/// Internal constructor. Use <see cref="BgGraphSpec.AddAggregate(BgString, BgList{BgFileSet})"/> to create an aggregate.
 		/// </summary>
-		internal BgAggregateSpec(BgString Name, BgList<BgFileSet> RequiredOutputs)
+		internal BgAggregateSpec(BgString name, BgList<BgFileSet> requiredOutputs)
 		{
-			this.Name = Name;
-			this.RequiredOutputs = RequiredOutputs;
+			Name = name;
+			RequiredOutputs = requiredOutputs;
 		}
 
 		/// <summary>
 		/// Creates a concrete aggregate object from this specification.
 		/// </summary>
-		internal void AddToGraph(BgExprContext Context, BgGraph Graph)
+		internal void AddToGraph(BgExprContext context, BgGraph graph)
 		{
-			BgAggregate Aggregate = new BgAggregate(Name.Compute(Context));
-			Aggregate.RequiredNodes.UnionWith(RequiredOutputs.ComputeTags(Context).Select(x => Graph.TagNameToNodeOutput[x].ProducingNode));
-			Graph.NameToAggregate.Add(Aggregate.Name, Aggregate);
+			BgAggregate aggregate = new BgAggregate(Name.Compute(context));
+			aggregate.RequiredNodes.UnionWith(RequiredOutputs.ComputeTags(context).Select(x => graph.TagNameToNodeOutput[x].ProducingNode));
+			graph.NameToAggregate.Add(aggregate.Name, aggregate);
 		}
 
 		/// <summary>
 		/// Adds a set of dependencies to this aggregate
 		/// </summary>
-		/// <param name="Tokens">List of token dependencies</param>
-		public void Requires(params BgFileSet[] Tokens)
+		/// <param name="tokens">List of token dependencies</param>
+		public void Requires(params BgFileSet[] tokens)
 		{
-			RequiredOutputs.Add(Tokens);
+			RequiredOutputs.Add(tokens);
 		}
 	}
 }
diff --git a/Engine/Source/Programs/Shared/EpicGames.BuildGraph/BgBadge.cs b/Engine/Source/Programs/Shared/EpicGames.BuildGraph/BgBadge.cs
--- a/Engine/Source/Programs/Shared/EpicGames.BuildGraph/BgBadge.cs
+++ b/Engine/Source/Programs/Shared/EpicGames.BuildGraph/BgBadge.cs
@@ -1,11 +1,6 @@
 // Copyright Epic Games, Inc. All Rights Reserved.
 
-using System;
 using System.Collections.Generic;
-using System.Diagnostics;
-using System.Linq;
-using System.Text;
-using System.Threading.Tasks;
 
 namespace EpicGames.BuildGraph
 {
@@ -37,14 +32,14 @@ namespace EpicGames.BuildGraph
 		/// <summary>
 		/// Constructor
 		/// </summary>
-		/// <param name="InName">Name of this report</param>
-		/// <param name="InProject">Depot path to the project that this badge applies to</param>
-		/// <param name="InChange">The changelist to post the badge for</param>
-		public BgBadge(string InName, string InProject, int InChange)
+		/// <param name="inName">Name of this report</param>
+		/// <param name="inProject">Depot path to the project that this badge applies to</param>
+		/// <param name="inChange">The changelist to post the badge for</param>
+		public BgBadge(string inName, string inProject, int inChange)
 		{
-			Name = InName;
-			Project = InProject;
-			Change = InChange;
+			Name = inName;
+			Project = inProject;
+			Change = inChange;
 		}
 
 		/// <summary>
diff --git a/Engine/Source/Programs/Shared/EpicGames.BuildGraph/BgCondition.cs b/Engine/Source/Programs/Shared/EpicGames.BuildGraph/BgCondition.cs
--- a/Engine/Source/Programs/Shared/EpicGames.BuildGraph/BgCondition.cs
+++ b/Engine/Source/Programs/Shared/EpicGames.BuildGraph/BgCondition.cs
@@ -4,7 +4,6 @@ using System;
 using System.Collections.Generic;
 using System.IO;
 using System.Linq;
-using System.Text;
 using System.Threading.Tasks;
 
 namespace EpicGames.BuildGraph
@@ -17,9 +16,9 @@ namespace EpicGames.BuildGraph
 		/// <summary>
 		/// Constructor; formats the exception message with the given String.Format() style parameters.
 		/// </summary>
-		/// <param name="Format">Formatting string, in String.Format syntax</param>
-		/// <param name="Args">Optional arguments for the string</param>
-		public BgConditionException(string Format, params object[] Args) : base(String.Format(Format, Args))
+		/// <param name="format">Formatting string, in String.Format syntax</param>
+		/// <param name="args">Optional arguments for the string</param>
+		public BgConditionException(string format, params object[] args) : base(String.Format(format, args))
 		{
 		}
 	}
@@ -64,39 +63,39 @@ namespace EpicGames.BuildGraph
 		/// <summary>
 		/// Tokens for the condition
 		/// </summary>
-		List<string> Tokens = new List<string>();
+		readonly List<string> _tokens = new List<string>();
 
 		/// <summary>
 		/// The current token index
 		/// </summary>
-		int Idx;
+		int _idx;
 
 		/// <summary>
 		/// Context for evaluating the expression
 		/// </summary>
-		IBgScriptReaderContext Context;
+		readonly IBgScriptReaderContext _context;
 
 		/// <summary>
 		/// Constructor
 		/// </summary>
-		/// <param name="Text">The condition text</param>
-		/// <param name="Context">Context for evaluating the expression</param>
-		private BgCondition(string Text, IBgScriptReaderContext Context)
+		/// <param name="text">The condition text</param>
+		/// <param name="context">Context for evaluating the expression</param>
+		private BgCondition(string text, IBgScriptReaderContext context)
 		{
-			this.Context = Context;
+			_context = context;
 
-			Tokenize(Text, Tokens);
+			Tokenize(text, _tokens);
 		}
 
 		/// <summary>
 		/// Evaluates the given string as a condition. Throws a ConditionException on a type or syntax error.
 		/// </summary>
-		/// <param name="Text"></param>
-		/// <param name="Context"></param>
+		/// <param name="text"></param>
+		/// <param name="context"></param>
 		/// <returns>The result of evaluating the condition</returns>
-		public static ValueTask<bool> EvaluateAsync(string Text, IBgScriptReaderContext Context)
+		public static ValueTask<bool> EvaluateAsync(string text, IBgScriptReaderContext context)
 		{
-			return new BgCondition(Text, Context).EvaluateAsync();
+			return new BgCondition(text, context).EvaluateAsync();
 		}
 
 		/// <summary>
@@ -106,15 +105,15 @@ namespace EpicGames.BuildGraph
 		async ValueTask<bool> EvaluateAsync()
 		{
 			bool bResult = true;
-			if(Tokens.Count > 1)
+			if (_tokens.Count > 1)
 			{
-				Idx = 0;
-				string Result = await EvaluateOrAsync();
-				if(Tokens[Idx] != EndToken)
+				_idx = 0;
+				string result = await EvaluateOrAsync();
+				if (_tokens[_idx] != EndToken)
 				{
-					throw new BgConditionException("Garbage after expression: {0}", String.Join("", Tokens.Skip(Idx)));
+					throw new BgConditionException("Garbage after expression: {0}", String.Join("", _tokens.Skip(_idx)));
 				}
-				bResult = CoerceToBool(Result);
+				bResult = CoerceToBool(result);
 			}
 			return bResult;
 		}
@@ -126,16 +125,16 @@ namespace EpicGames.BuildGraph
 		async ValueTask<string> EvaluateOrAsync()
 		{
 			// <Condition> Or <Condition> Or...
-			string Result = await EvaluateAndAsync();
-			while(String.Compare(Tokens[Idx], "Or", true) == 0)
+			string result = await EvaluateAndAsync();
+			while (String.Compare(_tokens[_idx], "Or", true) == 0)
 			{
 				// Evaluate this condition. We use a binary OR here, because we want to parse everything rather than short-circuit it.
-				Idx++;
-				string Lhs = Result;
-				string Rhs = await EvaluateAndAsync();
-				Result = (CoerceToBool(Lhs) | CoerceToBool(Rhs))? "true" : "false";
+				_idx++;
+				string lhs = result;
+				string rhs = await EvaluateAndAsync();
+				result = (CoerceToBool(lhs) | CoerceToBool(rhs)) ? "true" : "false";
 			}
-			return Result;
+			return result;
 		}
 
 		/// <summary>
@@ -145,16 +144,16 @@ namespace EpicGames.BuildGraph
 		async ValueTask<string> EvaluateAndAsync()
 		{
 			// <Condition> And <Condition> And...
-			string Result = await EvaluateComparisonAsync();
-			while(String.Compare(Tokens[Idx], "And", true) == 0)
+			string result = await EvaluateComparisonAsync();
+			while (String.Compare(_tokens[_idx], "And", true) == 0)
 			{
 				// Evaluate this condition. We use a binary AND here, because we want to parse everything rather than short-circuit it.
-				Idx++;
-				string Lhs = Result;
-				string Rhs = await EvaluateComparisonAsync();
-				Result = (CoerceToBool(Lhs) & CoerceToBool(Rhs))? "true" : "false";
+				_idx++;
+				string lhs = result;
+				string rhs = await EvaluateComparisonAsync();
+				result = (CoerceToBool(lhs) & CoerceToBool(rhs)) ? "true" : "false";
 			}
-			return Result;
+			return result;
 		}
 
 		/// <summary>
@@ -171,56 +170,56 @@ namespace EpicGames.BuildGraph
 			// scalar > scalar
 			// scalar >= scalar
 
-			string Result = await EvaluateScalarAsync();
-			if(Tokens[Idx] == "==")
+			string result = await EvaluateScalarAsync();
+			if (_tokens[_idx] == "==")
 			{
 				// Compare two scalars for equality
-				Idx++;
-				string Lhs = Result;
-				string Rhs = await EvaluateScalarAsync();
-				Result = (String.Compare(Lhs, Rhs, true) == 0)? "true" : "false";
+				_idx++;
+				string lhs = result;
+				string rhs = await EvaluateScalarAsync();
+				result = (String.Compare(lhs, rhs, true) == 0) ? "true" : "false";
 			}
-			else if(Tokens[Idx] == "!=")
+			else if (_tokens[_idx] == "!=")
 			{
 				// Compare two scalars for inequality
-				Idx++;
-				string Lhs = Result;
-				string Rhs = await EvaluateScalarAsync();
-				Result = (String.Compare(Lhs, Rhs, true) != 0)? "true" : "false";
+				_idx++;
+				string lhs = result;
+				string rhs = await EvaluateScalarAsync();
+				result = (String.Compare(lhs, rhs, true) != 0) ? "true" : "false";
 			}
-			else if(Tokens[Idx] == "<")
+			else if (_tokens[_idx] == "<")
 			{
 				// Compares whether the first integer is less than the second
-				Idx++;
-				int Lhs = CoerceToInteger(Result);
-				int Rhs = CoerceToInteger(await EvaluateScalarAsync());
-				Result = (Lhs < Rhs)? "true" : "false";
+				_idx++;
+				int lhs = CoerceToInteger(result);
+				int rhs = CoerceToInteger(await EvaluateScalarAsync());
+				result = (lhs < rhs) ? "true" : "false";
 			}
-			else if(Tokens[Idx] == "<=")
+			else if (_tokens[_idx] == "<=")
 			{
 				// Compares whether the first integer is less than the second
-				Idx++;
-				int Lhs = CoerceToInteger(Result);
-				int Rhs = CoerceToInteger(await EvaluateScalarAsync());
-				Result = (Lhs <= Rhs)? "true" : "false";
+				_idx++;
+				int lhs = CoerceToInteger(result);
+				int rhs = CoerceToInteger(await EvaluateScalarAsync());
+				result = (lhs <= rhs) ? "true" : "false";
 			}
-			else if(Tokens[Idx] == ">")
+			else if (_tokens[_idx] == ">")
 			{
 				// Compares whether the first integer is less than the second
-				Idx++;
-				int Lhs = CoerceToInteger(Result);
-				int Rhs = CoerceToInteger(await EvaluateScalarAsync());
-				Result = (Lhs > Rhs)? "true" : "false";
+				_idx++;
+				int lhs = CoerceToInteger(result);
+				int rhs = CoerceToInteger(await EvaluateScalarAsync());
+				result = (lhs > rhs) ? "true" : "false";
 			}
-			else if(Tokens[Idx] == ">=")
+			else if (_tokens[_idx] == ">=")
 			{
 				// Compares whether the first integer is less than the second
-				Idx++;
-				int Lhs = CoerceToInteger(Result);
-				int Rhs = CoerceToInteger(await EvaluateScalarAsync());
-				Result = (Lhs >= Rhs)? "true" : "false";
+				_idx++;
+				int lhs = CoerceToInteger(result);
+				int rhs = CoerceToInteger(await EvaluateScalarAsync());
+				result = (lhs >= rhs) ? "true" : "false";
 			}
-			return Result;
+			return result;
 		}
 
 		/// <summary>
@@ -229,46 +228,46 @@ namespace EpicGames.BuildGraph
 		/// <returns>The result of evaluating the expression</returns>
 		IEnumerable<string> EvaluateArguments()
 		{
-			List<string> Arguments = new List<string>();
+			List<string> arguments = new List<string>();
 
 			// skip opening bracket
-			if (Tokens[Idx++] != "(")
+			if (_tokens[_idx++] != "(")
 			{
 				throw new BgConditionException("Expected '('");
 			}
 
-			bool DidCloseBracket = false;
+			bool didCloseBracket = false;
 
-			while (Idx < Tokens.Count)
+			while (_idx < _tokens.Count)
 			{
-				string NextToken = Tokens.ElementAt(Idx++);
+				string nextToken = _tokens.ElementAt(_idx++);
 
-				if (NextToken == EndToken)
+				if (nextToken == EndToken)
 				{
 					// ran out of items
 					break;
 				}
-				else if (NextToken == ")")
+				else if (nextToken == ")")
 				{
-					DidCloseBracket = true;
+					didCloseBracket = true;
 					break;
 				}
-				else if (NextToken != ",")
+				else if (nextToken != ",")
 				{
-					if (NextToken.First() == '\'' && NextToken.Last() == '\'')
+					if (nextToken.First() == '\'' && nextToken.Last() == '\'')
 					{
-						NextToken = NextToken.Substring(1, NextToken.Length - 2);
+						nextToken = nextToken.Substring(1, nextToken.Length - 2);
 					}
-					Arguments.Add(NextToken);
+					arguments.Add(nextToken);
 				}
 			}
 
-			if (!DidCloseBracket)
+			if (!didCloseBracket)
 			{
 				throw new BgConditionException("Expected ')'");
 			}
 
-			return Arguments;
+			return arguments;
 		}
 
 		/// <summary>
@@ -277,98 +276,98 @@ namespace EpicGames.BuildGraph
 		/// <returns>The result of evaluating the expression</returns>
 		async ValueTask<string> EvaluateScalarAsync()
 		{
-			string Result;
-			if(Tokens[Idx] == "(")
+			string result;
+			if (_tokens[_idx] == "(")
 			{
 				// Subexpression
-				Idx++;
-				Result = await EvaluateOrAsync();
-				if(Tokens[Idx] != ")")
+				_idx++;
+				result = await EvaluateOrAsync();
+				if (_tokens[_idx] != ")")
 				{
 					throw new BgConditionException("Expected ')'");
 				}
-				Idx++;
+				_idx++;
 			}
-			else if(Tokens[Idx] == "!")
+			else if (_tokens[_idx] == "!")
 			{
 				// Logical not
-				Idx++;
-				string Rhs = await EvaluateScalarAsync();
-				Result = CoerceToBool(Rhs)? "false" : "true";
+				_idx++;
+				string rhs = await EvaluateScalarAsync();
+				result = CoerceToBool(rhs) ? "false" : "true";
 			}
-			else if(String.Compare(Tokens[Idx], "Exists", true) == 0 && Tokens[Idx + 1] == "(")
+			else if (String.Compare(_tokens[_idx], "Exists", true) == 0 && _tokens[_idx + 1] == "(")
 			{
 				// Check whether file or directory exists. Evaluate the argument as a subexpression.
-				Idx++;
-				string Argument = await EvaluateScalarAsync();
-				Result = await Context.ExistsAsync(Argument)? "true" : "false";
+				_idx++;
+				string argument = await EvaluateScalarAsync();
+				result = await _context.ExistsAsync(argument) ? "true" : "false";
 			}
-			else if(String.Compare(Tokens[Idx], "HasTrailingSlash", true) == 0 && Tokens[Idx + 1] == "(")
+			else if (String.Compare(_tokens[_idx], "HasTrailingSlash", true) == 0 && _tokens[_idx + 1] == "(")
 			{
 				// Check whether the given string ends with a slash
-				Idx++;
-				string Argument = await EvaluateScalarAsync();
-				Result = (Argument.Length > 0 && (Argument[Argument.Length - 1] == Path.DirectorySeparatorChar || Argument[Argument.Length - 1] == Path.AltDirectorySeparatorChar))? "true" : "false";
+				_idx++;
+				string argument = await EvaluateScalarAsync();
+				result = (argument.Length > 0 && (argument[argument.Length - 1] == Path.DirectorySeparatorChar || argument[argument.Length - 1] == Path.AltDirectorySeparatorChar)) ? "true" : "false";
 			}
-			else if (String.Compare(Tokens[Idx], "Contains", true) == 0 && Tokens[Idx + 1] == "(")
+			else if (String.Compare(_tokens[_idx], "Contains", true) == 0 && _tokens[_idx + 1] == "(")
 			{
 				// Check a string contains a substring. If a separator is supplied the string is first split
-				Idx++;
-				IEnumerable<string> Arguments = EvaluateArguments();
+				_idx++;
+				IEnumerable<string> arguments = EvaluateArguments();
 
-				if (Arguments.Count() != 2)
+				if (arguments.Count() != 2)
 				{
 					throw new BgConditionException("Invalid argument count for 'Contains'. Expected (Haystack,Needle)");
 				}
 
-				Result = Contains(Arguments.ElementAt(0), Arguments.ElementAt(1)) ? "true" : "false";
+				result = Contains(arguments.ElementAt(0), arguments.ElementAt(1)) ? "true" : "false";
 			}
-			else if (String.Compare(Tokens[Idx], "ContainsItem", true) == 0 && Tokens[Idx + 1] == "(")
+			else if (String.Compare(_tokens[_idx], "ContainsItem", true) == 0 && _tokens[_idx + 1] == "(")
 			{
 				// Check a string contains a substring. If a separator is supplied the string is first split
-				Idx++;
-				IEnumerable<string> Arguments = EvaluateArguments();
+				_idx++;
+				IEnumerable<string> arguments = EvaluateArguments();
 
-				if (Arguments.Count() != 3)
+				if (arguments.Count() != 3)
 				{
 					throw new BgConditionException("Invalid argument count for 'ContainsItem'. Expected (Haystack,Needle,HaystackSeparator)");
 				}
 
-				Result = ContainsItem(Arguments.ElementAt(0), Arguments.ElementAt(1), Arguments.ElementAt(2)) ? "true" : "false";
+				result = ContainsItem(arguments.ElementAt(0), arguments.ElementAt(1), arguments.ElementAt(2)) ? "true" : "false";
 			}
 			else
 			{
 				// Raw scalar. Remove quotes from strings, and allow literals and simple identifiers to pass through directly.
-				string Token = Tokens[Idx];
-				if(Token.Length >= 2 && (Token[0] == '\'' || Token[0] == '\"') && Token[Token.Length - 1] == Token[0])
+				string token = _tokens[_idx];
+				if (token.Length >= 2 && (token[0] == '\'' || token[0] == '\"') && token[token.Length - 1] == token[0])
 				{
-					Result = Token.Substring(1, Token.Length - 2);
-					Idx++;
+					result = token.Substring(1, token.Length - 2);
+					_idx++;
 				}
-				else if(Char.IsLetterOrDigit(Token[0]) || Token[0] == '_')
+				else if (Char.IsLetterOrDigit(token[0]) || token[0] == '_')
 				{
-					Result = Token;
-					Idx++;
+					result = token;
+					_idx++;
 				}
 				else
 				{
-					throw new BgConditionException("Token '{0}' is not a valid scalar", Token);
+					throw new BgConditionException("Token '{0}' is not a valid scalar", token);
 				}
 			}
-			return Result;
+			return result;
 		}
 
 		/// <summary>
 		/// Checks whether Haystack contains "Needle". 
 		/// </summary>
-		/// <param name="Haystack">The string to search</param>
-		/// <param name="Needle">The string to search for</param>
+		/// <param name="haystack">The string to search</param>
+		/// <param name="needle">The string to search for</param>
 		/// <returns>True if the path exists, false otherwise.</returns>
-		static bool Contains(string Haystack, string Needle)
+		static bool Contains(string haystack, string needle)
 		{
 			try
 			{
-				return Haystack.IndexOf(Needle, StringComparison.CurrentCultureIgnoreCase) >= 0;
+				return haystack.IndexOf(needle, StringComparison.CurrentCultureIgnoreCase) >= 0;
 			}
 			catch
 			{
@@ -379,16 +378,16 @@ namespace EpicGames.BuildGraph
 		/// <summary>
 		/// Checks whether HaystackItems contains "Needle". 
 		/// </summary>
-		/// <param name="Haystack">The separated list of items to check</param>
-		/// <param name="Needle">The item to check for</param>
-		/// <param name="HaystackSeparator">The separator used in Haystack</param>
+		/// <param name="haystack">The separated list of items to check</param>
+		/// <param name="needle">The item to check for</param>
+		/// <param name="haystackSeparator">The separator used in Haystack</param>
 		/// <returns>True if the path exists, false otherwise.</returns>
-		static bool ContainsItem(string Haystack, string Needle, string HaystackSeparator)
+		static bool ContainsItem(string haystack, string needle, string haystackSeparator)
 		{
 			try
 			{
-				IEnumerable<string> HaystackItems = Haystack.Split(new string[] { HaystackSeparator }, StringSplitOptions.RemoveEmptyEntries);
-				return HaystackItems.Any(I => I.ToLower() == Needle.ToLower());
+				IEnumerable<string> haystackItems = haystack.Split(new string[] { haystackSeparator }, StringSplitOptions.RemoveEmptyEntries);
+				return haystackItems.Any(i => i.ToLower() == needle.ToLower());
 			}
 			catch
 			{
@@ -399,93 +398,93 @@ namespace EpicGames.BuildGraph
 		/// <summary>
 		/// Converts a scalar to a boolean value.
 		/// </summary>
-		/// <param name="Scalar">The scalar to convert</param>
+		/// <param name="scalar">The scalar to convert</param>
 		/// <returns>The scalar converted to a boolean value.</returns>
-		static bool CoerceToBool(string Scalar)
+		static bool CoerceToBool(string scalar)
 		{
-			bool Result;
-			if(String.Compare(Scalar, "true", true) == 0)
+			bool result;
+			if (String.Compare(scalar, "true", true) == 0)
 			{
-				Result = true;
+				result = true;
 			}
-			else if(String.Compare(Scalar, "false", true) == 0)
+			else if (String.Compare(scalar, "false", true) == 0)
 			{
-				Result = false;
+				result = false;
 			}
 			else
 			{
-				throw new BgConditionException("Token '{0}' cannot be coerced to a bool", Scalar);
+				throw new BgConditionException("Token '{0}' cannot be coerced to a bool", scalar);
 			}
-			return Result;
+			return result;
 		}
 
 		/// <summary>
 		/// Converts a scalar to a boolean value.
 		/// </summary>
-		/// <param name="Scalar">The scalar to convert</param>
+		/// <param name="scalar">The scalar to convert</param>
 		/// <returns>The scalar converted to an integer value.</returns>
-		static int CoerceToInteger(string Scalar)
+		static int CoerceToInteger(string scalar)
 		{
-			int Value;
-			if(!Int32.TryParse(Scalar, out Value))
+			int value;
+			if (!Int32.TryParse(scalar, out value))
 			{
-				throw new BgConditionException("Token '{0}' cannot be coerced to an integer", Scalar);
+				throw new BgConditionException("Token '{0}' cannot be coerced to an integer", scalar);
 			}
-			return Value;
+			return value;
 		}
 
 		/// <summary>
 		/// Splits an input string up into expression tokens.
 		/// </summary>
-		/// <param name="Text">Text to be converted into tokens</param>
-		/// <param name="Tokens">List to receive a list of tokens</param>
-		static void Tokenize(string Text, List<string> Tokens)
+		/// <param name="text">Text to be converted into tokens</param>
+		/// <param name="tokens">List to receive a list of tokens</param>
+		static void Tokenize(string text, List<string> tokens)
 		{
-			int Idx = 0;
-			while(Idx < Text.Length)
+			int idx = 0;
+			while (idx < text.Length)
 			{
-				int EndIdx = Idx + 1;
-				if(!Char.IsWhiteSpace(Text[Idx]))
+				int endIdx = idx + 1;
+				if (!Char.IsWhiteSpace(text[idx]))
 				{
 					// Scan to the end of the current token
-					if(Char.IsNumber(Text[Idx]))
+					if (Char.IsNumber(text[idx]))
 					{
 						// Number
-						while(EndIdx < Text.Length && Char.IsNumber(Text[EndIdx]))
+						while (endIdx < text.Length && Char.IsNumber(text[endIdx]))
 						{
-							EndIdx++;
+							endIdx++;
 						}
 					}
-					else if(Char.IsLetter(Text[Idx]) || Text[Idx] == '_')
+					else if (Char.IsLetter(text[idx]) || text[idx] == '_')
 					{
 						// Identifier
-						while(EndIdx < Text.Length && (Char.IsLetterOrDigit(Text[EndIdx]) || Text[EndIdx] == '_'))
+						while (endIdx < text.Length && (Char.IsLetterOrDigit(text[endIdx]) || text[endIdx] == '_'))
 						{
-							EndIdx++;
+							endIdx++;
 						}
 					}
-					else if(Text[Idx] == '!' || Text[Idx] == '<' || Text[Idx] == '>' || Text[Idx] == '=')
+					else if (text[idx] == '!' || text[idx] == '<' || text[idx] == '>' || text[idx] == '=')
 					{
 						// Operator that can be followed by an equals character
-						if(EndIdx < Text.Length && Text[EndIdx] == '=')
+						if (endIdx < text.Length && text[endIdx] == '=')
 						{
-							EndIdx++;
+							endIdx++;
 						}
 					}
-					else if(Text[Idx] == '\'' || Text[Idx] == '\"')
+					else if (text[idx] == '\'' || text[idx] == '\"')
 					{
 						// String
-						if(EndIdx < Text.Length)
+						if (endIdx < text.Length)
 						{
-							EndIdx++;
-							while(EndIdx < Text.Length && Text[EndIdx - 1] != Text[Idx]) EndIdx++;
+							endIdx++;
+							while (endIdx < text.Length && text[endIdx - 1] != text[idx]) endIdx++;
 						}
 					}
-					Tokens.Add(Text.Substring(Idx, EndIdx - Idx));
+					tokens.Add(text.Substring(idx, endIdx - idx));
 				}
-				Idx = EndIdx;
+				idx = endIdx;
 			}
-			Tokens.Add(EndToken);
+			tokens.Add(EndToken);
 		}
 
 		/// <summary>
@@ -514,12 +513,12 @@ namespace EpicGames.BuildGraph
 		/// <summary>
 		/// Helper method to evaluate a condition and check it's the expected result
 		/// </summary>
-		/// <param name="Condition">Condition to evaluate</param>
-		/// <param name="ExpectedResult">The expected result</param>
-		static async Task TestConditionAsync(string Condition, bool ExpectedResult)
+		/// <param name="condition">Condition to evaluate</param>
+		/// <param name="expectedResult">The expected result</param>
+		static async Task TestConditionAsync(string condition, bool expectedResult)
 		{
-			bool Result = await new BgCondition(Condition, null!).EvaluateAsync();
-			Console.WriteLine("{0}: {1} = {2}", (Result == ExpectedResult)? "PASS" : "FAIL", Condition, Result);
+			bool result = await new BgCondition(condition, null!).EvaluateAsync();
+			Console.WriteLine("{0}: {1} = {2}", (result == expectedResult) ? "PASS" : "FAIL", condition, result);
 		}
 	}
 }
diff --git a/Engine/Source/Programs/Shared/EpicGames.BuildGraph/BgContext.cs b/Engine/Source/Programs/Shared/EpicGames.BuildGraph/BgContext.cs
--- a/Engine/Source/Programs/Shared/EpicGames.BuildGraph/BgContext.cs
+++ b/Engine/Source/Programs/Shared/EpicGames.BuildGraph/BgContext.cs
@@ -1,11 +1,9 @@
 // Copyright Epic Games, Inc. All Rights Reserved.
 
-using EpicGames.BuildGraph.Expressions;
-using EpicGames.Core;
-using System;
 using System.Collections.Generic;
 using System.Linq;
-using System.Text;
+using EpicGames.BuildGraph.Expressions;
+using EpicGames.Core;
 
 namespace EpicGames.BuildGraph
 {
@@ -47,95 +45,95 @@ namespace EpicGames.BuildGraph
 		/// <summary>
 		/// All outputs for the node
 		/// </summary>
-		public HashSet<FileReference> BuildProducts = new HashSet<FileReference>();
+		public HashSet<FileReference> BuildProducts { get; } = new HashSet<FileReference>();
 
 		/// <summary>
 		/// Constructor
 		/// </summary>
-		public BgContext(BgExprContext Context)
+		public BgContext(BgExprContext context)
 		{
-			this.Context = Context;
+			Context = context;
 		}
 
 		/// <summary>
 		/// Resolve a boolean expression to a value
 		/// </summary>
-		/// <param name="Bool">The boolean expression</param>
+		/// <param name="expr">The boolean expression</param>
 		/// <returns>Value of the expression</returns>
-		public bool Get(BgBool Bool) => Bool.Compute(Context);
+		public bool Get(BgBool expr) => expr.Compute(Context);
 
 		/// <summary>
 		/// Resolve an integer expression to a value
 		/// </summary>
-		/// <param name="Integer">The integer expression</param>
+		/// <param name="expr">The integer expression</param>
 		/// <returns>Value of the expression</returns>
-		public int Get(BgInt Integer) => Integer.Compute(Context);
+		public int Get(BgInt expr) => expr.Compute(Context);
 
 		/// <summary>
 		/// Resolve a string expression to a value
 		/// </summary>
-		/// <param name="String">The string expression</param>
+		/// <param name="expr">The string expression</param>
 		/// <returns>Value of the expression</returns>
-		public string Get(BgString String) => String.Compute(Context);
+		public string Get(BgString expr) => expr.Compute(Context);
 
 		/// <summary>
 		/// Resolves an enum expression to a value
 		/// </summary>
 		/// <typeparam name="TEnum">The enum type</typeparam>
-		/// <param name="Enum">Enum expression</param>
+		/// <param name="expr">Enum expression</param>
 		/// <returns>The enum value</returns>
-		public TEnum Get<TEnum>(BgEnum<TEnum> Enum) where TEnum : struct => Enum.Compute(Context);
+		public TEnum Get<TEnum>(BgEnum<TEnum> expr) where TEnum : struct => expr.Compute(Context);
 
 		/// <summary>
 		/// Resolve a list of enums to a value
 		/// </summary>
 		/// <typeparam name="TEnum">The enum type</typeparam>
-		/// <param name="List">Enum expression</param>
+		/// <param name="expr">Enum expression</param>
 		/// <returns>The enum value</returns>
-		public List<TEnum> Get<TEnum>(BgList<BgEnum<TEnum>> List) where TEnum : struct => List.GetEnumerable(Context).Select(x => Get(x)).ToList();
+		public List<TEnum> Get<TEnum>(BgList<BgEnum<TEnum>> expr) where TEnum : struct => expr.GetEnumerable(Context).Select(x => Get(x)).ToList();
 
 		/// <summary>
 		/// Resolve a list of strings
 		/// </summary>
-		/// <param name="List">List expression</param>
+		/// <param name="expr">List expression</param>
 		/// <returns></returns>
-		public List<string> Get(BgList<BgString> List) => List.Compute(Context);
+		public List<string> Get(BgList<BgString> expr) => expr.Compute(Context);
 
 		/// <summary>
 		/// Resolve a file set
 		/// </summary>
-		/// <param name="FileSet">The token expression</param>
+		/// <param name="fileSet">The token expression</param>
 		/// <returns>Set of files for the token</returns>
-		public FileSet Get(BgFileSet FileSet) => FileSet.ComputeValue(Context);
+		public FileSet Get(BgFileSet fileSet) => fileSet.ComputeValue(Context);
 
 		/// <summary>
 		/// Resolve a file set
 		/// </summary>
-		/// <param name="FileSets">The token expression</param>
+		/// <param name="fileSets">The token expression</param>
 		/// <returns>Set of files for the token</returns>
-		public FileSet Get(BgList<BgFileSet> FileSets)
+		public FileSet Get(BgList<BgFileSet> fileSets)
 		{
-			FileSet Result = FileSet.Empty;
-			foreach (BgFileSet FileSet in FileSets.GetEnumerable(Context))
+			FileSet result = FileSet.Empty;
+			foreach (BgFileSet fileSet in fileSets.GetEnumerable(Context))
 			{
-				Result += Get(FileSet);
+				result += Get(fileSet);
 			}
-			return Result;
+			return result;
 		}
 
 		/// <summary>
 		/// Resolve a file set
 		/// </summary>
-		/// <param name="FileSets">The token expression</param>
+		/// <param name="fileSets">The token expression</param>
 		/// <returns>Set of files for the token</returns>
-		public FileSet Get(IEnumerable<BgFileSet> FileSets)
+		public FileSet Get(IEnumerable<BgFileSet> fileSets)
 		{
-			FileSet Result = FileSet.Empty;
-			foreach (BgFileSet FileSet in FileSets)
+			FileSet result = FileSet.Empty;
+			foreach (BgFileSet fileSet in fileSets)
 			{
-				Result += FileSet.ComputeValue(Context);
+				result += fileSet.ComputeValue(Context);
 			}
-			return Result;
+			return result;
 		}
 	}
 }
diff --git a/Engine/Source/Programs/Shared/EpicGames.BuildGraph/BgDiagnostic.cs b/Engine/Source/Programs/Shared/EpicGames.BuildGraph/BgDiagnostic.cs
--- a/Engine/Source/Programs/Shared/EpicGames.BuildGraph/BgDiagnostic.cs
+++ b/Engine/Source/Programs/Shared/EpicGames.BuildGraph/BgDiagnostic.cs
@@ -1,9 +1,6 @@
 // Copyright Epic Games, Inc. All Rights Reserved.
 
 using EpicGames.Core;
-using System;
-using System.Collections.Generic;
-using System.Text;
 
 namespace EpicGames.BuildGraph
 {
@@ -16,38 +13,38 @@ namespace EpicGames.BuildGraph
 		/// <summary>
 		/// Location of the diagnostic
 		/// </summary>
-		public BgScriptLocation Location;
+		public BgScriptLocation Location { get; }
 
 		/// <summary>
 		/// The diagnostic event type
 		/// </summary>
-		public LogEventType EventType;
+		public LogEventType EventType { get; }
 
 		/// <summary>
 		/// The message to display
 		/// </summary>
-		public string Message;
+		public string Message { get; }
 
 		/// <summary>
 		/// The node which this diagnostic is declared in. If the node is culled from the graph, the message will not be displayed.
 		/// </summary>
-		public BgNode? EnclosingNode;
+		public BgNode? EnclosingNode { get; }
 
 		/// <summary>
 		/// The agent that this diagnostic is declared in. If the entire agent is culled from the graph, the message will not be displayed.
 		/// </summary>
-		public BgAgent? EnclosingAgent;
+		public BgAgent? EnclosingAgent { get; }
 
 		/// <summary>
 		/// Constructor
 		/// </summary>
-		public BgGraphDiagnostic(BgScriptLocation Location, LogEventType EventType, string Message, BgNode? EnclosingNode, BgAgent? EnclosingAgent)
+		public BgGraphDiagnostic(BgScriptLocation location, LogEventType eventType, string message, BgNode? enclosingNode, BgAgent? enclosingAgent)
 		{
-			this.Location = Location;
-			this.EventType = EventType;
-			this.Message = Message;
-			this.EnclosingNode = EnclosingNode;
-			this.EnclosingAgent = EnclosingAgent;
+			Location = location;
+			EventType = eventType;
+			Message = message;
+			EnclosingNode = enclosingNode;
+			EnclosingAgent = enclosingAgent;
 		}
 	}
 }
diff --git a/Engine/Source/Programs/Shared/EpicGames.BuildGraph/BgDiagnosticSpec.cs b/Engine/Source/Programs/Shared/EpicGames.BuildGraph/BgDiagnosticSpec.cs
--- a/Engine/Source/Programs/Shared/EpicGames.BuildGraph/BgDiagnosticSpec.cs
+++ b/Engine/Source/Programs/Shared/EpicGames.BuildGraph/BgDiagnosticSpec.cs
@@ -1,11 +1,9 @@
 // Copyright Epic Games, Inc. All Rights Reserved.
 
+using System;
 using EpicGames.BuildGraph.Expressions;
 using EpicGames.Core;
 using Microsoft.Extensions.Logging;
-using System;
-using System.Collections.Generic;
-using System.Text;
 
 namespace EpicGames.BuildGraph
 {
@@ -17,39 +15,39 @@ namespace EpicGames.BuildGraph
 		/// <summary>
 		/// Verbosity to output the message at
 		/// </summary>
-		public LogLevel Level;
+		public LogLevel Level { get; }
 
 		/// <summary>
 		/// The message to display
 		/// </summary>
-		public BgString Message;
+		public BgString Message { get; }
 
 		/// <summary>
 		/// Constructor
 		/// </summary>
-		/// <param name="Level">Output level</param>
-		/// <param name="Message"></param>
-		public BgDiagnosticSpec(LogLevel Level, BgString Message)
+		/// <param name="level">Output level</param>
+		/// <param name="message"></param>
+		public BgDiagnosticSpec(LogLevel level, BgString message)
 		{
-			this.Level = Level;
-			this.Message = Message;
+			Level = level;
+			Message = message;
 		}
 
 		/// <summary>
 		/// Adds a diagnostic to the graph
 		/// </summary>
-		internal void AddToGraph(BgExprContext Context, BgGraph Graph, BgAgent? EnclosingAgent, BgNode? EnclosingNode)
+		internal void AddToGraph(BgExprContext context, BgGraph graph, BgAgent? enclosingAgent, BgNode? enclosingNode)
 		{
-			BgScriptLocation Location = new BgScriptLocation("(unknown)", "(unknown)", 1);
-			string MessageValue = Message.Compute(Context);
-			Graph.Diagnostics.Add(new BgGraphDiagnostic(Location, (LogEventType)Level, MessageValue, EnclosingNode, EnclosingAgent));
+			BgScriptLocation location = new BgScriptLocation("(unknown)", "(unknown)", 1);
+			string messageValue = Message.Compute(context);
+			graph.Diagnostics.Add(new BgGraphDiagnostic(location, (LogEventType)Level, messageValue, enclosingNode, enclosingAgent));
 		}
 
 		/// <inheritdoc/>
-		public BgDiagnosticSpec IfThen(BgBool Condition, BgDiagnosticSpec ValueIfTrue) => throw new NotImplementedException();
+		public BgDiagnosticSpec IfThen(BgBool condition, BgDiagnosticSpec valueIfTrue) => throw new NotImplementedException();
 
 		/// <inheritdoc/>
-		object IBgExpr.Compute(BgExprContext Context) => this;
+		object IBgExpr.Compute(BgExprContext context) => this;
 
 		/// <inheritdoc/>
 		public BgString ToBgString() => Message;
diff --git a/Engine/Source/Programs/Shared/EpicGames.BuildGraph/BgGraph.cs b/Engine/Source/Programs/Shared/EpicGames.BuildGraph/BgGraph.cs
--- a/Engine/Source/Programs/Shared/EpicGames.BuildGraph/BgGraph.cs
+++ b/Engine/Source/Programs/Shared/EpicGames.BuildGraph/BgGraph.cs
@@ -2,15 +2,12 @@
 
 using System;
 using System.Collections.Generic;
+using System.Diagnostics.CodeAnalysis;
+using System.Linq;
 using System.Text;
-using System.IO;
-using System.Reflection;
 using System.Xml;
-using System.Linq;
-using System.Diagnostics;
 using EpicGames.Core;
 using Microsoft.Extensions.Logging;
-using System.Diagnostics.CodeAnalysis;
 
 namespace EpicGames.BuildGraph
 {
@@ -43,52 +40,52 @@ namespace EpicGames.BuildGraph
 		/// <summary>
 		/// List of options, in the order they were specified
 		/// </summary>
-		public List<BgOption> Options = new List<BgOption>();
+		public List<BgOption> Options { get; } = new List<BgOption>();
 
 		/// <summary>
 		/// List of agents containing nodes to execute
 		/// </summary>
-		public List<BgAgent> Agents = new List<BgAgent>();
+		public List<BgAgent> Agents { get; } = new List<BgAgent>();
 
 		/// <summary>
 		/// Mapping from name to agent
 		/// </summary>
-		public Dictionary<string, BgAgent> NameToAgent = new Dictionary<string, BgAgent>(StringComparer.OrdinalIgnoreCase);
+		public Dictionary<string, BgAgent> NameToAgent { get; set; } = new Dictionary<string, BgAgent>(StringComparer.OrdinalIgnoreCase);
 
 		/// <summary>
 		/// Mapping of names to the corresponding node.
 		/// </summary>
-		public Dictionary<string, BgNode> NameToNode = new Dictionary<string, BgNode>(StringComparer.OrdinalIgnoreCase);
+		public Dictionary<string, BgNode> NameToNode { get; set; } = new Dictionary<string, BgNode>(StringComparer.OrdinalIgnoreCase);
 
 		/// <summary>
 		/// Mapping of names to the corresponding report.
 		/// </summary>
-		public Dictionary<string, BgReport> NameToReport = new Dictionary<string, BgReport>(StringComparer.OrdinalIgnoreCase);
+		public Dictionary<string, BgReport> NameToReport { get; private set; } = new Dictionary<string, BgReport>(StringComparer.OrdinalIgnoreCase);
 
 		/// <summary>
 		/// Mapping of names to their corresponding node output.
 		/// </summary>
-		public Dictionary<string, BgNodeOutput> TagNameToNodeOutput = new Dictionary<string, BgNodeOutput>(StringComparer.OrdinalIgnoreCase);
+		public Dictionary<string, BgNodeOutput> TagNameToNodeOutput { get; private set; } = new Dictionary<string, BgNodeOutput>(StringComparer.OrdinalIgnoreCase);
 
 		/// <summary>
 		/// Mapping of aggregate names to their respective nodes
 		/// </summary>
-		public Dictionary<string, BgAggregate> NameToAggregate = new Dictionary<string, BgAggregate>(StringComparer.OrdinalIgnoreCase);
+		public Dictionary<string, BgAggregate> NameToAggregate { get; private set; } = new Dictionary<string, BgAggregate>(StringComparer.OrdinalIgnoreCase);
 
 		/// <summary>
 		/// List of badges that can be displayed for this build
 		/// </summary>
-		public List<BgBadge> Badges = new List<BgBadge>();
+		public List<BgBadge> Badges { get; } = new List<BgBadge>();
 
 		/// <summary>
 		/// List of labels that can be displayed for this build
 		/// </summary>
-		public List<BgLabel> Labels = new List<BgLabel>();
+		public List<BgLabel> Labels { get; } = new List<BgLabel>();
 
 		/// <summary>
 		/// Diagnostic messages for this graph
 		/// </summary>
-		public List<BgGraphDiagnostic> Diagnostics = new List<BgGraphDiagnostic>();
+		public List<BgGraphDiagnostic> Diagnostics { get; } = new List<BgGraphDiagnostic>();
 
 		/// <summary>
 		/// Default constructor
@@ -100,114 +97,114 @@ namespace EpicGames.BuildGraph
 		/// <summary>
 		/// Checks whether a given name already exists
 		/// </summary>
-		/// <param name="Name">The name to check.</param>
+		/// <param name="name">The name to check.</param>
 		/// <returns>True if the name exists, false otherwise.</returns>
-		public bool ContainsName(string Name)
+		public bool ContainsName(string name)
 		{
-			return NameToNode.ContainsKey(Name) || NameToReport.ContainsKey(Name) || NameToAggregate.ContainsKey(Name);
+			return NameToNode.ContainsKey(name) || NameToReport.ContainsKey(name) || NameToAggregate.ContainsKey(name);
 		}
 
 		/// <summary>
 		/// Tries to resolve the given name to one or more nodes. Checks for aggregates, and actual nodes.
 		/// </summary>
-		/// <param name="Name">The name to search for</param>
-		/// <param name="OutNodes">If the name is a match, receives an array of nodes and their output names</param>
+		/// <param name="name">The name to search for</param>
+		/// <param name="outNodes">If the name is a match, receives an array of nodes and their output names</param>
 		/// <returns>True if the name was found, false otherwise.</returns>
-		public bool TryResolveReference(string Name, [NotNullWhen(true)] out BgNode[]? OutNodes)
+		public bool TryResolveReference(string name, [NotNullWhen(true)] out BgNode[]? outNodes)
 		{
 			// Check if it's a tag reference or node reference
-			if (Name.StartsWith("#"))
+			if (name.StartsWith("#"))
 			{
 				// Check if it's a regular node or output name
-				BgNodeOutput? Output;
-				if (TagNameToNodeOutput.TryGetValue(Name, out Output))
+				BgNodeOutput? output;
+				if (TagNameToNodeOutput.TryGetValue(name, out output))
 				{
-					OutNodes = new BgNode[] { Output.ProducingNode };
+					outNodes = new BgNode[] { output.ProducingNode };
 					return true;
 				}
 			}
 			else
 			{
 				// Check if it's a regular node or output name
-				BgNode? Node;
-				if (NameToNode.TryGetValue(Name, out Node))
+				BgNode? node;
+				if (NameToNode.TryGetValue(name, out node))
 				{
-					OutNodes = new BgNode[] { Node };
+					outNodes = new BgNode[] { node };
 					return true;
 				}
 
 				// Check if it's an aggregate name
-				BgAggregate? Aggregate;
-				if (NameToAggregate.TryGetValue(Name, out Aggregate))
+				BgAggregate? aggregate;
+				if (NameToAggregate.TryGetValue(name, out aggregate))
 				{
-					OutNodes = Aggregate.RequiredNodes.ToArray();
+					outNodes = aggregate.RequiredNodes.ToArray();
 					return true;
 				}
 
 				// Check if it's a group name
-				BgAgent? Agent;
-				if (NameToAgent.TryGetValue(Name, out Agent))
+				BgAgent? agent;
+				if (NameToAgent.TryGetValue(name, out agent))
 				{
-					OutNodes = Agent.Nodes.ToArray();
+					outNodes = agent.Nodes.ToArray();
 					return true;
 				}
 			}
 
 			// Otherwise fail
-			OutNodes = null;
+			outNodes = null;
 			return false;
 		}
 
 		/// <summary>
 		/// Tries to resolve the given name to one or more node outputs. Checks for aggregates, and actual nodes.
 		/// </summary>
-		/// <param name="Name">The name to search for</param>
-		/// <param name="OutOutputs">If the name is a match, receives an array of nodes and their output names</param>
+		/// <param name="name">The name to search for</param>
+		/// <param name="outOutputs">If the name is a match, receives an array of nodes and their output names</param>
 		/// <returns>True if the name was found, false otherwise.</returns>
-		public bool TryResolveInputReference(string Name, [NotNullWhen(true)] out BgNodeOutput[]? OutOutputs)
+		public bool TryResolveInputReference(string name, [NotNullWhen(true)] out BgNodeOutput[]? outOutputs)
 		{
 			// Check if it's a tag reference or node reference
-			if (Name.StartsWith("#"))
+			if (name.StartsWith("#"))
 			{
 				// Check if it's a regular node or output name
-				BgNodeOutput? Output;
-				if (TagNameToNodeOutput.TryGetValue(Name, out Output))
+				BgNodeOutput? output;
+				if (TagNameToNodeOutput.TryGetValue(name, out output))
 				{
-					OutOutputs = new BgNodeOutput[] { Output };
+					outOutputs = new BgNodeOutput[] { output };
 					return true;
 				}
 			}
 			else
 			{
 				// Check if it's a regular node or output name
-				BgNode? Node;
-				if (NameToNode.TryGetValue(Name, out Node))
+				BgNode? node;
+				if (NameToNode.TryGetValue(name, out node))
 				{
-					OutOutputs = Node.Outputs.Union(Node.Inputs).ToArray();
+					outOutputs = node.Outputs.Union(node.Inputs).ToArray();
 					return true;
 				}
 
 				// Check if it's an aggregate name
-				BgAggregate? Aggregate;
-				if (NameToAggregate.TryGetValue(Name, out Aggregate))
+				BgAggregate? aggregate;
+				if (NameToAggregate.TryGetValue(name, out aggregate))
 				{
-					OutOutputs = Aggregate.RequiredNodes.SelectMany(x => x.Outputs.Union(x.Inputs)).Distinct().ToArray();
+					outOutputs = aggregate.RequiredNodes.SelectMany(x => x.Outputs.Union(x.Inputs)).Distinct().ToArray();
 					return true;
 				}
 			}
 
 			// Otherwise fail
-			OutOutputs = null;
+			outOutputs = null;
 			return false;
 		}
 
-		static void AddDependencies(BgNode Node, HashSet<BgNode> RetainNodes)
+		static void AddDependencies(BgNode node, HashSet<BgNode> retainNodes)
 		{
-			if (RetainNodes.Add(Node))
+			if (retainNodes.Add(node))
 			{
-				foreach(BgNode InputDependency in Node.InputDependencies)
+				foreach (BgNode inputDependency in node.InputDependencies)
 				{
-					AddDependencies(InputDependency, RetainNodes);
+					AddDependencies(inputDependency, retainNodes);
 				}
 			}
 		}
@@ -215,508 +212,508 @@ namespace EpicGames.BuildGraph
 		/// <summary>
 		/// Cull the graph to only include the given nodes and their dependencies
 		/// </summary>
-		/// <param name="TargetNodes">A set of target nodes to build</param>
-		public void Select(IEnumerable<BgNode> TargetNodes)
+		/// <param name="targetNodes">A set of target nodes to build</param>
+		public void Select(IEnumerable<BgNode> targetNodes)
 		{
 			// Find this node and all its dependencies
-			HashSet<BgNode> RetainNodes = new HashSet<BgNode>();
-			foreach (BgNode TargetNode in TargetNodes)
+			HashSet<BgNode> retainNodes = new HashSet<BgNode>();
+			foreach (BgNode targetNode in targetNodes)
 			{
-				AddDependencies(TargetNode, RetainNodes);
+				AddDependencies(targetNode, retainNodes);
 			}
 
 			// Remove all the nodes which are not marked to be kept
-			foreach (BgAgent Agent in Agents)
+			foreach (BgAgent agent in Agents)
 			{
-				Agent.Nodes = Agent.Nodes.Where(x => RetainNodes.Contains(x)).ToList();
+				agent.Nodes = agent.Nodes.Where(x => retainNodes.Contains(x)).ToList();
 			}
 
 			// Remove all the empty agents
 			Agents.RemoveAll(x => x.Nodes.Count == 0);
 
 			// Trim down the list of nodes for each report to the ones that are being built
-			foreach (BgReport Report in NameToReport.Values)
+			foreach (BgReport report in NameToReport.Values)
 			{
-				Report.Nodes.RemoveWhere(x => !RetainNodes.Contains(x));
+				report.Nodes.RemoveWhere(x => !retainNodes.Contains(x));
 			}
 
 			// Remove all the empty reports
-			NameToReport = NameToReport.Where(x => x.Value.Nodes.Count > 0).ToDictionary(Pair => Pair.Key, Pair => Pair.Value, StringComparer.InvariantCultureIgnoreCase);
+			NameToReport = NameToReport.Where(x => x.Value.Nodes.Count > 0).ToDictionary(pair => pair.Key, pair => pair.Value, StringComparer.InvariantCultureIgnoreCase);
 
 			// Remove all the order dependencies which are no longer part of the graph. Since we don't need to build them, we don't need to wait for them
-			foreach (BgNode Node in RetainNodes)
+			foreach (BgNode node in retainNodes)
 			{
-				Node.OrderDependencies = Node.OrderDependencies.Where(x => RetainNodes.Contains(x)).ToArray();
+				node.OrderDependencies = node.OrderDependencies.Where(x => retainNodes.Contains(x)).ToArray();
 			}
 
 			// Create a new list of aggregates for everything that's left
-			Dictionary<string, BgAggregate> NewNameToAggregate = new Dictionary<string, BgAggregate>(NameToAggregate.Comparer);
-			foreach (BgAggregate Aggregate in NameToAggregate.Values)
+			Dictionary<string, BgAggregate> newNameToAggregate = new Dictionary<string, BgAggregate>(NameToAggregate.Comparer);
+			foreach (BgAggregate aggregate in NameToAggregate.Values)
 			{
-				if (Aggregate.RequiredNodes.All(x => RetainNodes.Contains(x)))
+				if (aggregate.RequiredNodes.All(x => retainNodes.Contains(x)))
 				{
-					NewNameToAggregate[Aggregate.Name] = Aggregate;
+					newNameToAggregate[aggregate.Name] = aggregate;
 				}
 			}
-			NameToAggregate = NewNameToAggregate;
+			NameToAggregate = newNameToAggregate;
 
 			// Remove any labels that are no longer value
-			foreach (BgLabel Label in Labels)
+			foreach (BgLabel label in Labels)
 			{
-				Label.RequiredNodes.RemoveWhere(x => !RetainNodes.Contains(x));
-				Label.IncludedNodes.RemoveWhere(x => !RetainNodes.Contains(x));
+				label.RequiredNodes.RemoveWhere(x => !retainNodes.Contains(x));
+				label.IncludedNodes.RemoveWhere(x => !retainNodes.Contains(x));
 			}
 			Labels.RemoveAll(x => x.RequiredNodes.Count == 0);
 
 			// Remove any badges which do not have all their dependencies
-			Badges.RemoveAll(x => x.Nodes.Any(y => !RetainNodes.Contains(y)));
+			Badges.RemoveAll(x => x.Nodes.Any(y => !retainNodes.Contains(y)));
 
 			// Remove any diagnostics which are no longer part of the graph
-			Diagnostics.RemoveAll(x => (x.EnclosingNode != null && !RetainNodes.Contains(x.EnclosingNode)) || (x.EnclosingAgent != null && !Agents.Contains(x.EnclosingAgent)));
+			Diagnostics.RemoveAll(x => (x.EnclosingNode != null && !retainNodes.Contains(x.EnclosingNode)) || (x.EnclosingAgent != null && !Agents.Contains(x.EnclosingAgent)));
 		}
 
 		/// <summary>
 		/// Writes a preprocessed build graph to a script file
 		/// </summary>
-		/// <param name="File">The file to load</param>
-		/// <param name="SchemaFile">Schema file for validation</param>
-		public void Write(FileReference File, FileReference SchemaFile)
+		/// <param name="file">The file to load</param>
+		/// <param name="schemaFile">Schema file for validation</param>
+		public void Write(FileReference file, FileReference schemaFile)
 		{
-			XmlWriterSettings Settings = new XmlWriterSettings();
-			Settings.Indent = true;
-			Settings.IndentChars = "\t";
+			XmlWriterSettings settings = new XmlWriterSettings();
+			settings.Indent = true;
+			settings.IndentChars = "\t";
 
-			using (XmlWriter Writer = XmlWriter.Create(File.FullName, Settings))
+			using (XmlWriter writer = XmlWriter.Create(file.FullName, settings))
 			{
-				Writer.WriteStartElement("BuildGraph", "http://www.epicgames.com/BuildGraph");
+				writer.WriteStartElement("BuildGraph", "http://www.epicgames.com/BuildGraph");
 
-				if (SchemaFile != null)
+				if (schemaFile != null)
 				{
-					Writer.WriteAttributeString("schemaLocation", "http://www.w3.org/2001/XMLSchema-instance", "http://www.epicgames.com/BuildGraph " + SchemaFile.MakeRelativeTo(File.Directory));
+					writer.WriteAttributeString("schemaLocation", "http://www.w3.org/2001/XMLSchema-instance", "http://www.epicgames.com/BuildGraph " + schemaFile.MakeRelativeTo(file.Directory));
 				}
 
-				foreach (BgAgent Agent in Agents)
+				foreach (BgAgent agent in Agents)
 				{
-					Agent.Write(Writer);
+					agent.Write(writer);
 				}
 
-				foreach (BgAggregate Aggregate in NameToAggregate.Values)
+				foreach (BgAggregate aggregate in NameToAggregate.Values)
 				{
 					// If the aggregate has no required elements, skip it.
-					if (Aggregate.RequiredNodes.Count == 0)
+					if (aggregate.RequiredNodes.Count == 0)
 					{
 						continue;
 					}
 
-					Writer.WriteStartElement("Aggregate");
-					Writer.WriteAttributeString("Name", Aggregate.Name);
-					Writer.WriteAttributeString("Requires", String.Join(";", Aggregate.RequiredNodes.Select(x => x.Name)));
-					Writer.WriteEndElement();
+					writer.WriteStartElement("Aggregate");
+					writer.WriteAttributeString("Name", aggregate.Name);
+					writer.WriteAttributeString("Requires", String.Join(";", aggregate.RequiredNodes.Select(x => x.Name)));
+					writer.WriteEndElement();
 				}
 
-				foreach (BgLabel Label in Labels)
+				foreach (BgLabel label in Labels)
 				{
-					Writer.WriteStartElement("Label");
-					if (Label.DashboardCategory != null)
+					writer.WriteStartElement("Label");
+					if (label.DashboardCategory != null)
 					{
-						Writer.WriteAttributeString("Category", Label.DashboardCategory);
+						writer.WriteAttributeString("Category", label.DashboardCategory);
 					}
-					Writer.WriteAttributeString("Name", Label.DashboardName);
-					Writer.WriteAttributeString("Requires", String.Join(";", Label.RequiredNodes.Select(x => x.Name)));
+					writer.WriteAttributeString("Name", label.DashboardName);
+					writer.WriteAttributeString("Requires", String.Join(";", label.RequiredNodes.Select(x => x.Name)));
 
-					HashSet<BgNode> IncludedNodes = new HashSet<BgNode>(Label.IncludedNodes);
-					IncludedNodes.ExceptWith(Label.IncludedNodes.SelectMany(x => x.InputDependencies));
-					IncludedNodes.ExceptWith(Label.RequiredNodes);
-					if (IncludedNodes.Count > 0)
+					HashSet<BgNode> includedNodes = new HashSet<BgNode>(label.IncludedNodes);
+					includedNodes.ExceptWith(label.IncludedNodes.SelectMany(x => x.InputDependencies));
+					includedNodes.ExceptWith(label.RequiredNodes);
+					if (includedNodes.Count > 0)
 					{
-						Writer.WriteAttributeString("Include", String.Join(";", IncludedNodes.Select(x => x.Name)));
+						writer.WriteAttributeString("Include", String.Join(";", includedNodes.Select(x => x.Name)));
 					}
 
-					HashSet<BgNode> ExcludedNodes = new HashSet<BgNode>(Label.IncludedNodes);
-					ExcludedNodes.UnionWith(Label.IncludedNodes.SelectMany(x => x.InputDependencies));
-					ExcludedNodes.ExceptWith(Label.IncludedNodes);
-					ExcludedNodes.ExceptWith(ExcludedNodes.ToArray().SelectMany(x => x.InputDependencies));
-					if (ExcludedNodes.Count > 0)
+					HashSet<BgNode> excludedNodes = new HashSet<BgNode>(label.IncludedNodes);
+					excludedNodes.UnionWith(label.IncludedNodes.SelectMany(x => x.InputDependencies));
+					excludedNodes.ExceptWith(label.IncludedNodes);
+					excludedNodes.ExceptWith(excludedNodes.ToArray().SelectMany(x => x.InputDependencies));
+					if (excludedNodes.Count > 0)
 					{
-						Writer.WriteAttributeString("Exclude", String.Join(";", ExcludedNodes.Select(x => x.Name)));
+						writer.WriteAttributeString("Exclude", String.Join(";", excludedNodes.Select(x => x.Name)));
 					}
-					Writer.WriteEndElement();
+					writer.WriteEndElement();
 				}
 
-				foreach (BgReport Report in NameToReport.Values)
+				foreach (BgReport report in NameToReport.Values)
 				{
-					Writer.WriteStartElement("Report");
-					Writer.WriteAttributeString("Name", Report.Name);
-					Writer.WriteAttributeString("Requires", String.Join(";", Report.Nodes.Select(x => x.Name)));
-					Writer.WriteEndElement();
+					writer.WriteStartElement("Report");
+					writer.WriteAttributeString("Name", report.Name);
+					writer.WriteAttributeString("Requires", String.Join(";", report.Nodes.Select(x => x.Name)));
+					writer.WriteEndElement();
 				}
 
-				foreach (BgBadge Badge in Badges)
+				foreach (BgBadge badge in Badges)
 				{
-					Writer.WriteStartElement("Badge");
-					Writer.WriteAttributeString("Name", Badge.Name);
-					if (Badge.Project != null)
+					writer.WriteStartElement("Badge");
+					writer.WriteAttributeString("Name", badge.Name);
+					if (badge.Project != null)
 					{
-						Writer.WriteAttributeString("Project", Badge.Project);
+						writer.WriteAttributeString("Project", badge.Project);
 					}
-					if (Badge.Change != 0)
+					if (badge.Change != 0)
 					{
-						Writer.WriteAttributeString("Change", Badge.Change.ToString());
+						writer.WriteAttributeString("Change", badge.Change.ToString());
 					}
-					Writer.WriteAttributeString("Requires", String.Join(";", Badge.Nodes.Select(x => x.Name)));
-					Writer.WriteEndElement();
+					writer.WriteAttributeString("Requires", String.Join(";", badge.Nodes.Select(x => x.Name)));
+					writer.WriteEndElement();
 				}
 
-				Writer.WriteEndElement();
+				writer.WriteEndElement();
 			}
 		}
 
 		/// <summary>
 		/// Export the build graph to a Json file, for parallel execution by the build system
 		/// </summary>
-		/// <param name="File">Output file to write</param>
-		/// <param name="CompletedNodes">Set of nodes which have been completed</param>
-		public void Export(FileReference File, HashSet<BgNode> CompletedNodes)
+		/// <param name="file">Output file to write</param>
+		/// <param name="completedNodes">Set of nodes which have been completed</param>
+		public void Export(FileReference file, HashSet<BgNode> completedNodes)
 		{
 			// Find all the nodes which we're actually going to execute. We'll use this to filter the graph.
-			HashSet<BgNode> NodesToExecute = new HashSet<BgNode>();
-			foreach (BgNode Node in Agents.SelectMany(x => x.Nodes))
+			HashSet<BgNode> nodesToExecute = new HashSet<BgNode>();
+			foreach (BgNode node in Agents.SelectMany(x => x.Nodes))
 			{
-				if (!CompletedNodes.Contains(Node))
+				if (!completedNodes.Contains(node))
 				{
-					NodesToExecute.Add(Node);
+					nodesToExecute.Add(node);
 				}
 			}
 
 			// Open the output file
-			using (JsonWriter JsonWriter = new JsonWriter(File.FullName))
+			using (JsonWriter jsonWriter = new JsonWriter(file.FullName))
 			{
-				JsonWriter.WriteObjectStart();
+				jsonWriter.WriteObjectStart();
 
 				// Write all the agents
-				JsonWriter.WriteArrayStart("Groups");
-				foreach (BgAgent Agent in Agents)
+				jsonWriter.WriteArrayStart("Groups");
+				foreach (BgAgent agent in Agents)
 				{
-					BgNode[] Nodes = Agent.Nodes.Where(x => NodesToExecute.Contains(x)).ToArray();
-					if (Nodes.Length > 0)
+					BgNode[] nodes = agent.Nodes.Where(x => nodesToExecute.Contains(x)).ToArray();
+					if (nodes.Length > 0)
 					{
-						JsonWriter.WriteObjectStart();
-						JsonWriter.WriteValue("Name", Agent.Name);
-						JsonWriter.WriteArrayStart("Agent Types");
-						foreach (string AgentType in Agent.PossibleTypes)
+						jsonWriter.WriteObjectStart();
+						jsonWriter.WriteValue("Name", agent.Name);
+						jsonWriter.WriteArrayStart("Agent Types");
+						foreach (string agentType in agent.PossibleTypes)
 						{
-							JsonWriter.WriteValue(AgentType);
+							jsonWriter.WriteValue(agentType);
 						}
-						JsonWriter.WriteArrayEnd();
-						JsonWriter.WriteArrayStart("Nodes");
-						foreach (BgNode Node in Nodes)
+						jsonWriter.WriteArrayEnd();
+						jsonWriter.WriteArrayStart("Nodes");
+						foreach (BgNode node in nodes)
 						{
-							JsonWriter.WriteObjectStart();
-							JsonWriter.WriteValue("Name", Node.Name);
-							JsonWriter.WriteValue("DependsOn", String.Join(";", Node.GetDirectOrderDependencies().Where(x => NodesToExecute.Contains(x))));
-							JsonWriter.WriteValue("RunEarly", Node.bRunEarly);
-							JsonWriter.WriteObjectStart("Notify");
-							JsonWriter.WriteValue("Default", String.Join(";", Node.NotifyUsers));
-							JsonWriter.WriteValue("Submitters", String.Join(";", Node.NotifySubmitters));
-							JsonWriter.WriteValue("Warnings", Node.bNotifyOnWarnings);
-							JsonWriter.WriteObjectEnd();
-							JsonWriter.WriteObjectEnd();
+							jsonWriter.WriteObjectStart();
+							jsonWriter.WriteValue("Name", node.Name);
+							jsonWriter.WriteValue("DependsOn", String.Join(";", node.GetDirectOrderDependencies().Where(x => nodesToExecute.Contains(x))));
+							jsonWriter.WriteValue("RunEarly", node.BRunEarly);
+							jsonWriter.WriteObjectStart("Notify");
+							jsonWriter.WriteValue("Default", String.Join(";", node.NotifyUsers));
+							jsonWriter.WriteValue("Submitters", String.Join(";", node.NotifySubmitters));
+							jsonWriter.WriteValue("Warnings", node.BNotifyOnWarnings);
+							jsonWriter.WriteObjectEnd();
+							jsonWriter.WriteObjectEnd();
 						}
-						JsonWriter.WriteArrayEnd();
-						JsonWriter.WriteObjectEnd();
+						jsonWriter.WriteArrayEnd();
+						jsonWriter.WriteObjectEnd();
 					}
 				}
-				JsonWriter.WriteArrayEnd();
+				jsonWriter.WriteArrayEnd();
 
 				// Write all the badges
-				JsonWriter.WriteArrayStart("Badges");
-				foreach (BgBadge Badge in Badges)
+				jsonWriter.WriteArrayStart("Badges");
+				foreach (BgBadge badge in Badges)
 				{
-					BgNode[] Dependencies = Badge.Nodes.Where(x => NodesToExecute.Contains(x)).ToArray();
-					if (Dependencies.Length > 0)
+					BgNode[] dependencies = badge.Nodes.Where(x => nodesToExecute.Contains(x)).ToArray();
+					if (dependencies.Length > 0)
 					{
 						// Reduce that list to the smallest subset of direct dependencies
-						HashSet<BgNode> DirectDependencies = new HashSet<BgNode>(Dependencies);
-						foreach (BgNode Dependency in Dependencies)
+						HashSet<BgNode> directDependencies = new HashSet<BgNode>(dependencies);
+						foreach (BgNode dependency in dependencies)
 						{
-							DirectDependencies.ExceptWith(Dependency.OrderDependencies);
+							directDependencies.ExceptWith(dependency.OrderDependencies);
 						}
 
-						JsonWriter.WriteObjectStart();
-						JsonWriter.WriteValue("Name", Badge.Name);
-						if (!String.IsNullOrEmpty(Badge.Project))
+						jsonWriter.WriteObjectStart();
+						jsonWriter.WriteValue("Name", badge.Name);
+						if (!String.IsNullOrEmpty(badge.Project))
 						{
-							JsonWriter.WriteValue("Project", Badge.Project);
+							jsonWriter.WriteValue("Project", badge.Project);
 						}
-						if (Badge.Change != 0)
+						if (badge.Change != 0)
 						{
-							JsonWriter.WriteValue("Change", Badge.Change);
+							jsonWriter.WriteValue("Change", badge.Change);
 						}
-						JsonWriter.WriteValue("AllDependencies", String.Join(";", Agents.SelectMany(x => x.Nodes).Where(x => Dependencies.Contains(x)).Select(x => x.Name)));
-						JsonWriter.WriteValue("DirectDependencies", String.Join(";", DirectDependencies.Select(x => x.Name)));
-						JsonWriter.WriteObjectEnd();
+						jsonWriter.WriteValue("AllDependencies", String.Join(";", Agents.SelectMany(x => x.Nodes).Where(x => dependencies.Contains(x)).Select(x => x.Name)));
+						jsonWriter.WriteValue("DirectDependencies", String.Join(";", directDependencies.Select(x => x.Name)));
+						jsonWriter.WriteObjectEnd();
 					}
 				}
-				JsonWriter.WriteArrayEnd();
+				jsonWriter.WriteArrayEnd();
 
 				// Write all the triggers and reports. 
-				JsonWriter.WriteArrayStart("Reports");
-				foreach (BgReport Report in NameToReport.Values)
+				jsonWriter.WriteArrayStart("Reports");
+				foreach (BgReport report in NameToReport.Values)
 				{
-					BgNode[] Dependencies = Report.Nodes.Where(x => NodesToExecute.Contains(x)).ToArray();
-					if (Dependencies.Length > 0)
+					BgNode[] dependencies = report.Nodes.Where(x => nodesToExecute.Contains(x)).ToArray();
+					if (dependencies.Length > 0)
 					{
 						// Reduce that list to the smallest subset of direct dependencies
-						HashSet<BgNode> DirectDependencies = new HashSet<BgNode>(Dependencies);
-						foreach (BgNode Dependency in Dependencies)
+						HashSet<BgNode> directDependencies = new HashSet<BgNode>(dependencies);
+						foreach (BgNode dependency in dependencies)
 						{
-							DirectDependencies.ExceptWith(Dependency.OrderDependencies);
+							directDependencies.ExceptWith(dependency.OrderDependencies);
 						}
 
-						JsonWriter.WriteObjectStart();
-						JsonWriter.WriteValue("Name", Report.Name);
-						JsonWriter.WriteValue("AllDependencies", String.Join(";", Agents.SelectMany(x => x.Nodes).Where(x => Dependencies.Contains(x)).Select(x => x.Name)));
-						JsonWriter.WriteValue("DirectDependencies", String.Join(";", DirectDependencies.Select(x => x.Name)));
-						JsonWriter.WriteValue("Notify", String.Join(";", Report.NotifyUsers));
-						JsonWriter.WriteValue("IsTrigger", false);
-						JsonWriter.WriteObjectEnd();
+						jsonWriter.WriteObjectStart();
+						jsonWriter.WriteValue("Name", report.Name);
+						jsonWriter.WriteValue("AllDependencies", String.Join(";", Agents.SelectMany(x => x.Nodes).Where(x => dependencies.Contains(x)).Select(x => x.Name)));
+						jsonWriter.WriteValue("DirectDependencies", String.Join(";", directDependencies.Select(x => x.Name)));
+						jsonWriter.WriteValue("Notify", String.Join(";", report.NotifyUsers));
+						jsonWriter.WriteValue("IsTrigger", false);
+						jsonWriter.WriteObjectEnd();
 					}
 				}
-				JsonWriter.WriteArrayEnd();
+				jsonWriter.WriteArrayEnd();
 
-				JsonWriter.WriteObjectEnd();
+				jsonWriter.WriteObjectEnd();
 			}
 		}
 
 		/// <summary>
 		/// Export the build graph to a Json file for parsing by Horde
 		/// </summary>
-		/// <param name="File">Output file to write</param>
-		public void ExportForHorde(FileReference File)
+		/// <param name="file">Output file to write</param>
+		public void ExportForHorde(FileReference file)
 		{
-			DirectoryReference.CreateDirectory(File.Directory);
-			using (JsonWriter JsonWriter = new JsonWriter(File.FullName))
+			DirectoryReference.CreateDirectory(file.Directory);
+			using (JsonWriter jsonWriter = new JsonWriter(file.FullName))
 			{
-				JsonWriter.WriteObjectStart();
-				JsonWriter.WriteArrayStart("Groups");
-				foreach (BgAgent Agent in Agents)
+				jsonWriter.WriteObjectStart();
+				jsonWriter.WriteArrayStart("Groups");
+				foreach (BgAgent agent in Agents)
 				{
-					JsonWriter.WriteObjectStart();
-					JsonWriter.WriteArrayStart("Types");
-					foreach (string PossibleType in Agent.PossibleTypes)
+					jsonWriter.WriteObjectStart();
+					jsonWriter.WriteArrayStart("Types");
+					foreach (string possibleType in agent.PossibleTypes)
 					{
-						JsonWriter.WriteValue(PossibleType);
+						jsonWriter.WriteValue(possibleType);
 					}
-					JsonWriter.WriteArrayEnd();
-					JsonWriter.WriteArrayStart("Nodes");
-					foreach (BgNode Node in Agent.Nodes)
+					jsonWriter.WriteArrayEnd();
+					jsonWriter.WriteArrayStart("Nodes");
+					foreach (BgNode node in agent.Nodes)
 					{
-						JsonWriter.WriteObjectStart();
-						JsonWriter.WriteValue("Name", Node.Name);
-						JsonWriter.WriteValue("RunEarly", Node.bRunEarly);
-						JsonWriter.WriteValue("Warnings", Node.bNotifyOnWarnings);
+						jsonWriter.WriteObjectStart();
+						jsonWriter.WriteValue("Name", node.Name);
+						jsonWriter.WriteValue("RunEarly", node.BRunEarly);
+						jsonWriter.WriteValue("Warnings", node.BNotifyOnWarnings);
 
-						JsonWriter.WriteArrayStart("InputDependencies");
-						foreach (string InputDependency in Node.GetDirectInputDependencies().Select(x => x.Name))
+						jsonWriter.WriteArrayStart("InputDependencies");
+						foreach (string inputDependency in node.GetDirectInputDependencies().Select(x => x.Name))
 						{
-							JsonWriter.WriteValue(InputDependency);
+							jsonWriter.WriteValue(inputDependency);
 						}
-						JsonWriter.WriteArrayEnd();
+						jsonWriter.WriteArrayEnd();
 
-						JsonWriter.WriteArrayStart("OrderDependencies");
-						foreach (string OrderDependency in Node.GetDirectOrderDependencies().Select(x => x.Name))
+						jsonWriter.WriteArrayStart("OrderDependencies");
+						foreach (string orderDependency in node.GetDirectOrderDependencies().Select(x => x.Name))
 						{
-							JsonWriter.WriteValue(OrderDependency);
+							jsonWriter.WriteValue(orderDependency);
 						}
-						JsonWriter.WriteArrayEnd();
+						jsonWriter.WriteArrayEnd();
 
-						JsonWriter.WriteObjectEnd();
+						jsonWriter.WriteObjectEnd();
 					}
-					JsonWriter.WriteArrayEnd();
-					JsonWriter.WriteObjectEnd();
+					jsonWriter.WriteArrayEnd();
+					jsonWriter.WriteObjectEnd();
 				}
-				JsonWriter.WriteArrayEnd();
+				jsonWriter.WriteArrayEnd();
 
-				JsonWriter.WriteArrayStart("Aggregates");
-				foreach (BgAggregate Aggregate in NameToAggregate.Values)
+				jsonWriter.WriteArrayStart("Aggregates");
+				foreach (BgAggregate aggregate in NameToAggregate.Values)
 				{
-					JsonWriter.WriteObjectStart();
-					JsonWriter.WriteValue("Name", Aggregate.Name);
-					JsonWriter.WriteArrayStart("Nodes");
-					foreach (BgNode RequiredNode in Aggregate.RequiredNodes.OrderBy(x => x.Name))
+					jsonWriter.WriteObjectStart();
+					jsonWriter.WriteValue("Name", aggregate.Name);
+					jsonWriter.WriteArrayStart("Nodes");
+					foreach (BgNode requiredNode in aggregate.RequiredNodes.OrderBy(x => x.Name))
 					{
-						JsonWriter.WriteValue(RequiredNode.Name);
+						jsonWriter.WriteValue(requiredNode.Name);
 					}
-					JsonWriter.WriteArrayEnd();
-					JsonWriter.WriteObjectEnd();
+					jsonWriter.WriteArrayEnd();
+					jsonWriter.WriteObjectEnd();
 				}
-				JsonWriter.WriteArrayEnd();
+				jsonWriter.WriteArrayEnd();
 
-				JsonWriter.WriteArrayStart("Labels");
-				foreach (BgLabel Label in Labels)
+				jsonWriter.WriteArrayStart("Labels");
+				foreach (BgLabel label in Labels)
 				{
-					JsonWriter.WriteObjectStart();
-					if (!String.IsNullOrEmpty(Label.DashboardName))
+					jsonWriter.WriteObjectStart();
+					if (!String.IsNullOrEmpty(label.DashboardName))
 					{
-						JsonWriter.WriteValue("Name", Label.DashboardName);
+						jsonWriter.WriteValue("Name", label.DashboardName);
 					}
-					if (!String.IsNullOrEmpty(Label.DashboardCategory))
+					if (!String.IsNullOrEmpty(label.DashboardCategory))
 					{
-						JsonWriter.WriteValue("Category", Label.DashboardCategory);
+						jsonWriter.WriteValue("Category", label.DashboardCategory);
 					}
-					if (!String.IsNullOrEmpty(Label.UgsBadge))
+					if (!String.IsNullOrEmpty(label.UgsBadge))
 					{
-						JsonWriter.WriteValue("UgsBadge", Label.UgsBadge);
+						jsonWriter.WriteValue("UgsBadge", label.UgsBadge);
 					}
-					if (!String.IsNullOrEmpty(Label.UgsProject))
+					if (!String.IsNullOrEmpty(label.UgsProject))
 					{
-						JsonWriter.WriteValue("UgsProject", Label.UgsProject);
+						jsonWriter.WriteValue("UgsProject", label.UgsProject);
 					}
-					if (Label.Change != BgLabelChange.Current)
+					if (label.Change != BgLabelChange.Current)
 					{
-						JsonWriter.WriteValue("Change", Label.Change.ToString());
+						jsonWriter.WriteValue("Change", label.Change.ToString());
 					}
 
-					JsonWriter.WriteArrayStart("RequiredNodes");
-					foreach (BgNode RequiredNode in Label.RequiredNodes.OrderBy(x => x.Name))
+					jsonWriter.WriteArrayStart("RequiredNodes");
+					foreach (BgNode requiredNode in label.RequiredNodes.OrderBy(x => x.Name))
 					{
-						JsonWriter.WriteValue(RequiredNode.Name);
+						jsonWriter.WriteValue(requiredNode.Name);
 					}
-					JsonWriter.WriteArrayEnd();
-					JsonWriter.WriteArrayStart("IncludedNodes");
-					foreach (BgNode IncludedNode in Label.IncludedNodes.OrderBy(x => x.Name))
+					jsonWriter.WriteArrayEnd();
+					jsonWriter.WriteArrayStart("IncludedNodes");
+					foreach (BgNode includedNode in label.IncludedNodes.OrderBy(x => x.Name))
 					{
-						JsonWriter.WriteValue(IncludedNode.Name);
+						jsonWriter.WriteValue(includedNode.Name);
 					}
-					JsonWriter.WriteArrayEnd();
-					JsonWriter.WriteObjectEnd();
+					jsonWriter.WriteArrayEnd();
+					jsonWriter.WriteObjectEnd();
 				}
-				JsonWriter.WriteArrayEnd();
+				jsonWriter.WriteArrayEnd();
 
-				JsonWriter.WriteArrayStart("Badges");
-				foreach (BgBadge Badge in Badges)
+				jsonWriter.WriteArrayStart("Badges");
+				foreach (BgBadge badge in Badges)
 				{
-					HashSet<BgNode> Dependencies = Badge.Nodes;
-					if (Dependencies.Count > 0)
+					HashSet<BgNode> dependencies = badge.Nodes;
+					if (dependencies.Count > 0)
 					{
 						// Reduce that list to the smallest subset of direct dependencies
-						HashSet<BgNode> DirectDependencies = new HashSet<BgNode>(Dependencies);
-						foreach (BgNode Dependency in Dependencies)
+						HashSet<BgNode> directDependencies = new HashSet<BgNode>(dependencies);
+						foreach (BgNode dependency in dependencies)
 						{
-							DirectDependencies.ExceptWith(Dependency.OrderDependencies);
+							directDependencies.ExceptWith(dependency.OrderDependencies);
 						}
 
-						JsonWriter.WriteObjectStart();
-						JsonWriter.WriteValue("Name", Badge.Name);
-						if (!String.IsNullOrEmpty(Badge.Project))
+						jsonWriter.WriteObjectStart();
+						jsonWriter.WriteValue("Name", badge.Name);
+						if (!String.IsNullOrEmpty(badge.Project))
 						{
-							JsonWriter.WriteValue("Project", Badge.Project);
+							jsonWriter.WriteValue("Project", badge.Project);
 						}
-						if (Badge.Change != 0)
+						if (badge.Change != 0)
 						{
-							JsonWriter.WriteValue("Change", Badge.Change);
+							jsonWriter.WriteValue("Change", badge.Change);
 						}
-						JsonWriter.WriteValue("Dependencies", String.Join(";", DirectDependencies.Select(x => x.Name)));
-						JsonWriter.WriteObjectEnd();
+						jsonWriter.WriteValue("Dependencies", String.Join(";", directDependencies.Select(x => x.Name)));
+						jsonWriter.WriteObjectEnd();
 					}
 				}
-				JsonWriter.WriteArrayEnd();
+				jsonWriter.WriteArrayEnd();
 
-				JsonWriter.WriteObjectEnd();
+				jsonWriter.WriteObjectEnd();
 			}
 		}
 
 		/// <summary>
 		/// Print the contents of the graph
 		/// </summary>
-		/// <param name="CompletedNodes">Set of nodes which are already complete</param>
-		/// <param name="PrintOptions">Options for how to print the graph</param>
-		/// <param name="Logger"></param>
-		public void Print(HashSet<BgNode> CompletedNodes, GraphPrintOptions PrintOptions, ILogger Logger)
+		/// <param name="completedNodes">Set of nodes which are already complete</param>
+		/// <param name="printOptions">Options for how to print the graph</param>
+		/// <param name="logger"></param>
+		public void Print(HashSet<BgNode> completedNodes, GraphPrintOptions printOptions, ILogger logger)
 		{
 			// Print the options
-			if ((PrintOptions & GraphPrintOptions.ShowCommandLineOptions) != 0)
+			if ((printOptions & GraphPrintOptions.ShowCommandLineOptions) != 0)
 			{
 				// Get the list of messages
-				List<KeyValuePair<string, string>> Parameters = new List<KeyValuePair<string, string>>();
-				foreach (BgOption Option in Options)
+				List<KeyValuePair<string, string>> parameters = new List<KeyValuePair<string, string>>();
+				foreach (BgOption option in Options)
 				{
-					string Name = String.Format("-set:{0}=...", Option.Name);
+					string name = String.Format("-set:{0}=...", option.Name);
 
-					StringBuilder Description = new StringBuilder(Option.Description);
-					if (!String.IsNullOrEmpty(Option.DefaultValue))
+					StringBuilder description = new StringBuilder(option.Description);
+					if (!String.IsNullOrEmpty(option.DefaultValue))
 					{
-						Description.AppendFormat(" (Default: {0})", Option.DefaultValue);
+						description.AppendFormat(" (Default: {0})", option.DefaultValue);
 					}
 
-					Parameters.Add(new KeyValuePair<string, string>(Name, Description.ToString()));
+					parameters.Add(new KeyValuePair<string, string>(name, description.ToString()));
 				}
 
 				// Format them to the log
-				if (Parameters.Count > 0)
+				if (parameters.Count > 0)
 				{
-					Logger.LogInformation("");
-					Logger.LogInformation("Options:");
-					Logger.LogInformation("");
+					logger.LogInformation("");
+					logger.LogInformation("Options:");
+					logger.LogInformation("");
 
-					List<string> Lines = new List<string>();
-					HelpUtils.FormatTable(Parameters, 4, 24, HelpUtils.WindowWidth - 20, Lines);
+					List<string> lines = new List<string>();
+					HelpUtils.FormatTable(parameters, 4, 24, HelpUtils.WindowWidth - 20, lines);
 
-					foreach (string Line in Lines)
+					foreach (string line in lines)
 					{
-						Logger.Log(LogLevel.Information, Line);
+						logger.Log(LogLevel.Information, line);
 					}
 				}
 			}
 
 			// Output all the triggers in order
-			Logger.LogInformation("");
-			Logger.LogInformation("Graph:");
-			foreach (BgAgent Agent in Agents)
+			logger.LogInformation("");
+			logger.LogInformation("Graph:");
+			foreach (BgAgent agent in Agents)
 			{
-				Logger.LogInformation("        Agent: {0} ({1})", Agent.Name, String.Join(";", Agent.PossibleTypes));
-				foreach (BgNode Node in Agent.Nodes)
+				logger.LogInformation("        Agent: {0} ({1})", agent.Name, String.Join(";", agent.PossibleTypes));
+				foreach (BgNode node in agent.Nodes)
 				{
-					Logger.LogInformation("            Node: {0}{1}", Node.Name, CompletedNodes.Contains(Node) ? " (completed)" : Node.bRunEarly ? " (early)" : "");
-					if (PrintOptions.HasFlag(GraphPrintOptions.ShowDependencies))
+					logger.LogInformation("            Node: {0}{1}", node.Name, completedNodes.Contains(node) ? " (completed)" : node.BRunEarly ? " (early)" : "");
+					if (printOptions.HasFlag(GraphPrintOptions.ShowDependencies))
 					{
-						HashSet<BgNode> InputDependencies = new HashSet<BgNode>(Node.GetDirectInputDependencies());
-						foreach (BgNode InputDependency in InputDependencies)
+						HashSet<BgNode> inputDependencies = new HashSet<BgNode>(node.GetDirectInputDependencies());
+						foreach (BgNode inputDependency in inputDependencies)
 						{
-							Logger.LogInformation("                input> {0}", InputDependency.Name);
+							logger.LogInformation("                input> {0}", inputDependency.Name);
 						}
-						HashSet<BgNode> OrderDependencies = new HashSet<BgNode>(Node.GetDirectOrderDependencies());
-						foreach (BgNode OrderDependency in OrderDependencies.Except(InputDependencies))
+						HashSet<BgNode> orderDependencies = new HashSet<BgNode>(node.GetDirectOrderDependencies());
+						foreach (BgNode orderDependency in orderDependencies.Except(inputDependencies))
 						{
-							Logger.LogInformation("                after> {0}", OrderDependency.Name);
+							logger.LogInformation("                after> {0}", orderDependency.Name);
 						}
 					}
-					if (PrintOptions.HasFlag(GraphPrintOptions.ShowNotifications))
+					if (printOptions.HasFlag(GraphPrintOptions.ShowNotifications))
 					{
-						string Label = Node.bNotifyOnWarnings ? "warnings" : "errors";
-						foreach (string User in Node.NotifyUsers)
+						string label = node.BNotifyOnWarnings ? "warnings" : "errors";
+						foreach (string user in node.NotifyUsers)
 						{
-							Logger.LogInformation("                {0}> {1}", Label, User);
+							logger.LogInformation("                {0}> {1}", label, user);
 						}
-						foreach (string Submitter in Node.NotifySubmitters)
+						foreach (string submitter in node.NotifySubmitters)
 						{
-							Logger.LogInformation("                {0}> submitters to {1}", Label, Submitter);
+							logger.LogInformation("                {0}> submitters to {1}", label, submitter);
 						}
 					}
 				}
 			}
-			Logger.LogInformation("");
+			logger.LogInformation("");
 
 			// Print out all the non-empty aggregates
-			BgAggregate[] Aggregates = NameToAggregate.Values.OrderBy(x => x.Name).ToArray();
-			if (Aggregates.Length > 0)
+			BgAggregate[] aggregates = NameToAggregate.Values.OrderBy(x => x.Name).ToArray();
+			if (aggregates.Length > 0)
 			{
-				Logger.LogInformation("Aggregates:");
-				foreach (string AggregateName in Aggregates.Select(x => x.Name))
+				logger.LogInformation("Aggregates:");
+				foreach (string aggregateName in aggregates.Select(x => x.Name))
 				{
-					Logger.LogInformation("    {0}", AggregateName);
+					logger.LogInformation("    {0}", aggregateName);
 				}
-				Logger.LogInformation("");
+				logger.LogInformation("");
 			}
 		}
 	}
