Reformat EpicGames.BuildGraph according to style guidelines.

diff --git a/Engine/Source/Programs/Shared/EpicGames.BuildGraph/BgGraphSpec.cs b/Engine/Source/Programs/Shared/EpicGames.BuildGraph/BgGraphSpec.cs
--- a/Engine/Source/Programs/Shared/EpicGames.BuildGraph/BgGraphSpec.cs
+++ b/Engine/Source/Programs/Shared/EpicGames.BuildGraph/BgGraphSpec.cs
@@ -1,11 +1,9 @@
 // Copyright Epic Games, Inc. All Rights Reserved.
 
-using EpicGames.BuildGraph.Expressions;
-using Microsoft.Extensions.Logging;
 using System;
 using System.Collections.Generic;
-using System.Text;
-using System.Linq;
+using EpicGames.BuildGraph.Expressions;
+using Microsoft.Extensions.Logging;
 
 namespace EpicGames.BuildGraph
 {
@@ -45,189 +43,189 @@ namespace EpicGames.BuildGraph
 		/// </summary>
 		public IBgGraphConstants Constants { get; }
 
-		List<IBgOption> Options = new List<IBgOption>();
-		List<BgAgentSpec> AgentSpecs = new List<BgAgentSpec>();
-		List<BgAggregateSpec> AggregateSpecs = new List<BgAggregateSpec>();
-		List<BgLabelSpec> LabelSpecs = new List<BgLabelSpec>();
-		BgList<BgDiagnosticSpec> DiagnosticSpecs = new List<BgDiagnosticSpec>();
+		readonly List<IBgOption> _options = new List<IBgOption>();
+		readonly List<BgAgentSpec> _agentSpecs = new List<BgAgentSpec>();
+		readonly List<BgAggregateSpec> _aggregateSpecs = new List<BgAggregateSpec>();
+		readonly List<BgLabelSpec> _labelSpecs = new List<BgLabelSpec>();
+		BgList<BgDiagnosticSpec> _diagnosticSpecs = new List<BgDiagnosticSpec>();
 
 		/// <summary>
 		/// Constructor
 		/// </summary>
-		/// <param name="Constants">Constant values for use during graph evaluation</param>
-		public BgGraphSpec(IBgGraphConstants Constants)
+		/// <param name="constants">Constant values for use during graph evaluation</param>
+		public BgGraphSpec(IBgGraphConstants constants)
 		{
-			this.Constants = Constants;
+			Constants = constants;
 		}
 
 		/// <summary>
 		/// Adds a boolean option
 		/// </summary>
-		/// <param name="Name">Name of the option. Must be unique.</param>
-		/// <param name="Description">Description for the option</param>
-		/// <param name="DefaultValue">Default value for the option.</param>
+		/// <param name="name">Name of the option. Must be unique.</param>
+		/// <param name="description">Description for the option</param>
+		/// <param name="defaultValue">Default value for the option.</param>
 		/// <returns>Boolean expression</returns>
-		public BgBoolOption AddOption(string Name, BgString Description, BgBool DefaultValue)
+		public BgBoolOption AddOption(string name, BgString description, BgBool defaultValue)
 		{
-			BgBoolOption Option = new BgBoolOption(Name, Description, DefaultValue);
-			Options.Add(Option);
-			return Option;
+			BgBoolOption option = new BgBoolOption(name, description, defaultValue);
+			_options.Add(option);
+			return option;
 		}
 
 		/// <summary>
 		/// Adds an integer option
 		/// </summary>
-		/// <param name="Name">Name of the option. Must be unique.</param>
-		/// <param name="Description">Description for the option</param>
-		/// <param name="DefaultValue">Default value for the option.</param>
+		/// <param name="name">Name of the option. Must be unique.</param>
+		/// <param name="description">Description for the option</param>
+		/// <param name="defaultValue">Default value for the option.</param>
 		/// <returns>Integer expression</returns>
-		public BgIntOption AddIntOption(string Name, BgString Description, BgInt DefaultValue)
+		public BgIntOption AddIntOption(string name, BgString description, BgInt defaultValue)
 		{
-			BgIntOption Option = new BgIntOption(Name, Description, DefaultValue);
-			Options.Add(Option);
-			return Option;
+			BgIntOption option = new BgIntOption(name, description, defaultValue);
+			_options.Add(option);
+			return option;
 		}
 
 		/// <summary>
 		/// Adds a string option
 		/// </summary>
-		/// <param name="Name">Name of the option. Must be unique.</param>
-		/// <param name="Description">Description for the option</param>
-		/// <param name="DefaultValue">Default value for the option.</param>
+		/// <param name="name">Name of the option. Must be unique.</param>
+		/// <param name="description">Description for the option</param>
+		/// <param name="defaultValue">Default value for the option.</param>
 		/// <returns>String expression</returns>
-		public BgStringOption AddOption(string Name, BgString Description, BgString DefaultValue)
+		public BgStringOption AddOption(string name, BgString description, BgString defaultValue)
 		{
-			BgStringOption Option = new BgStringOption(Name, Description, DefaultValue);
-			Options.Add(Option);
-			return Option;
+			BgStringOption option = new BgStringOption(name, description, defaultValue);
+			_options.Add(option);
+			return option;
 		}
 
 		/// <summary>
 		/// Adds an enum option
 		/// </summary>
 		/// <typeparam name="TEnum">Type of enum value</typeparam>
-		/// <param name="Name">Name of the option. Must be unique.</param>
-		/// <param name="Description">Description for the option</param>
-		/// <param name="DefaultValue">Default value for the option.</param>
+		/// <param name="name">Name of the option. Must be unique.</param>
+		/// <param name="description">Description for the option</param>
+		/// <param name="defaultValue">Default value for the option.</param>
 		/// <returns>Enum option</returns>
-		public BgEnumOption<TEnum> AddOption<TEnum>(string Name, BgString Description, BgEnum<TEnum> DefaultValue) where TEnum : struct
+		public BgEnumOption<TEnum> AddOption<TEnum>(string name, BgString description, BgEnum<TEnum> defaultValue) where TEnum : struct
 		{
-			BgEnumOption<TEnum> Option = new BgEnumOption<TEnum>(Name, Description, DefaultValue);
-			Options.Add(Option);
-			return Option;
+			BgEnumOption<TEnum> option = new BgEnumOption<TEnum>(name, description, defaultValue);
+			_options.Add(option);
+			return option;
 		}
 
 		/// <summary>
 		/// Adds an option which may be a list of enum values
 		/// </summary>
 		/// <typeparam name="TEnum">Type of enum value</typeparam>
-		/// <param name="Name">Name of the option. Must be unique.</param>
-		/// <param name="Description">Description for the option</param>
-		/// <param name="DefaultValue">Default value for the option.</param>
+		/// <param name="name">Name of the option. Must be unique.</param>
+		/// <param name="description">Description for the option</param>
+		/// <param name="defaultValue">Default value for the option.</param>
 		/// <returns></returns>
-		public BgEnumListOption<TEnum> AddOption<TEnum>(string Name, BgString Description, BgList<BgEnum<TEnum>> DefaultValue) where TEnum : struct
+		public BgEnumListOption<TEnum> AddOption<TEnum>(string name, BgString description, BgList<BgEnum<TEnum>> defaultValue) where TEnum : struct
 		{
-			BgEnumListOption<TEnum> Option = new BgEnumListOption<TEnum>(Name, Description, DefaultValue);
-			Options.Add(Option);
-			return Option;
+			BgEnumListOption<TEnum> option = new BgEnumListOption<TEnum>(name, description, defaultValue);
+			_options.Add(option);
+			return option;
 		}
 
 		/// <summary>
 		/// Adds a new agent to the graph
 		/// </summary>
-		/// <param name="Name">Name of the agent</param>
+		/// <param name="name">Name of the agent</param>
 		/// <returns>Agent specification object</returns>
-		public BgAgentSpec AddAgent(BgString Name)
+		public BgAgentSpec AddAgent(BgString name)
 		{
-			BgAgentSpec Agent = new BgAgentSpec(Name);
-			AgentSpecs.Add(Agent);
-			return Agent;
+			BgAgentSpec agent = new BgAgentSpec(name);
+			_agentSpecs.Add(agent);
+			return agent;
 		}
 
 		/// <summary>
 		/// Adds a new aggregate to the graph
 		/// </summary>
-		/// <param name="Name">Name of the agent</param>
-		/// <param name="Outputs">Initial set of tokens to include in the aggregate</param>
+		/// <param name="name">Name of the agent</param>
+		/// <param name="outputs">Initial set of tokens to include in the aggregate</param>
 		/// <returns>Agent specification object</returns>
-		public BgAggregateSpec AddAggregate(BgString Name, BgList<BgFileSet> Outputs)
+		public BgAggregateSpec AddAggregate(BgString name, BgList<BgFileSet> outputs)
 		{
-			BgAggregateSpec Aggregate = new BgAggregateSpec(Name, Outputs);
-			AggregateSpecs.Add(Aggregate);
-			return Aggregate;
+			BgAggregateSpec aggregate = new BgAggregateSpec(name, outputs);
+			_aggregateSpecs.Add(aggregate);
+			return aggregate;
 		}
 
 		/// <summary>
 		/// Adds a new aggregate to the graph
 		/// </summary>
-		/// <param name="Name">Name of the agent</param>
-		/// <param name="Outputs">Initial set of tokens to include in the aggregate</param>
-		/// <param name="Label"></param>
+		/// <param name="name">Name of the agent</param>
+		/// <param name="outputs">Initial set of tokens to include in the aggregate</param>
+		/// <param name="label"></param>
 		/// <returns>Agent specification object</returns>
-		public BgAggregateSpec AddAggregate(BgString Name, BgList<BgFileSet> Outputs, BgString Label)
+		public BgAggregateSpec AddAggregate(BgString name, BgList<BgFileSet> outputs, BgString label)
 		{
-			AddLabel(Config => Config.RequiredNodes = Config.RequiredNodes.Add(Outputs).Distinct());
-			return AddAggregate(Name, Outputs);
+			AddLabel(config => config.RequiredNodes = config.RequiredNodes.Add(outputs).Distinct());
+			return AddAggregate(name, outputs);
 		}
 
 		/// <summary>
 		/// Adds a new label to the graph
 		/// </summary>
-		/// <param name="Configure">Callback for configuring the label</param>
+		/// <param name="configure">Callback for configuring the label</param>
 		/// <returns>Label specification object</returns>
-		public BgLabelSpec AddLabel(Action<BgLabelConfig> Configure)
+		public BgLabelSpec AddLabel(Action<BgLabelConfig> configure)
 		{
-			BgLabelConfig Config = new BgLabelConfig();
-			Configure(Config);
+			BgLabelConfig config = new BgLabelConfig();
+			configure(config);
 
-			BgLabelSpec LabelSpec = new BgLabelSpec(Config);
-			LabelSpecs.Add(LabelSpec);
+			BgLabelSpec labelSpec = new BgLabelSpec(config);
+			_labelSpecs.Add(labelSpec);
 
-			return LabelSpec;
+			return labelSpec;
 		}
 
 		/// <summary>
 		/// Adds a diagnostic to the graph
 		/// </summary>
-		/// <param name="Condition">Condition for adding the diagnostic to the graph</param>
-		/// <param name="Level"></param>
-		/// <param name="Message"></param>
-		public void AddDiagnostic(BgBool Condition, LogLevel Level, BgString Message)
+		/// <param name="condition">Condition for adding the diagnostic to the graph</param>
+		/// <param name="level"></param>
+		/// <param name="message"></param>
+		public void AddDiagnostic(BgBool condition, LogLevel level, BgString message)
 		{
-			BgDiagnosticSpec Diagnostic = new BgDiagnosticSpec(Level, Message);
-			DiagnosticSpecs = DiagnosticSpecs.AddIf(Condition, Diagnostic);
+			BgDiagnosticSpec diagnostic = new BgDiagnosticSpec(level, message);
+			_diagnosticSpecs = _diagnosticSpecs.AddIf(condition, diagnostic);
 		}
 
 		/// <summary>
 		/// Creates a concrete graph from the specifications
 		/// </summary>
-		/// <param name="Targets"></param>
-		/// <param name="Options"></param>
+		/// <param name="targets"></param>
+		/// <param name="options"></param>
 		/// <returns>New graph instance</returns>
-		public BgGraph CreateGraph(HashSet<string> Targets, Dictionary<string, string> Options)
+		public BgGraph CreateGraph(HashSet<string> targets, Dictionary<string, string> options)
 		{
-			BgExprContext Context = new BgExprContext();
-			Context.Options = Options;
+			BgExprContext context = new BgExprContext();
+			context.Options = options;
 
-			BgGraph Graph = new BgGraph();
-			foreach (BgAgentSpec AgentSpec in AgentSpecs)
+			BgGraph graph = new BgGraph();
+			foreach (BgAgentSpec agentSpec in _agentSpecs)
 			{
-				AgentSpec.AddToGraph(Context, Graph);
+				agentSpec.AddToGraph(context, graph);
 			}
-			foreach (BgAggregateSpec AggregateSpec in AggregateSpecs)
+			foreach (BgAggregateSpec aggregateSpec in _aggregateSpecs)
 			{
-				AggregateSpec.AddToGraph(Context, Graph);
+				aggregateSpec.AddToGraph(context, graph);
 			}
-			foreach (BgDiagnosticSpec DiagnosticSpec in DiagnosticSpecs.GetEnumerable(Context))
+			foreach (BgDiagnosticSpec diagnosticSpec in _diagnosticSpecs.GetEnumerable(context))
 			{
-				DiagnosticSpec.AddToGraph(Context, Graph, null, null);
+				diagnosticSpec.AddToGraph(context, graph, null, null);
 			}
-			foreach (BgLabelSpec LabelSpec in LabelSpecs)
+			foreach (BgLabelSpec labelSpec in _labelSpecs)
 			{
-				LabelSpec.AddToGraph(Context, Graph);
+				labelSpec.AddToGraph(context, graph);
 			}
 
-			return Graph;
+			return graph;
 		}
 	}
 }
diff --git a/Engine/Source/Programs/Shared/EpicGames.BuildGraph/BgLabel.cs b/Engine/Source/Programs/Shared/EpicGames.BuildGraph/BgLabel.cs
--- a/Engine/Source/Programs/Shared/EpicGames.BuildGraph/BgLabel.cs
+++ b/Engine/Source/Programs/Shared/EpicGames.BuildGraph/BgLabel.cs
@@ -2,9 +2,6 @@
 
 using System;
 using System.Collections.Generic;
-using System.Linq;
-using System.Text;
-using System.Threading.Tasks;
 
 namespace EpicGames.BuildGraph
 {
@@ -53,33 +50,33 @@ namespace EpicGames.BuildGraph
 		/// <summary>
 		/// Which change to show the badge for
 		/// </summary>
-		public readonly BgLabelChange Change;
+		public BgLabelChange Change { get; }
 
 		/// <summary>
 		/// Set of nodes that must be run for this label to be shown.
 		/// </summary>
-		public HashSet<BgNode> RequiredNodes = new HashSet<BgNode>();
+		public HashSet<BgNode> RequiredNodes { get; } = new HashSet<BgNode>();
 
 		/// <summary>
 		/// Set of nodes that will be included in this label if present.
 		/// </summary>
-		public HashSet<BgNode> IncludedNodes = new HashSet<BgNode>();
+		public HashSet<BgNode> IncludedNodes { get; } = new HashSet<BgNode>();
 
 		/// <summary>
 		/// Constructor
 		/// </summary>
-		/// <param name="InDashboardName">Name of this label</param>
-		/// <param name="InDashboardCategory">Type of this label</param>
-		/// <param name="InUgsBadge">The UGS badge name</param>
-		/// <param name="InUgsProject">Project to display this badge for</param>
-		/// <param name="InChange">The change to show this badge on in UGS</param>
-		public BgLabel(string? InDashboardName, string? InDashboardCategory, string? InUgsBadge, string? InUgsProject, BgLabelChange InChange)
+		/// <param name="inDashboardName">Name of this label</param>
+		/// <param name="inDashboardCategory">Type of this label</param>
+		/// <param name="inUgsBadge">The UGS badge name</param>
+		/// <param name="inUgsProject">Project to display this badge for</param>
+		/// <param name="inChange">The change to show this badge on in UGS</param>
+		public BgLabel(string? inDashboardName, string? inDashboardCategory, string? inUgsBadge, string? inUgsProject, BgLabelChange inChange)
 		{
-			DashboardName = InDashboardName;
-			DashboardCategory = InDashboardCategory;
-			UgsBadge = InUgsBadge;
-			UgsProject = InUgsProject;
-			Change = InChange;
+			DashboardName = inDashboardName;
+			DashboardCategory = inDashboardCategory;
+			UgsBadge = inUgsBadge;
+			UgsProject = inUgsProject;
+			Change = inChange;
 		}
 
 		/// <summary>
diff --git a/Engine/Source/Programs/Shared/EpicGames.BuildGraph/BgLabelSpec.cs b/Engine/Source/Programs/Shared/EpicGames.BuildGraph/BgLabelSpec.cs
--- a/Engine/Source/Programs/Shared/EpicGames.BuildGraph/BgLabelSpec.cs
+++ b/Engine/Source/Programs/Shared/EpicGames.BuildGraph/BgLabelSpec.cs
@@ -1,10 +1,8 @@
 // Copyright Epic Games, Inc. All Rights Reserved.
 
-using EpicGames.BuildGraph.Expressions;
 using System;
-using System.Collections.Generic;
 using System.Linq;
-using System.Text;
+using EpicGames.BuildGraph.Expressions;
 
 namespace EpicGames.BuildGraph
 {
@@ -21,32 +19,32 @@ namespace EpicGames.BuildGraph
 		/// <summary>
 		/// Category for this label
 		/// </summary>
-		public BgString? DashboardCategory { get; }
+		public BgString? DashboardCategory { get; set; }
 
 		/// <summary>
 		/// Name of the badge in UGS
 		/// </summary>
-		public BgString? UgsBadge { get; }
+		public BgString? UgsBadge { get; set; }
 
 		/// <summary>
 		/// Path to the project folder in UGS
 		/// </summary>
-		public BgString? UgsProject { get; }
+		public BgString? UgsProject { get; set; }
 
 		/// <summary>
 		/// Which change to show the badge for
 		/// </summary>
-		public BgString? Change;
+		public BgString? Change { get; set; }
 
 		/// <summary>
 		/// Set of nodes that must be run for this label to be shown.
 		/// </summary>
-		public BgList<BgFileSet> RequiredNodes = BgList<BgFileSet>.Empty;
+		public BgList<BgFileSet> RequiredNodes { get; set; } = BgList<BgFileSet>.Empty;
 
 		/// <summary>
 		/// Set of nodes that will be included in this label if present.
 		/// </summary>
-		public BgList<BgFileSet> IncludedNodes = BgList<BgFileSet>.Empty;
+		public BgList<BgFileSet> IncludedNodes { get; set; } = BgList<BgFileSet>.Empty;
 	}
 
 	/// <summary>
@@ -54,34 +52,34 @@ namespace EpicGames.BuildGraph
 	/// </summary>
 	public class BgLabelSpec
 	{
-		BgLabelConfig Config { get; }
+		private readonly BgLabelConfig _config;
 
-		internal BgLabelSpec(BgLabelConfig Config)
+		internal BgLabelSpec(BgLabelConfig config)
 		{
-			this.Config = Config;
+			_config = config;
 		}
 
-		internal void AddToGraph(BgExprContext Context, BgGraph Graph)
+		internal void AddToGraph(BgExprContext context, BgGraph graph)
 		{
-			string? DashboardName = Config.DashboardName?.Compute(Context);
-			string? DashboardCategory = Config.DashboardCategory?.Compute(Context);
-			string? UgsBadge = Config.UgsBadge?.Compute(Context);
-			string? UgsProject = Config.UgsBadge?.Compute(Context);
+			string? dashboardName = _config.DashboardName?.Compute(context);
+			string? dashboardCategory = _config.DashboardCategory?.Compute(context);
+			string? ugsBadge = _config.UgsBadge?.Compute(context);
+			string? ugsProject = _config.UgsBadge?.Compute(context);
 
-			BgLabelChange LabelChange;
-			if ((object?)Config.Change == null)
+			BgLabelChange labelChange;
+			if ((object?)_config.Change == null)
 			{
-				LabelChange = BgLabelChange.Current;
+				labelChange = BgLabelChange.Current;
 			}
 			else
 			{
-				LabelChange = Enum.Parse<BgLabelChange>(Config.Change.Compute(Context));
+				labelChange = Enum.Parse<BgLabelChange>(_config.Change.Compute(context));
 			}
 
-			BgLabel Label = new BgLabel(DashboardName, DashboardCategory, UgsBadge, UgsProject, LabelChange);
-			Label.RequiredNodes.UnionWith(Config.RequiredNodes.ComputeTags(Context).Select(x => Graph.TagNameToNodeOutput[x].ProducingNode));
-			Label.IncludedNodes.UnionWith(Config.IncludedNodes.ComputeTags(Context).Select(x => Graph.TagNameToNodeOutput[x].ProducingNode));
-			Graph.Labels.Add(Label);
+			BgLabel label = new BgLabel(dashboardName, dashboardCategory, ugsBadge, ugsProject, labelChange);
+			label.RequiredNodes.UnionWith(_config.RequiredNodes.ComputeTags(context).Select(x => graph.TagNameToNodeOutput[x].ProducingNode));
+			label.IncludedNodes.UnionWith(_config.IncludedNodes.ComputeTags(context).Select(x => graph.TagNameToNodeOutput[x].ProducingNode));
+			graph.Labels.Add(label);
 		}
 	}
 }
diff --git a/Engine/Source/Programs/Shared/EpicGames.BuildGraph/BgNode.cs b/Engine/Source/Programs/Shared/EpicGames.BuildGraph/BgNode.cs
--- a/Engine/Source/Programs/Shared/EpicGames.BuildGraph/BgNode.cs
+++ b/Engine/Source/Programs/Shared/EpicGames.BuildGraph/BgNode.cs
@@ -3,15 +3,8 @@
 using System;
 using System.Collections.Generic;
 using System.Linq;
-using System.Text;
-using System.Threading.Tasks;
 using System.Xml;
 using EpicGames.Core;
-using OpenTracing;
-using OpenTracing.Util;
-using Microsoft.Extensions.Logging;
-using System.Reflection;
-using System.ComponentModel;
 
 namespace EpicGames.BuildGraph
 {
@@ -33,12 +26,12 @@ namespace EpicGames.BuildGraph
 		/// <summary>
 		/// Constructor
 		/// </summary>
-		/// <param name="InProducingNode">Node which produces the given output</param>
-		/// <param name="InTagName">Name of the tag</param>
-		public BgNodeOutput(BgNode InProducingNode, string InTagName)
+		/// <param name="inProducingNode">Node which produces the given output</param>
+		/// <param name="inTagName">Name of the tag</param>
+		public BgNodeOutput(BgNode inProducingNode, string inTagName)
 		{
-			ProducingNode = InProducingNode;
-			TagName = InTagName;
+			ProducingNode = inProducingNode;
+			TagName = inTagName;
 		}
 
 		/// <summary>
@@ -104,35 +97,35 @@ namespace EpicGames.BuildGraph
 		/// <summary>
 		/// Whether to start this node as soon as its dependencies are satisfied, rather than waiting for all of its agent's dependencies to be met.
 		/// </summary>
-		public bool bRunEarly { get; set; } = false;
+		public bool BRunEarly { get; set; } = false;
 
 		/// <summary>
 		/// Whether to ignore warnings produced by this node
 		/// </summary>
-		public bool bNotifyOnWarnings { get; set; } = true;
+		public bool BNotifyOnWarnings { get; set; } = true;
 
 		/// <summary>
 		/// Constructor
 		/// </summary>
-		/// <param name="InName">The name of this node</param>
-		/// <param name="InInputs">Inputs that this node depends on</param>
-		/// <param name="InOutputNames">Names of the outputs that this node produces</param>
-		/// <param name="InInputDependencies">Nodes which this node is dependent on for its inputs</param>
-		/// <param name="InOrderDependencies">Nodes which this node needs to run after. Should include all input dependencies.</param>
-		/// <param name="InRequiredTokens">Optional tokens which must be required for this node to run</param>
-		public BgNode(string InName, BgNodeOutput[] InInputs, string[] InOutputNames, BgNode[] InInputDependencies, BgNode[] InOrderDependencies, FileReference[] InRequiredTokens)
+		/// <param name="inName">The name of this node</param>
+		/// <param name="inInputs">Inputs that this node depends on</param>
+		/// <param name="inOutputNames">Names of the outputs that this node produces</param>
+		/// <param name="inInputDependencies">Nodes which this node is dependent on for its inputs</param>
+		/// <param name="inOrderDependencies">Nodes which this node needs to run after. Should include all input dependencies.</param>
+		/// <param name="inRequiredTokens">Optional tokens which must be required for this node to run</param>
+		public BgNode(string inName, BgNodeOutput[] inInputs, string[] inOutputNames, BgNode[] inInputDependencies, BgNode[] inOrderDependencies, FileReference[] inRequiredTokens)
 		{
-			Name = InName;
-			Inputs = InInputs;
+			Name = inName;
+			Inputs = inInputs;
 
-			List<BgNodeOutput> AllOutputs = new List<BgNodeOutput>();
-			AllOutputs.Add(new BgNodeOutput(this, "#" + Name));
-			AllOutputs.AddRange(InOutputNames.Where(x => String.Compare(x, Name, StringComparison.InvariantCultureIgnoreCase) != 0).Select(x => new BgNodeOutput(this, x)));
-			Outputs = AllOutputs.ToArray();
+			List<BgNodeOutput> allOutputs = new List<BgNodeOutput>();
+			allOutputs.Add(new BgNodeOutput(this, "#" + Name));
+			allOutputs.AddRange(inOutputNames.Where(x => String.Compare(x, Name, StringComparison.InvariantCultureIgnoreCase) != 0).Select(x => new BgNodeOutput(this, x)));
+			Outputs = allOutputs.ToArray();
 
-			InputDependencies = InInputDependencies;
-			OrderDependencies = InOrderDependencies;
-			RequiredTokens = InRequiredTokens;
+			InputDependencies = inInputDependencies;
+			OrderDependencies = inOrderDependencies;
+			RequiredTokens = inRequiredTokens;
 		}
 
 		/// <summary>
@@ -149,12 +142,12 @@ namespace EpicGames.BuildGraph
 		/// <returns>Sequence of nodes that are direct inputs to this node</returns>
 		public IEnumerable<BgNode> GetDirectInputDependencies()
 		{
-			HashSet<BgNode> DirectDependencies = new HashSet<BgNode>(InputDependencies);
-			foreach(BgNode InputDependency in InputDependencies)
+			HashSet<BgNode> directDependencies = new HashSet<BgNode>(InputDependencies);
+			foreach (BgNode inputDependency in InputDependencies)
 			{
-				DirectDependencies.ExceptWith(InputDependency.InputDependencies);
+				directDependencies.ExceptWith(inputDependency.InputDependencies);
 			}
-			return DirectDependencies;
+			return directDependencies;
 		}
 
 		/// <summary>
@@ -163,56 +156,56 @@ namespace EpicGames.BuildGraph
 		/// <returns>Sequence of nodes that are direct order dependencies of this node</returns>
 		public IEnumerable<BgNode> GetDirectOrderDependencies()
 		{
-			HashSet<BgNode> DirectDependencies = new HashSet<BgNode>(OrderDependencies);
-			foreach(BgNode OrderDependency in OrderDependencies)
+			HashSet<BgNode> directDependencies = new HashSet<BgNode>(OrderDependencies);
+			foreach (BgNode orderDependency in OrderDependencies)
 			{
-				DirectDependencies.ExceptWith(OrderDependency.OrderDependencies);
+				directDependencies.ExceptWith(orderDependency.OrderDependencies);
 			}
-			return DirectDependencies;
+			return directDependencies;
 		}
 
 		/// <summary>
 		/// Write this node to an XML writer
 		/// </summary>
-		/// <param name="Writer">The writer to output the node to</param>
-		public void Write(XmlWriter Writer)
+		/// <param name="writer">The writer to output the node to</param>
+		public void Write(XmlWriter writer)
 		{
-			Writer.WriteStartElement("Node");
-			Writer.WriteAttributeString("Name", Name);
+			writer.WriteStartElement("Node");
+			writer.WriteAttributeString("Name", Name);
 
-			string[] RequireNames = Inputs.Select(x => x.TagName).ToArray();
-			if (RequireNames.Length > 0)
+			string[] requireNames = Inputs.Select(x => x.TagName).ToArray();
+			if (requireNames.Length > 0)
 			{
-				Writer.WriteAttributeString("Requires", String.Join(";", RequireNames));
+				writer.WriteAttributeString("Requires", String.Join(";", requireNames));
 			}
 
-			string[] ProducesNames = Outputs.Where(x => x != DefaultOutput).Select(x => x.TagName).ToArray();
-			if (ProducesNames.Length > 0)
+			string[] producesNames = Outputs.Where(x => x != DefaultOutput).Select(x => x.TagName).ToArray();
+			if (producesNames.Length > 0)
 			{
-				Writer.WriteAttributeString("Produces", String.Join(";", ProducesNames));
+				writer.WriteAttributeString("Produces", String.Join(";", producesNames));
 			}
 
-			string[] AfterNames = GetDirectOrderDependencies().Except(InputDependencies).Select(x => x.Name).ToArray();
-			if (AfterNames.Length > 0)
+			string[] afterNames = GetDirectOrderDependencies().Except(InputDependencies).Select(x => x.Name).ToArray();
+			if (afterNames.Length > 0)
 			{
-				Writer.WriteAttributeString("After", String.Join(";", AfterNames));
+				writer.WriteAttributeString("After", String.Join(";", afterNames));
 			}
 
-			if (!bNotifyOnWarnings)
+			if (!BNotifyOnWarnings)
 			{
-				Writer.WriteAttributeString("NotifyOnWarnings", bNotifyOnWarnings.ToString());
+				writer.WriteAttributeString("NotifyOnWarnings", BNotifyOnWarnings.ToString());
 			}
 
-			if(bRunEarly)
+			if (BRunEarly)
 			{
-				Writer.WriteAttributeString("RunEarly", bRunEarly.ToString());
+				writer.WriteAttributeString("RunEarly", BRunEarly.ToString());
 			}
 
-			foreach (BgTask Task in Tasks)
+			foreach (BgTask task in Tasks)
 			{
-				Task.Write(Writer);
+				task.Write(writer);
 			}
-			Writer.WriteEndElement();
+			writer.WriteEndElement();
 		}
 
 		/// <summary>
diff --git a/Engine/Source/Programs/Shared/EpicGames.BuildGraph/BgNodeSpec.cs b/Engine/Source/Programs/Shared/EpicGames.BuildGraph/BgNodeSpec.cs
--- a/Engine/Source/Programs/Shared/EpicGames.BuildGraph/BgNodeSpec.cs
+++ b/Engine/Source/Programs/Shared/EpicGames.BuildGraph/BgNodeSpec.cs
@@ -1,16 +1,15 @@
 // Copyright Epic Games, Inc. All Rights Reserved.
 
-using EpicGames.BuildGraph.Expressions;
-using EpicGames.Core;
-using Microsoft.Extensions.Logging;
 using System;
 using System.Collections.Generic;
 using System.Linq;
 using System.Linq.Expressions;
 using System.Reflection;
 using System.Text;
-using System.Text.Json;
 using System.Threading.Tasks;
+using EpicGames.BuildGraph.Expressions;
+using EpicGames.Core;
+using Microsoft.Extensions.Logging;
 
 namespace EpicGames.BuildGraph
 {
@@ -22,8 +21,8 @@ namespace EpicGames.BuildGraph
 		/// <summary>
 		/// Constructor
 		/// </summary>
-		/// <param name="Message"></param>
-		public BgNodeException(string Message) : base(Message)
+		/// <param name="message"></param>
+		public BgNodeException(string message) : base(message)
 		{
 		}
 	}
@@ -42,10 +41,10 @@ namespace EpicGames.BuildGraph
 		/// <summary>
 		/// Constructor
 		/// </summary>
-		/// <param name="Template">Format string for the name</param>
-		public BgNodeNameAttribute(string Template)
+		/// <param name="template">Format string for the name</param>
+		public BgNodeNameAttribute(string template)
 		{
-			this.Template = Template;
+			Template = template;
 		}
 	}
 
@@ -54,9 +53,9 @@ namespace EpicGames.BuildGraph
 	/// </summary>
 	public class BgNodeSpec
 	{
-		internal string ClassName;
-		internal string MethodName;
-		internal IBgExpr?[] ArgumentExprs;
+		internal string _className;
+		internal string _methodName;
+		internal IBgExpr?[] _argumentExprs;
 		internal BgFileSet[] ResultExprs { get; }
 
 		/// <summary>
@@ -97,26 +96,26 @@ namespace EpicGames.BuildGraph
 		/// <summary>
 		/// Constructor
 		/// </summary>
-		internal BgNodeSpec(MethodCallExpression Call)
+		internal BgNodeSpec(MethodCallExpression call)
 		{
-			string MethodSignature = $"{Call.Method.DeclaringType?.Name}.{Call.Method.Name}";
-			if (!Call.Method.IsStatic)
+			string methodSignature = $"{call.Method.DeclaringType?.Name}.{call.Method.Name}";
+			if (!call.Method.IsStatic)
 			{
-				throw new BgNodeException($"Node {Call.Method.DeclaringType?.Name}.{Call.Method.Name} must be static");
+				throw new BgNodeException($"Node {call.Method.DeclaringType?.Name}.{call.Method.Name} must be static");
 			}
 
-			this.ClassName = Call.Method.DeclaringType!.AssemblyQualifiedName!;
-			this.MethodName = Call.Method.Name;
+			_className = call.Method.DeclaringType!.AssemblyQualifiedName!;
+			_methodName = call.Method.Name;
 
 			try
 			{
-				ArgumentExprs = CreateArgumentExprs(Call);
-				Name = GetDefaultNodeName(Call.Method, ArgumentExprs);
-				ResultExprs = CreateReturnExprs(Call.Method, Name, Call.Method.ReturnType);
+				_argumentExprs = CreateArgumentExprs(call);
+				Name = GetDefaultNodeName(call.Method, _argumentExprs);
+				ResultExprs = CreateReturnExprs(call.Method, Name, call.Method.ReturnType);
 			}
-			catch (Exception Ex)
+			catch (Exception ex)
 			{
-				ExceptionUtils.AddContext(Ex, $"while calling method {Call.Method.DeclaringType?.Name}.{Call.Method.Name}");
+				ExceptionUtils.AddContext(ex, $"while calling method {call.Method.DeclaringType?.Name}.{call.Method.Name}");
 				throw;
 			}
 
@@ -126,61 +125,61 @@ namespace EpicGames.BuildGraph
 		/// <summary>
 		/// Creates a node specification
 		/// </summary>
-		internal static BgNodeSpec Create(Expression<Func<BgContext, Task>> Action)
+		internal static BgNodeSpec Create(Expression<Func<BgContext, Task>> action)
 		{
-			MethodCallExpression Call = (MethodCallExpression)Action.Body;
-			return new BgNodeSpec(Call);
+			MethodCallExpression call = (MethodCallExpression)action.Body;
+			return new BgNodeSpec(call);
 		}
 
 		/// <summary>
 		/// Creates a node specification
 		/// </summary>
-		internal static BgNodeSpec<T> Create<T>(Expression<Func<BgContext, Task<T>>> Function)
+		internal static BgNodeSpec<T> Create<T>(Expression<Func<BgContext, Task<T>>> function)
 		{
-			MethodCallExpression Call = (MethodCallExpression)Function.Body;
-			return new BgNodeSpec<T>(Call);
+			MethodCallExpression call = (MethodCallExpression)function.Body;
+			return new BgNodeSpec<T>(call);
 		}
 
-		static IBgExpr?[] CreateArgumentExprs(MethodCallExpression Call)
+		static IBgExpr?[] CreateArgumentExprs(MethodCallExpression call)
 		{
-			IBgExpr?[] Args = new IBgExpr?[Call.Arguments.Count];
-			for (int Idx = 0; Idx < Call.Arguments.Count; Idx++)
+			IBgExpr?[] args = new IBgExpr?[call.Arguments.Count];
+			for (int idx = 0; idx < call.Arguments.Count; idx++)
 			{
-				Expression Expr = Call.Arguments[Idx];
-				if (Expr is ParameterExpression ParameterExpr)
+				Expression expr = call.Arguments[idx];
+				if (expr is ParameterExpression parameterExpr)
 				{
-					if (ParameterExpr.Type != typeof(BgContext))
+					if (parameterExpr.Type != typeof(BgContext))
 					{
-						throw new BgNodeException($"Unable to determine type of parameter '{ParameterExpr.Name}'");
+						throw new BgNodeException($"Unable to determine type of parameter '{parameterExpr.Name}'");
 					}
 				}
 				else
 				{
-					Delegate Compiled = Expression.Lambda(Expr).Compile();
+					Delegate compiled = Expression.Lambda(expr).Compile();
 
-					object? Result = Compiled.DynamicInvoke();
-					if (Result is IBgExpr Computable)
+					object? result = compiled.DynamicInvoke();
+					if (result is IBgExpr computable)
 					{
-						Args[Idx] = Computable;
+						args[idx] = computable;
 					}
 					else
 					{
-						Args[Idx] = (BgString)(Result?.ToString() ?? String.Empty);
+						args[idx] = (BgString)(result?.ToString() ?? String.Empty);
 					}
 				}
 			}
-			return Args;
+			return args;
 		}
 
-		static BgFileSet[] CreateReturnExprs(MethodInfo MethodInfo, BgString Name, Type Type)
+		static BgFileSet[] CreateReturnExprs(MethodInfo methodInfo, BgString name, Type type)
 		{
-			if (Type == typeof(Task))
+			if (type == typeof(Task))
 			{
 				return Array.Empty<BgFileSet>();
 			}
-			else if (Type.IsGenericType && Type.GetGenericTypeDefinition() == typeof(Task<>))
+			else if (type.IsGenericType && type.GetGenericTypeDefinition() == typeof(Task<>))
 			{
-				return CreateReturnExprsInner(MethodInfo, Name, Type.GetGenericArguments()[0]);
+				return CreateReturnExprsInner(methodInfo, name, type.GetGenericArguments()[0]);
 			}
 			else
 			{
@@ -188,44 +187,44 @@ namespace EpicGames.BuildGraph
 			}
 		}
 
-		static BgFileSet[] CreateReturnExprsInner(MethodInfo MethodInfo, BgString Name, Type Type)
+		static BgFileSet[] CreateReturnExprsInner(MethodInfo methodInfo, BgString name, Type type)
 		{
-			Type[] OutputTypes;
-			if (IsValueTuple(Type))
+			Type[] outputTypes;
+			if (IsValueTuple(type))
 			{
-				OutputTypes = Type.GetGenericArguments();
+				outputTypes = type.GetGenericArguments();
 			}
 			else
 			{
-				OutputTypes = new[] { Type };
+				outputTypes = new[] { type };
 			}
 
-			BgFileSet[] OutputExprs = new BgFileSet[OutputTypes.Length];
-			for (int Idx = 0; Idx < OutputTypes.Length; Idx++)
+			BgFileSet[] outputExprs = new BgFileSet[outputTypes.Length];
+			for (int idx = 0; idx < outputTypes.Length; idx++)
 			{
-				OutputExprs[Idx] = CreateOutputExpr(MethodInfo, Name, Idx, OutputTypes[Idx]);
+				outputExprs[idx] = CreateOutputExpr(methodInfo, name, idx, outputTypes[idx]);
 			}
-			return OutputExprs;
+			return outputExprs;
 		}
 
-		static BgFileSet CreateOutputExpr(MethodInfo MethodInfo, BgString Name, int Index, Type Type)
+		static BgFileSet CreateOutputExpr(MethodInfo methodInfo, BgString name, int index, Type type)
 		{
-			if (Type == typeof(BgFileSet))
+			if (type == typeof(BgFileSet))
 			{
-				return new BgFileSetTagFromStringExpr(BgString.Format("#{0}${1}", Name, Index));
+				return new BgFileSetTagFromStringExpr(BgString.Format("#{0}${1}", name, index));
 			}
 			else
 			{
-				throw new BgNodeException($"Unsupported return type for {MethodInfo.Name}: {Type.Name}");
+				throw new BgNodeException($"Unsupported return type for {methodInfo.Name}: {type.Name}");
 			}
 		}
 
-		internal static bool IsValueTuple(Type ReturnType)
+		internal static bool IsValueTuple(Type returnType)
 		{
-			if (ReturnType.IsGenericType)
+			if (returnType.IsGenericType)
 			{
-				Type GenericType = ReturnType.GetGenericTypeDefinition();
-				if(GenericType.FullName != null && GenericType.FullName.StartsWith("System.ValueTuple`", StringComparison.Ordinal))
+				Type genericType = returnType.GetGenericTypeDefinition();
+				if (genericType.FullName != null && genericType.FullName.StartsWith("System.ValueTuple`", StringComparison.Ordinal))
 				{
 					return true;
 				}
@@ -233,196 +232,196 @@ namespace EpicGames.BuildGraph
 			return false;
 		}
 
-		static BgString GetDefaultNodeName(MethodInfo MethodInfo, IBgExpr?[] Args)
+		static BgString GetDefaultNodeName(MethodInfo methodInfo, IBgExpr?[] args)
 		{
 			// Check if it's got an attribute override for the node name
-			BgNodeNameAttribute? NameAttr = MethodInfo.GetCustomAttribute<BgNodeNameAttribute>();
-			if (NameAttr != null)
+			BgNodeNameAttribute? nameAttr = methodInfo.GetCustomAttribute<BgNodeNameAttribute>();
+			if (nameAttr != null)
 			{
-				return GetNodeNameFromTemplate(NameAttr.Template, MethodInfo.GetParameters(), Args);
+				return GetNodeNameFromTemplate(nameAttr.Template, methodInfo.GetParameters(), args);
 			}
 			else
 			{
-				return GetNodeNameFromMethodName(MethodInfo.Name);
+				return GetNodeNameFromMethodName(methodInfo.Name);
 			}
 		}
 
-		static BgString GetNodeNameFromTemplate(string Template, ParameterInfo[] Parameters, IBgExpr?[] Args)
+		static BgString GetNodeNameFromTemplate(string template, ParameterInfo[] parameters, IBgExpr?[] args)
 		{
 			// Create a list of lazily computed string fragments which comprise the evaluated name
-			List<BgString> Fragments = new List<BgString>();
+			List<BgString> fragments = new List<BgString>();
 
-			int LastIdx = 0;
-			for (int NextIdx = 0; NextIdx < Template.Length; NextIdx++)
+			int lastIdx = 0;
+			for (int nextIdx = 0; nextIdx < template.Length; nextIdx++)
 			{
-				if (Template[NextIdx] == '{')
+				if (template[nextIdx] == '{')
 				{
-					if (NextIdx + 1 < Template.Length && Template[NextIdx + 1] == '{')
+					if (nextIdx + 1 < template.Length && template[nextIdx + 1] == '{')
 					{
-						Fragments.Add(Template.Substring(LastIdx, NextIdx - LastIdx));
-						LastIdx = ++NextIdx;
+						fragments.Add(template.Substring(lastIdx, nextIdx - lastIdx));
+						lastIdx = ++nextIdx;
 					}
 					else
 					{
-						Fragments.Add(Template.Substring(LastIdx, NextIdx - LastIdx));
-						NextIdx++;
+						fragments.Add(template.Substring(lastIdx, nextIdx - lastIdx));
+						nextIdx++;
 
-						int EndIdx = Template.IndexOf('}', NextIdx);
-						if (EndIdx == -1)
+						int endIdx = template.IndexOf('}', nextIdx);
+						if (endIdx == -1)
 						{
-							throw new BgNodeException($"Unterminated parameter expression for {nameof(BgNodeNameAttribute)} in {Template}");
+							throw new BgNodeException($"Unterminated parameter expression for {nameof(BgNodeNameAttribute)} in {template}");
 						}
 
-						StringView ParamName = new StringView(Template, NextIdx, EndIdx - NextIdx);
+						StringView paramName = new StringView(template, nextIdx, endIdx - nextIdx);
 
-						int ParamIdx = Array.FindIndex(Parameters, x => x.Name != null && ParamName.Equals(x.Name, StringComparison.Ordinal));
-						if (ParamIdx == -1)
+						int paramIdx = Array.FindIndex(parameters, x => x.Name != null && paramName.Equals(x.Name, StringComparison.Ordinal));
+						if (paramIdx == -1)
 						{
-							throw new BgNodeException($"Unable to find parameter named {ParamName} in {Template}");
+							throw new BgNodeException($"Unable to find parameter named {paramName} in {template}");
 						}
 
-						IBgExpr? Arg = Args[ParamIdx];
-						if(Arg != null)
+						IBgExpr? arg = args[paramIdx];
+						if (arg != null)
 						{
-							Fragments.Add(Arg.ToBgString());
+							fragments.Add(arg.ToBgString());
 						}
 
-						LastIdx = NextIdx = EndIdx + 1;
+						lastIdx = nextIdx = endIdx + 1;
 					}
 				}
-				else if (Template[NextIdx] == '}')
+				else if (template[nextIdx] == '}')
 				{
-					if (NextIdx + 1 < Template.Length && Template[NextIdx + 1] == '{')
+					if (nextIdx + 1 < template.Length && template[nextIdx + 1] == '{')
 					{
-						Fragments.Add(Template.Substring(LastIdx, NextIdx - LastIdx));
-						LastIdx = ++NextIdx;
+						fragments.Add(template.Substring(lastIdx, nextIdx - lastIdx));
+						lastIdx = ++nextIdx;
 					}
 				}
 			}
-			Fragments.Add(Template.Substring(LastIdx, Template.Length - LastIdx));
+			fragments.Add(template.Substring(lastIdx, template.Length - lastIdx));
 
-			return BgString.Join(BgString.Empty, Fragments);
+			return BgString.Join(BgString.Empty, fragments);
 		}
 
 		/// <summary>
 		/// Inserts spaces into a PascalCase method name to create a node name
 		/// </summary>
-		public static string GetNodeNameFromMethodName(string MethodName)
+		public static string GetNodeNameFromMethodName(string methodName)
 		{
-			StringBuilder Name = new StringBuilder();
-			Name.Append(MethodName[0]);
+			StringBuilder name = new StringBuilder();
+			name.Append(methodName[0]);
 
-			int Length = MethodName.Length;
-			if (Length > 5 && MethodName.EndsWith("Async", StringComparison.Ordinal))
+			int length = methodName.Length;
+			if (length > 5 && methodName.EndsWith("Async", StringComparison.Ordinal))
 			{
-				Length -= 5;
+				length -= 5;
 			}
 
 			bool bIsAcronym = false;
-			for (int Idx = 1; Idx < Length; Idx++)
+			for (int idx = 1; idx < length; idx++)
 			{
-				bool bLastIsUpper = Char.IsUpper(MethodName[Idx - 1]);
-				bool bNextIsUpper = Char.IsUpper(MethodName[Idx]);
+				bool bLastIsUpper = Char.IsUpper(methodName[idx - 1]);
+				bool bNextIsUpper = Char.IsUpper(methodName[idx]);
 				if (bLastIsUpper && bNextIsUpper)
 				{
 					bIsAcronym = true;
 				}
 				else if (bIsAcronym)
 				{
-					Name.Insert(Name.Length - 2, ' ');
+					name.Insert(name.Length - 2, ' ');
 					bIsAcronym = false;
 				}
 				else if (!bLastIsUpper && bNextIsUpper)
 				{
-					Name.Append(' ');
+					name.Append(' ');
 				}
-				Name.Append(MethodName[Idx]);
+				name.Append(methodName[idx]);
 			}
 
-			return Name.ToString();
+			return name.ToString();
 		}
 
 		/// <summary>
 		/// Gets the signature for a method
 		/// </summary>
-		public static string GetSignature(MethodInfo MethodInfo)
+		public static string GetSignature(MethodInfo methodInfo)
 		{
-			StringBuilder Arguments = new StringBuilder();
-			foreach (ParameterInfo ParameterInfo in MethodInfo.GetParameters())
+			StringBuilder arguments = new StringBuilder();
+			foreach (ParameterInfo parameterInfo in methodInfo.GetParameters())
 			{
-				Arguments.AppendLine(ParameterInfo.ParameterType.FullName);
+				arguments.AppendLine(parameterInfo.ParameterType.FullName);
 			}
-			return $"{MethodInfo.Name}:{Digest.Compute<Sha1>(Arguments.ToString())}";
+			return $"{methodInfo.Name}:{Digest.Compute<Sha1>(arguments.ToString())}";
 		}
 
 		/// <summary>
 		/// Creates a concrete node
 		/// </summary>
-		/// <param name="Context"></param>
-		/// <param name="Graph"></param>
-		/// <param name="Agent"></param>
+		/// <param name="context"></param>
+		/// <param name="graph"></param>
+		/// <param name="agent"></param>
 		/// <returns></returns>
-		internal void AddToGraph(BgExprContext Context, BgGraph Graph, BgAgent Agent)
+		internal void AddToGraph(BgExprContext context, BgGraph graph, BgAgent agent)
 		{
-			HashSet<string> InputTags = new HashSet<string>();
-			InputTags.UnionWith(this.InputDependencies.ComputeTags(Context));
-			InputTags.UnionWith(ArgumentExprs.OfType<BgFileSet>().Select(x => x.ComputeTag(Context)));
-			InputTags.UnionWith(ArgumentExprs.OfType<BgList<BgFileSet>>().SelectMany(x => x.GetEnumerable(Context)).Select(x => x.ComputeTag(Context)));
-
-			HashSet<string> AfterTags = new HashSet<string>(InputTags);
-			AfterTags.UnionWith(this.AfterDependencies.ComputeTags(Context));
-
-			string Name = this.Name.Compute(Context);
-			BgNodeOutput[] InputDependencies = InputTags.Select(x => Graph.TagNameToNodeOutput[x]).Distinct().ToArray();
-			BgNodeOutput[] AfterDependencies = AfterTags.Select(x => Graph.TagNameToNodeOutput[x]).Distinct().ToArray();
-			string[] OutputNames = Array.ConvertAll(ResultExprs, x => x.ComputeTag(Context));
-			BgNode[] InputNodes = InputDependencies.Select(x => x.ProducingNode).Distinct().ToArray();
-			BgNode[] AfterNodes = AfterDependencies.Select(x => x.ProducingNode).Distinct().ToArray();
-			bool RunEarly = CanRunEarly.Compute(Context);
-
-			BgNode Node = new BgNode(Name, InputDependencies, OutputNames, InputNodes, AfterNodes, Array.Empty<FileReference>());
-			Node.bRunEarly = RunEarly;
-
-			BgScriptLocation Location = new BgScriptLocation("(unknown)", "(unknown)", 1);
-
-			BgTask Task = new BgTask(Location, "CallMethod");
-			Task.Arguments["Class"] = ClassName;
-			Task.Arguments["Method"] = MethodName;
-			for (int Idx = 0; Idx < ArgumentExprs.Length; Idx++)
+			HashSet<string> inputTags = new HashSet<string>();
+			inputTags.UnionWith(InputDependencies.ComputeTags(context));
+			inputTags.UnionWith(_argumentExprs.OfType<BgFileSet>().Select(x => x.ComputeTag(context)));
+			inputTags.UnionWith(_argumentExprs.OfType<BgList<BgFileSet>>().SelectMany(x => x.GetEnumerable(context)).Select(x => x.ComputeTag(context)));
+
+			HashSet<string> afterTags = new HashSet<string>(inputTags);
+			afterTags.UnionWith(AfterDependencies.ComputeTags(context));
+
+			string name = Name.Compute(context);
+			BgNodeOutput[] inputDependencies = inputTags.Select(x => graph.TagNameToNodeOutput[x]).Distinct().ToArray();
+			BgNodeOutput[] afterDependencies = afterTags.Select(x => graph.TagNameToNodeOutput[x]).Distinct().ToArray();
+			string[] outputNames = Array.ConvertAll(ResultExprs, x => x.ComputeTag(context));
+			BgNode[] inputNodes = inputDependencies.Select(x => x.ProducingNode).Distinct().ToArray();
+			BgNode[] afterNodes = afterDependencies.Select(x => x.ProducingNode).Distinct().ToArray();
+			bool runEarly = CanRunEarly.Compute(context);
+
+			BgNode node = new BgNode(name, inputDependencies, outputNames, inputNodes, afterNodes, Array.Empty<FileReference>());
+			node.BRunEarly = runEarly;
+
+			BgScriptLocation location = new BgScriptLocation("(unknown)", "(unknown)", 1);
+
+			BgTask task = new BgTask(location, "CallMethod");
+			task.Arguments["Class"] = _className;
+			task.Arguments["Method"] = _methodName;
+			for (int idx = 0; idx < _argumentExprs.Length; idx++)
 			{
-				IBgExpr? ArgumentExpr = ArgumentExprs[Idx];
-				if (ArgumentExpr != null)
+				IBgExpr? argumentExpr = _argumentExprs[idx];
+				if (argumentExpr != null)
 				{
-					Task.Arguments[$"Arg{Idx + 1}"] = BgType.Get(ArgumentExpr.GetType()).SerializeArgument(ArgumentExpr, Context);
+					task.Arguments[$"Arg{idx + 1}"] = BgType.Get(argumentExpr.GetType()).SerializeArgument(argumentExpr, context);
 				}
 			}
-			if (OutputNames != null)
+			if (outputNames != null)
 			{
-				Task.Arguments["Tags"] = String.Join(";", OutputNames);
+				task.Arguments["Tags"] = String.Join(";", outputNames);
 			}
-			Node.Tasks.Add(Task);
+			node.Tasks.Add(task);
 
-			Agent.Nodes.Add(Node);
-			Graph.NameToNode.Add(Name, Node);
+			agent.Nodes.Add(node);
+			graph.NameToNode.Add(name, node);
 
-			foreach (BgDiagnosticSpec Precondition in Diagnostics.GetEnumerable(Context))
+			foreach (BgDiagnosticSpec precondition in Diagnostics.GetEnumerable(context))
 			{
-				Precondition.AddToGraph(Context, Graph, Agent, Node);
+				precondition.AddToGraph(context, graph, agent, node);
 			}
 
-			foreach(BgNodeOutput Output in Node.Outputs)
+			foreach (BgNodeOutput output in node.Outputs)
 			{
-				Graph.TagNameToNodeOutput.Add(Output.TagName, Output);
+				graph.TagNameToNodeOutput.Add(output.TagName, output);
 			}
 		}
 
 		/// <summary>
 		/// Allows using the 
 		/// </summary>
-		/// <param name="NodeSpec"></param>
-		public static implicit operator BgList<BgFileSet>(BgNodeSpec NodeSpec)
+		/// <param name="nodeSpec"></param>
+		public static implicit operator BgList<BgFileSet>(BgNodeSpec nodeSpec)
 		{
-			return NodeSpec.InputsAndOutputs.Add(NodeSpec.Output);
+			return nodeSpec.InputsAndOutputs.Add(nodeSpec.Output);
 		}
 	}
 
@@ -440,13 +439,13 @@ namespace EpicGames.BuildGraph
 		/// <summary>
 		/// Constructor
 		/// </summary>
-		internal BgNodeSpec(MethodCallExpression Call)
-			: base(Call)
+		internal BgNodeSpec(MethodCallExpression call)
+			: base(call)
 		{
-			Type ReturnType = Call.Method.ReturnType;
-			if (IsValueTuple(ReturnType))
+			Type returnType = call.Method.ReturnType;
+			if (IsValueTuple(returnType))
 			{
-				Output = (T)Activator.CreateInstance(ReturnType, ResultExprs)!;
+				Output = (T)Activator.CreateInstance(returnType, ResultExprs)!;
 			}
 			else
 			{
@@ -463,58 +462,58 @@ namespace EpicGames.BuildGraph
 		/// <summary>
 		/// Add dependencies onto other nodes or outputs. Outputs from the given tokens will be copied to the current machine before execution of the node.
 		/// </summary>
-		/// <param name="NodeSpec">The node specification</param>
-		/// <param name="Tokens">Tokens to add dependencies on</param>
+		/// <param name="nodeSpec">The node specification</param>
+		/// <param name="tokens">Tokens to add dependencies on</param>
 		/// <returns>The current node spec, to allow chaining calls</returns>
-		public static T Requires<T>(this T NodeSpec, params BgList<BgFileSet>[] Tokens) where T : BgNodeSpec
+		public static T Requires<T>(this T nodeSpec, params BgList<BgFileSet>[] tokens) where T : BgNodeSpec
 		{
-			foreach (BgList<BgFileSet> TokenSet in Tokens)
+			foreach (BgList<BgFileSet> tokenSet in tokens)
 			{
-				NodeSpec.InputDependencies = NodeSpec.InputDependencies.Add(TokenSet);
+				nodeSpec.InputDependencies = nodeSpec.InputDependencies.Add(tokenSet);
 			}
-			return NodeSpec;
+			return nodeSpec;
 		}
 
 		/// <summary>
 		/// Add weak dependencies onto other nodes or outputs. The producing nodes must complete successfully if they are part of the graph, but outputs from them will not be 
 		/// transferred to the machine running this node.
 		/// </summary>
-		/// <param name="NodeSpec">The node specification</param>
-		/// <param name="Tokens">Tokens to add dependencies on</param>
+		/// <param name="nodeSpec">The node specification</param>
+		/// <param name="tokens">Tokens to add dependencies on</param>
 		/// <returns>The current node spec, to allow chaining calls</returns>
-		public static T After<T>(this T NodeSpec, params BgList<BgFileSet>[] Tokens) where T : BgNodeSpec
+		public static T After<T>(this T nodeSpec, params BgList<BgFileSet>[] tokens) where T : BgNodeSpec
 		{
-			foreach (BgList<BgFileSet> TokenSet in Tokens)
+			foreach (BgList<BgFileSet> tokenSet in tokens)
 			{
-				NodeSpec.AfterDependencies = NodeSpec.AfterDependencies.Add(TokenSet);
+				nodeSpec.AfterDependencies = nodeSpec.AfterDependencies.Add(tokenSet);
 			}
-			return NodeSpec;
+			return nodeSpec;
 		}
 
 		/// <summary>
 		/// Adds a warning when executing this node
 		/// </summary>
-		/// <param name="NodeSpec">The node specification</param>
-		/// <param name="Condition">Condition for writing the message</param>
-		/// <param name="Message">Message to output</param>
+		/// <param name="nodeSpec">The node specification</param>
+		/// <param name="condition">Condition for writing the message</param>
+		/// <param name="message">Message to output</param>
 		/// <returns>The current node spec, to allow chaining calls</returns>
-		public static BgNodeSpec WarnIf<T>(this T NodeSpec, BgBool Condition, BgString Message) where T : BgNodeSpec
+		public static BgNodeSpec WarnIf<T>(this T nodeSpec, BgBool condition, BgString message) where T : BgNodeSpec
 		{
-			NodeSpec.Diagnostics = NodeSpec.Diagnostics.AddIf(Condition, new BgDiagnosticSpec(LogLevel.Warning, Message));
-			return NodeSpec;
+			nodeSpec.Diagnostics = nodeSpec.Diagnostics.AddIf(condition, new BgDiagnosticSpec(LogLevel.Warning, message));
+			return nodeSpec;
 		}
 
 		/// <summary>
 		/// Adds an error when executing this node
 		/// </summary>
-		/// <param name="NodeSpec">The node specification</param>
-		/// <param name="Condition">Condition for writing the message</param>
-		/// <param name="Message">Message to output</param>
+		/// <param name="nodeSpec">The node specification</param>
+		/// <param name="condition">Condition for writing the message</param>
+		/// <param name="message">Message to output</param>
 		/// <returns>The current node spec, to allow chaining calls</returns>
-		public static BgNodeSpec ErrorIf<T>(this T NodeSpec, BgBool Condition, BgString Message) where T : BgNodeSpec
+		public static BgNodeSpec ErrorIf<T>(this T nodeSpec, BgBool condition, BgString message) where T : BgNodeSpec
 		{
-			NodeSpec.Diagnostics = NodeSpec.Diagnostics.AddIf(Condition, new BgDiagnosticSpec(LogLevel.Error, Message));
-			return NodeSpec;
+			nodeSpec.Diagnostics = nodeSpec.Diagnostics.AddIf(condition, new BgDiagnosticSpec(LogLevel.Error, message));
+			return nodeSpec;
 		}
 	}
 }
diff --git a/Engine/Source/Programs/Shared/EpicGames.BuildGraph/BgOption.cs b/Engine/Source/Programs/Shared/EpicGames.BuildGraph/BgOption.cs
--- a/Engine/Source/Programs/Shared/EpicGames.BuildGraph/BgOption.cs
+++ b/Engine/Source/Programs/Shared/EpicGames.BuildGraph/BgOption.cs
@@ -1,9 +1,5 @@
 // Copyright Epic Games, Inc. All Rights Reserved.
 
-using System;
-using System.Collections.Generic;
-using System.Text;
-
 namespace EpicGames.BuildGraph
 {
 	/// <summary>
@@ -14,29 +10,29 @@ namespace EpicGames.BuildGraph
 		/// <summary>
 		/// Name of this option
 		/// </summary>
-		public string Name;
+		public string Name { get; }
 
 		/// <summary>
 		/// Description for this option
 		/// </summary>
-		public string Description;
+		public string Description { get; }
 
 		/// <summary>
 		/// Default value for this option
 		/// </summary>
-		public string DefaultValue;
+		public string DefaultValue { get; }
 
 		/// <summary>
 		/// Constructor
 		/// </summary>
-		/// <param name="Name">The name of this option</param>
-		/// <param name="Description">Description of the option, for display on help pages</param>
-		/// <param name="DefaultValue">Default value for the option</param>
-		public BgOption(string Name, string Description, string DefaultValue)
+		/// <param name="name">The name of this option</param>
+		/// <param name="description">Description of the option, for display on help pages</param>
+		/// <param name="defaultValue">Default value for the option</param>
+		public BgOption(string name, string description, string defaultValue)
 		{
-			this.Name = Name;
-			this.Description = Description;
-			this.DefaultValue = DefaultValue;
+			Name = name;
+			Description = description;
+			DefaultValue = defaultValue;
 		}
 
 		/// <summary>
diff --git a/Engine/Source/Programs/Shared/EpicGames.BuildGraph/BgOptionSpec.cs b/Engine/Source/Programs/Shared/EpicGames.BuildGraph/BgOptionSpec.cs
--- a/Engine/Source/Programs/Shared/EpicGames.BuildGraph/BgOptionSpec.cs
+++ b/Engine/Source/Programs/Shared/EpicGames.BuildGraph/BgOptionSpec.cs
@@ -1,11 +1,10 @@
 // Copyright Epic Games, Inc. All Rights Reserved.
 
-using EpicGames.BuildGraph.Expressions;
 using System;
 using System.Collections.Generic;
 using System.Linq;
-using System.Text;
 using System.Text.RegularExpressions;
+using EpicGames.BuildGraph.Expressions;
 
 namespace EpicGames.BuildGraph
 {
@@ -17,9 +16,9 @@ namespace EpicGames.BuildGraph
 		/// <summary>
 		/// Constructor
 		/// </summary>
-		/// <param name="Message"></param>
-		public BgOptionValidationException(string Message)
-			: base(Message)
+		/// <param name="message"></param>
+		public BgOptionValidationException(string message)
+			: base(message)
 		{
 		}
 	}
@@ -67,27 +66,27 @@ namespace EpicGames.BuildGraph
 		/// <summary>
 		/// Constructor
 		/// </summary>
-		internal BgBoolOption(string Name, BgString Description, BgBool DefaultValue)
+		internal BgBoolOption(string name, BgString description, BgBool defaultValue)
 		{
-			this.Name = Name;
-			this.Description = Description;
-			this.DefaultValue = DefaultValue;
+			Name = name;
+			Description = description;
+			DefaultValue = defaultValue;
 		}
 
 		/// <inheritdoc/>
-		public override bool Compute(BgExprContext Context)
+		public override bool Compute(BgExprContext context)
 		{
-			string? Value;
-			if (Context.Options.TryGetValue(Name, out Value))
+			string? value;
+			if (context.Options.TryGetValue(Name, out value))
 			{
-				bool BoolValue;
-				if (!bool.TryParse(Value, out BoolValue))
+				bool boolValue;
+				if (!bool.TryParse(value, out boolValue))
 				{
-					throw new BgOptionValidationException($"Argument for {Name} is not a valid bool ({Value})");
+					throw new BgOptionValidationException($"Argument for {Name} is not a valid bool ({value})");
 				}
-				return BoolValue;
+				return boolValue;
 			}
-			return DefaultValue.Compute(Context);
+			return DefaultValue.Compute(context);
 		}
 	}
 
@@ -123,42 +122,42 @@ namespace EpicGames.BuildGraph
 		/// <summary>
 		/// Constructor
 		/// </summary>
-		internal BgIntOption(string Name, BgString Description, BgInt DefaultValue)
+		internal BgIntOption(string name, BgString description, BgInt defaultValue)
 		{
-			this.Name = Name;
-			this.Description = Description;
-			this.DefaultValue = DefaultValue;
+			Name = name;
+			Description = description;
+			DefaultValue = defaultValue;
 		}
 
 		/// <inheritdoc/>
-		public override int Compute(BgExprContext Context)
+		public override int Compute(BgExprContext context)
 		{
-			string? Value;
-			if (Context.Options.TryGetValue(Name, out Value))
+			string? value;
+			if (context.Options.TryGetValue(Name, out value))
 			{
-				int IntValue;
-				if (!int.TryParse(Value, out IntValue))
+				int intValue;
+				if (!int.TryParse(value, out intValue))
 				{
 					throw new BgOptionValidationException($"Argument for '{Name}' is not a valid integer");
 				}
 				if (!ReferenceEquals(MinValue, null))
 				{
-					int IntMinValue = MinValue.Compute(Context);
-					if (IntValue < IntMinValue)
+					int intMinValue = MinValue.Compute(context);
+					if (intValue < intMinValue)
 					{
-						throw new BgOptionValidationException($"Argument for '{Name}' is less than the allowed minimum ({IntValue} < {IntMinValue})");
+						throw new BgOptionValidationException($"Argument for '{Name}' is less than the allowed minimum ({intValue} < {intMinValue})");
 					}
 				}
 				if (!ReferenceEquals(MaxValue, null))
 				{
-					int IntMaxValue = MaxValue.Compute(Context);
-					if (IntValue > IntMaxValue)
+					int intMaxValue = MaxValue.Compute(context);
+					if (intValue > intMaxValue)
 					{
-						throw new BgOptionValidationException($"Argument for '{Name}' is greater than the allowed maximum ({IntValue} > {IntMaxValue})");
+						throw new BgOptionValidationException($"Argument for '{Name}' is greater than the allowed maximum ({intValue} > {intMaxValue})");
 					}
 				}
 			}
-			return DefaultValue.Compute(Context);
+			return DefaultValue.Compute(context);
 		}
 	}
 
@@ -184,27 +183,27 @@ namespace EpicGames.BuildGraph
 		/// <summary>
 		/// Constructor
 		/// </summary>
-		internal BgEnumOption(string Name, BgString Description, BgEnum<TEnum> DefaultValue)
+		internal BgEnumOption(string name, BgString description, BgEnum<TEnum> defaultValue)
 		{
-			this.Name = Name;
-			this.Description = Description;
-			this.DefaultValue = DefaultValue;
+			Name = name;
+			Description = description;
+			DefaultValue = defaultValue;
 		}
 
 		/// <inheritdoc/>
-		public override TEnum Compute(BgExprContext Context)
+		public override TEnum Compute(BgExprContext context)
 		{
-			string? Value;
-			if (Context.Options.TryGetValue(Name, out Value))
+			string? value;
+			if (context.Options.TryGetValue(Name, out value))
 			{
-				TEnum EnumValue;
-				if (!Enum.TryParse<TEnum>(Value, true, out EnumValue))
+				TEnum enumValue;
+				if (!Enum.TryParse<TEnum>(value, true, out enumValue))
 				{
 					throw new BgOptionValidationException($"Argument '{Name}' is not a valid value for {typeof(TEnum).Name}");
 				}
-				return EnumValue;
+				return enumValue;
 			}
-			return DefaultValue.Compute(Context);
+			return DefaultValue.Compute(context);
 		}
 	}
 
@@ -245,39 +244,39 @@ namespace EpicGames.BuildGraph
 		/// <summary>
 		/// Constructor
 		/// </summary>
-		internal BgStringOption(string Name, BgString Description, BgString DefaultValue)
+		internal BgStringOption(string name, BgString description, BgString defaultValue)
 		{
-			this.Name = Name;
-			this.Description = Description;
-			this.DefaultValue = DefaultValue;
+			Name = name;
+			Description = description;
+			DefaultValue = defaultValue;
 		}
 
 		/// <inheritdoc/>
-		public override string Compute(BgExprContext Context)
+		public override string Compute(BgExprContext context)
 		{
-			string? Value;
-			if (Context.Options.TryGetValue(Name, out Value))
+			string? value;
+			if (context.Options.TryGetValue(Name, out value))
 			{
 				if (!Object.ReferenceEquals(Pattern, null))
 				{
-					string PatternValue = Pattern.Compute(Context);
-					if (!Regex.IsMatch(Value, PatternValue))
+					string patternValue = Pattern.Compute(context);
+					if (!Regex.IsMatch(value, patternValue))
 					{
-						string PatternFailedValue = PatternFailed?.Compute(Context) ?? $"Argument '{Name}' does not match the required pattern: '{PatternValue}'";
-						throw new BgOptionValidationException(PatternFailedValue);
+						string patternFailedValue = PatternFailed?.Compute(context) ?? $"Argument '{Name}' does not match the required pattern: '{patternValue}'";
+						throw new BgOptionValidationException(patternFailedValue);
 					}
 				}
 				if (!Object.ReferenceEquals(Enum, null))
 				{
-					List<string> EnumValues = Enum.Compute(Context);
-					if (!EnumValues.Any(x => x.Equals(Value, StringComparison.OrdinalIgnoreCase)))
+					List<string> enumValues = Enum.Compute(context);
+					if (!enumValues.Any(x => x.Equals(value, StringComparison.OrdinalIgnoreCase)))
 					{
 						throw new BgOptionValidationException($"Argument '{Name}' is invalid");
 					}
 				}
-				return Value;
+				return value;
 			}
-			return DefaultValue.Compute(Context);
+			return DefaultValue.Compute(context);
 		}
 	}
 
@@ -303,22 +302,22 @@ namespace EpicGames.BuildGraph
 		/// <summary>
 		/// Constructor
 		/// </summary>
-		internal BgEnumListOption(string Name, BgString Description, BgList<BgEnum<TEnum>> DefaultValue)
+		internal BgEnumListOption(string name, BgString description, BgList<BgEnum<TEnum>> defaultValue)
 		{
-			this.Name = Name;
-			this.Description = Description;
-			this.DefaultValue = DefaultValue;
+			Name = name;
+			Description = description;
+			DefaultValue = defaultValue;
 		}
 
 		/// <inheritdoc/>
-		public override IEnumerable<BgEnum<TEnum>> GetEnumerable(BgExprContext Context)
+		public override IEnumerable<BgEnum<TEnum>> GetEnumerable(BgExprContext context)
 		{
-			BgList<BgEnum<TEnum>> Value = DefaultValue;
-			if (Context.Options.TryGetValue(Name, out string? ValueText))
+			BgList<BgEnum<TEnum>> value = DefaultValue;
+			if (context.Options.TryGetValue(Name, out string? valueText))
 			{
-				Value = BgType.Get<BgList<BgEnum<TEnum>>>().DeserializeArgument(ValueText);
+				value = BgType.Get<BgList<BgEnum<TEnum>>>().DeserializeArgument(valueText);
 			}
-			return Value.GetEnumerable(Context);
+			return value.GetEnumerable(context);
 		}
 	}
 }
diff --git a/Engine/Source/Programs/Shared/EpicGames.BuildGraph/BgReport.cs b/Engine/Source/Programs/Shared/EpicGames.BuildGraph/BgReport.cs
--- a/Engine/Source/Programs/Shared/EpicGames.BuildGraph/BgReport.cs
+++ b/Engine/Source/Programs/Shared/EpicGames.BuildGraph/BgReport.cs
@@ -2,10 +2,6 @@
 
 using System;
 using System.Collections.Generic;
-using System.Diagnostics;
-using System.Linq;
-using System.Text;
-using System.Threading.Tasks;
 
 namespace EpicGames.BuildGraph
 {
@@ -17,25 +13,25 @@ namespace EpicGames.BuildGraph
 		/// <summary>
 		/// Name of this trigger
 		/// </summary>
-		public readonly string Name;
+		public string Name { get; }
 
 		/// <summary>
 		/// Set of nodes to include in the report
 		/// </summary>
-		public HashSet<BgNode> Nodes = new HashSet<BgNode>();
+		public HashSet<BgNode> Nodes { get; } = new HashSet<BgNode>();
 
 		/// <summary>
 		/// List of users to notify with this report
 		/// </summary>
-		public HashSet<string> NotifyUsers = new HashSet<string>(StringComparer.InvariantCultureIgnoreCase);
+		public HashSet<string> NotifyUsers { get; } = new HashSet<string>(StringComparer.InvariantCultureIgnoreCase);
 
 		/// <summary>
 		/// Constructor
 		/// </summary>
-		/// <param name="InName">Name of this report</param>
-		public BgReport(string InName)
+		/// <param name="inName">Name of this report</param>
+		public BgReport(string inName)
 		{
-			Name = InName;
+			Name = inName;
 		}
 
 		/// <summary>
