Reformat EpicGames.BuildGraph according to style guidelines.

diff --git a/Engine/Source/Programs/Shared/EpicGames.BuildGraph/BgSchema.cs b/Engine/Source/Programs/Shared/EpicGames.BuildGraph/BgSchema.cs
--- a/Engine/Source/Programs/Shared/EpicGames.BuildGraph/BgSchema.cs
+++ b/Engine/Source/Programs/Shared/EpicGames.BuildGraph/BgSchema.cs
@@ -3,14 +3,9 @@
 using System;
 using System.Collections.Generic;
 using System.Diagnostics;
-using System.IO;
 using System.Linq;
-using System.Reflection;
-using System.Text;
-using System.Threading.Tasks;
 using System.Xml;
 using System.Xml.Schema;
-using System.Xml.Serialization;
 using EpicGames.Core;
 
 namespace EpicGames.BuildGraph
@@ -65,35 +60,35 @@ namespace EpicGames.BuildGraph
 		/// <summary>
 		/// Whether this parameter is optional
 		/// </summary>
-		public bool bOptional { get; }
+		public bool Optional { get; }
 
 		/// <summary>
 		/// Constructor
 		/// </summary>
-		public BgScriptTaskParameter(string InName, Type InValueType, TaskParameterValidationType InValidationType, bool bInOptional)
+		public BgScriptTaskParameter(string inName, Type inValueType, TaskParameterValidationType inValidationType, bool bInOptional)
 		{
-			Name = InName;
-			ValueType = InValueType;
-			ValidationType = InValidationType;
-			bOptional = bInOptional;
+			Name = inName;
+			ValueType = inValueType;
+			ValidationType = inValidationType;
+			Optional = bInOptional;
 
 			if (ValueType.IsGenericType && ValueType.GetGenericTypeDefinition() == typeof(Nullable<>))
 			{
 				ValueType = ValueType.GetGenericArguments()[0];
-				bOptional = true;
+				Optional = true;
 			}
 
-			if(ValueType.IsClass)
+			if (ValueType.IsClass)
 			{
-				foreach(Type InterfaceType in ValueType.GetInterfaces())
+				foreach (Type interfaceType in ValueType.GetInterfaces())
 				{
-					if(InterfaceType.IsGenericType)
+					if (interfaceType.IsGenericType)
 					{
-						Type GenericInterfaceType = InterfaceType.GetGenericTypeDefinition();
-						if(GenericInterfaceType == typeof(ICollection<>))
+						Type genericInterfaceType = interfaceType.GetGenericTypeDefinition();
+						if (genericInterfaceType == typeof(ICollection<>))
 						{
-							CollectionType = InterfaceType;
-							ValueType = InterfaceType.GetGenericArguments()[0];
+							CollectionType = interfaceType;
+							ValueType = interfaceType.GetGenericArguments()[0];
 						}
 					}
 				}
@@ -120,12 +115,12 @@ namespace EpicGames.BuildGraph
 		/// <summary>
 		/// Constructor
 		/// </summary>
-		/// <param name="Name">Name of the task</param>
-		/// <param name="Parameters">Parameters for the task</param>
-		public BgScriptTask(string Name, List<BgScriptTaskParameter> Parameters)
+		/// <param name="name">Name of the task</param>
+		/// <param name="parameters">Parameters for the task</param>
+		public BgScriptTask(string name, List<BgScriptTaskParameter> parameters)
 		{
-			this.Name = Name;
-			this.Parameters = new List<BgScriptTaskParameter>(Parameters);
+			Name = name;
+			Parameters = new List<BgScriptTaskParameter>(parameters);
 		}
 	}
 
@@ -184,12 +179,12 @@ namespace EpicGames.BuildGraph
 		/// <summary>
 		/// Namespace for the schema
 		/// </summary>
-		public const string NamespaceURI = "http://www.epicgames.com/BuildGraph";
+		public const string NamespaceUri = "http://www.epicgames.com/BuildGraph";
 
 		/// <summary>
 		/// Qualified name for the string type
 		/// </summary>
-		static readonly XmlQualifiedName StringTypeName = new XmlQualifiedName("string", "http://www.w3.org/2001/XMLSchema");
+		static readonly XmlQualifiedName s_stringTypeName = new XmlQualifiedName("string", "http://www.w3.org/2001/XMLSchema");
 
 		/// <summary>
 		/// The inner xml schema
@@ -204,7 +199,7 @@ namespace EpicGames.BuildGraph
 		/// <summary>
 		/// Pattern which matches any name; alphanumeric characters, with single embedded spaces.
 		/// </summary>
-		const string NamePattern = "[^ " + IllegalNameCharacters +"]+( [^ " + IllegalNameCharacters + "]+)*";
+		const string NamePattern = "[^ " + IllegalNameCharacters + "]+( [^ " + IllegalNameCharacters + "]+)*";
 
 		/// <summary>
 		/// Pattern which matches a list of names, separated by semicolons.
@@ -251,208 +246,208 @@ namespace EpicGames.BuildGraph
 		/// </summary>
 		const string BalancedStringPattern = "[^\\$]*" + "(" + "(" + PropertyPattern + "|" + "\\$[^\\(]" + ")" + "[^\\$]*" + ")*" + "\\$?";
 
-		private BgScriptSchema(XmlSchema Schema)
+		private BgScriptSchema(XmlSchema schema)
 		{
-			CompiledSchema = Schema;
+			CompiledSchema = schema;
 		}
 
 		/// <summary>
 		/// Constructor
 		/// </summary>
-		/// <param name="Tasks">Set of known tasks</param>
-		/// <param name="PrimitiveTypes">Mapping of task name to information about how to construct it</param>
-		public BgScriptSchema(IEnumerable<BgScriptTask> Tasks, List<(Type, ScriptSchemaStandardType)> PrimitiveTypes)
+		/// <param name="tasks">Set of known tasks</param>
+		/// <param name="primitiveTypes">Mapping of task name to information about how to construct it</param>
+		public BgScriptSchema(IEnumerable<BgScriptTask> tasks, List<(Type, ScriptSchemaStandardType)> primitiveTypes)
 		{
 			// Create a lookup from standard types to their qualified names
-			Dictionary<Type, XmlQualifiedName> TypeToSchemaTypeName = new Dictionary<Type,XmlQualifiedName>();
-			TypeToSchemaTypeName.Add(typeof(String), GetQualifiedTypeName(ScriptSchemaStandardType.BalancedString));
-			TypeToSchemaTypeName.Add(typeof(Boolean), GetQualifiedTypeName(ScriptSchemaStandardType.Boolean));
-			TypeToSchemaTypeName.Add(typeof(Int32), GetQualifiedTypeName(ScriptSchemaStandardType.Integer));
-			foreach ((Type Type, ScriptSchemaStandardType SchemaType) in PrimitiveTypes)
+			Dictionary<Type, XmlQualifiedName> typeToSchemaTypeName = new Dictionary<Type, XmlQualifiedName>();
+			typeToSchemaTypeName.Add(typeof(String), GetQualifiedTypeName(ScriptSchemaStandardType.BalancedString));
+			typeToSchemaTypeName.Add(typeof(Boolean), GetQualifiedTypeName(ScriptSchemaStandardType.Boolean));
+			typeToSchemaTypeName.Add(typeof(Int32), GetQualifiedTypeName(ScriptSchemaStandardType.Integer));
+			foreach ((Type type, ScriptSchemaStandardType schemaType) in primitiveTypes)
 			{
-				TypeToSchemaTypeName.Add(Type, GetQualifiedTypeName(SchemaType));
+				typeToSchemaTypeName.Add(type, GetQualifiedTypeName(schemaType));
 			}
 
 			// Create all the custom user types, and add them to the qualified name lookup
-			List<XmlSchemaType> UserTypes = new List<XmlSchemaType>();
-			foreach (Type Type in Tasks.SelectMany(x => x.Parameters).Select(x => x.ValueType))
+			List<XmlSchemaType> userTypes = new List<XmlSchemaType>();
+			foreach (Type type in tasks.SelectMany(x => x.Parameters).Select(x => x.ValueType))
 			{
-				if (!TypeToSchemaTypeName.ContainsKey(Type))
+				if (!typeToSchemaTypeName.ContainsKey(type))
 				{
-					if (Type.IsClass && Type.GetInterfaces().Any(x => x.GetGenericTypeDefinition() == typeof(ICollection<>)))
+					if (type.IsClass && type.GetInterfaces().Any(x => x.GetGenericTypeDefinition() == typeof(ICollection<>)))
 					{
-						TypeToSchemaTypeName.Add(Type, GetQualifiedTypeName(ScriptSchemaStandardType.BalancedString));
+						typeToSchemaTypeName.Add(type, GetQualifiedTypeName(ScriptSchemaStandardType.BalancedString));
 					}
 					else
 					{
-						string Name = Type.Name + "UserType";
-						XmlSchemaType SchemaType = CreateUserType(Name, Type);
-						UserTypes.Add(SchemaType);
-						TypeToSchemaTypeName.Add(Type, new XmlQualifiedName(Name, NamespaceURI));
+						string name = type.Name + "UserType";
+						XmlSchemaType schemaType = CreateUserType(name, type);
+						userTypes.Add(schemaType);
+						typeToSchemaTypeName.Add(type, new XmlQualifiedName(name, NamespaceUri));
 					}
 				}
 			}
 
 			// Create all the task types
-			Dictionary<string, XmlSchemaComplexType>? TaskNameToType = new Dictionary<string, XmlSchemaComplexType>();
-			foreach (BgScriptTask Task in Tasks)
+			Dictionary<string, XmlSchemaComplexType>? taskNameToType = new Dictionary<string, XmlSchemaComplexType>();
+			foreach (BgScriptTask task in tasks)
 			{
-				XmlSchemaComplexType TaskType = new XmlSchemaComplexType();
-				TaskType.Name = Task.Name + "TaskType";
-				foreach (BgScriptTaskParameter Parameter in Task.Parameters)
+				XmlSchemaComplexType taskType = new XmlSchemaComplexType();
+				taskType.Name = task.Name + "TaskType";
+				foreach (BgScriptTaskParameter parameter in task.Parameters)
 				{
-					XmlQualifiedName? SchemaTypeName = GetQualifiedTypeName(Parameter.ValidationType);
-					if (SchemaTypeName == null)
+					XmlQualifiedName? schemaTypeName = GetQualifiedTypeName(parameter.ValidationType);
+					if (schemaTypeName == null)
 					{
-						SchemaTypeName = TypeToSchemaTypeName[Parameter.ValueType];
+						schemaTypeName = typeToSchemaTypeName[parameter.ValueType];
 					}
-					TaskType.Attributes.Add(CreateSchemaAttribute(Parameter.Name, SchemaTypeName, Parameter.bOptional ? XmlSchemaUse.Optional : XmlSchemaUse.Required));
+					taskType.Attributes.Add(CreateSchemaAttribute(parameter.Name, schemaTypeName, parameter.Optional ? XmlSchemaUse.Optional : XmlSchemaUse.Required));
 				}
-				TaskType.Attributes.Add(CreateSchemaAttribute("If", ScriptSchemaStandardType.BalancedString, XmlSchemaUse.Optional));
-				TaskNameToType.Add(Task.Name, TaskType);
+				taskType.Attributes.Add(CreateSchemaAttribute("If", ScriptSchemaStandardType.BalancedString, XmlSchemaUse.Optional));
+				taskNameToType.Add(task.Name, taskType);
 			}
 
 			// Create the schema object
-			XmlSchema NewSchema = new XmlSchema();
-			NewSchema.TargetNamespace = NamespaceURI;
-			NewSchema.ElementFormDefault = XmlSchemaForm.Qualified;
-			NewSchema.Items.Add(CreateSchemaElement(RootElementName, ScriptSchemaStandardType.Graph));
-			NewSchema.Items.Add(CreateGraphType());
-			NewSchema.Items.Add(CreateAgentType());
-			NewSchema.Items.Add(CreateAgentBodyType());
-			NewSchema.Items.Add(CreateNodeType());
-			NewSchema.Items.Add(CreateNodeBodyType(TaskNameToType));
-			NewSchema.Items.Add(CreateAggregateType());
-			NewSchema.Items.Add(CreateReportType());
-			NewSchema.Items.Add(CreateBadgeType());
-			NewSchema.Items.Add(CreateLabelType());
-			NewSchema.Items.Add(CreateEnumType(GetTypeName(ScriptSchemaStandardType.LabelChange), typeof(BgLabelChange)));
-			NewSchema.Items.Add(CreateNotifyType());
-			NewSchema.Items.Add(CreateIncludeType());
-			NewSchema.Items.Add(CreateOptionType());
-			NewSchema.Items.Add(CreateEnvVarType());
-			NewSchema.Items.Add(CreatePropertyType());
-			NewSchema.Items.Add(CreateRegexType());
-			NewSchema.Items.Add(CreateMacroType());
-			NewSchema.Items.Add(CreateMacroBodyType(TaskNameToType));
-			NewSchema.Items.Add(CreateExtendType());
-			NewSchema.Items.Add(CreateExpandType());
-			NewSchema.Items.Add(CreateDiagnosticType(ScriptSchemaStandardType.Trace));
-			NewSchema.Items.Add(CreateDiagnosticType(ScriptSchemaStandardType.Warning));
-			NewSchema.Items.Add(CreateDiagnosticType(ScriptSchemaStandardType.Error));
-			NewSchema.Items.Add(CreateSimpleTypeFromRegex(GetTypeName(ScriptSchemaStandardType.Name), "(" + NamePattern + "|" + StringWithPropertiesPattern + ")"));
-			NewSchema.Items.Add(CreateSimpleTypeFromRegex(GetTypeName(ScriptSchemaStandardType.NameList), "(" + NameListPattern + "|" + StringWithPropertiesPattern + ")"));
-			NewSchema.Items.Add(CreateSimpleTypeFromRegex(GetTypeName(ScriptSchemaStandardType.Tag), "(" + TagPattern + "|" + StringWithPropertiesPattern + ")"));
-			NewSchema.Items.Add(CreateSimpleTypeFromRegex(GetTypeName(ScriptSchemaStandardType.TagList), "(" + TagListPattern + "|" + StringWithPropertiesPattern + ")"));
-			NewSchema.Items.Add(CreateSimpleTypeFromRegex(GetTypeName(ScriptSchemaStandardType.NameOrTag), "(" + NameOrTagPattern + "|" + StringWithPropertiesPattern + ")"));
-			NewSchema.Items.Add(CreateSimpleTypeFromRegex(GetTypeName(ScriptSchemaStandardType.NameOrTagList), "(" + NameOrTagListPattern + "|" + StringWithPropertiesPattern + ")"));
-			NewSchema.Items.Add(CreateSimpleTypeFromRegex(GetTypeName(ScriptSchemaStandardType.QualifiedName), "(" + QualifiedNamePattern + "|" + StringWithPropertiesPattern + ")"));
-			NewSchema.Items.Add(CreateSimpleTypeFromRegex(GetTypeName(ScriptSchemaStandardType.BalancedString), BalancedStringPattern));
-			NewSchema.Items.Add(CreateSimpleTypeFromRegex(GetTypeName(ScriptSchemaStandardType.Boolean), "(true|True|false|False|" + StringWithPropertiesPattern + ")"));
-			NewSchema.Items.Add(CreateSimpleTypeFromRegex(GetTypeName(ScriptSchemaStandardType.Integer), "(" + "(-?[1-9][0-9]*|0)" + "|" + StringWithPropertiesPattern + ")"));
-			if (TaskNameToType != null)
+			XmlSchema newSchema = new XmlSchema();
+			newSchema.TargetNamespace = NamespaceUri;
+			newSchema.ElementFormDefault = XmlSchemaForm.Qualified;
+			newSchema.Items.Add(CreateSchemaElement(RootElementName, ScriptSchemaStandardType.Graph));
+			newSchema.Items.Add(CreateGraphType());
+			newSchema.Items.Add(CreateAgentType());
+			newSchema.Items.Add(CreateAgentBodyType());
+			newSchema.Items.Add(CreateNodeType());
+			newSchema.Items.Add(CreateNodeBodyType(taskNameToType));
+			newSchema.Items.Add(CreateAggregateType());
+			newSchema.Items.Add(CreateReportType());
+			newSchema.Items.Add(CreateBadgeType());
+			newSchema.Items.Add(CreateLabelType());
+			newSchema.Items.Add(CreateEnumType(GetTypeName(ScriptSchemaStandardType.LabelChange), typeof(BgLabelChange)));
+			newSchema.Items.Add(CreateNotifyType());
+			newSchema.Items.Add(CreateIncludeType());
+			newSchema.Items.Add(CreateOptionType());
+			newSchema.Items.Add(CreateEnvVarType());
+			newSchema.Items.Add(CreatePropertyType());
+			newSchema.Items.Add(CreateRegexType());
+			newSchema.Items.Add(CreateMacroType());
+			newSchema.Items.Add(CreateMacroBodyType(taskNameToType));
+			newSchema.Items.Add(CreateExtendType());
+			newSchema.Items.Add(CreateExpandType());
+			newSchema.Items.Add(CreateDiagnosticType(ScriptSchemaStandardType.Trace));
+			newSchema.Items.Add(CreateDiagnosticType(ScriptSchemaStandardType.Warning));
+			newSchema.Items.Add(CreateDiagnosticType(ScriptSchemaStandardType.Error));
+			newSchema.Items.Add(CreateSimpleTypeFromRegex(GetTypeName(ScriptSchemaStandardType.Name), "(" + NamePattern + "|" + StringWithPropertiesPattern + ")"));
+			newSchema.Items.Add(CreateSimpleTypeFromRegex(GetTypeName(ScriptSchemaStandardType.NameList), "(" + NameListPattern + "|" + StringWithPropertiesPattern + ")"));
+			newSchema.Items.Add(CreateSimpleTypeFromRegex(GetTypeName(ScriptSchemaStandardType.Tag), "(" + TagPattern + "|" + StringWithPropertiesPattern + ")"));
+			newSchema.Items.Add(CreateSimpleTypeFromRegex(GetTypeName(ScriptSchemaStandardType.TagList), "(" + TagListPattern + "|" + StringWithPropertiesPattern + ")"));
+			newSchema.Items.Add(CreateSimpleTypeFromRegex(GetTypeName(ScriptSchemaStandardType.NameOrTag), "(" + NameOrTagPattern + "|" + StringWithPropertiesPattern + ")"));
+			newSchema.Items.Add(CreateSimpleTypeFromRegex(GetTypeName(ScriptSchemaStandardType.NameOrTagList), "(" + NameOrTagListPattern + "|" + StringWithPropertiesPattern + ")"));
+			newSchema.Items.Add(CreateSimpleTypeFromRegex(GetTypeName(ScriptSchemaStandardType.QualifiedName), "(" + QualifiedNamePattern + "|" + StringWithPropertiesPattern + ")"));
+			newSchema.Items.Add(CreateSimpleTypeFromRegex(GetTypeName(ScriptSchemaStandardType.BalancedString), BalancedStringPattern));
+			newSchema.Items.Add(CreateSimpleTypeFromRegex(GetTypeName(ScriptSchemaStandardType.Boolean), "(true|True|false|False|" + StringWithPropertiesPattern + ")"));
+			newSchema.Items.Add(CreateSimpleTypeFromRegex(GetTypeName(ScriptSchemaStandardType.Integer), "(" + "(-?[1-9][0-9]*|0)" + "|" + StringWithPropertiesPattern + ")"));
+			if (taskNameToType != null)
 			{
-				foreach (XmlSchemaComplexType Type in TaskNameToType.Values)
+				foreach (XmlSchemaComplexType type in taskNameToType.Values)
 				{
-					NewSchema.Items.Add(Type);
+					newSchema.Items.Add(type);
 				}
 			}
-			foreach(XmlSchemaSimpleType Type in UserTypes)
+			foreach (XmlSchemaSimpleType type in userTypes)
 			{
-				NewSchema.Items.Add(Type);
+				newSchema.Items.Add(type);
 			}
 
 			// Now that we've finished, compile it and store it to the class
-			XmlSchemaSet NewSchemaSet = new XmlSchemaSet();
-			NewSchemaSet.Add(NewSchema);
-			NewSchemaSet.Compile();
+			XmlSchemaSet newSchemaSet = new XmlSchemaSet();
+			newSchemaSet.Add(newSchema);
+			newSchemaSet.Compile();
 
-			XmlSchema? Schema = null;
-			foreach(XmlSchema? NewCompiledSchema in NewSchemaSet.Schemas())
+			XmlSchema? schema = null;
+			foreach (XmlSchema? newCompiledSchema in newSchemaSet.Schemas())
 			{
-				Schema = NewCompiledSchema!;
+				schema = newCompiledSchema!;
 			}
-			CompiledSchema = Schema!;
+			CompiledSchema = schema!;
 		}
 
 		/// <summary>
 		/// Export the schema to a file
 		/// </summary>
-		/// <param name="File"></param>
-		public void Export(FileReference File)
+		/// <param name="file"></param>
+		public void Export(FileReference file)
 		{
-			DirectoryReference.CreateDirectory(File.Directory);
+			DirectoryReference.CreateDirectory(file.Directory);
 
-			XmlWriterSettings Settings = new XmlWriterSettings();
-			Settings.Indent = true;
-			Settings.IndentChars = "  ";
-			Settings.NewLineChars = "\n";
+			XmlWriterSettings settings = new XmlWriterSettings();
+			settings.Indent = true;
+			settings.IndentChars = "  ";
+			settings.NewLineChars = "\n";
 
-			DirectoryReference.CreateDirectory(File.Directory);
-			using(XmlWriter Writer = XmlWriter.Create(File.FullName, Settings))
+			DirectoryReference.CreateDirectory(file.Directory);
+			using (XmlWriter writer = XmlWriter.Create(file.FullName, settings))
 			{
-				CompiledSchema.Write(Writer);
+				CompiledSchema.Write(writer);
 			}
 		}
 
 		/// <summary>
 		/// Imports a schema from a file
 		/// </summary>
-		/// <param name="File">The XML file to import</param>
+		/// <param name="file">The XML file to import</param>
 		/// <returns>A <see cref="BgScriptSchema"/> deserialized from the XML file, or null if file doesn't exist</returns>
-		public static BgScriptSchema? Import(FileReference File)
+		public static BgScriptSchema? Import(FileReference file)
 		{
-			if (!FileReference.Exists(File))
+			if (!FileReference.Exists(file))
 			{
 				return null;
 			}
-			
-			using (XmlReader SchemaFile = XmlReader.Create(File.FullName))
+
+			using (XmlReader schemaFile = XmlReader.Create(file.FullName))
 			{
-				BgScriptSchema ImportedSchema = new BgScriptSchema(XmlSchema.Read(SchemaFile, ValidationCallback));
-				return ImportedSchema;
+				BgScriptSchema importedSchema = new BgScriptSchema(XmlSchema.Read(schemaFile, ValidationCallback));
+				return importedSchema;
 			}
 		}
 
-		static void ValidationCallback(object Sender, ValidationEventArgs EventArgs)
+		static void ValidationCallback(object sender, ValidationEventArgs eventArgs)
 		{
-			if (EventArgs.Severity == XmlSeverityType.Warning)
+			if (eventArgs.Severity == XmlSeverityType.Warning)
 			{
-				Log.WriteLine(LogEventType.Warning, "WARNING: {0}", EventArgs.Message);
+				Log.WriteLine(LogEventType.Warning, "WARNING: {0}", eventArgs.Message);
 			}
-			else if (EventArgs.Severity == XmlSeverityType.Error)
+			else if (eventArgs.Severity == XmlSeverityType.Error)
 			{
-				Log.WriteLine(LogEventType.Error, "ERROR: {0}", EventArgs.Message);
+				Log.WriteLine(LogEventType.Error, "ERROR: {0}", eventArgs.Message);
 			}
 		}
 
 		/// <summary>
 		/// Gets the bare name for the given script type
 		/// </summary>
-		/// <param name="Type">Script type to find the name of</param>
+		/// <param name="type">Script type to find the name of</param>
 		/// <returns>Name of the schema type that matches the given script type</returns>
-		static string GetTypeName(ScriptSchemaStandardType Type)
+		static string GetTypeName(ScriptSchemaStandardType type)
 		{
-			return Type.ToString() + "Type";
+			return type.ToString() + "Type";
 		}
-		
+
 		/// <summary>
 		/// Gets the qualified name for the given script type
 		/// </summary>
-		/// <param name="Type">Script type to find the qualified name for</param>
+		/// <param name="type">Script type to find the qualified name for</param>
 		/// <returns>Qualified name of the schema type that matches the given script type</returns>
-		static XmlQualifiedName GetQualifiedTypeName(ScriptSchemaStandardType Type)
+		static XmlQualifiedName GetQualifiedTypeName(ScriptSchemaStandardType type)
 		{
-			return new XmlQualifiedName(GetTypeName(Type), NamespaceURI);
+			return new XmlQualifiedName(GetTypeName(type), NamespaceUri);
 		}
 
 		/// <summary>
 		/// Gets the qualified name of the schema type for the given type of validation
 		/// </summary>
 		/// <returns>Qualified name for the corresponding schema type</returns>
-		static XmlQualifiedName? GetQualifiedTypeName(TaskParameterValidationType Type)
+		static XmlQualifiedName? GetQualifiedTypeName(TaskParameterValidationType type)
 		{
-			switch(Type)
+			switch (type)
 			{
 				case TaskParameterValidationType.TagList:
 					return GetQualifiedTypeName(ScriptSchemaStandardType.TagList);
@@ -466,34 +461,34 @@ namespace EpicGames.BuildGraph
 		/// <returns>Type definition for a graph</returns>
 		static XmlSchemaType CreateGraphType()
 		{
-			XmlSchemaChoice GraphChoice = new XmlSchemaChoice();
-			GraphChoice.MinOccurs = 0;
-			GraphChoice.MaxOccursString = "unbounded";
-			GraphChoice.Items.Add(CreateSchemaElement("Include", ScriptSchemaStandardType.Include));
-			GraphChoice.Items.Add(CreateSchemaElement("Option", ScriptSchemaStandardType.Option));
-			GraphChoice.Items.Add(CreateSchemaElement("EnvVar", ScriptSchemaStandardType.EnvVar));
-			GraphChoice.Items.Add(CreateSchemaElement("Property", ScriptSchemaStandardType.Property));
-			GraphChoice.Items.Add(CreateSchemaElement("Regex", ScriptSchemaStandardType.Regex));
-			GraphChoice.Items.Add(CreateSchemaElement("Macro", ScriptSchemaStandardType.Macro));
-			GraphChoice.Items.Add(CreateSchemaElement("Extend", ScriptSchemaStandardType.Extend));
-			GraphChoice.Items.Add(CreateSchemaElement("Agent", ScriptSchemaStandardType.Agent));
-			GraphChoice.Items.Add(CreateSchemaElement("Aggregate", ScriptSchemaStandardType.Aggregate));
-			GraphChoice.Items.Add(CreateSchemaElement("Report", ScriptSchemaStandardType.Report));
-			GraphChoice.Items.Add(CreateSchemaElement("Badge", ScriptSchemaStandardType.Badge));
-			GraphChoice.Items.Add(CreateSchemaElement("Label", ScriptSchemaStandardType.Label));
-			GraphChoice.Items.Add(CreateSchemaElement("Notify", ScriptSchemaStandardType.Notify));
-			GraphChoice.Items.Add(CreateSchemaElement("Trace", ScriptSchemaStandardType.Trace));
-			GraphChoice.Items.Add(CreateSchemaElement("Warning", ScriptSchemaStandardType.Warning));
-			GraphChoice.Items.Add(CreateSchemaElement("Error", ScriptSchemaStandardType.Error));
-			GraphChoice.Items.Add(CreateSchemaElement("Expand", ScriptSchemaStandardType.Expand));
-			GraphChoice.Items.Add(CreateDoElement(ScriptSchemaStandardType.Graph));
-			GraphChoice.Items.Add(CreateSwitchElement(ScriptSchemaStandardType.Graph));
-			GraphChoice.Items.Add(CreateForEachElement(ScriptSchemaStandardType.Graph));
-
-			XmlSchemaComplexType GraphType = new XmlSchemaComplexType();
-			GraphType.Name = GetTypeName(ScriptSchemaStandardType.Graph);
-			GraphType.Particle = GraphChoice;
-			return GraphType;
+			XmlSchemaChoice graphChoice = new XmlSchemaChoice();
+			graphChoice.MinOccurs = 0;
+			graphChoice.MaxOccursString = "unbounded";
+			graphChoice.Items.Add(CreateSchemaElement("Include", ScriptSchemaStandardType.Include));
+			graphChoice.Items.Add(CreateSchemaElement("Option", ScriptSchemaStandardType.Option));
+			graphChoice.Items.Add(CreateSchemaElement("EnvVar", ScriptSchemaStandardType.EnvVar));
+			graphChoice.Items.Add(CreateSchemaElement("Property", ScriptSchemaStandardType.Property));
+			graphChoice.Items.Add(CreateSchemaElement("Regex", ScriptSchemaStandardType.Regex));
+			graphChoice.Items.Add(CreateSchemaElement("Macro", ScriptSchemaStandardType.Macro));
+			graphChoice.Items.Add(CreateSchemaElement("Extend", ScriptSchemaStandardType.Extend));
+			graphChoice.Items.Add(CreateSchemaElement("Agent", ScriptSchemaStandardType.Agent));
+			graphChoice.Items.Add(CreateSchemaElement("Aggregate", ScriptSchemaStandardType.Aggregate));
+			graphChoice.Items.Add(CreateSchemaElement("Report", ScriptSchemaStandardType.Report));
+			graphChoice.Items.Add(CreateSchemaElement("Badge", ScriptSchemaStandardType.Badge));
+			graphChoice.Items.Add(CreateSchemaElement("Label", ScriptSchemaStandardType.Label));
+			graphChoice.Items.Add(CreateSchemaElement("Notify", ScriptSchemaStandardType.Notify));
+			graphChoice.Items.Add(CreateSchemaElement("Trace", ScriptSchemaStandardType.Trace));
+			graphChoice.Items.Add(CreateSchemaElement("Warning", ScriptSchemaStandardType.Warning));
+			graphChoice.Items.Add(CreateSchemaElement("Error", ScriptSchemaStandardType.Error));
+			graphChoice.Items.Add(CreateSchemaElement("Expand", ScriptSchemaStandardType.Expand));
+			graphChoice.Items.Add(CreateDoElement(ScriptSchemaStandardType.Graph));
+			graphChoice.Items.Add(CreateSwitchElement(ScriptSchemaStandardType.Graph));
+			graphChoice.Items.Add(CreateForEachElement(ScriptSchemaStandardType.Graph));
+
+			XmlSchemaComplexType graphType = new XmlSchemaComplexType();
+			graphType.Name = GetTypeName(ScriptSchemaStandardType.Graph);
+			graphType.Particle = graphChoice;
+			return graphType;
 		}
 
 		/// <summary>
@@ -502,19 +497,19 @@ namespace EpicGames.BuildGraph
 		/// <returns>Type definition for an agent</returns>
 		static XmlSchemaType CreateAgentType()
 		{
-			XmlSchemaComplexContentExtension Extension = new XmlSchemaComplexContentExtension();
-			Extension.BaseTypeName = GetQualifiedTypeName(ScriptSchemaStandardType.AgentBody);
-			Extension.Attributes.Add(CreateSchemaAttribute("Name", StringTypeName, XmlSchemaUse.Required));
-			Extension.Attributes.Add(CreateSchemaAttribute("Type", ScriptSchemaStandardType.NameList, XmlSchemaUse.Optional));
-			Extension.Attributes.Add(CreateSchemaAttribute("If", ScriptSchemaStandardType.BalancedString, XmlSchemaUse.Optional));
-
-			XmlSchemaComplexContent ContentModel = new XmlSchemaComplexContent();
-			ContentModel.Content = Extension;
-
-			XmlSchemaComplexType ComplexType = new XmlSchemaComplexType();
-			ComplexType.Name = GetTypeName(ScriptSchemaStandardType.Agent);
-			ComplexType.ContentModel = ContentModel;
-			return ComplexType;
+			XmlSchemaComplexContentExtension extension = new XmlSchemaComplexContentExtension();
+			extension.BaseTypeName = GetQualifiedTypeName(ScriptSchemaStandardType.AgentBody);
+			extension.Attributes.Add(CreateSchemaAttribute("Name", s_stringTypeName, XmlSchemaUse.Required));
+			extension.Attributes.Add(CreateSchemaAttribute("Type", ScriptSchemaStandardType.NameList, XmlSchemaUse.Optional));
+			extension.Attributes.Add(CreateSchemaAttribute("If", ScriptSchemaStandardType.BalancedString, XmlSchemaUse.Optional));
+
+			XmlSchemaComplexContent contentModel = new XmlSchemaComplexContent();
+			contentModel.Content = extension;
+
+			XmlSchemaComplexType complexType = new XmlSchemaComplexType();
+			complexType.Name = GetTypeName(ScriptSchemaStandardType.Agent);
+			complexType.ContentModel = contentModel;
+			return complexType;
 		}
 
 		/// <summary>
@@ -523,26 +518,26 @@ namespace EpicGames.BuildGraph
 		/// <returns>Type definition for an agent</returns>
 		static XmlSchemaType CreateAgentBodyType()
 		{
-			XmlSchemaChoice AgentChoice = new XmlSchemaChoice();
-			AgentChoice.MinOccurs = 0;
-			AgentChoice.MaxOccursString = "unbounded";
-			AgentChoice.Items.Add(CreateSchemaElement("Property", ScriptSchemaStandardType.Property));
-			AgentChoice.Items.Add(CreateSchemaElement("Regex", ScriptSchemaStandardType.Regex));
-			AgentChoice.Items.Add(CreateSchemaElement("EnvVar", ScriptSchemaStandardType.EnvVar));
-			AgentChoice.Items.Add(CreateSchemaElement("Node", ScriptSchemaStandardType.Node));
-			AgentChoice.Items.Add(CreateSchemaElement("Trace", ScriptSchemaStandardType.Trace));
-			AgentChoice.Items.Add(CreateSchemaElement("Label", ScriptSchemaStandardType.Label));
-			AgentChoice.Items.Add(CreateSchemaElement("Warning", ScriptSchemaStandardType.Warning));
-			AgentChoice.Items.Add(CreateSchemaElement("Error", ScriptSchemaStandardType.Error));
-			AgentChoice.Items.Add(CreateSchemaElement("Expand", ScriptSchemaStandardType.Expand));
-			AgentChoice.Items.Add(CreateDoElement(ScriptSchemaStandardType.AgentBody));
-			AgentChoice.Items.Add(CreateSwitchElement(ScriptSchemaStandardType.AgentBody));
-			AgentChoice.Items.Add(CreateForEachElement(ScriptSchemaStandardType.AgentBody));
-
-			XmlSchemaComplexType AgentType = new XmlSchemaComplexType();
-			AgentType.Name = GetTypeName(ScriptSchemaStandardType.AgentBody);
-			AgentType.Particle = AgentChoice;
-			return AgentType;
+			XmlSchemaChoice agentChoice = new XmlSchemaChoice();
+			agentChoice.MinOccurs = 0;
+			agentChoice.MaxOccursString = "unbounded";
+			agentChoice.Items.Add(CreateSchemaElement("Property", ScriptSchemaStandardType.Property));
+			agentChoice.Items.Add(CreateSchemaElement("Regex", ScriptSchemaStandardType.Regex));
+			agentChoice.Items.Add(CreateSchemaElement("EnvVar", ScriptSchemaStandardType.EnvVar));
+			agentChoice.Items.Add(CreateSchemaElement("Node", ScriptSchemaStandardType.Node));
+			agentChoice.Items.Add(CreateSchemaElement("Trace", ScriptSchemaStandardType.Trace));
+			agentChoice.Items.Add(CreateSchemaElement("Label", ScriptSchemaStandardType.Label));
+			agentChoice.Items.Add(CreateSchemaElement("Warning", ScriptSchemaStandardType.Warning));
+			agentChoice.Items.Add(CreateSchemaElement("Error", ScriptSchemaStandardType.Error));
+			agentChoice.Items.Add(CreateSchemaElement("Expand", ScriptSchemaStandardType.Expand));
+			agentChoice.Items.Add(CreateDoElement(ScriptSchemaStandardType.AgentBody));
+			agentChoice.Items.Add(CreateSwitchElement(ScriptSchemaStandardType.AgentBody));
+			agentChoice.Items.Add(CreateForEachElement(ScriptSchemaStandardType.AgentBody));
+
+			XmlSchemaComplexType agentType = new XmlSchemaComplexType();
+			agentType.Name = GetTypeName(ScriptSchemaStandardType.AgentBody);
+			agentType.Particle = agentChoice;
+			return agentType;
 		}
 
 		/// <summary>
@@ -551,62 +546,62 @@ namespace EpicGames.BuildGraph
 		/// <returns>Type definition for a node</returns>
 		static XmlSchemaType CreateNodeType()
 		{
-			XmlSchemaComplexContentExtension Extension = new XmlSchemaComplexContentExtension();
-			Extension.BaseTypeName = GetQualifiedTypeName(ScriptSchemaStandardType.NodeBody);
-			Extension.Attributes.Add(CreateSchemaAttribute("Name", ScriptSchemaStandardType.Name, XmlSchemaUse.Required));
-			Extension.Attributes.Add(CreateSchemaAttribute("Requires", ScriptSchemaStandardType.NameOrTagList, XmlSchemaUse.Optional));
-			Extension.Attributes.Add(CreateSchemaAttribute("Produces", ScriptSchemaStandardType.TagList, XmlSchemaUse.Optional));
-			Extension.Attributes.Add(CreateSchemaAttribute("After", ScriptSchemaStandardType.NameOrTagList, XmlSchemaUse.Optional));
-			Extension.Attributes.Add(CreateSchemaAttribute("Token", ScriptSchemaStandardType.BalancedString, XmlSchemaUse.Optional));
-			Extension.Attributes.Add(CreateSchemaAttribute("If", ScriptSchemaStandardType.BalancedString, XmlSchemaUse.Optional));
-			Extension.Attributes.Add(CreateSchemaAttribute("RunEarly", ScriptSchemaStandardType.Boolean, XmlSchemaUse.Optional));
-			Extension.Attributes.Add(CreateSchemaAttribute("NotifyOnWarnings", ScriptSchemaStandardType.Boolean, XmlSchemaUse.Optional));
-
-			XmlSchemaComplexContent ContentModel = new XmlSchemaComplexContent();
-			ContentModel.Content = Extension;
-
-			XmlSchemaComplexType ComplexType = new XmlSchemaComplexType();
-			ComplexType.Name = GetTypeName(ScriptSchemaStandardType.Node);
-			ComplexType.ContentModel = ContentModel;
-			return ComplexType;
+			XmlSchemaComplexContentExtension extension = new XmlSchemaComplexContentExtension();
+			extension.BaseTypeName = GetQualifiedTypeName(ScriptSchemaStandardType.NodeBody);
+			extension.Attributes.Add(CreateSchemaAttribute("Name", ScriptSchemaStandardType.Name, XmlSchemaUse.Required));
+			extension.Attributes.Add(CreateSchemaAttribute("Requires", ScriptSchemaStandardType.NameOrTagList, XmlSchemaUse.Optional));
+			extension.Attributes.Add(CreateSchemaAttribute("Produces", ScriptSchemaStandardType.TagList, XmlSchemaUse.Optional));
+			extension.Attributes.Add(CreateSchemaAttribute("After", ScriptSchemaStandardType.NameOrTagList, XmlSchemaUse.Optional));
+			extension.Attributes.Add(CreateSchemaAttribute("Token", ScriptSchemaStandardType.BalancedString, XmlSchemaUse.Optional));
+			extension.Attributes.Add(CreateSchemaAttribute("If", ScriptSchemaStandardType.BalancedString, XmlSchemaUse.Optional));
+			extension.Attributes.Add(CreateSchemaAttribute("RunEarly", ScriptSchemaStandardType.Boolean, XmlSchemaUse.Optional));
+			extension.Attributes.Add(CreateSchemaAttribute("NotifyOnWarnings", ScriptSchemaStandardType.Boolean, XmlSchemaUse.Optional));
+
+			XmlSchemaComplexContent contentModel = new XmlSchemaComplexContent();
+			contentModel.Content = extension;
+
+			XmlSchemaComplexType complexType = new XmlSchemaComplexType();
+			complexType.Name = GetTypeName(ScriptSchemaStandardType.Node);
+			complexType.ContentModel = contentModel;
+			return complexType;
 		}
 
 		/// <summary>
 		/// Creates the schema type representing the body of the node type
 		/// </summary>
 		/// <returns>Type definition for a node</returns>
-		static XmlSchemaType CreateNodeBodyType(Dictionary<string, XmlSchemaComplexType> TaskNameToType)
+		static XmlSchemaType CreateNodeBodyType(Dictionary<string, XmlSchemaComplexType> taskNameToType)
 		{
-			XmlSchemaChoice NodeChoice = new XmlSchemaChoice();
-			NodeChoice.MinOccurs = 0;
-			NodeChoice.MaxOccursString = "unbounded";
-			NodeChoice.Items.Add(CreateSchemaElement("Property", ScriptSchemaStandardType.Property));
-			NodeChoice.Items.Add(CreateSchemaElement("Regex", ScriptSchemaStandardType.Regex));
-			NodeChoice.Items.Add(CreateSchemaElement("EnvVar", ScriptSchemaStandardType.EnvVar));
-			NodeChoice.Items.Add(CreateSchemaElement("Trace", ScriptSchemaStandardType.Trace));
-			NodeChoice.Items.Add(CreateSchemaElement("Warning", ScriptSchemaStandardType.Warning));
-			NodeChoice.Items.Add(CreateSchemaElement("Error", ScriptSchemaStandardType.Error));
-			NodeChoice.Items.Add(CreateSchemaElement("Expand", ScriptSchemaStandardType.Expand));
-			NodeChoice.Items.Add(CreateDoElement(ScriptSchemaStandardType.NodeBody));
-			NodeChoice.Items.Add(CreateSwitchElement(ScriptSchemaStandardType.NodeBody));
-			NodeChoice.Items.Add(CreateForEachElement(ScriptSchemaStandardType.NodeBody));
-
-			if (TaskNameToType == null)
+			XmlSchemaChoice nodeChoice = new XmlSchemaChoice();
+			nodeChoice.MinOccurs = 0;
+			nodeChoice.MaxOccursString = "unbounded";
+			nodeChoice.Items.Add(CreateSchemaElement("Property", ScriptSchemaStandardType.Property));
+			nodeChoice.Items.Add(CreateSchemaElement("Regex", ScriptSchemaStandardType.Regex));
+			nodeChoice.Items.Add(CreateSchemaElement("EnvVar", ScriptSchemaStandardType.EnvVar));
+			nodeChoice.Items.Add(CreateSchemaElement("Trace", ScriptSchemaStandardType.Trace));
+			nodeChoice.Items.Add(CreateSchemaElement("Warning", ScriptSchemaStandardType.Warning));
+			nodeChoice.Items.Add(CreateSchemaElement("Error", ScriptSchemaStandardType.Error));
+			nodeChoice.Items.Add(CreateSchemaElement("Expand", ScriptSchemaStandardType.Expand));
+			nodeChoice.Items.Add(CreateDoElement(ScriptSchemaStandardType.NodeBody));
+			nodeChoice.Items.Add(CreateSwitchElement(ScriptSchemaStandardType.NodeBody));
+			nodeChoice.Items.Add(CreateForEachElement(ScriptSchemaStandardType.NodeBody));
+
+			if (taskNameToType == null)
 			{
-				NodeChoice.Items.Add(new XmlSchemaAny());
+				nodeChoice.Items.Add(new XmlSchemaAny());
 			}
 			else
 			{
-				foreach (KeyValuePair<string, XmlSchemaComplexType> Pair in TaskNameToType.OrderBy(x => x.Key))
+				foreach (KeyValuePair<string, XmlSchemaComplexType> pair in taskNameToType.OrderBy(x => x.Key))
 				{
-					NodeChoice.Items.Add(CreateSchemaElement(Pair.Key, new XmlQualifiedName(Pair.Value.Name, NamespaceURI)));
+					nodeChoice.Items.Add(CreateSchemaElement(pair.Key, new XmlQualifiedName(pair.Value.Name, NamespaceUri)));
 				}
 			}
 
-			XmlSchemaComplexType NodeType = new XmlSchemaComplexType();
-			NodeType.Name = GetTypeName(ScriptSchemaStandardType.NodeBody);
-			NodeType.Particle = NodeChoice;
-			return NodeType;
+			XmlSchemaComplexType nodeType = new XmlSchemaComplexType();
+			nodeType.Name = GetTypeName(ScriptSchemaStandardType.NodeBody);
+			nodeType.Particle = nodeChoice;
+			return nodeType;
 		}
 
 		/// <summary>
@@ -615,15 +610,15 @@ namespace EpicGames.BuildGraph
 		/// <returns>Type definition for an aggregate</returns>
 		static XmlSchemaType CreateAggregateType()
 		{
-			XmlSchemaComplexType AggregateType = new XmlSchemaComplexType();
-			AggregateType.Name = GetTypeName(ScriptSchemaStandardType.Aggregate);
-			AggregateType.Attributes.Add(CreateSchemaAttribute("Name", ScriptSchemaStandardType.Name, XmlSchemaUse.Required));
-			AggregateType.Attributes.Add(CreateSchemaAttribute("Label", ScriptSchemaStandardType.BalancedString, XmlSchemaUse.Optional));
-			AggregateType.Attributes.Add(CreateSchemaAttribute("Requires", ScriptSchemaStandardType.NameOrTagList, XmlSchemaUse.Required));
-			AggregateType.Attributes.Add(CreateSchemaAttribute("Include", ScriptSchemaStandardType.NameOrTagList, XmlSchemaUse.Optional));
-			AggregateType.Attributes.Add(CreateSchemaAttribute("Exclude", ScriptSchemaStandardType.NameOrTagList, XmlSchemaUse.Optional));
-			AggregateType.Attributes.Add(CreateSchemaAttribute("If", ScriptSchemaStandardType.BalancedString, XmlSchemaUse.Optional));
-			return AggregateType;
+			XmlSchemaComplexType aggregateType = new XmlSchemaComplexType();
+			aggregateType.Name = GetTypeName(ScriptSchemaStandardType.Aggregate);
+			aggregateType.Attributes.Add(CreateSchemaAttribute("Name", ScriptSchemaStandardType.Name, XmlSchemaUse.Required));
+			aggregateType.Attributes.Add(CreateSchemaAttribute("Label", ScriptSchemaStandardType.BalancedString, XmlSchemaUse.Optional));
+			aggregateType.Attributes.Add(CreateSchemaAttribute("Requires", ScriptSchemaStandardType.NameOrTagList, XmlSchemaUse.Required));
+			aggregateType.Attributes.Add(CreateSchemaAttribute("Include", ScriptSchemaStandardType.NameOrTagList, XmlSchemaUse.Optional));
+			aggregateType.Attributes.Add(CreateSchemaAttribute("Exclude", ScriptSchemaStandardType.NameOrTagList, XmlSchemaUse.Optional));
+			aggregateType.Attributes.Add(CreateSchemaAttribute("If", ScriptSchemaStandardType.BalancedString, XmlSchemaUse.Optional));
+			return aggregateType;
 		}
 
 		/// <summary>
@@ -632,12 +627,12 @@ namespace EpicGames.BuildGraph
 		/// <returns>Type definition for a report</returns>
 		static XmlSchemaType CreateReportType()
 		{
-			XmlSchemaComplexType ReportType = new XmlSchemaComplexType();
-			ReportType.Name = GetTypeName(ScriptSchemaStandardType.Report);
-			ReportType.Attributes.Add(CreateSchemaAttribute("Name", ScriptSchemaStandardType.Name, XmlSchemaUse.Required));
-			ReportType.Attributes.Add(CreateSchemaAttribute("Requires", ScriptSchemaStandardType.NameOrTagList, XmlSchemaUse.Required));
-			ReportType.Attributes.Add(CreateSchemaAttribute("If", ScriptSchemaStandardType.BalancedString, XmlSchemaUse.Optional));
-			return ReportType;
+			XmlSchemaComplexType reportType = new XmlSchemaComplexType();
+			reportType.Name = GetTypeName(ScriptSchemaStandardType.Report);
+			reportType.Attributes.Add(CreateSchemaAttribute("Name", ScriptSchemaStandardType.Name, XmlSchemaUse.Required));
+			reportType.Attributes.Add(CreateSchemaAttribute("Requires", ScriptSchemaStandardType.NameOrTagList, XmlSchemaUse.Required));
+			reportType.Attributes.Add(CreateSchemaAttribute("If", ScriptSchemaStandardType.BalancedString, XmlSchemaUse.Optional));
+			return reportType;
 		}
 
 		/// <summary>
@@ -646,15 +641,15 @@ namespace EpicGames.BuildGraph
 		/// <returns>Type definition for a badge</returns>
 		static XmlSchemaType CreateBadgeType()
 		{
-			XmlSchemaComplexType BadgeType = new XmlSchemaComplexType();
-			BadgeType.Name = GetTypeName(ScriptSchemaStandardType.Badge);
-			BadgeType.Attributes.Add(CreateSchemaAttribute("Name", ScriptSchemaStandardType.Name, XmlSchemaUse.Required));
-			BadgeType.Attributes.Add(CreateSchemaAttribute("Requires", ScriptSchemaStandardType.NameOrTagList, XmlSchemaUse.Optional));
-			BadgeType.Attributes.Add(CreateSchemaAttribute("Targets", ScriptSchemaStandardType.NameOrTagList, XmlSchemaUse.Optional));
-			BadgeType.Attributes.Add(CreateSchemaAttribute("Project", ScriptSchemaStandardType.BalancedString, XmlSchemaUse.Required));
-			BadgeType.Attributes.Add(CreateSchemaAttribute("Change", ScriptSchemaStandardType.BalancedString, XmlSchemaUse.Optional));
-			BadgeType.Attributes.Add(CreateSchemaAttribute("If", ScriptSchemaStandardType.BalancedString, XmlSchemaUse.Optional));
-			return BadgeType;
+			XmlSchemaComplexType badgeType = new XmlSchemaComplexType();
+			badgeType.Name = GetTypeName(ScriptSchemaStandardType.Badge);
+			badgeType.Attributes.Add(CreateSchemaAttribute("Name", ScriptSchemaStandardType.Name, XmlSchemaUse.Required));
+			badgeType.Attributes.Add(CreateSchemaAttribute("Requires", ScriptSchemaStandardType.NameOrTagList, XmlSchemaUse.Optional));
+			badgeType.Attributes.Add(CreateSchemaAttribute("Targets", ScriptSchemaStandardType.NameOrTagList, XmlSchemaUse.Optional));
+			badgeType.Attributes.Add(CreateSchemaAttribute("Project", ScriptSchemaStandardType.BalancedString, XmlSchemaUse.Required));
+			badgeType.Attributes.Add(CreateSchemaAttribute("Change", ScriptSchemaStandardType.BalancedString, XmlSchemaUse.Optional));
+			badgeType.Attributes.Add(CreateSchemaAttribute("If", ScriptSchemaStandardType.BalancedString, XmlSchemaUse.Optional));
+			return badgeType;
 		}
 
 		/// <summary>
@@ -663,18 +658,18 @@ namespace EpicGames.BuildGraph
 		/// <returns>Type definition for a label</returns>
 		static XmlSchemaType CreateLabelType()
 		{
-			XmlSchemaComplexType LabelType = new XmlSchemaComplexType();
-			LabelType.Name = GetTypeName(ScriptSchemaStandardType.Label);
-			LabelType.Attributes.Add(CreateSchemaAttribute("Name", ScriptSchemaStandardType.Name, XmlSchemaUse.Optional));
-			LabelType.Attributes.Add(CreateSchemaAttribute("Category", ScriptSchemaStandardType.Name, XmlSchemaUse.Optional));
-			LabelType.Attributes.Add(CreateSchemaAttribute("UgsBadge", ScriptSchemaStandardType.Name, XmlSchemaUse.Optional));
-			LabelType.Attributes.Add(CreateSchemaAttribute("UgsProject", ScriptSchemaStandardType.BalancedString, XmlSchemaUse.Optional));
-			LabelType.Attributes.Add(CreateSchemaAttribute("Change", ScriptSchemaStandardType.LabelChange, XmlSchemaUse.Optional));
-			LabelType.Attributes.Add(CreateSchemaAttribute("Requires", ScriptSchemaStandardType.NameOrTagList, XmlSchemaUse.Required));
-			LabelType.Attributes.Add(CreateSchemaAttribute("Include", ScriptSchemaStandardType.NameOrTagList, XmlSchemaUse.Optional));
-			LabelType.Attributes.Add(CreateSchemaAttribute("Exclude", ScriptSchemaStandardType.NameOrTagList, XmlSchemaUse.Optional));
-			LabelType.Attributes.Add(CreateSchemaAttribute("If", ScriptSchemaStandardType.BalancedString, XmlSchemaUse.Optional));
-			return LabelType;
+			XmlSchemaComplexType labelType = new XmlSchemaComplexType();
+			labelType.Name = GetTypeName(ScriptSchemaStandardType.Label);
+			labelType.Attributes.Add(CreateSchemaAttribute("Name", ScriptSchemaStandardType.Name, XmlSchemaUse.Optional));
+			labelType.Attributes.Add(CreateSchemaAttribute("Category", ScriptSchemaStandardType.Name, XmlSchemaUse.Optional));
+			labelType.Attributes.Add(CreateSchemaAttribute("UgsBadge", ScriptSchemaStandardType.Name, XmlSchemaUse.Optional));
+			labelType.Attributes.Add(CreateSchemaAttribute("UgsProject", ScriptSchemaStandardType.BalancedString, XmlSchemaUse.Optional));
+			labelType.Attributes.Add(CreateSchemaAttribute("Change", ScriptSchemaStandardType.LabelChange, XmlSchemaUse.Optional));
+			labelType.Attributes.Add(CreateSchemaAttribute("Requires", ScriptSchemaStandardType.NameOrTagList, XmlSchemaUse.Required));
+			labelType.Attributes.Add(CreateSchemaAttribute("Include", ScriptSchemaStandardType.NameOrTagList, XmlSchemaUse.Optional));
+			labelType.Attributes.Add(CreateSchemaAttribute("Exclude", ScriptSchemaStandardType.NameOrTagList, XmlSchemaUse.Optional));
+			labelType.Attributes.Add(CreateSchemaAttribute("If", ScriptSchemaStandardType.BalancedString, XmlSchemaUse.Optional));
+			return labelType;
 		}
 
 		/// <summary>
@@ -683,18 +678,18 @@ namespace EpicGames.BuildGraph
 		/// <returns>Type definition for a notifier</returns>
 		static XmlSchemaType CreateNotifyType()
 		{
-			XmlSchemaComplexType AggregateType = new XmlSchemaComplexType();
-			AggregateType.Name = GetTypeName(ScriptSchemaStandardType.Notify);
-			AggregateType.Attributes.Add(CreateSchemaAttribute("Targets", ScriptSchemaStandardType.NameOrTagList, XmlSchemaUse.Optional));
-			AggregateType.Attributes.Add(CreateSchemaAttribute("Except", ScriptSchemaStandardType.NameOrTagList, XmlSchemaUse.Optional));
-			AggregateType.Attributes.Add(CreateSchemaAttribute("Nodes", ScriptSchemaStandardType.NameOrTagList, XmlSchemaUse.Optional));
-			AggregateType.Attributes.Add(CreateSchemaAttribute("Reports", ScriptSchemaStandardType.NameList, XmlSchemaUse.Optional));
-			AggregateType.Attributes.Add(CreateSchemaAttribute("Users", ScriptSchemaStandardType.BalancedString, XmlSchemaUse.Optional));
-			AggregateType.Attributes.Add(CreateSchemaAttribute("Submitters", ScriptSchemaStandardType.BalancedString, XmlSchemaUse.Optional));
-			AggregateType.Attributes.Add(CreateSchemaAttribute("Warnings", ScriptSchemaStandardType.Boolean, XmlSchemaUse.Optional));
-			AggregateType.Attributes.Add(CreateSchemaAttribute("If", ScriptSchemaStandardType.BalancedString, XmlSchemaUse.Optional));
-			AggregateType.Attributes.Add(CreateSchemaAttribute("Absolute", ScriptSchemaStandardType.BalancedString, XmlSchemaUse.Optional));
-			return AggregateType;
+			XmlSchemaComplexType aggregateType = new XmlSchemaComplexType();
+			aggregateType.Name = GetTypeName(ScriptSchemaStandardType.Notify);
+			aggregateType.Attributes.Add(CreateSchemaAttribute("Targets", ScriptSchemaStandardType.NameOrTagList, XmlSchemaUse.Optional));
+			aggregateType.Attributes.Add(CreateSchemaAttribute("Except", ScriptSchemaStandardType.NameOrTagList, XmlSchemaUse.Optional));
+			aggregateType.Attributes.Add(CreateSchemaAttribute("Nodes", ScriptSchemaStandardType.NameOrTagList, XmlSchemaUse.Optional));
+			aggregateType.Attributes.Add(CreateSchemaAttribute("Reports", ScriptSchemaStandardType.NameList, XmlSchemaUse.Optional));
+			aggregateType.Attributes.Add(CreateSchemaAttribute("Users", ScriptSchemaStandardType.BalancedString, XmlSchemaUse.Optional));
+			aggregateType.Attributes.Add(CreateSchemaAttribute("Submitters", ScriptSchemaStandardType.BalancedString, XmlSchemaUse.Optional));
+			aggregateType.Attributes.Add(CreateSchemaAttribute("Warnings", ScriptSchemaStandardType.Boolean, XmlSchemaUse.Optional));
+			aggregateType.Attributes.Add(CreateSchemaAttribute("If", ScriptSchemaStandardType.BalancedString, XmlSchemaUse.Optional));
+			aggregateType.Attributes.Add(CreateSchemaAttribute("Absolute", ScriptSchemaStandardType.BalancedString, XmlSchemaUse.Optional));
+			return aggregateType;
 		}
 
 		/// <summary>
@@ -703,11 +698,11 @@ namespace EpicGames.BuildGraph
 		/// <returns>Type definition for an include directive</returns>
 		static XmlSchemaType CreateIncludeType()
 		{
-			XmlSchemaComplexType PropertyType = new XmlSchemaComplexType();
-			PropertyType.Name = GetTypeName(ScriptSchemaStandardType.Include);
-			PropertyType.Attributes.Add(CreateSchemaAttribute("Script", ScriptSchemaStandardType.BalancedString, XmlSchemaUse.Required));
-			PropertyType.Attributes.Add(CreateSchemaAttribute("If", ScriptSchemaStandardType.BalancedString, XmlSchemaUse.Optional));
-			return PropertyType;
+			XmlSchemaComplexType propertyType = new XmlSchemaComplexType();
+			propertyType.Name = GetTypeName(ScriptSchemaStandardType.Include);
+			propertyType.Attributes.Add(CreateSchemaAttribute("Script", ScriptSchemaStandardType.BalancedString, XmlSchemaUse.Required));
+			propertyType.Attributes.Add(CreateSchemaAttribute("If", ScriptSchemaStandardType.BalancedString, XmlSchemaUse.Optional));
+			return propertyType;
 		}
 
 		/// <summary>
@@ -716,14 +711,14 @@ namespace EpicGames.BuildGraph
 		/// <returns>Type definition for a parameter</returns>
 		static XmlSchemaType CreateOptionType()
 		{
-			XmlSchemaComplexType OptionType = new XmlSchemaComplexType();
-			OptionType.Name = GetTypeName(ScriptSchemaStandardType.Option);
-			OptionType.Attributes.Add(CreateSchemaAttribute("Name", ScriptSchemaStandardType.Name, XmlSchemaUse.Required));
-			OptionType.Attributes.Add(CreateSchemaAttribute("Restrict", StringTypeName, XmlSchemaUse.Optional));
-			OptionType.Attributes.Add(CreateSchemaAttribute("DefaultValue", StringTypeName, XmlSchemaUse.Required));
-			OptionType.Attributes.Add(CreateSchemaAttribute("Description", StringTypeName, XmlSchemaUse.Required));
-			OptionType.Attributes.Add(CreateSchemaAttribute("If", ScriptSchemaStandardType.BalancedString, XmlSchemaUse.Optional));
-			return OptionType;
+			XmlSchemaComplexType optionType = new XmlSchemaComplexType();
+			optionType.Name = GetTypeName(ScriptSchemaStandardType.Option);
+			optionType.Attributes.Add(CreateSchemaAttribute("Name", ScriptSchemaStandardType.Name, XmlSchemaUse.Required));
+			optionType.Attributes.Add(CreateSchemaAttribute("Restrict", s_stringTypeName, XmlSchemaUse.Optional));
+			optionType.Attributes.Add(CreateSchemaAttribute("DefaultValue", s_stringTypeName, XmlSchemaUse.Required));
+			optionType.Attributes.Add(CreateSchemaAttribute("Description", s_stringTypeName, XmlSchemaUse.Required));
+			optionType.Attributes.Add(CreateSchemaAttribute("If", ScriptSchemaStandardType.BalancedString, XmlSchemaUse.Optional));
+			return optionType;
 		}
 
 		/// <summary>
@@ -732,11 +727,11 @@ namespace EpicGames.BuildGraph
 		/// <returns>Type definition for an environment variable property</returns>
 		static XmlSchemaType CreateEnvVarType()
 		{
-			XmlSchemaComplexType EnvVarType = new XmlSchemaComplexType();
-			EnvVarType.Name = GetTypeName(ScriptSchemaStandardType.EnvVar);
-			EnvVarType.Attributes.Add(CreateSchemaAttribute("Name", ScriptSchemaStandardType.Name, XmlSchemaUse.Required));
-			EnvVarType.Attributes.Add(CreateSchemaAttribute("If", ScriptSchemaStandardType.BalancedString, XmlSchemaUse.Optional));
-			return EnvVarType;
+			XmlSchemaComplexType envVarType = new XmlSchemaComplexType();
+			envVarType.Name = GetTypeName(ScriptSchemaStandardType.EnvVar);
+			envVarType.Attributes.Add(CreateSchemaAttribute("Name", ScriptSchemaStandardType.Name, XmlSchemaUse.Required));
+			envVarType.Attributes.Add(CreateSchemaAttribute("If", ScriptSchemaStandardType.BalancedString, XmlSchemaUse.Optional));
+			return envVarType;
 		}
 
 		/// <summary>
@@ -745,19 +740,19 @@ namespace EpicGames.BuildGraph
 		/// <returns>Type definition for a property</returns>
 		static XmlSchemaType CreatePropertyType()
 		{
-			XmlSchemaSimpleContentExtension Extension = new XmlSchemaSimpleContentExtension();
-			Extension.BaseTypeName = StringTypeName;
-			Extension.Attributes.Add(CreateSchemaAttribute("Name", ScriptSchemaStandardType.Name, XmlSchemaUse.Required));
-			Extension.Attributes.Add(CreateSchemaAttribute("Value", StringTypeName, XmlSchemaUse.Optional));
-			Extension.Attributes.Add(CreateSchemaAttribute("If", ScriptSchemaStandardType.BalancedString, XmlSchemaUse.Optional));
-
-			XmlSchemaSimpleContent ContentModel = new XmlSchemaSimpleContent();
-			ContentModel.Content = Extension;
-
-			XmlSchemaComplexType PropertyType = new XmlSchemaComplexType();
-			PropertyType.Name = GetTypeName(ScriptSchemaStandardType.Property);
-			PropertyType.ContentModel = ContentModel;
-			return PropertyType;
+			XmlSchemaSimpleContentExtension extension = new XmlSchemaSimpleContentExtension();
+			extension.BaseTypeName = s_stringTypeName;
+			extension.Attributes.Add(CreateSchemaAttribute("Name", ScriptSchemaStandardType.Name, XmlSchemaUse.Required));
+			extension.Attributes.Add(CreateSchemaAttribute("Value", s_stringTypeName, XmlSchemaUse.Optional));
+			extension.Attributes.Add(CreateSchemaAttribute("If", ScriptSchemaStandardType.BalancedString, XmlSchemaUse.Optional));
+
+			XmlSchemaSimpleContent contentModel = new XmlSchemaSimpleContent();
+			contentModel.Content = extension;
+
+			XmlSchemaComplexType propertyType = new XmlSchemaComplexType();
+			propertyType.Name = GetTypeName(ScriptSchemaStandardType.Property);
+			propertyType.ContentModel = contentModel;
+			return propertyType;
 		}
 
 		/// <summary>
@@ -766,22 +761,22 @@ namespace EpicGames.BuildGraph
 		/// <returns>Type definition for a regex</returns>
 		static XmlSchemaType CreateRegexType()
 		{
-			XmlSchemaSimpleContentExtension Extension = new XmlSchemaSimpleContentExtension();
-			Extension.BaseTypeName = StringTypeName;
-
-			Extension.Attributes.Add(CreateSchemaAttribute("Input", StringTypeName, XmlSchemaUse.Required));
-			Extension.Attributes.Add(CreateSchemaAttribute("Pattern", StringTypeName, XmlSchemaUse.Required));
-			Extension.Attributes.Add(CreateSchemaAttribute("Capture", ScriptSchemaStandardType.NameList, XmlSchemaUse.Required));
-			Extension.Attributes.Add(CreateSchemaAttribute("Optional", ScriptSchemaStandardType.Boolean, XmlSchemaUse.Optional));
-			Extension.Attributes.Add(CreateSchemaAttribute("If", ScriptSchemaStandardType.BalancedString, XmlSchemaUse.Optional));
-
-			XmlSchemaSimpleContent ContentModel = new XmlSchemaSimpleContent();
-			ContentModel.Content = Extension;
-
-			XmlSchemaComplexType RegexType  = new XmlSchemaComplexType();
-			RegexType.Name = GetTypeName(ScriptSchemaStandardType.Regex);
-			RegexType.ContentModel = ContentModel;
-			return RegexType;
+			XmlSchemaSimpleContentExtension extension = new XmlSchemaSimpleContentExtension();
+			extension.BaseTypeName = s_stringTypeName;
+
+			extension.Attributes.Add(CreateSchemaAttribute("Input", s_stringTypeName, XmlSchemaUse.Required));
+			extension.Attributes.Add(CreateSchemaAttribute("Pattern", s_stringTypeName, XmlSchemaUse.Required));
+			extension.Attributes.Add(CreateSchemaAttribute("Capture", ScriptSchemaStandardType.NameList, XmlSchemaUse.Required));
+			extension.Attributes.Add(CreateSchemaAttribute("Optional", ScriptSchemaStandardType.Boolean, XmlSchemaUse.Optional));
+			extension.Attributes.Add(CreateSchemaAttribute("If", ScriptSchemaStandardType.BalancedString, XmlSchemaUse.Optional));
+
+			XmlSchemaSimpleContent contentModel = new XmlSchemaSimpleContent();
+			contentModel.Content = extension;
+
+			XmlSchemaComplexType regexType = new XmlSchemaComplexType();
+			regexType.Name = GetTypeName(ScriptSchemaStandardType.Regex);
+			regexType.ContentModel = contentModel;
+			return regexType;
 		}
 
 		/// <summary>
@@ -790,65 +785,65 @@ namespace EpicGames.BuildGraph
 		/// <returns>Type definition for a node</returns>
 		static XmlSchemaType CreateMacroType()
 		{
-			XmlSchemaComplexContentExtension Extension = new XmlSchemaComplexContentExtension();
-			Extension.BaseTypeName = GetQualifiedTypeName(ScriptSchemaStandardType.MacroBody);
-			Extension.Attributes.Add(CreateSchemaAttribute("Name", ScriptSchemaStandardType.Name, XmlSchemaUse.Required));
-			Extension.Attributes.Add(CreateSchemaAttribute("Arguments", ScriptSchemaStandardType.BalancedString, XmlSchemaUse.Optional));
-			Extension.Attributes.Add(CreateSchemaAttribute("OptionalArguments", ScriptSchemaStandardType.BalancedString, XmlSchemaUse.Optional));
-
-			XmlSchemaComplexContent ContentModel = new XmlSchemaComplexContent();
-			ContentModel.Content = Extension;
-
-			XmlSchemaComplexType ComplexType = new XmlSchemaComplexType();
-			ComplexType.Name = GetTypeName(ScriptSchemaStandardType.Macro);
-			ComplexType.ContentModel = ContentModel;
-			return ComplexType;
+			XmlSchemaComplexContentExtension extension = new XmlSchemaComplexContentExtension();
+			extension.BaseTypeName = GetQualifiedTypeName(ScriptSchemaStandardType.MacroBody);
+			extension.Attributes.Add(CreateSchemaAttribute("Name", ScriptSchemaStandardType.Name, XmlSchemaUse.Required));
+			extension.Attributes.Add(CreateSchemaAttribute("Arguments", ScriptSchemaStandardType.BalancedString, XmlSchemaUse.Optional));
+			extension.Attributes.Add(CreateSchemaAttribute("OptionalArguments", ScriptSchemaStandardType.BalancedString, XmlSchemaUse.Optional));
+
+			XmlSchemaComplexContent contentModel = new XmlSchemaComplexContent();
+			contentModel.Content = extension;
+
+			XmlSchemaComplexType complexType = new XmlSchemaComplexType();
+			complexType.Name = GetTypeName(ScriptSchemaStandardType.Macro);
+			complexType.ContentModel = contentModel;
+			return complexType;
 		}
 
 		/// <summary>
 		/// Creates the schema type representing the macro type
 		/// </summary>
 		/// <returns>Type definition for a node</returns>
-		static XmlSchemaType CreateMacroBodyType(Dictionary<string, XmlSchemaComplexType> TaskNameToType)
+		static XmlSchemaType CreateMacroBodyType(Dictionary<string, XmlSchemaComplexType> taskNameToType)
 		{
-			XmlSchemaChoice MacroChoice = new XmlSchemaChoice();
-			MacroChoice.MinOccurs = 0;
-			MacroChoice.MaxOccursString = "unbounded";
+			XmlSchemaChoice macroChoice = new XmlSchemaChoice();
+			macroChoice.MinOccurs = 0;
+			macroChoice.MaxOccursString = "unbounded";
 
 			// Graph scope
-			MacroChoice.Items.Add(CreateSchemaElement("Include", ScriptSchemaStandardType.Include));
-			MacroChoice.Items.Add(CreateSchemaElement("Option", ScriptSchemaStandardType.Option));
-			MacroChoice.Items.Add(CreateSchemaElement("EnvVar", ScriptSchemaStandardType.EnvVar));
-			MacroChoice.Items.Add(CreateSchemaElement("Property", ScriptSchemaStandardType.Property));
-			MacroChoice.Items.Add(CreateSchemaElement("Regex", ScriptSchemaStandardType.Regex));
-			MacroChoice.Items.Add(CreateSchemaElement("Macro", ScriptSchemaStandardType.Macro));
-			MacroChoice.Items.Add(CreateSchemaElement("Agent", ScriptSchemaStandardType.Agent));
-			MacroChoice.Items.Add(CreateSchemaElement("Aggregate", ScriptSchemaStandardType.Aggregate));
-			MacroChoice.Items.Add(CreateSchemaElement("Report", ScriptSchemaStandardType.Report));
-			MacroChoice.Items.Add(CreateSchemaElement("Badge", ScriptSchemaStandardType.Badge));
-			MacroChoice.Items.Add(CreateSchemaElement("Notify", ScriptSchemaStandardType.Notify));
-			MacroChoice.Items.Add(CreateSchemaElement("Trace", ScriptSchemaStandardType.Trace));
-			MacroChoice.Items.Add(CreateSchemaElement("Warning", ScriptSchemaStandardType.Warning));
-			MacroChoice.Items.Add(CreateSchemaElement("Error", ScriptSchemaStandardType.Error));
-			MacroChoice.Items.Add(CreateSchemaElement("Expand", ScriptSchemaStandardType.Expand));
-			MacroChoice.Items.Add(CreateSchemaElement("Label", ScriptSchemaStandardType.Label));
+			macroChoice.Items.Add(CreateSchemaElement("Include", ScriptSchemaStandardType.Include));
+			macroChoice.Items.Add(CreateSchemaElement("Option", ScriptSchemaStandardType.Option));
+			macroChoice.Items.Add(CreateSchemaElement("EnvVar", ScriptSchemaStandardType.EnvVar));
+			macroChoice.Items.Add(CreateSchemaElement("Property", ScriptSchemaStandardType.Property));
+			macroChoice.Items.Add(CreateSchemaElement("Regex", ScriptSchemaStandardType.Regex));
+			macroChoice.Items.Add(CreateSchemaElement("Macro", ScriptSchemaStandardType.Macro));
+			macroChoice.Items.Add(CreateSchemaElement("Agent", ScriptSchemaStandardType.Agent));
+			macroChoice.Items.Add(CreateSchemaElement("Aggregate", ScriptSchemaStandardType.Aggregate));
+			macroChoice.Items.Add(CreateSchemaElement("Report", ScriptSchemaStandardType.Report));
+			macroChoice.Items.Add(CreateSchemaElement("Badge", ScriptSchemaStandardType.Badge));
+			macroChoice.Items.Add(CreateSchemaElement("Notify", ScriptSchemaStandardType.Notify));
+			macroChoice.Items.Add(CreateSchemaElement("Trace", ScriptSchemaStandardType.Trace));
+			macroChoice.Items.Add(CreateSchemaElement("Warning", ScriptSchemaStandardType.Warning));
+			macroChoice.Items.Add(CreateSchemaElement("Error", ScriptSchemaStandardType.Error));
+			macroChoice.Items.Add(CreateSchemaElement("Expand", ScriptSchemaStandardType.Expand));
+			macroChoice.Items.Add(CreateSchemaElement("Label", ScriptSchemaStandardType.Label));
 
 			// Agent scope
-			MacroChoice.Items.Add(CreateSchemaElement("Node", ScriptSchemaStandardType.Node));
+			macroChoice.Items.Add(CreateSchemaElement("Node", ScriptSchemaStandardType.Node));
 
 			// Node scope
-			MacroChoice.Items.Add(CreateDoElement(ScriptSchemaStandardType.NodeBody));
-			MacroChoice.Items.Add(CreateSwitchElement(ScriptSchemaStandardType.NodeBody));
-			MacroChoice.Items.Add(CreateForEachElement(ScriptSchemaStandardType.NodeBody));
-			foreach (KeyValuePair<string, XmlSchemaComplexType> Pair in TaskNameToType.OrderBy(x => x.Key))
+			macroChoice.Items.Add(CreateDoElement(ScriptSchemaStandardType.NodeBody));
+			macroChoice.Items.Add(CreateSwitchElement(ScriptSchemaStandardType.NodeBody));
+			macroChoice.Items.Add(CreateForEachElement(ScriptSchemaStandardType.NodeBody));
+			foreach (KeyValuePair<string, XmlSchemaComplexType> pair in taskNameToType.OrderBy(x => x.Key))
 			{
-				MacroChoice.Items.Add(CreateSchemaElement(Pair.Key, new XmlQualifiedName(Pair.Value.Name, NamespaceURI)));
+				macroChoice.Items.Add(CreateSchemaElement(pair.Key, new XmlQualifiedName(pair.Value.Name, NamespaceUri)));
 			}
 
-			XmlSchemaComplexType NodeType = new XmlSchemaComplexType();
-			NodeType.Name = GetTypeName(ScriptSchemaStandardType.MacroBody);
-			NodeType.Particle = MacroChoice;
-			return NodeType;
+			XmlSchemaComplexType nodeType = new XmlSchemaComplexType();
+			nodeType.Name = GetTypeName(ScriptSchemaStandardType.MacroBody);
+			nodeType.Particle = macroChoice;
+			return nodeType;
 		}
 
 		/// <summary>
@@ -857,17 +852,17 @@ namespace EpicGames.BuildGraph
 		/// <returns>Type definition for a node</returns>
 		static XmlSchemaType CreateExtendType()
 		{
-			XmlSchemaComplexContentExtension Extension = new XmlSchemaComplexContentExtension();
-			Extension.BaseTypeName = GetQualifiedTypeName(ScriptSchemaStandardType.MacroBody);
-			Extension.Attributes.Add(CreateSchemaAttribute("Name", ScriptSchemaStandardType.Name, XmlSchemaUse.Required));
+			XmlSchemaComplexContentExtension extension = new XmlSchemaComplexContentExtension();
+			extension.BaseTypeName = GetQualifiedTypeName(ScriptSchemaStandardType.MacroBody);
+			extension.Attributes.Add(CreateSchemaAttribute("Name", ScriptSchemaStandardType.Name, XmlSchemaUse.Required));
 
-			XmlSchemaComplexContent ContentModel = new XmlSchemaComplexContent();
-			ContentModel.Content = Extension;
+			XmlSchemaComplexContent contentModel = new XmlSchemaComplexContent();
+			contentModel.Content = extension;
 
-			XmlSchemaComplexType ComplexType = new XmlSchemaComplexType();
-			ComplexType.Name = GetTypeName(ScriptSchemaStandardType.Extend);
-			ComplexType.ContentModel = ContentModel;
-			return ComplexType;
+			XmlSchemaComplexType complexType = new XmlSchemaComplexType();
+			complexType.Name = GetTypeName(ScriptSchemaStandardType.Extend);
+			complexType.ContentModel = contentModel;
+			return complexType;
 		}
 
 		/// <summary>
@@ -876,256 +871,256 @@ namespace EpicGames.BuildGraph
 		/// <returns>Type definition for expanding a macro</returns>
 		static XmlSchemaType CreateExpandType()
 		{
-			XmlSchemaAnyAttribute AnyAttribute = new XmlSchemaAnyAttribute();
-			AnyAttribute.ProcessContents = XmlSchemaContentProcessing.Skip;
-
-			XmlSchemaComplexType PropertyType = new XmlSchemaComplexType();
-			PropertyType.Name = GetTypeName(ScriptSchemaStandardType.Expand);
-			PropertyType.Attributes.Add(CreateSchemaAttribute("Name", ScriptSchemaStandardType.BalancedString, XmlSchemaUse.Required));
-			PropertyType.Attributes.Add(CreateSchemaAttribute("If", ScriptSchemaStandardType.BalancedString, XmlSchemaUse.Optional));
-			PropertyType.AnyAttribute = AnyAttribute;
-			return PropertyType;
+			XmlSchemaAnyAttribute anyAttribute = new XmlSchemaAnyAttribute();
+			anyAttribute.ProcessContents = XmlSchemaContentProcessing.Skip;
+
+			XmlSchemaComplexType propertyType = new XmlSchemaComplexType();
+			propertyType.Name = GetTypeName(ScriptSchemaStandardType.Expand);
+			propertyType.Attributes.Add(CreateSchemaAttribute("Name", ScriptSchemaStandardType.BalancedString, XmlSchemaUse.Required));
+			propertyType.Attributes.Add(CreateSchemaAttribute("If", ScriptSchemaStandardType.BalancedString, XmlSchemaUse.Optional));
+			propertyType.AnyAttribute = anyAttribute;
+			return propertyType;
 		}
 
 		/// <summary>
 		/// Creates the schema type representing a warning or error type
 		/// </summary>
 		/// <returns>Type definition for a warning</returns>
-		static XmlSchemaType CreateDiagnosticType(ScriptSchemaStandardType StandardType)
+		static XmlSchemaType CreateDiagnosticType(ScriptSchemaStandardType standardType)
 		{
-			XmlSchemaComplexType PropertyType = new XmlSchemaComplexType();
-			PropertyType.Name = GetTypeName(StandardType);
-			PropertyType.Attributes.Add(CreateSchemaAttribute("Message", ScriptSchemaStandardType.BalancedString, XmlSchemaUse.Required));
-			PropertyType.Attributes.Add(CreateSchemaAttribute("If", ScriptSchemaStandardType.BalancedString, XmlSchemaUse.Optional));
-			return PropertyType;
+			XmlSchemaComplexType propertyType = new XmlSchemaComplexType();
+			propertyType.Name = GetTypeName(standardType);
+			propertyType.Attributes.Add(CreateSchemaAttribute("Message", ScriptSchemaStandardType.BalancedString, XmlSchemaUse.Required));
+			propertyType.Attributes.Add(CreateSchemaAttribute("If", ScriptSchemaStandardType.BalancedString, XmlSchemaUse.Optional));
+			return propertyType;
 		}
 
 		/// <summary>
 		/// Creates an element representing a conditional "Do" block, which recursively contains another type
 		/// </summary>
-		/// <param name="InnerType">The base type for the do block to contain</param>
+		/// <param name="innerType">The base type for the do block to contain</param>
 		/// <returns>New schema element for the block</returns>
-		static XmlSchemaElement CreateDoElement(ScriptSchemaStandardType InnerType)
+		static XmlSchemaElement CreateDoElement(ScriptSchemaStandardType innerType)
 		{
-			XmlSchemaComplexContentExtension Extension = new XmlSchemaComplexContentExtension();
-			Extension.BaseTypeName = GetQualifiedTypeName(InnerType);
-			Extension.Attributes.Add(CreateSchemaAttribute("If", ScriptSchemaStandardType.BalancedString, XmlSchemaUse.Optional));
+			XmlSchemaComplexContentExtension extension = new XmlSchemaComplexContentExtension();
+			extension.BaseTypeName = GetQualifiedTypeName(innerType);
+			extension.Attributes.Add(CreateSchemaAttribute("If", ScriptSchemaStandardType.BalancedString, XmlSchemaUse.Optional));
 
-			XmlSchemaComplexContent ContentModel = new XmlSchemaComplexContent();
-			ContentModel.Content = Extension;
+			XmlSchemaComplexContent contentModel = new XmlSchemaComplexContent();
+			contentModel.Content = extension;
 
-			XmlSchemaComplexType SchemaType = new XmlSchemaComplexType();
-			SchemaType.ContentModel = ContentModel;
+			XmlSchemaComplexType schemaType = new XmlSchemaComplexType();
+			schemaType.ContentModel = contentModel;
 
-			XmlSchemaElement Element = new XmlSchemaElement();
-			Element.Name = "Do";
-			Element.SchemaType = SchemaType;
-			return Element;
+			XmlSchemaElement element = new XmlSchemaElement();
+			element.Name = "Do";
+			element.SchemaType = schemaType;
+			return element;
 		}
 
 		/// <summary>
 		/// Creates an element representing a conditional "Switch" block, which recursively contains another type
 		/// </summary>
-		/// <param name="InnerType">The base type for the do block to contain</param>
+		/// <param name="innerType">The base type for the do block to contain</param>
 		/// <returns>New schema element for the block</returns>
-		static XmlSchemaElement CreateSwitchElement(ScriptSchemaStandardType InnerType)
+		static XmlSchemaElement CreateSwitchElement(ScriptSchemaStandardType innerType)
 		{
 			// Create the "Option" element
-			XmlSchemaComplexContentExtension CaseExtension = new XmlSchemaComplexContentExtension();
-			CaseExtension.BaseTypeName = GetQualifiedTypeName(InnerType);
-			CaseExtension.Attributes.Add(CreateSchemaAttribute("If", ScriptSchemaStandardType.BalancedString, XmlSchemaUse.Required));
+			XmlSchemaComplexContentExtension caseExtension = new XmlSchemaComplexContentExtension();
+			caseExtension.BaseTypeName = GetQualifiedTypeName(innerType);
+			caseExtension.Attributes.Add(CreateSchemaAttribute("If", ScriptSchemaStandardType.BalancedString, XmlSchemaUse.Required));
 
-			XmlSchemaComplexContent CaseContentModel = new XmlSchemaComplexContent();
-			CaseContentModel.Content = CaseExtension;
+			XmlSchemaComplexContent caseContentModel = new XmlSchemaComplexContent();
+			caseContentModel.Content = caseExtension;
 
-			XmlSchemaComplexType CaseSchemaType = new XmlSchemaComplexType();
-			CaseSchemaType.ContentModel = CaseContentModel;
+			XmlSchemaComplexType caseSchemaType = new XmlSchemaComplexType();
+			caseSchemaType.ContentModel = caseContentModel;
 
-			XmlSchemaElement CaseElement = new XmlSchemaElement();
-			CaseElement.Name = "Case";
-			CaseElement.SchemaType = CaseSchemaType;
-			CaseElement.MinOccurs = 0;
-			CaseElement.MaxOccursString = "unbounded";
+			XmlSchemaElement caseElement = new XmlSchemaElement();
+			caseElement.Name = "Case";
+			caseElement.SchemaType = caseSchemaType;
+			caseElement.MinOccurs = 0;
+			caseElement.MaxOccursString = "unbounded";
 
 			// Create the "Otherwise" element
-			XmlSchemaElement OtherwiseElement = new XmlSchemaElement();
-			OtherwiseElement.Name = "Default";
-			OtherwiseElement.SchemaTypeName = GetQualifiedTypeName(InnerType);
-			OtherwiseElement.MinOccurs = 0;
-			OtherwiseElement.MaxOccurs = 1;
+			XmlSchemaElement otherwiseElement = new XmlSchemaElement();
+			otherwiseElement.Name = "Default";
+			otherwiseElement.SchemaTypeName = GetQualifiedTypeName(innerType);
+			otherwiseElement.MinOccurs = 0;
+			otherwiseElement.MaxOccurs = 1;
 
 			// Create the "Switch" element
-			XmlSchemaSequence SwitchSequence = new XmlSchemaSequence();
-			SwitchSequence.Items.Add(CaseElement);
-			SwitchSequence.Items.Add(OtherwiseElement);
+			XmlSchemaSequence switchSequence = new XmlSchemaSequence();
+			switchSequence.Items.Add(caseElement);
+			switchSequence.Items.Add(otherwiseElement);
 
-			XmlSchemaComplexType SwitchSchemaType = new XmlSchemaComplexType();
-			SwitchSchemaType.Particle = SwitchSequence;
+			XmlSchemaComplexType switchSchemaType = new XmlSchemaComplexType();
+			switchSchemaType.Particle = switchSequence;
 
-			XmlSchemaElement SwitchElement = new XmlSchemaElement();
-			SwitchElement.Name = "Switch";
-			SwitchElement.SchemaType = SwitchSchemaType;
-			return SwitchElement;
+			XmlSchemaElement switchElement = new XmlSchemaElement();
+			switchElement.Name = "Switch";
+			switchElement.SchemaType = switchSchemaType;
+			return switchElement;
 		}
 
 		/// <summary>
 		/// Creates an element representing a conditional "ForEach" block, which recursively contains another type
 		/// </summary>
-		/// <param name="InnerType">The base type for the foreach block to contain</param>
+		/// <param name="innerType">The base type for the foreach block to contain</param>
 		/// <returns>New schema element for the block</returns>
-		static XmlSchemaElement CreateForEachElement(ScriptSchemaStandardType InnerType)
+		static XmlSchemaElement CreateForEachElement(ScriptSchemaStandardType innerType)
 		{
-			XmlSchemaComplexContentExtension Extension = new XmlSchemaComplexContentExtension();
-			Extension.BaseTypeName = GetQualifiedTypeName(InnerType);
-			Extension.Attributes.Add(CreateSchemaAttribute("Name", ScriptSchemaStandardType.BalancedString, XmlSchemaUse.Required));
-			Extension.Attributes.Add(CreateSchemaAttribute("Values", ScriptSchemaStandardType.BalancedString, XmlSchemaUse.Required));
-			Extension.Attributes.Add(CreateSchemaAttribute("Separator", ScriptSchemaStandardType.BalancedString, XmlSchemaUse.Optional));
-			Extension.Attributes.Add(CreateSchemaAttribute("If", ScriptSchemaStandardType.BalancedString, XmlSchemaUse.Optional));
-
-			XmlSchemaComplexContent ContentModel = new XmlSchemaComplexContent();
-			ContentModel.Content = Extension;
-
-			XmlSchemaComplexType SchemaType = new XmlSchemaComplexType();
-			SchemaType.ContentModel = ContentModel;
-
-			XmlSchemaElement Element = new XmlSchemaElement();
-			Element.Name = "ForEach";
-			Element.SchemaType = SchemaType;
-			return Element;
+			XmlSchemaComplexContentExtension extension = new XmlSchemaComplexContentExtension();
+			extension.BaseTypeName = GetQualifiedTypeName(innerType);
+			extension.Attributes.Add(CreateSchemaAttribute("Name", ScriptSchemaStandardType.BalancedString, XmlSchemaUse.Required));
+			extension.Attributes.Add(CreateSchemaAttribute("Values", ScriptSchemaStandardType.BalancedString, XmlSchemaUse.Required));
+			extension.Attributes.Add(CreateSchemaAttribute("Separator", ScriptSchemaStandardType.BalancedString, XmlSchemaUse.Optional));
+			extension.Attributes.Add(CreateSchemaAttribute("If", ScriptSchemaStandardType.BalancedString, XmlSchemaUse.Optional));
+
+			XmlSchemaComplexContent contentModel = new XmlSchemaComplexContent();
+			contentModel.Content = extension;
+
+			XmlSchemaComplexType schemaType = new XmlSchemaComplexType();
+			schemaType.ContentModel = contentModel;
+
+			XmlSchemaElement element = new XmlSchemaElement();
+			element.Name = "ForEach";
+			element.SchemaType = schemaType;
+			return element;
 		}
 
 		/// <summary>
 		/// Constructs an XmlSchemaElement and initializes it with the given parameters
 		/// </summary>
-		/// <param name="Name">Element name</param>
-		/// <param name="SchemaType">Type enumeration for the attribute</param>
+		/// <param name="name">Element name</param>
+		/// <param name="schemaType">Type enumeration for the attribute</param>
 		/// <returns>A new XmlSchemaElement object</returns>
-		static XmlSchemaElement CreateSchemaElement(string Name, ScriptSchemaStandardType SchemaType)
+		static XmlSchemaElement CreateSchemaElement(string name, ScriptSchemaStandardType schemaType)
 		{
-			return CreateSchemaElement(Name, GetQualifiedTypeName(SchemaType));
+			return CreateSchemaElement(name, GetQualifiedTypeName(schemaType));
 		}
 
 		/// <summary>
 		/// Constructs an XmlSchemaElement and initializes it with the given parameters
 		/// </summary>
-		/// <param name="Name">Element name</param>
-		/// <param name="SchemaTypeName">Qualified name of the type for this element</param>
+		/// <param name="name">Element name</param>
+		/// <param name="schemaTypeName">Qualified name of the type for this element</param>
 		/// <returns>A new XmlSchemaElement object</returns>
-		static XmlSchemaElement CreateSchemaElement(string Name, XmlQualifiedName SchemaTypeName)
+		static XmlSchemaElement CreateSchemaElement(string name, XmlQualifiedName schemaTypeName)
 		{
-			XmlSchemaElement Element = new XmlSchemaElement();
-			Element.Name = Name;
-			Element.SchemaTypeName = SchemaTypeName;
-			return Element;
+			XmlSchemaElement element = new XmlSchemaElement();
+			element.Name = name;
+			element.SchemaTypeName = schemaTypeName;
+			return element;
 		}
 
 		/// <summary>
 		/// Constructs an XmlSchemaAttribute and initialize it with the given parameters
 		/// </summary>
-		/// <param name="Name">The attribute name</param>
-		/// <param name="SchemaType">Type enumeration for the attribute</param>
-		/// <param name="Use">Whether the attribute is required or optional</param>
+		/// <param name="name">The attribute name</param>
+		/// <param name="schemaType">Type enumeration for the attribute</param>
+		/// <param name="use">Whether the attribute is required or optional</param>
 		/// <returns>A new XmlSchemaAttribute object</returns>
-		static XmlSchemaAttribute CreateSchemaAttribute(string Name, ScriptSchemaStandardType SchemaType, XmlSchemaUse Use)
+		static XmlSchemaAttribute CreateSchemaAttribute(string name, ScriptSchemaStandardType schemaType, XmlSchemaUse use)
 		{
-			return CreateSchemaAttribute(Name, GetQualifiedTypeName(SchemaType), Use);
+			return CreateSchemaAttribute(name, GetQualifiedTypeName(schemaType), use);
 		}
 
 		/// <summary>
 		/// Constructs an XmlSchemaAttribute and initialize it with the given parameters
 		/// </summary>
-		/// <param name="Name">The attribute name</param>
-		/// <param name="SchemaTypeName">Qualified name of the type for this attribute</param>
-		/// <param name="Use">Whether the attribute is required or optional</param>
+		/// <param name="name">The attribute name</param>
+		/// <param name="schemaTypeName">Qualified name of the type for this attribute</param>
+		/// <param name="use">Whether the attribute is required or optional</param>
 		/// <returns>The new attribute</returns>
-		static XmlSchemaAttribute CreateSchemaAttribute(string Name, XmlQualifiedName SchemaTypeName, XmlSchemaUse Use)
+		static XmlSchemaAttribute CreateSchemaAttribute(string name, XmlQualifiedName schemaTypeName, XmlSchemaUse use)
 		{
-			XmlSchemaAttribute Attribute = new XmlSchemaAttribute();
-			Attribute.Name = Name;
-			Attribute.SchemaTypeName = SchemaTypeName;
-			Attribute.Use = Use;
-			return Attribute;
+			XmlSchemaAttribute attribute = new XmlSchemaAttribute();
+			attribute.Name = name;
+			attribute.SchemaTypeName = schemaTypeName;
+			attribute.Use = use;
+			return attribute;
 		}
 
 		/// <summary>
 		/// Creates a simple type that is the union of two other types
 		/// </summary>
-		/// <param name="Name">The name of the type</param>
-		/// <param name="ValidTypes">List of valid types for the union</param>
+		/// <param name="name">The name of the type</param>
+		/// <param name="validTypes">List of valid types for the union</param>
 		/// <returns>A simple type which will match the given pattern</returns>
-		static XmlSchemaSimpleType CreateSimpleTypeFromUnion(string Name, params XmlSchemaType[] ValidTypes)
+		static XmlSchemaSimpleType CreateSimpleTypeFromUnion(string name, params XmlSchemaType[] validTypes)
 		{
-			XmlSchemaSimpleTypeUnion Union = new XmlSchemaSimpleTypeUnion();
-			foreach (XmlSchemaType ValidType in ValidTypes)
+			XmlSchemaSimpleTypeUnion union = new XmlSchemaSimpleTypeUnion();
+			foreach (XmlSchemaType validType in validTypes)
 			{
-				Union.BaseTypes.Add(ValidType);
+				union.BaseTypes.Add(validType);
 			}
 
-			XmlSchemaSimpleType UnionType = new XmlSchemaSimpleType();
-			UnionType.Name = Name;
-			UnionType.Content = Union;
-			return UnionType;
+			XmlSchemaSimpleType unionType = new XmlSchemaSimpleType();
+			unionType.Name = name;
+			unionType.Content = union;
+			return unionType;
 		}
 
 		/// <summary>
 		/// Creates a simple type that matches a regex
 		/// </summary>
-		/// <param name="Name">Name of the new type</param>
-		/// <param name="Pattern">Regex pattern to match</param>
+		/// <param name="name">Name of the new type</param>
+		/// <param name="pattern">Regex pattern to match</param>
 		/// <returns>A simple type which will match the given pattern</returns>
-		static XmlSchemaSimpleType CreateSimpleTypeFromRegex(string? Name, string Pattern)
+		static XmlSchemaSimpleType CreateSimpleTypeFromRegex(string? name, string pattern)
 		{
-			XmlSchemaPatternFacet PatternFacet = new XmlSchemaPatternFacet();
-			PatternFacet.Value = Pattern;
+			XmlSchemaPatternFacet patternFacet = new XmlSchemaPatternFacet();
+			patternFacet.Value = pattern;
 
-			XmlSchemaSimpleTypeRestriction Restriction = new XmlSchemaSimpleTypeRestriction();
-			Restriction.BaseTypeName = StringTypeName;
-			Restriction.Facets.Add(PatternFacet);
+			XmlSchemaSimpleTypeRestriction restriction = new XmlSchemaSimpleTypeRestriction();
+			restriction.BaseTypeName = s_stringTypeName;
+			restriction.Facets.Add(patternFacet);
 
-			XmlSchemaSimpleType SimpleType = new XmlSchemaSimpleType();
-			SimpleType.Name = Name;
-			SimpleType.Content = Restriction;
-			return SimpleType;
+			XmlSchemaSimpleType simpleType = new XmlSchemaSimpleType();
+			simpleType.Name = name;
+			simpleType.Content = restriction;
+			return simpleType;
 		}
 
 		/// <summary>
 		/// Create a schema type for the given user type. Currently only handles enumerations.
 		/// </summary>
-		/// <param name="Name">Name for the new type</param>
-		/// <param name="Type">CLR type information to create a schema type for</param>
-		static XmlSchemaType CreateUserType(string Name, Type Type)
+		/// <param name="name">Name for the new type</param>
+		/// <param name="type">CLR type information to create a schema type for</param>
+		static XmlSchemaType CreateUserType(string name, Type type)
 		{
-			if(Type.IsEnum)
+			if (type.IsEnum)
 			{
-				return CreateSimpleTypeFromUnion(Name, CreateEnumType(null, Type), CreateSimpleTypeFromRegex(null, StringWithPropertiesPattern));
+				return CreateSimpleTypeFromUnion(name, CreateEnumType(null, type), CreateSimpleTypeFromRegex(null, StringWithPropertiesPattern));
 			}
 			else
 			{
-				throw new Exception($"Cannot create custom type in schema for '{Type.Name}'");
+				throw new Exception($"Cannot create custom type in schema for '{type.Name}'");
 			}
 		}
 
 		/// <summary>
 		/// Create a schema type for the given enum.
 		/// </summary>
-		/// <param name="Name">Name for the new type</param>
-		/// <param name="Type">CLR type information to create a schema type for</param>
-		static XmlSchemaType CreateEnumType(string? Name, Type Type)
+		/// <param name="name">Name for the new type</param>
+		/// <param name="type">CLR type information to create a schema type for</param>
+		static XmlSchemaType CreateEnumType(string? name, Type type)
 		{
-			XmlSchemaSimpleTypeRestriction Restriction = new XmlSchemaSimpleTypeRestriction();
-			Restriction.BaseTypeName = StringTypeName;
+			XmlSchemaSimpleTypeRestriction restriction = new XmlSchemaSimpleTypeRestriction();
+			restriction.BaseTypeName = s_stringTypeName;
 
-			foreach(string EnumName in Enum.GetNames(Type))
+			foreach (string enumName in Enum.GetNames(type))
 			{
-				XmlSchemaEnumerationFacet Facet = new XmlSchemaEnumerationFacet();
-				Facet.Value = EnumName;
-				Restriction.Facets.Add(Facet);
+				XmlSchemaEnumerationFacet facet = new XmlSchemaEnumerationFacet();
+				facet.Value = enumName;
+				restriction.Facets.Add(facet);
 			}
 
-			XmlSchemaSimpleType SchemaType = new XmlSchemaSimpleType();
-			SchemaType.Name = Name;
-			SchemaType.Content = Restriction;
-			return SchemaType;
+			XmlSchemaSimpleType schemaType = new XmlSchemaSimpleType();
+			schemaType.Name = name;
+			schemaType.Content = restriction;
+			return schemaType;
 		}
 	}
 }
\ No newline at end of file
diff --git a/Engine/Source/Programs/Shared/EpicGames.BuildGraph/BgScript.cs b/Engine/Source/Programs/Shared/EpicGames.BuildGraph/BgScript.cs
--- a/Engine/Source/Programs/Shared/EpicGames.BuildGraph/BgScript.cs
+++ b/Engine/Source/Programs/Shared/EpicGames.BuildGraph/BgScript.cs
@@ -1,17 +1,5 @@
 // Copyright Epic Games, Inc. All Rights Reserved.
 
-using System;
-using System.Collections.Generic;
-using System.Text;
-using System.IO;
-using System.Reflection;
-using System.Xml;
-using System.Linq;
-using System.Diagnostics;
-using EpicGames.Core;
-using Microsoft.Extensions.Logging;
-using System.Diagnostics.CodeAnalysis;
-
 namespace EpicGames.BuildGraph
 {
 }
