Reformat EpicGames.BuildGraph according to style guidelines.

diff --git a/Engine/Source/Programs/Shared/EpicGames.BuildGraph/BgScriptReader.cs b/Engine/Source/Programs/Shared/EpicGames.BuildGraph/BgScriptReader.cs
--- a/Engine/Source/Programs/Shared/EpicGames.BuildGraph/BgScriptReader.cs
+++ b/Engine/Source/Programs/Shared/EpicGames.BuildGraph/BgScriptReader.cs
@@ -2,11 +2,9 @@
 
 using System;
 using System.Collections.Generic;
-using System.ComponentModel;
 using System.Diagnostics.CodeAnalysis;
 using System.IO;
 using System.Linq;
-using System.Reflection;
 using System.Text;
 using System.Text.RegularExpressions;
 using System.Threading.Tasks;
@@ -40,14 +38,14 @@ namespace EpicGames.BuildGraph
 		/// <summary>
 		/// Constructor
 		/// </summary>
-		/// <param name="File"></param>
-		/// <param name="NativeFile"></param>
-		/// <param name="LineNumber"></param>
-		public BgScriptLocation(string File, object NativeFile, int LineNumber)
+		/// <param name="file"></param>
+		/// <param name="nativeFile"></param>
+		/// <param name="lineNumber"></param>
+		public BgScriptLocation(string file, object nativeFile, int lineNumber)
 		{
-			this.File = File;
-			this.NativeFile = NativeFile;
-			this.LineNumber = LineNumber;
+			File = file;
+			NativeFile = nativeFile;
+			LineNumber = lineNumber;
 		}
 	}
 
@@ -69,12 +67,12 @@ namespace EpicGames.BuildGraph
 		/// <summary>
 		/// Interface to the LineInfo on the active XmlReader
 		/// </summary>
-		IXmlLineInfo? LineInfo;
+		IXmlLineInfo? _lineInfo;
 
 		/// <summary>
 		/// Set to true if the reader encounters an error
 		/// </summary>
-		bool bHasErrors;
+		bool _bHasErrors;
 
 		/// <summary>
 		/// Logger for validation errors
@@ -84,107 +82,107 @@ namespace EpicGames.BuildGraph
 		/// <summary>
 		/// Private constructor. Use ScriptDocument.Load to read an XML document.
 		/// </summary>
-		BgScriptDocument(string InFile, object InNativeFile, ILogger InLogger)
+		BgScriptDocument(string inFile, object inNativeFile, ILogger inLogger)
 		{
-			File = InFile;
-			NativeFile = InNativeFile;
-			Logger = InLogger;
+			File = inFile;
+			NativeFile = inNativeFile;
+			Logger = inLogger;
 		}
 
 		/// <summary>
 		/// Overrides XmlDocument.CreateElement() to construct ScriptElements rather than XmlElements
 		/// </summary>
-		public override XmlElement CreateElement(string Prefix, string LocalName, string NamespaceUri)
+		public override XmlElement CreateElement(string prefix, string localName, string namespaceUri)
 		{
-			BgScriptLocation Location = new BgScriptLocation(File, NativeFile, LineInfo!.LineNumber);
-			return new BgScriptElement(Location, Prefix, LocalName, NamespaceUri, this);
+			BgScriptLocation location = new BgScriptLocation(File, NativeFile, _lineInfo!.LineNumber);
+			return new BgScriptElement(location, prefix, localName, namespaceUri, this);
 		}
 
 		/// <summary>
 		/// Loads a script document from the given file
 		/// </summary>
-		/// <param name="File">The file to load</param>
-		/// <param name="NativeFile"></param>
-		/// <param name="Data"></param>
-		/// <param name="Schema">The schema to validate against</param>
-		/// <param name="Logger">Logger for output messages</param>
-		/// <param name="OutDocument">If successful, the document that was read</param>
+		/// <param name="file">The file to load</param>
+		/// <param name="nativeFile"></param>
+		/// <param name="data"></param>
+		/// <param name="schema">The schema to validate against</param>
+		/// <param name="logger">Logger for output messages</param>
+		/// <param name="outDocument">If successful, the document that was read</param>
 		/// <returns>True if the document could be read, false otherwise</returns>
-		public static bool TryRead(string File, object NativeFile, byte[] Data, BgScriptSchema Schema, ILogger Logger, [NotNullWhen(true)] out BgScriptDocument? OutDocument)
+		public static bool TryRead(string file, object nativeFile, byte[] data, BgScriptSchema schema, ILogger logger, [NotNullWhen(true)] out BgScriptDocument? outDocument)
 		{
-			BgScriptDocument Document = new BgScriptDocument(File, NativeFile, Logger);
+			BgScriptDocument document = new BgScriptDocument(file, nativeFile, logger);
 
-			XmlReaderSettings Settings = new XmlReaderSettings();
-			if (Schema != null)
+			XmlReaderSettings settings = new XmlReaderSettings();
+			if (schema != null)
 			{
-				Settings.Schemas.Add(Schema.CompiledSchema);
-				Settings.ValidationType = ValidationType.Schema;
-				Settings.ValidationEventHandler += Document.ValidationEvent;
+				settings.Schemas.Add(schema.CompiledSchema);
+				settings.ValidationType = ValidationType.Schema;
+				settings.ValidationEventHandler += document.ValidationEvent;
 			}
 
-			using (MemoryStream Stream = new MemoryStream(Data))
-			using (XmlReader Reader = XmlReader.Create(Stream, Settings))
+			using (MemoryStream stream = new MemoryStream(data))
+			using (XmlReader reader = XmlReader.Create(stream, settings))
 			{
 				// Read the document
-				Document.LineInfo = (IXmlLineInfo)Reader;
+				document._lineInfo = (IXmlLineInfo)reader;
 				try
 				{
-					Document.Load(Reader);
+					document.Load(reader);
 				}
-				catch (XmlException Ex)
+				catch (XmlException ex)
 				{
-					if (!Document.bHasErrors)
+					if (!document._bHasErrors)
 					{
-						BgScriptLocation Location = new BgScriptLocation(File, NativeFile, Ex.LineNumber);
-						Logger.LogScriptError(Location, "{Message}", Ex.Message);
-						Document.bHasErrors = true;
+						BgScriptLocation location = new BgScriptLocation(file, nativeFile, ex.LineNumber);
+						logger.LogScriptError(location, "{Message}", ex.Message);
+						document._bHasErrors = true;
 					}
 				}
 
 				// If we hit any errors while parsing
-				if (Document.bHasErrors)
+				if (document._bHasErrors)
 				{
-					OutDocument = null;
+					outDocument = null;
 					return false;
 				}
 
 				// Check that the root element is valid. If not, we didn't actually validate against the schema.
-				if (Document.DocumentElement.Name != BgScriptSchema.RootElementName)
+				if (document.DocumentElement.Name != BgScriptSchema.RootElementName)
 				{
-					BgScriptLocation Location = new BgScriptLocation(File, NativeFile, 1);
-					Logger.LogScriptError(Location, "Script does not have a root element called '{ElementName}'", BgScriptSchema.RootElementName);
-					OutDocument = null;
+					BgScriptLocation location = new BgScriptLocation(file, nativeFile, 1);
+					logger.LogScriptError(location, "Script does not have a root element called '{ElementName}'", BgScriptSchema.RootElementName);
+					outDocument = null;
 					return false;
 				}
-				if (Document.DocumentElement.NamespaceURI != BgScriptSchema.NamespaceURI)
+				if (document.DocumentElement.NamespaceURI != BgScriptSchema.NamespaceUri)
 				{
-					BgScriptLocation Location = new BgScriptLocation(File, NativeFile, 1);
-					Logger.LogScriptError(Location, "Script root element is not in the '{Namespace}' namespace (add the xmlns=\"{NewNamespace}\" attribute)", BgScriptSchema.NamespaceURI, BgScriptSchema.NamespaceURI);
-					OutDocument = null;
+					BgScriptLocation location = new BgScriptLocation(file, nativeFile, 1);
+					logger.LogScriptError(location, "Script root element is not in the '{Namespace}' namespace (add the xmlns=\"{NewNamespace}\" attribute)", BgScriptSchema.NamespaceUri, BgScriptSchema.NamespaceUri);
+					outDocument = null;
 					return false;
 				}
 			}
 
-			OutDocument = Document;
+			outDocument = document;
 			return true;
 		}
 
 		/// <summary>
 		/// Callback for validation errors in the document
 		/// </summary>
-		/// <param name="Sender">Standard argument for ValidationEventHandler</param>
-		/// <param name="Args">Standard argument for ValidationEventHandler</param>
-		void ValidationEvent(object Sender, ValidationEventArgs Args)
+		/// <param name="sender">Standard argument for ValidationEventHandler</param>
+		/// <param name="args">Standard argument for ValidationEventHandler</param>
+		void ValidationEvent(object sender, ValidationEventArgs args)
 		{
-			BgScriptLocation Location = new BgScriptLocation(File, NativeFile, Args.Exception.LineNumber);
-			if (Args.Severity == XmlSeverityType.Warning)
+			BgScriptLocation location = new BgScriptLocation(File, NativeFile, args.Exception.LineNumber);
+			if (args.Severity == XmlSeverityType.Warning)
 			{
-				Logger.LogScriptWarning(Location, "{Message}", Args.Message);
+				Logger.LogScriptWarning(location, "{Message}", args.Message);
 			}
 			else
 			{
-				Logger.LogScriptError(Location, "{Message}", Args.Message);
-				bHasErrors = true;
+				Logger.LogScriptError(location, "{Message}", args.Message);
+				_bHasErrors = true;
 			}
 		}
 	}
@@ -202,10 +200,10 @@ namespace EpicGames.BuildGraph
 		/// <summary>
 		/// Constructor
 		/// </summary>
-		public BgScriptElement(BgScriptLocation Location, string Prefix, string LocalName, string NamespaceUri, BgScriptDocument Document)
-			: base(Prefix, LocalName, NamespaceUri, Document)
+		public BgScriptElement(BgScriptLocation location, string prefix, string localName, string namespaceUri, BgScriptDocument document)
+			: base(prefix, localName, namespaceUri, document)
 		{
-			this.Location = Location;
+			Location = location;
 		}
 	}
 
@@ -217,42 +215,42 @@ namespace EpicGames.BuildGraph
 		/// <summary>
 		/// Name of the function
 		/// </summary>
-		public readonly string Name;
+		public string Name { get; }
 
 		/// <summary>
 		/// Element where the function was declared
 		/// </summary>
-		public List<BgScriptElement> Elements = new List<BgScriptElement>();
+		public List<BgScriptElement> Elements { get; } = new List<BgScriptElement>();
 
 		/// <summary>
 		/// The total number of arguments
 		/// </summary>
-		public readonly int NumArguments;
+		public int NumArguments { get; }
 
 		/// <summary>
 		/// Number of arguments that are required
 		/// </summary>
-		public readonly int NumRequiredArguments;
+		public int NumRequiredArguments { get; }
 
 		/// <summary>
 		/// Maps an argument name to its type
 		/// </summary>
-		public readonly Dictionary<string, int> ArgumentNameToIndex;
+		public IReadOnlyDictionary<string, int> ArgumentNameToIndex { get; }
 
 		/// <summary>
 		/// Constructor
 		/// </summary>
-		/// <param name="Name">Name of the function</param>
-		/// <param name="Element">Element containing the function definition</param>
-		/// <param name="ArgumentNameToIndex">Map of argument name to index</param>
-		/// <param name="NumRequiredArguments">Number of arguments that are required. Indices 0 to NumRequiredArguments - 1 are required.</param>
-		public BgScriptMacro(string Name, BgScriptElement Element, Dictionary<string, int> ArgumentNameToIndex, int NumRequiredArguments)
+		/// <param name="name">Name of the function</param>
+		/// <param name="element">Element containing the function definition</param>
+		/// <param name="argumentNameToIndex">Map of argument name to index</param>
+		/// <param name="numRequiredArguments">Number of arguments that are required. Indices 0 to NumRequiredArguments - 1 are required.</param>
+		public BgScriptMacro(string name, BgScriptElement element, Dictionary<string, int> argumentNameToIndex, int numRequiredArguments)
 		{
-			this.Name = Name;
-			this.Elements.Add(Element);
-			this.NumArguments = ArgumentNameToIndex.Count;
-			this.NumRequiredArguments = NumRequiredArguments;
-			this.ArgumentNameToIndex = ArgumentNameToIndex;
+			Name = name;
+			Elements.Add(element);
+			NumArguments = argumentNameToIndex.Count;
+			NumRequiredArguments = numRequiredArguments;
+			ArgumentNameToIndex = argumentNameToIndex;
 		}
 	}
 
@@ -264,30 +262,30 @@ namespace EpicGames.BuildGraph
 		/// <summary>
 		/// Tests whether the given file or directory exists
 		/// </summary>
-		/// <param name="Path">Path to a file</param>
+		/// <param name="path">Path to a file</param>
 		/// <returns>True if the file exists</returns>
-		Task<bool> ExistsAsync(string Path);
+		Task<bool> ExistsAsync(string path);
 
 		/// <summary>
 		/// Tries to read a file from the given path
 		/// </summary>
-		/// <param name="Path">Path of the file to read</param>
+		/// <param name="path">Path of the file to read</param>
 		/// <returns></returns>
-		Task<byte[]?> ReadAsync(string Path);
+		Task<byte[]?> ReadAsync(string path);
 
 		/// <summary>
 		/// Finds files matching the given pattern
 		/// </summary>
-		/// <param name="Pattern"></param>
+		/// <param name="pattern"></param>
 		/// <returns></returns>
-		Task<string[]> FindAsync(string Pattern);
+		Task<string[]> FindAsync(string pattern);
 
 		/// <summary>
 		/// Converts a path to its native form, for display to the user
 		/// </summary>
-		/// <param name="Path">Path to format</param>
+		/// <param name="path">Path to format</param>
 		/// <returns></returns>
-		object GetNativePath(string Path);
+		object GetNativePath(string path);
 	}
 
 	/// <summary>
@@ -298,25 +296,25 @@ namespace EpicGames.BuildGraph
 		/// <summary>
 		/// Utility method to log a script error at a particular location
 		/// </summary>
-		public static void LogScriptError(this ILogger Logger, BgScriptLocation Location, string Format, params object[] Args)
+		public static void LogScriptError(this ILogger logger, BgScriptLocation location, string format, params object[] args)
 		{
-			object[] AllArgs = new object[Args.Length + 2];
-			AllArgs[0] = Location.NativeFile;
-			AllArgs[1] = Location.LineNumber;
-			Args.CopyTo(AllArgs, 2);
-			Logger.LogError($"{{Script}}({{Line}}): error: {Format}", AllArgs);
+			object[] allArgs = new object[args.Length + 2];
+			allArgs[0] = location.NativeFile;
+			allArgs[1] = location.LineNumber;
+			args.CopyTo(allArgs, 2);
+			logger.LogError($"{{Script}}({{Line}}): error: {format}", allArgs);
 		}
 
 		/// <summary>
 		/// Utility method to log a script warning at a particular location
 		/// </summary>
-		public static void LogScriptWarning(this ILogger Logger, BgScriptLocation Location, string Format, params object[] Args)
+		public static void LogScriptWarning(this ILogger logger, BgScriptLocation location, string format, params object[] args)
 		{
-			object[] AllArgs = new object[Args.Length + 2];
-			AllArgs[0] = Location.NativeFile;
-			AllArgs[1] = Location.LineNumber;
-			Args.CopyTo(AllArgs, 2);
-			Logger.LogWarning($"{{Script}}({{Line}}): warning: {Format}", AllArgs);
+			object[] allArgs = new object[args.Length + 2];
+			allArgs[0] = location.NativeFile;
+			allArgs[1] = location.LineNumber;
+			args.CopyTo(allArgs, 2);
+			logger.LogWarning($"{{Script}}({{Line}}): warning: {format}", allArgs);
 		}
 	}
 
@@ -340,12 +338,12 @@ namespace EpicGames.BuildGraph
 		/// When declaring a property in a nested scope, we enter its name into a set for each parent scope which prevents redeclaration in an OUTER scope later. Subsequent NESTED scopes can redeclare it.
 		/// The former is likely a coding error, since it implies that the scope of the variable was meant to be further out, whereas the latter is common for temporary and loop variables.
 		/// </summary>
-		List<HashSet<string>> ShadowProperties = new List<HashSet<string>>();
+		readonly List<HashSet<string>> _shadowProperties = new List<HashSet<string>>();
 
 		/// <summary>
 		/// Maps from a function name to its definition
 		/// </summary>
-		Dictionary<string, BgScriptMacro> MacroNameToDefinition = new Dictionary<string, BgScriptMacro>();
+		readonly Dictionary<string, BgScriptMacro> _macroNameToDefinition = new Dictionary<string, BgScriptMacro>();
 
 		/// <summary>
 		/// Schema for the script
@@ -365,14 +363,14 @@ namespace EpicGames.BuildGraph
 		/// <summary>
 		/// Private constructor. Use ScriptReader.TryRead() to read a script file.
 		/// </summary>
-		/// <param name="Context">Context object</param>
-		/// <param name="Schema">Schema for the script</param>
-		/// <param name="Logger">Logger for diagnostic messages</param>
-		protected BgScriptReaderBase(IBgScriptReaderContext Context, BgScriptSchema Schema, ILogger Logger)
+		/// <param name="context">Context object</param>
+		/// <param name="schema">Schema for the script</param>
+		/// <param name="logger">Logger for diagnostic messages</param>
+		protected BgScriptReaderBase(IBgScriptReaderContext context, BgScriptSchema schema, ILogger logger)
 		{
-			this.Context = Context;
-			this.Schema = Schema;
-			this.Logger = Logger;
+			Context = context;
+			Schema = schema;
+			Logger = logger;
 
 			EnterScope();
 		}
@@ -380,103 +378,103 @@ namespace EpicGames.BuildGraph
 		/// <summary>
 		/// Read the script from the given file
 		/// </summary>
-		/// <param name="File">File to read from</param>
-		protected async Task<bool> TryReadAsync(string File)
+		/// <param name="file">File to read from</param>
+		protected async Task<bool> TryReadAsync(string file)
 		{
 			// Get the data for this file
-			byte[]? Data = await Context.ReadAsync(File);
-			if (Data == null)
+			byte[]? data = await Context.ReadAsync(file);
+			if (data == null)
 			{
-				Logger.LogError("Unable to open file {File}", File);
+				Logger.LogError("Unable to open file {File}", file);
 				NumErrors++;
 				return false;
 			}
 
 			// Read the document and validate it against the schema
-			BgScriptDocument? Document;
-			if (!BgScriptDocument.TryRead(File, Context.GetNativePath(File), Data, Schema, Logger, out Document))
+			BgScriptDocument? document;
+			if (!BgScriptDocument.TryRead(file, Context.GetNativePath(file), data, Schema, Logger, out document))
 			{
 				NumErrors++;
 				return false;
 			}
 
 			// Read the root BuildGraph element
-			await ReadGraphBodyAsync(Document.DocumentElement);
+			await ReadGraphBodyAsync(document.DocumentElement);
 			return true;
 		}
 
 		/// <summary>
 		/// Reads the contents of a graph
 		/// </summary>
-		/// <param name="Element">The parent element to read from</param>
-		async Task ReadGraphBodyAsync(XmlElement Element)
+		/// <param name="element">The parent element to read from</param>
+		async Task ReadGraphBodyAsync(XmlElement element)
 		{
-			foreach (BgScriptElement ChildElement in Element.ChildNodes.OfType<BgScriptElement>())
+			foreach (BgScriptElement childElement in element.ChildNodes.OfType<BgScriptElement>())
 			{
-				switch (ChildElement.Name)
+				switch (childElement.Name)
 				{
 					case "Include":
-						await ReadIncludeAsync(ChildElement);
+						await ReadIncludeAsync(childElement);
 						break;
 					case "Option":
-						await ReadOptionAsync(ChildElement);
+						await ReadOptionAsync(childElement);
 						break;
 					case "Property":
-						await ReadPropertyAsync(ChildElement);
+						await ReadPropertyAsync(childElement);
 						break;
 					case "Regex":
-						await ReadRegexAsync(ChildElement);
+						await ReadRegexAsync(childElement);
 						break;
 					case "EnvVar":
-						await ReadEnvVarAsync(ChildElement);
+						await ReadEnvVarAsync(childElement);
 						break;
 					case "Macro":
-						ReadMacro(ChildElement);
+						ReadMacro(childElement);
 						break;
 					case "Extend":
-						await ReadExtendAsync(ChildElement);
+						await ReadExtendAsync(childElement);
 						break;
 					case "Agent":
-						await ReadAgentAsync(ChildElement);
+						await ReadAgentAsync(childElement);
 						break;
 					case "Aggregate":
-						await ReadAggregateAsync(ChildElement);
+						await ReadAggregateAsync(childElement);
 						break;
 					case "Report":
-						await ReadReportAsync(ChildElement);
+						await ReadReportAsync(childElement);
 						break;
 					case "Badge":
-						await ReadBadgeAsync(ChildElement);
+						await ReadBadgeAsync(childElement);
 						break;
 					case "Label":
-						await ReadLabelAsync(ChildElement);
+						await ReadLabelAsync(childElement);
 						break;
 					case "Notify":
-						await ReadNotifierAsync(ChildElement);
+						await ReadNotifierAsync(childElement);
 						break;
 					case "Trace":
-						await ReadDiagnosticAsync(ChildElement, LogEventType.Console);
+						await ReadDiagnosticAsync(childElement, LogEventType.Console);
 						break;
 					case "Warning":
-						await ReadDiagnosticAsync(ChildElement, LogEventType.Warning);
+						await ReadDiagnosticAsync(childElement, LogEventType.Warning);
 						break;
 					case "Error":
-						await ReadDiagnosticAsync(ChildElement, LogEventType.Error);
+						await ReadDiagnosticAsync(childElement, LogEventType.Error);
 						break;
 					case "Do":
-						await ReadBlockAsync(ChildElement, ReadGraphBodyAsync);
+						await ReadBlockAsync(childElement, ReadGraphBodyAsync);
 						break;
 					case "Switch":
-						await ReadSwitchAsync(ChildElement, ReadGraphBodyAsync);
+						await ReadSwitchAsync(childElement, ReadGraphBodyAsync);
 						break;
 					case "ForEach":
-						await ReadForEachAsync(ChildElement, ReadGraphBodyAsync);
+						await ReadForEachAsync(childElement, ReadGraphBodyAsync);
 						break;
 					case "Expand":
-						await ReadExpandAsync(ChildElement, ReadGraphBodyAsync);
+						await ReadExpandAsync(childElement, ReadGraphBodyAsync);
 						break;
 					default:
-						LogError(ChildElement, "Invalid element '{0}'", ChildElement.Name);
+						LogError(childElement, "Invalid element '{0}'", childElement.Name);
 						break;
 				}
 			}
@@ -488,7 +486,7 @@ namespace EpicGames.BuildGraph
 		protected void EnterScope()
 		{
 			ScopedProperties.Add(new Dictionary<string, string>(StringComparer.OrdinalIgnoreCase));
-			ShadowProperties.Add(new HashSet<string>(StringComparer.OrdinalIgnoreCase));
+			_shadowProperties.Add(new HashSet<string>(StringComparer.OrdinalIgnoreCase));
 		}
 
 		/// <summary>
@@ -497,75 +495,75 @@ namespace EpicGames.BuildGraph
 		protected void LeaveScope()
 		{
 			ScopedProperties.RemoveAt(ScopedProperties.Count - 1);
-			ShadowProperties.RemoveAt(ShadowProperties.Count - 1);
+			_shadowProperties.RemoveAt(_shadowProperties.Count - 1);
 		}
 
 		/// <summary>
 		/// Sets a property value in the current scope
 		/// </summary>
-		/// <param name="Element">Element containing the property assignment. Used for error messages if the property is shadowed in another scope.</param>
-		/// <param name="Name">Name of the property</param>
-		/// <param name="Value">Value for the property</param>
-		protected void SetPropertyValue(BgScriptElement Element, string Name, string Value)
+		/// <param name="element">Element containing the property assignment. Used for error messages if the property is shadowed in another scope.</param>
+		/// <param name="name">Name of the property</param>
+		/// <param name="value">Value for the property</param>
+		protected void SetPropertyValue(BgScriptElement element, string name, string value)
 		{
 			// Find the scope containing this property, defaulting to the current scope
-			int ScopeIdx = 0;
-			while(ScopeIdx < ScopedProperties.Count - 1 && !ScopedProperties[ScopeIdx].ContainsKey(Name))
+			int scopeIdx = 0;
+			while (scopeIdx < ScopedProperties.Count - 1 && !ScopedProperties[scopeIdx].ContainsKey(name))
 			{
-				ScopeIdx++;
+				scopeIdx++;
 			}
 
 			// Make sure this property name was not already used in a child scope; it likely indicates an error.
-			if(ShadowProperties[ScopeIdx].Contains(Name))
+			if (_shadowProperties[scopeIdx].Contains(name))
 			{
-				LogError(Element, "Property '{0}' was already used in a child scope. Move this definition before the previous usage if they are intended to share scope, or use a different name.", Name);
+				LogError(element, "Property '{0}' was already used in a child scope. Move this definition before the previous usage if they are intended to share scope, or use a different name.", name);
 			}
 			else
 			{
 				// Make sure it's added to the shadow property list for every parent scope
-				for(int Idx = 0; Idx < ScopeIdx; Idx++)
+				for (int idx = 0; idx < scopeIdx; idx++)
 				{
-					ShadowProperties[Idx].Add(Name);
+					_shadowProperties[idx].Add(name);
 				}
-				ScopedProperties[ScopeIdx][Name] = Value;
+				ScopedProperties[scopeIdx][name] = value;
 			}
 		}
 
 		/// <summary>
 		/// Tries to get the value of a property
 		/// </summary>
-		/// <param name="Name">Name of the property</param>
-		/// <param name="Value">On success, contains the value of the property. Set to null otherwise.</param>
+		/// <param name="name">Name of the property</param>
+		/// <param name="value">On success, contains the value of the property. Set to null otherwise.</param>
 		/// <returns>True if the property was found, false otherwise</returns>
-		protected bool TryGetPropertyValue(string Name, out string? Value)
+		protected bool TryGetPropertyValue(string name, out string? value)
 		{
-			int ValueLength = 0;
-			if (Name.Contains(":"))
+			int valueLength = 0;
+			if (name.Contains(":"))
 			{
-				string[] Tokens = Name.Split(':');
-				Name = Tokens[0];
-				ValueLength = int.Parse(Tokens[1]);
+				string[] tokens = name.Split(':');
+				name = tokens[0];
+				valueLength = int.Parse(tokens[1]);
 			}
 
 			// Check each scope for the property
-			for (int ScopeIdx = ScopedProperties.Count - 1; ScopeIdx >= 0; ScopeIdx--)
+			for (int scopeIdx = ScopedProperties.Count - 1; scopeIdx >= 0; scopeIdx--)
 			{
-				string? ScopeValue;
-				if (ScopedProperties[ScopeIdx].TryGetValue(Name, out ScopeValue))
+				string? scopeValue;
+				if (ScopedProperties[scopeIdx].TryGetValue(name, out scopeValue))
 				{
-					Value = ScopeValue;
+					value = scopeValue;
 
 					// It's valid for a property to exist but have a null value. It won't be expanded
 					// Handle $(PropName:-6) where PropName might be "Foo"
-					if (Value != null && Value.Length > Math.Abs(ValueLength))
+					if (value != null && value.Length > Math.Abs(valueLength))
 					{
-						if (ValueLength > 0)
+						if (valueLength > 0)
 						{
-							Value = Value.Substring(0, ValueLength);
+							value = value.Substring(0, valueLength);
 						}
-						if (ValueLength < 0)
+						if (valueLength < 0)
 						{
-							Value = Value.Substring(Value.Length + ValueLength, -ValueLength);
+							value = value.Substring(value.Length + valueLength, -valueLength);
 						}
 					}
 					return true;
@@ -573,66 +571,66 @@ namespace EpicGames.BuildGraph
 			}
 
 			// If we didn't find it, return false.
-			Value = null;
+			value = null;
 			return false;
 		}
 
-		static string CombinePaths(string BasePath, string NextPath)
+		static string CombinePaths(string basePath, string nextPath)
 		{
-			List<string> Fragments = new List<string>(BasePath.Split('/'));
-			Fragments.RemoveAt(Fragments.Count - 1);
+			List<string> fragments = new List<string>(basePath.Split('/'));
+			fragments.RemoveAt(fragments.Count - 1);
 
-			foreach (string AppendFragment in NextPath.Split('/'))
+			foreach (string appendFragment in nextPath.Split('/'))
 			{
-				if (AppendFragment.Equals(".", StringComparison.Ordinal))
+				if (appendFragment.Equals(".", StringComparison.Ordinal))
 				{
 					continue;
 				}
-				else if (AppendFragment.Equals("..", StringComparison.Ordinal))
+				else if (appendFragment.Equals("..", StringComparison.Ordinal))
 				{
-					if (Fragments.Count > 0)
+					if (fragments.Count > 0)
 					{
-						Fragments.RemoveAt(Fragments.Count - 1);
+						fragments.RemoveAt(fragments.Count - 1);
 					}
 					else
 					{
-						throw new Exception($"Path '{NextPath}' cannot be combined with '{BasePath}'");
+						throw new Exception($"Path '{nextPath}' cannot be combined with '{basePath}'");
 					}
 				}
 				else
 				{
-					Fragments.Add(AppendFragment);
+					fragments.Add(appendFragment);
 				}
 			}
-			return String.Join('/', Fragments);
+			return String.Join('/', fragments);
 		}
 
 		/// <summary>
 		/// Read an include directive, and the contents of the target file
 		/// </summary>
-		/// <param name="Element">Xml element to read the definition from</param>
-		async Task ReadIncludeAsync(BgScriptElement Element)
+		/// <param name="element">Xml element to read the definition from</param>
+		async Task ReadIncludeAsync(BgScriptElement element)
 		{
-			if (await EvaluateConditionAsync(Element))
+			if (await EvaluateConditionAsync(element))
 			{
-				HashSet<string> Files = new HashSet<string>();
-				foreach (string Script in ReadListAttribute(Element, "Script"))
+				HashSet<string> files = new HashSet<string>();
+				foreach (string script in ReadListAttribute(element, "Script"))
 				{
-					string IncludePath = CombinePaths(Element.Location.File, Script);
-					if (Regex.IsMatch(IncludePath, @"\*|\?|\.\.\."))
+					string includePath = CombinePaths(element.Location.File, script);
+					if (Regex.IsMatch(includePath, @"\*|\?|\.\.\."))
 					{
-						Files.UnionWith(await Context.FindAsync(IncludePath));
+						files.UnionWith(await Context.FindAsync(includePath));
 					}
 					else
 					{
-						Files.Add(IncludePath);
+						files.Add(includePath);
 					}
 				}
 
-				foreach(string File in Files.OrderBy(x => x, StringComparer.OrdinalIgnoreCase))
+				foreach (string file in files.OrderBy(x => x, StringComparer.OrdinalIgnoreCase))
 				{
-					Logger.LogDebug("Including file {File}", File);
-					await TryReadAsync(File);
+					Logger.LogDebug("Including file {File}", file);
+					await TryReadAsync(file);
 				}
 			}
 		}
@@ -640,51 +638,51 @@ namespace EpicGames.BuildGraph
 		/// <summary>
 		/// Reads the definition of a graph option; a parameter which can be set by the user on the command-line or via an environment variable.
 		/// </summary>
-		/// <param name="Element">Xml element to read the definition from</param>
-		protected abstract Task ReadOptionAsync(BgScriptElement Element);
+		/// <param name="element">Xml element to read the definition from</param>
+		protected abstract Task ReadOptionAsync(BgScriptElement element);
 
 		/// <summary>
 		/// Reads a property assignment.
 		/// </summary>
-		/// <param name="Element">Xml element to read the definition from</param>
-		protected abstract Task ReadPropertyAsync(BgScriptElement Element);
+		/// <param name="element">Xml element to read the definition from</param>
+		protected abstract Task ReadPropertyAsync(BgScriptElement element);
 
 		/// <summary>
 		/// Reads a Regex assignment.
 		/// </summary>
-		/// <param name="Element">Xml element to read the definition from</param>
-		protected abstract Task ReadRegexAsync(BgScriptElement Element);
+		/// <param name="element">Xml element to read the definition from</param>
+		protected abstract Task ReadRegexAsync(BgScriptElement element);
 
 		/// <summary>
 		/// Reads a property assignment from an environment variable.
 		/// </summary>
-		/// <param name="Element">Xml element to read the definition from</param>
-		protected abstract Task ReadEnvVarAsync(BgScriptElement Element);
+		/// <param name="element">Xml element to read the definition from</param>
+		protected abstract Task ReadEnvVarAsync(BgScriptElement element);
 
 		/// <summary>
 		/// Reads a macro definition
 		/// </summary>
-		/// <param name="Element">Xml element to read the definition from</param>
-		void ReadMacro(BgScriptElement Element)
+		/// <param name="element">Xml element to read the definition from</param>
+		void ReadMacro(BgScriptElement element)
 		{
-			string Name = Element.GetAttribute("Name");
-			if (ValidateName(Element, Name))
+			string name = element.GetAttribute("Name");
+			if (ValidateName(element, name))
 			{
-				BgScriptMacro? OriginalDefinition;
-				if(MacroNameToDefinition.TryGetValue(Name, out OriginalDefinition))
+				BgScriptMacro? originalDefinition;
+				if (_macroNameToDefinition.TryGetValue(name, out originalDefinition))
 				{
-					BgScriptLocation Location = OriginalDefinition.Elements[0].Location;
-					LogError(Element, "Macro '{0}' has already been declared (see {1} line {2})", Name, Location.File, Location.LineNumber);
+					BgScriptLocation location = originalDefinition.Elements[0].Location;
+					LogError(element, "Macro '{0}' has already been declared (see {1} line {2})", name, location.File, location.LineNumber);
 				}
 				else
 				{
-					Dictionary<string, int> ArgumentNameToIndex = new Dictionary<string, int>();
-					ReadMacroArguments(Element, "Arguments", ArgumentNameToIndex);
+					Dictionary<string, int> argumentNameToIndex = new Dictionary<string, int>();
+					ReadMacroArguments(element, "Arguments", argumentNameToIndex);
 
-					int NumRequiredArguments = ArgumentNameToIndex.Count;
-					ReadMacroArguments(Element, "OptionalArguments", ArgumentNameToIndex);
+					int numRequiredArguments = argumentNameToIndex.Count;
+					ReadMacroArguments(element, "OptionalArguments", argumentNameToIndex);
 
-					MacroNameToDefinition.Add(Name, new BgScriptMacro(Name, Element, ArgumentNameToIndex, NumRequiredArguments));
+					_macroNameToDefinition.Add(name, new BgScriptMacro(name, element, argumentNameToIndex, numRequiredArguments));
 				}
 			}
 		}
@@ -692,23 +690,23 @@ namespace EpicGames.BuildGraph
 		/// <summary>
 		/// Reads a list of macro arguments from an attribute
 		/// </summary>
-		/// <param name="Element">The element containing the attributes</param>
-		/// <param name="AttributeName">Name of the attribute containing the arguments</param>
-		/// <param name="ArgumentNameToIndex">List of arguments to add to</param>
-		void ReadMacroArguments(BgScriptElement Element, string AttributeName, Dictionary<string, int> ArgumentNameToIndex)
+		/// <param name="element">The element containing the attributes</param>
+		/// <param name="attributeName">Name of the attribute containing the arguments</param>
+		/// <param name="argumentNameToIndex">List of arguments to add to</param>
+		void ReadMacroArguments(BgScriptElement element, string attributeName, Dictionary<string, int> argumentNameToIndex)
 		{
-			string AttributeValue = ReadAttribute(Element, AttributeName);
-			if(AttributeValue != null)
+			string attributeValue = ReadAttribute(element, attributeName);
+			if (attributeValue != null)
 			{
-				foreach(string ArgumentName in AttributeValue.Split(new char[]{ ';' }, StringSplitOptions.RemoveEmptyEntries))
+				foreach (string argumentName in attributeValue.Split(new char[] { ';' }, StringSplitOptions.RemoveEmptyEntries))
 				{
-					if(ArgumentNameToIndex.ContainsKey(ArgumentName))
+					if (argumentNameToIndex.ContainsKey(argumentName))
 					{
-						LogWarning(Element, "Argument '{0}' is listed multiple times", ArgumentName);
+						LogWarning(element, "Argument '{0}' is listed multiple times", argumentName);
 					}
 					else
 					{
-						ArgumentNameToIndex.Add(ArgumentName, ArgumentNameToIndex.Count);
+						argumentNameToIndex.Add(argumentName, argumentNameToIndex.Count);
 					}
 				}
 			}
@@ -717,21 +715,21 @@ namespace EpicGames.BuildGraph
 		/// <summary>
 		/// Reads a macro definition
 		/// </summary>
-		/// <param name="Element">Xml element to read the definition from</param>
-		async Task ReadExtendAsync(BgScriptElement Element)
+		/// <param name="element">Xml element to read the definition from</param>
+		async Task ReadExtendAsync(BgScriptElement element)
 		{
-			if (await EvaluateConditionAsync(Element))
+			if (await EvaluateConditionAsync(element))
 			{
-				string Name = ReadAttribute(Element, "Name");
+				string name = ReadAttribute(element, "Name");
 
-				BgScriptMacro? OriginalDefinition;
-				if (MacroNameToDefinition.TryGetValue(Name, out OriginalDefinition))
+				BgScriptMacro? originalDefinition;
+				if (_macroNameToDefinition.TryGetValue(name, out originalDefinition))
 				{
-					OriginalDefinition.Elements.Add(Element);
+					originalDefinition.Elements.Add(element);
 				}
 				else
 				{
-					LogError(Element, "Macro '{0}' has not been declared", Name);
+					LogError(element, "Macro '{0}' has not been declared", name);
 				}
 			}
 		}
@@ -739,58 +737,58 @@ namespace EpicGames.BuildGraph
 		/// <summary>
 		/// Reads the definition for an agent.
 		/// </summary>
-		/// <param name="Element">Xml element to read the definition from</param>
-		protected abstract Task ReadAgentAsync(BgScriptElement Element);
+		/// <param name="element">Xml element to read the definition from</param>
+		protected abstract Task ReadAgentAsync(BgScriptElement element);
 
 		/// <summary>
 		/// Read the contents of an agent definition
 		/// </summary>
-		/// <param name="Element">Xml element to read the definition from</param>
-		protected async Task ReadAgentBodyAsync(BgScriptElement Element)
+		/// <param name="element">Xml element to read the definition from</param>
+		protected async Task ReadAgentBodyAsync(BgScriptElement element)
 		{
 			EnterScope();
-			foreach (BgScriptElement ChildElement in Element.ChildNodes.OfType<BgScriptElement>())
+			foreach (BgScriptElement childElement in element.ChildNodes.OfType<BgScriptElement>())
 			{
-				switch (ChildElement.Name)
+				switch (childElement.Name)
 				{
 					case "Property":
-						await ReadPropertyAsync(ChildElement);
+						await ReadPropertyAsync(childElement);
 						break;
 					case "Regex":
-						await ReadRegexAsync(ChildElement);
+						await ReadRegexAsync(childElement);
 						break;
 					case "Node":
-						await ReadNodeAsync(ChildElement);
+						await ReadNodeAsync(childElement);
 						break;
 					case "Aggregate":
-						await ReadAggregateAsync(ChildElement);
+						await ReadAggregateAsync(childElement);
 						break;
 					case "Trace":
-						await ReadDiagnosticAsync(ChildElement, LogEventType.Console);
+						await ReadDiagnosticAsync(childElement, LogEventType.Console);
 						break;
 					case "Warning":
-						await ReadDiagnosticAsync(ChildElement, LogEventType.Warning);
+						await ReadDiagnosticAsync(childElement, LogEventType.Warning);
 						break;
 					case "Error":
-						await ReadDiagnosticAsync(ChildElement, LogEventType.Error);
+						await ReadDiagnosticAsync(childElement, LogEventType.Error);
 						break;
 					case "Label":
-						await ReadLabelAsync(ChildElement);
+						await ReadLabelAsync(childElement);
 						break;
 					case "Do":
-						await ReadBlockAsync(ChildElement, ReadAgentBodyAsync);
+						await ReadBlockAsync(childElement, ReadAgentBodyAsync);
 						break;
 					case "Switch":
-						await ReadSwitchAsync(ChildElement, ReadAgentBodyAsync);
+						await ReadSwitchAsync(childElement, ReadAgentBodyAsync);
 						break;
 					case "ForEach":
-						await ReadForEachAsync(ChildElement, ReadAgentBodyAsync);
+						await ReadForEachAsync(childElement, ReadAgentBodyAsync);
 						break;
 					case "Expand":
-						await ReadExpandAsync(ChildElement, ReadAgentBodyAsync);
+						await ReadExpandAsync(childElement, ReadAgentBodyAsync);
 						break;
 					default:
-						LogError(ChildElement, "Unexpected element type '{0}'", ChildElement.Name);
+						LogError(childElement, "Unexpected element type '{0}'", childElement.Name);
 						break;
 				}
 			}
@@ -800,73 +798,73 @@ namespace EpicGames.BuildGraph
 		/// <summary>
 		/// Reads the definition for an aggregate
 		/// </summary>
-		/// <param name="Element">Xml element to read the definition from</param>
-		protected abstract Task ReadAggregateAsync(BgScriptElement Element);
+		/// <param name="element">Xml element to read the definition from</param>
+		protected abstract Task ReadAggregateAsync(BgScriptElement element);
 
 		/// <summary>
 		/// Reads the definition for a report
 		/// </summary>
-		/// <param name="Element">Xml element to read the definition from</param>
-		protected abstract Task ReadReportAsync(BgScriptElement Element);
+		/// <param name="element">Xml element to read the definition from</param>
+		protected abstract Task ReadReportAsync(BgScriptElement element);
 
 		/// <summary>
 		/// Reads the definition for a badge
 		/// </summary>
-		/// <param name="Element">Xml element to read the definition from</param>
-		protected abstract Task ReadBadgeAsync(BgScriptElement Element);
+		/// <param name="element">Xml element to read the definition from</param>
+		protected abstract Task ReadBadgeAsync(BgScriptElement element);
 
 		/// <summary>
 		/// Reads the definition for a label
 		/// </summary>
-		/// <param name="Element">Xml element to read the definition from</param>
-		protected abstract Task ReadLabelAsync(BgScriptElement Element);
+		/// <param name="element">Xml element to read the definition from</param>
+		protected abstract Task ReadLabelAsync(BgScriptElement element);
 
 		/// <summary>
 		/// Reads the definition for a node, and adds it to the given agent
 		/// </summary>
-		/// <param name="Element">Xml element to read the definition from</param>
-		protected abstract Task ReadNodeAsync(BgScriptElement Element);
+		/// <param name="element">Xml element to read the definition from</param>
+		protected abstract Task ReadNodeAsync(BgScriptElement element);
 
 		/// <summary>
 		/// Reads the contents of a node element
 		/// </summary>
-		/// <param name="Element">Xml element to read the definition from</param>
-		protected async Task ReadNodeBodyAsync(XmlElement Element)
+		/// <param name="element">Xml element to read the definition from</param>
+		protected async Task ReadNodeBodyAsync(XmlElement element)
 		{
 			EnterScope();
-			foreach (BgScriptElement ChildElement in Element.ChildNodes.OfType<BgScriptElement>())
+			foreach (BgScriptElement childElement in element.ChildNodes.OfType<BgScriptElement>())
 			{
-				switch (ChildElement.Name)
+				switch (childElement.Name)
 				{
 					case "Property":
-						await ReadPropertyAsync(ChildElement);
+						await ReadPropertyAsync(childElement);
 						break;
 					case "Regex":
-						await ReadRegexAsync(ChildElement);
+						await ReadRegexAsync(childElement);
 						break;
 					case "Trace":
-						await ReadDiagnosticAsync(ChildElement, LogEventType.Console);
+						await ReadDiagnosticAsync(childElement, LogEventType.Console);
 						break;
 					case "Warning":
-						await ReadDiagnosticAsync(ChildElement, LogEventType.Warning);
+						await ReadDiagnosticAsync(childElement, LogEventType.Warning);
 						break;
 					case "Error":
-						await ReadDiagnosticAsync(ChildElement, LogEventType.Error);
+						await ReadDiagnosticAsync(childElement, LogEventType.Error);
 						break;
 					case "Do":
-						await ReadBlockAsync(ChildElement, ReadNodeBodyAsync);
+						await ReadBlockAsync(childElement, ReadNodeBodyAsync);
 						break;
 					case "Switch":
-						await ReadSwitchAsync(ChildElement, ReadNodeBodyAsync);
+						await ReadSwitchAsync(childElement, ReadNodeBodyAsync);
 						break;
 					case "ForEach":
-						await ReadForEachAsync(ChildElement, ReadNodeBodyAsync);
+						await ReadForEachAsync(childElement, ReadNodeBodyAsync);
 						break;
 					case "Expand":
-						await ReadExpandAsync(ChildElement, ReadNodeBodyAsync);
+						await ReadExpandAsync(childElement, ReadNodeBodyAsync);
 						break;
 					default:
-						await ReadTaskAsync(ChildElement);
+						await ReadTaskAsync(childElement);
 						break;
 				}
 			}
@@ -876,58 +874,58 @@ namespace EpicGames.BuildGraph
 		/// <summary>
 		/// Reads a block element
 		/// </summary>
-		/// <param name="Element">Xml element to read the definition from</param>
-		/// <param name="ReadContentsAsync">Delegate to read the contents of the element, if the condition evaluates to true</param>
-		async Task ReadBlockAsync(BgScriptElement Element, Func<BgScriptElement, Task> ReadContentsAsync)
+		/// <param name="element">Xml element to read the definition from</param>
+		/// <param name="readContentsAsync">Delegate to read the contents of the element, if the condition evaluates to true</param>
+		async Task ReadBlockAsync(BgScriptElement element, Func<BgScriptElement, Task> readContentsAsync)
 		{
-			if (await EvaluateConditionAsync(Element))
+			if (await EvaluateConditionAsync(element))
 			{
-				await ReadContentsAsync(Element);
+				await readContentsAsync(element);
 			}
 		}
 
 		/// <summary>
 		/// Reads a "Switch" element 
 		/// </summary>
-		/// <param name="Element">Xml element to read the definition from</param>
-		/// <param name="ReadContentsAsync">Delegate to read the contents of the element, if the condition evaluates to true</param>
-		protected abstract Task ReadSwitchAsync(BgScriptElement Element, Func<BgScriptElement, Task> ReadContentsAsync);
+		/// <param name="element">Xml element to read the definition from</param>
+		/// <param name="readContentsAsync">Delegate to read the contents of the element, if the condition evaluates to true</param>
+		protected abstract Task ReadSwitchAsync(BgScriptElement element, Func<BgScriptElement, Task> readContentsAsync);
 
 		/// <summary>
 		/// Reads a "ForEach" element 
 		/// </summary>
-		/// <param name="Element">Xml element to read the definition from</param>
-		/// <param name="ReadContentsAsync">Delegate to read the contents of the element, if the condition evaluates to true</param>
-		async Task ReadForEachAsync(BgScriptElement Element, Func<BgScriptElement, Task> ReadContentsAsync)
+		/// <param name="element">Xml element to read the definition from</param>
+		/// <param name="readContentsAsync">Delegate to read the contents of the element, if the condition evaluates to true</param>
+		async Task ReadForEachAsync(BgScriptElement element, Func<BgScriptElement, Task> readContentsAsync)
 		{
 			EnterScope();
-			if(await EvaluateConditionAsync(Element))
+			if (await EvaluateConditionAsync(element))
 			{
-				string Name = ReadAttribute(Element, "Name");
-				string Separator = ReadAttribute(Element, "Separator");
-				if(Separator.Length > 1)
+				string name = ReadAttribute(element, "Name");
+				string separator = ReadAttribute(element, "Separator");
+				if (separator.Length > 1)
 				{
-					LogWarning(Element, "Node {0}'s Separator attribute is more than one character ({1}). Defaulting to ;", Name, Separator);
-					Separator = ";";
+					LogWarning(element, "Node {0}'s Separator attribute is more than one character ({1}). Defaulting to ;", name, separator);
+					separator = ";";
 				}
-				if(string.IsNullOrEmpty(Separator))
+				if (string.IsNullOrEmpty(separator))
 				{
-					Separator = ";";
+					separator = ";";
 				}
-				if (ValidateName(Element, Name))
+				if (ValidateName(element, name))
 				{
-					if(ScopedProperties.Any(x => x.ContainsKey(Name)))
+					if (ScopedProperties.Any(x => x.ContainsKey(name)))
 					{
-						LogError(Element, "Loop variable '{0}' already exists as a local property in an outer scope", Name);
+						LogError(element, "Loop variable '{0}' already exists as a local property in an outer scope", name);
 					}
 					else
 					{
 						// Loop through all the values
-						string[] Values = ReadListAttribute(Element, "Values", Convert.ToChar(Separator));
-						foreach(string Value in Values)
+						string[] values = ReadListAttribute(element, "Values", Convert.ToChar(separator));
+						foreach (string value in values)
 						{
-							ScopedProperties[ScopedProperties.Count - 1][Name] = Value;
-							await ReadContentsAsync(Element);
+							ScopedProperties[ScopedProperties.Count - 1][name] = value;
+							await readContentsAsync(element);
 						}
 					}
 				}
@@ -938,47 +936,47 @@ namespace EpicGames.BuildGraph
 		/// <summary>
 		/// Reads an "Expand" element 
 		/// </summary>
-		/// <param name="Element">Xml element to read the definition from</param>
-		/// <param name="ReadContentsAsync">Delegate to read the contents of the element, if the condition evaluates to true</param>
-		async Task ReadExpandAsync(BgScriptElement Element, Func<BgScriptElement, Task> ReadContentsAsync)
+		/// <param name="element">Xml element to read the definition from</param>
+		/// <param name="readContentsAsync">Delegate to read the contents of the element, if the condition evaluates to true</param>
+		async Task ReadExpandAsync(BgScriptElement element, Func<BgScriptElement, Task> readContentsAsync)
 		{
-			if(await EvaluateConditionAsync(Element))
+			if (await EvaluateConditionAsync(element))
 			{
-				string Name = ReadAttribute(Element, "Name");
-				if (ValidateName(Element, Name))
+				string name = ReadAttribute(element, "Name");
+				if (ValidateName(element, name))
 				{
-					BgScriptMacro? Macro;
-					if (!MacroNameToDefinition.TryGetValue(Name, out Macro))
+					BgScriptMacro? macro;
+					if (!_macroNameToDefinition.TryGetValue(name, out macro))
 					{
-						LogError(Element, "Macro '{0}' does not exist", Name);
+						LogError(element, "Macro '{0}' does not exist", name);
 					}
 					else
 					{
 						// Parse the argument list
-						string[] Arguments = new string[Macro.ArgumentNameToIndex.Count];
-						foreach (XmlAttribute? Attribute in Element.Attributes)
+						string[] arguments = new string[macro.ArgumentNameToIndex.Count];
+						foreach (XmlAttribute? attribute in element.Attributes)
 						{
-							if (Attribute != null && Attribute.Name != "Name" && Attribute.Name != "If")
+							if (attribute != null && attribute.Name != "Name" && attribute.Name != "If")
 							{
-								int Index;
-								if (Macro.ArgumentNameToIndex.TryGetValue(Attribute.Name, out Index))
+								int index;
+								if (macro.ArgumentNameToIndex.TryGetValue(attribute.Name, out index))
 								{
-									Arguments[Index] = ExpandProperties(Element, Attribute.Value);
+									arguments[index] = ExpandProperties(element, attribute.Value);
 								}
 								else
 								{
-									LogWarning(Element, "Macro '{0}' does not take an argument '{1}'", Name, Attribute.Name);
+									LogWarning(element, "Macro '{0}' does not take an argument '{1}'", name, attribute.Name);
 								}
 							}
 						}
 
 						// Make sure none of the required arguments are missing
 						bool bHasMissingArguments = false;
-						for (int Idx = 0; Idx < Macro.NumRequiredArguments; Idx++)
+						for (int idx = 0; idx < macro.NumRequiredArguments; idx++)
 						{
-							if (Arguments[Idx] == null)
+							if (arguments[idx] == null)
 							{
-								LogWarning(Element, "Macro '{0}' is missing argument '{1}'", Macro.Name, Macro.ArgumentNameToIndex.First(x => x.Value == Idx).Key);
+								LogWarning(element, "Macro '{0}' is missing argument '{1}'", macro.Name, macro.ArgumentNameToIndex.First(x => x.Value == idx).Key);
 								bHasMissingArguments = true;
 							}
 						}
@@ -987,13 +985,13 @@ namespace EpicGames.BuildGraph
 						if (!bHasMissingArguments)
 						{
 							EnterScope();
-							foreach (KeyValuePair<string, int> Pair in Macro.ArgumentNameToIndex)
+							foreach (KeyValuePair<string, int> pair in macro.ArgumentNameToIndex)
 							{
-								ScopedProperties[ScopedProperties.Count - 1][Pair.Key] = Arguments[Pair.Value] ?? "";
+								ScopedProperties[ScopedProperties.Count - 1][pair.Key] = arguments[pair.Value] ?? "";
 							}
-							foreach (BgScriptElement MacroElement in Macro.Elements)
+							foreach (BgScriptElement macroElement in macro.Elements)
 							{
-								await ReadContentsAsync(MacroElement);
+								await readContentsAsync(macroElement);
 							}
 							LeaveScope();
 						}
@@ -1005,77 +1003,77 @@ namespace EpicGames.BuildGraph
 		/// <summary>
 		/// Reads a task definition from the given element, and add it to the given list
 		/// </summary>
-		/// <param name="Element">Xml element to read the definition from</param>
-		protected abstract Task ReadTaskAsync(BgScriptElement Element);
+		/// <param name="element">Xml element to read the definition from</param>
+		protected abstract Task ReadTaskAsync(BgScriptElement element);
 
 		/// <summary>
 		/// Reads the definition for an email notifier
 		/// </summary>
-		/// <param name="Element">Xml element to read the definition from</param>
-		protected abstract Task ReadNotifierAsync(BgScriptElement Element);
+		/// <param name="element">Xml element to read the definition from</param>
+		protected abstract Task ReadNotifierAsync(BgScriptElement element);
 
 		/// <summary>
 		/// Reads a warning from the given element, evaluates the condition on it, and writes it to the log if the condition passes.
 		/// </summary>
-		/// <param name="Element">Xml element to read the definition from</param>
-		/// <param name="EventType">The diagnostic event type</param>
-		protected abstract Task ReadDiagnosticAsync(BgScriptElement Element, LogEventType EventType);
+		/// <param name="element">Xml element to read the definition from</param>
+		/// <param name="eventType">The diagnostic event type</param>
+		protected abstract Task ReadDiagnosticAsync(BgScriptElement element, LogEventType eventType);
 
 		/// <summary>
 		/// Reads an object name from its defining element. Outputs an error if the name is missing.
 		/// </summary>
-		/// <param name="Element">Element to read the name for</param>
-		/// <param name="Name">Output variable to receive the name of the object</param>
+		/// <param name="element">Element to read the name for</param>
+		/// <param name="name">Output variable to receive the name of the object</param>
 		/// <returns>True if the object had a valid name (assigned to the Name variable), false if the name was invalid or missing.</returns>
-		protected bool TryReadObjectName(BgScriptElement Element, [NotNullWhen(true)] out string? Name)
+		protected bool TryReadObjectName(BgScriptElement element, [NotNullWhen(true)] out string? name)
 		{
 			// Check the name attribute is present
-			if (!Element.HasAttribute("Name"))
+			if (!element.HasAttribute("Name"))
 			{
-				LogError(Element, "Missing 'Name' attribute");
-				Name = null;
+				LogError(element, "Missing 'Name' attribute");
+				name = null;
 				return false;
 			}
 
 			// Get the value of it, strip any leading or trailing whitespace, and make sure it's not empty
-			string Value = ReadAttribute(Element, "Name");
-			if (!ValidateName(Element, Value))
+			string value = ReadAttribute(element, "Name");
+			if (!ValidateName(element, value))
 			{
-				Name = null;
+				name = null;
 				return false;
 			}
 
 			// Return it
-			Name = Value;
+			name = value;
 			return true;
 		}
 
 		/// <summary>
 		/// Checks that the given name is valid syntax
 		/// </summary>
-		/// <param name="Element">The element that contains the name</param>
-		/// <param name="Name">The name to check</param>
+		/// <param name="element">The element that contains the name</param>
+		/// <param name="name">The name to check</param>
 		/// <returns>True if the name is valid</returns>
-		protected bool ValidateName(BgScriptElement Element, string Name)
+		protected bool ValidateName(BgScriptElement element, string name)
 		{
 			// Check it's not empty
-			if (Name.Length == 0)
+			if (name.Length == 0)
 			{
-				LogError(Element, "Name is empty");
+				LogError(element, "Name is empty");
 				return false;
 			}
 
 			// Check there are no invalid characters
-			for (int Idx = 0; Idx < Name.Length; Idx++)
+			for (int idx = 0; idx < name.Length; idx++)
 			{
-				if (Idx > 0 && Name[Idx] == ' ' && Name[Idx - 1] == ' ')
+				if (idx > 0 && name[idx] == ' ' && name[idx - 1] == ' ')
 				{
-					LogError(Element, "Consecutive spaces in object name - '{0}'", Name);
+					LogError(element, "Consecutive spaces in object name - '{0}'", name);
 					return false;
 				}
-				if(Char.IsControl(Name[Idx]) || BgScriptSchema.IllegalNameCharacters.IndexOf(Name[Idx]) != -1)
+				if (Char.IsControl(name[idx]) || BgScriptSchema.IllegalNameCharacters.IndexOf(name[idx]) != -1)
 				{
-					LogError(Element, "Invalid character in object name - '{0}'", Name[Idx]);
+					LogError(element, "Invalid character in object name - '{0}'", name[idx]);
 					return false;
 				}
 			}
@@ -1085,23 +1083,23 @@ namespace EpicGames.BuildGraph
 		/// <summary>
 		/// Constructs a regex from a regex string and returns it
 		/// </summary>
-		/// <param name="Element">The element that contains the regex</param>
-		/// <param name="Regex">The pattern to construct</param>
+		/// <param name="element">The element that contains the regex</param>
+		/// <param name="regex">The pattern to construct</param>
 		/// <returns>The regex if is valid, otherwise null</returns>
-		protected Regex? ParseRegex(BgScriptElement Element, string Regex)
+		protected Regex? ParseRegex(BgScriptElement element, string regex)
 		{
-			if(Regex.Length == 0)
+			if (regex.Length == 0)
 			{
-				LogError(Element, "Regex is empty");
+				LogError(element, "Regex is empty");
 				return null;
 			}
 			try
 			{
-				return new Regex(Regex);
+				return new Regex(regex);
 			}
-			catch(ArgumentException InvalidRegex)
+			catch (ArgumentException invalidRegex)
 			{
-				LogError(Element, "Could not construct the Regex, reason: {0}", InvalidRegex.Message);
+				LogError(element, "Could not construct the Regex, reason: {0}", invalidRegex.Message);
 				return null;
 			}
 		}
@@ -1109,51 +1107,51 @@ namespace EpicGames.BuildGraph
 		/// <summary>
 		/// Expands any properties and reads an attribute.
 		/// </summary>
-		/// <param name="Element">Element to read the attribute from</param>
-		/// <param name="Name">Name of the attribute</param>
+		/// <param name="element">Element to read the attribute from</param>
+		/// <param name="name">Name of the attribute</param>
 		/// <returns>Array of names, with all leading and trailing whitespace removed</returns>
-		protected string ReadAttribute(BgScriptElement Element, string Name)
+		protected string ReadAttribute(BgScriptElement element, string name)
 		{
-			return ExpandProperties(Element, Element.GetAttribute(Name));
+			return ExpandProperties(element, element.GetAttribute(name));
 		}
 
 		/// <summary>
 		/// Expands any properties and reads a list of strings from an attribute, separated by semi-colon characters
 		/// </summary>
-		/// <param name="Element"></param>
-		/// <param name="Name"></param>
-		/// <param name="Separator"></param>
+		/// <param name="element"></param>
+		/// <param name="name"></param>
+		/// <param name="separator"></param>
 		/// <returns>Array of names, with all leading and trailing whitespace removed</returns>
-		protected string[] ReadListAttribute(BgScriptElement Element, string Name, char Separator = ';')
+		protected string[] ReadListAttribute(BgScriptElement element, string name, char separator = ';')
 		{
-			string Value = ReadAttribute(Element, Name);
-			return Value.Split(new char[] { Separator }).Select(x => x.Trim()).Where(x => x.Length > 0).ToArray();
+			string value = ReadAttribute(element, name);
+			return value.Split(new char[] { separator }).Select(x => x.Trim()).Where(x => x.Length > 0).ToArray();
 		}
 
 		/// <summary>
 		/// Reads an attribute from the given XML element, expands any properties in it, and parses it as a boolean.
 		/// </summary>
-		/// <param name="Element">Element to read the attribute from</param>
-		/// <param name="Name">Name of the attribute</param>
+		/// <param name="element">Element to read the attribute from</param>
+		/// <param name="name">Name of the attribute</param>
 		/// <param name="bDefaultValue">Default value if the attribute is missing</param>
 		/// <returns>The value of the attribute field</returns>
-		protected bool ReadBooleanAttribute(BgScriptElement Element, string Name, bool bDefaultValue)
+		protected bool ReadBooleanAttribute(BgScriptElement element, string name, bool bDefaultValue)
 		{
 			bool bResult = bDefaultValue;
-			if (Element.HasAttribute(Name))
+			if (element.HasAttribute(name))
 			{
-				string Value = ReadAttribute(Element, Name).Trim();
-				if (Value.Equals("true", StringComparison.InvariantCultureIgnoreCase))
+				string value = ReadAttribute(element, name).Trim();
+				if (value.Equals("true", StringComparison.InvariantCultureIgnoreCase))
 				{
 					bResult = true;
 				}
-				else if (Value.Equals("false", StringComparison.InvariantCultureIgnoreCase))
+				else if (value.Equals("false", StringComparison.InvariantCultureIgnoreCase))
 				{
 					bResult = false;
 				}
 				else
 				{
-					LogError(Element, "Invalid boolean value '{0}' - expected 'true' or 'false'", Value);
+					LogError(element, "Invalid boolean value '{0}' - expected 'true' or 'false'", value);
 				}
 			}
 			return bResult;
@@ -1162,91 +1160,91 @@ namespace EpicGames.BuildGraph
 		/// <summary>
 		/// Reads an attribute from the given XML element, expands any properties in it, and parses it as an integer.
 		/// </summary>
-		/// <param name="Element">Element to read the attribute from</param>
-		/// <param name="Name">Name of the attribute</param>
-		/// <param name="DefaultValue">Default value for the integer, if the attribute is missing</param>
+		/// <param name="element">Element to read the attribute from</param>
+		/// <param name="name">Name of the attribute</param>
+		/// <param name="defaultValue">Default value for the integer, if the attribute is missing</param>
 		/// <returns>The value of the attribute field</returns>
-		protected int ReadIntegerAttribute(BgScriptElement Element, string Name, int DefaultValue)
+		protected int ReadIntegerAttribute(BgScriptElement element, string name, int defaultValue)
 		{
-			int Result = DefaultValue;
-			if (Element.HasAttribute(Name))
+			int result = defaultValue;
+			if (element.HasAttribute(name))
 			{
-				string Value = ReadAttribute(Element, Name).Trim();
+				string value = ReadAttribute(element, name).Trim();
 
-				int IntValue;
-				if (Int32.TryParse(Value, out IntValue))
+				int intValue;
+				if (Int32.TryParse(value, out intValue))
 				{
-					Result = IntValue;
+					result = intValue;
 				}
 				else
 				{
-					LogError(Element, "Invalid integer value '{0}'", Value);
+					LogError(element, "Invalid integer value '{0}'", value);
 				}
 			}
-			return Result;
+			return result;
 		}
 
 		/// <summary>
 		/// Reads an attribute from the given XML element, expands any properties in it, and parses it as an enum of the given type.
 		/// </summary>
 		/// <typeparam name="T">The enum type to parse the attribute as</typeparam>
-		/// <param name="Element">Element to read the attribute from</param>
-		/// <param name="Name">Name of the attribute</param>
-		/// <param name="DefaultValue">Default value for the enum, if the attribute is missing</param>
+		/// <param name="element">Element to read the attribute from</param>
+		/// <param name="name">Name of the attribute</param>
+		/// <param name="defaultValue">Default value for the enum, if the attribute is missing</param>
 		/// <returns>The value of the attribute field</returns>
-		protected T ReadEnumAttribute<T>(BgScriptElement Element, string Name, T DefaultValue) where T : struct
+		protected T ReadEnumAttribute<T>(BgScriptElement element, string name, T defaultValue) where T : struct
 		{
-			T Result = DefaultValue;
-			if (Element.HasAttribute(Name))
+			T result = defaultValue;
+			if (element.HasAttribute(name))
 			{
-				string Value = ReadAttribute(Element, Name).Trim();
+				string value = ReadAttribute(element, name).Trim();
 
-				T EnumValue;
-				if (Enum.TryParse(Value, true, out EnumValue))
+				T enumValue;
+				if (Enum.TryParse(value, true, out enumValue))
 				{
-					Result = EnumValue;
+					result = enumValue;
 				}
 				else
 				{
-					LogError(Element, "Invalid value '{0}' - expected {1}", Value, String.Join("/", Enum.GetNames(typeof(T))));
+					LogError(element, "Invalid value '{0}' - expected {1}", value, String.Join("/", Enum.GetNames(typeof(T))));
 				}
 			}
-			return Result;
+			return result;
 		}
 
 		/// <summary>
 		/// Outputs an error message to the log and increments the number of errors, referencing the file and line number of the element that caused it.
 		/// </summary>
-		/// <param name="Element">The script element causing the error</param>
-		/// <param name="Format">Standard String.Format()-style format string</param>
-		/// <param name="Args">Optional arguments</param>
-		protected void LogError(BgScriptElement Element, string Format, params object[] Args)
+		/// <param name="element">The script element causing the error</param>
+		/// <param name="format">Standard String.Format()-style format string</param>
+		/// <param name="args">Optional arguments</param>
+		protected void LogError(BgScriptElement element, string format, params object[] args)
 		{
-			Logger.LogScriptError(Element.Location, Format, Args);
+			Logger.LogScriptError(element.Location, format, args);
 			NumErrors++;
 		}
 
 		/// <summary>
 		/// Outputs a warning message to the log and increments the number of errors, referencing the file and line number of the element that caused it.
 		/// </summary>
-		/// <param name="Element">The script element causing the error</param>
-		/// <param name="Format">Standard String.Format()-style format string</param>
-		/// <param name="Args">Optional arguments</param>
-		protected void LogWarning(BgScriptElement Element, string Format, params object[] Args)
+		/// <param name="element">The script element causing the error</param>
+		/// <param name="format">Standard String.Format()-style format string</param>
+		/// <param name="args">Optional arguments</param>
+		protected void LogWarning(BgScriptElement element, string format, params object[] args)
 		{
-			Logger.LogScriptWarning(Element.Location, Format, Args);
+			Logger.LogScriptWarning(element.Location, format, args);
 		}
 
 		/// <summary>
 		/// Evaluates the (optional) conditional expression on a given XML element via the If="..." attribute, and returns true if the element is enabled.
 		/// </summary>
-		/// <param name="Element">The element to check</param>
+		/// <param name="element">The element to check</param>
 		/// <returns>True if the element's condition evaluates to true (or doesn't have a conditional expression), false otherwise</returns>
-		protected async Task<bool> EvaluateConditionAsync(BgScriptElement Element)
+		protected async Task<bool> EvaluateConditionAsync(BgScriptElement element)
 		{
 			// Check if the element has a conditional attribute
 			const string AttributeName = "If";
-			if (!Element.HasAttribute(AttributeName))
+			if (!element.HasAttribute(AttributeName))
 			{
 				return true;
 			}
@@ -1254,12 +1252,12 @@ namespace EpicGames.BuildGraph
 			// If it does, try to evaluate it.
 			try
 			{
-				string Text = ExpandProperties(Element, Element.GetAttribute("If"));
-				return await BgCondition.EvaluateAsync(Text, Context);
+				string text = ExpandProperties(element, element.GetAttribute("If"));
+				return await BgCondition.EvaluateAsync(text, Context);
 			}
-			catch (BgConditionException Ex)
+			catch (BgConditionException ex)
 			{
-				LogError(Element, "Error in condition: {0}", Ex.Message);
+				LogError(element, "Error in condition: {0}", ex.Message);
 				return false;
 			}
 		}
@@ -1267,41 +1265,41 @@ namespace EpicGames.BuildGraph
 		/// <summary>
 		/// Expand all the property references (of the form $(PropertyName)) in a string.
 		/// </summary>
-		/// <param name="Element">The element containing the string. Used for diagnostic messages.</param>
-		/// <param name="Text">The input string to expand properties in</param>
+		/// <param name="element">The element containing the string. Used for diagnostic messages.</param>
+		/// <param name="text">The input string to expand properties in</param>
 		/// <returns>The expanded string</returns>
-		protected string ExpandProperties(BgScriptElement Element, string Text)
+		protected string ExpandProperties(BgScriptElement element, string text)
 		{
-			string Result = Text;
+			string result = text;
 			// Iterate in reverse order to handle cases where there are nested expansions like $(Outer$(Inner))
-			for (int Idx = Result.LastIndexOf("$("); Idx != -1; Idx = Result.LastIndexOf("$(", Idx, Idx+1))
+			for (int idx = result.LastIndexOf("$("); idx != -1; idx = result.LastIndexOf("$(", idx, idx + 1))
 			{
 				// Find the end of the variable name
-				int EndIdx = Result.IndexOf(')', Idx + 2);
-				if (EndIdx == -1)
+				int endIdx = result.IndexOf(')', idx + 2);
+				if (endIdx == -1)
 				{
 					break;
 				}
 
 				// Extract the variable name from the string
-				string Name = Result.Substring(Idx + 2, EndIdx - (Idx + 2));
+				string name = result.Substring(idx + 2, endIdx - (idx + 2));
 
 				// Find the value for it, either from the dictionary or the environment block
-				string? Value;
-				if (!TryGetPropertyValue(Name, out Value))
+				string? value;
+				if (!TryGetPropertyValue(name, out value))
 				{
-					LogWarning(Element, "Property '{0}' is not defined", Name);
-					Value = "";
+					LogWarning(element, "Property '{0}' is not defined", name);
+					value = "";
 				}
 
 				// Check if we've got a value for this variable
-				if (Value != null)
+				if (value != null)
 				{
 					// Replace the variable, or skip past it
-					Result = Result.Substring(0, Idx) + Value + Result.Substring(EndIdx + 1);
+					result = result.Substring(0, idx) + value + result.Substring(endIdx + 1);
 				}
 			}
-			return Result;
+			return result;
 		}
 	}
 
@@ -1313,88 +1311,88 @@ namespace EpicGames.BuildGraph
 		/// <summary>
 		/// The current graph
 		/// </summary>
-		BgGraph Graph = new BgGraph();
+		readonly BgGraph _graph = new BgGraph();
 
 		/// <summary>
 		/// Arguments for evaluating the graph
 		/// </summary>
-		Dictionary<string, string> Arguments;
+		readonly Dictionary<string, string> _arguments;
 
 		/// <summary>
 		/// The name of the node if only a single node is going to be built, otherwise null.
 		/// </summary>
-		string? SingleNodeName;
+		readonly string? _singleNodeName;
 
-		BgAgent? EnclosingAgent;
-		BgNode? EnclosingNode;
+		BgAgent? _enclosingAgent;
+		BgNode? _enclosingNode;
 
 		/// <summary>
 		/// Private constructor. Use ScriptReader.TryRead() to read a script file.
 		/// </summary>
-		/// <param name="Context">Context object</param>
-		/// <param name="DefaultProperties">Default properties available to the script</param>
-		/// <param name="Arguments">Arguments passed in to the graph on the command line</param>
-		/// <param name="SingleNodeName">If a single node will be processed, the name of that node.</param>
-		/// <param name="Schema">Schema for the script</param>
-		/// <param name="Logger">Logger for diagnostic messages</param>
-		private BgScriptReader(IBgScriptReaderContext Context, IDictionary<string, string> DefaultProperties, IReadOnlyDictionary<string, string> Arguments, string? SingleNodeName, BgScriptSchema Schema, ILogger Logger)
-			: base(Context, Schema, Logger)
+		/// <param name="context">Context object</param>
+		/// <param name="defaultProperties">Default properties available to the script</param>
+		/// <param name="arguments">Arguments passed in to the graph on the command line</param>
+		/// <param name="singleNodeName">If a single node will be processed, the name of that node.</param>
+		/// <param name="schema">Schema for the script</param>
+		/// <param name="logger">Logger for diagnostic messages</param>
+		private BgScriptReader(IBgScriptReaderContext context, IDictionary<string, string> defaultProperties, IReadOnlyDictionary<string, string> arguments, string? singleNodeName, BgScriptSchema schema, ILogger logger)
+			: base(context, schema, logger)
 		{
-			this.Arguments = new Dictionary<string, string>(Arguments, StringComparer.OrdinalIgnoreCase);
-			this.SingleNodeName = SingleNodeName;
+			_arguments = new Dictionary<string, string>(arguments, StringComparer.OrdinalIgnoreCase);
+			_singleNodeName = singleNodeName;
 
-			foreach (KeyValuePair<string, string> Pair in DefaultProperties)
+			foreach (KeyValuePair<string, string> pair in defaultProperties)
 			{
-				SetPropertyValue(null!, Pair.Key, Pair.Value);
+				SetPropertyValue(null!, pair.Key, pair.Value);
 			}
 		}
 
 		/// <summary>
 		/// Try to read a script file from the given file.
 		/// </summary>
-		/// <param name="Context">Supplies context about the parse</param>
-		/// <param name="File">File to read from</param>
-		/// <param name="Arguments">Arguments passed in to the graph on the command line</param>
-		/// <param name="DefaultProperties">Default properties available to the script</param>
-		/// <param name="Schema">Schema for the script</param>
-		/// <param name="Logger">Logger for output messages</param>
-		/// <param name="SingleNodeName">If a single node will be processed, the name of that node.</param>
+		/// <param name="context">Supplies context about the parse</param>
+		/// <param name="file">File to read from</param>
+		/// <param name="arguments">Arguments passed in to the graph on the command line</param>
+		/// <param name="defaultProperties">Default properties available to the script</param>
+		/// <param name="schema">Schema for the script</param>
+		/// <param name="logger">Logger for output messages</param>
+		/// <param name="singleNodeName">If a single node will be processed, the name of that node.</param>
 		/// <returns>True if the graph was read, false if there were errors</returns>
-		public static async Task<BgGraph?> ReadAsync(IBgScriptReaderContext Context, string File, Dictionary<string, string> Arguments, Dictionary<string, string> DefaultProperties, BgScriptSchema Schema, ILogger Logger, string? SingleNodeName = null)
+		public static async Task<BgGraph?> ReadAsync(IBgScriptReaderContext context, string file, Dictionary<string, string> arguments, Dictionary<string, string> defaultProperties, BgScriptSchema schema, ILogger logger, string? singleNodeName = null)
 		{
 			// Read the file and build the graph
-			BgScriptReader Reader = new BgScriptReader(Context, DefaultProperties, Arguments, SingleNodeName, Schema, Logger);
-			if (!await Reader.TryReadAsync(File) || Reader.NumErrors > 0)
+			BgScriptReader reader = new BgScriptReader(context, defaultProperties, arguments, singleNodeName, schema, logger);
+			if (!await reader.TryReadAsync(file) || reader.NumErrors > 0)
 			{
 				return null;
 			}
 
 			// Make sure all the arguments were valid
-			HashSet<string> ValidArgumentNames = new HashSet<string>(Reader.Graph.Options.Select(x => x.Name), StringComparer.OrdinalIgnoreCase);
-			ValidArgumentNames.Add("PreflightChange");
+			HashSet<string> validArgumentNames = new HashSet<string>(reader._graph.Options.Select(x => x.Name), StringComparer.OrdinalIgnoreCase);
+			validArgumentNames.Add("PreflightChange");
 
-			foreach (string ArgumentName in Arguments.Keys)
+			foreach (string argumentName in arguments.Keys)
 			{
-				if (!ValidArgumentNames.Contains(ArgumentName))
+				if (!validArgumentNames.Contains(argumentName))
 				{
-					Logger.LogWarning("Unknown argument '{ArgumentName}' for '{Script}'", ArgumentName, Context.GetNativePath(File));
+					logger.LogWarning("Unknown argument '{ArgumentName}' for '{Script}'", argumentName, context.GetNativePath(file));
 				}
 			}
 
 			// Return the constructed graph
-			return Reader.Graph;
+			return reader._graph;
 		}
 
 		/// <summary>
 		/// Reads the definition of a graph option; a parameter which can be set by the user on the command-line or via an environment variable.
 		/// </summary>
-		/// <param name="Element">Xml element to read the definition from</param>
-		protected override async Task ReadOptionAsync(BgScriptElement Element)
+		/// <param name="element">Xml element to read the definition from</param>
+		protected override async Task ReadOptionAsync(BgScriptElement element)
 		{
-			if (await EvaluateConditionAsync(Element))
+			if (await EvaluateConditionAsync(element))
 			{
-				string Name = ReadAttribute(Element, "Name");
-				if (ValidateName(Element, Name))
+				string name = ReadAttribute(element, "Name");
+				if (ValidateName(element, name))
 				{
 					// Make sure we're at global scope
 					if (ScopedProperties.Count > 1)
@@ -1403,44 +1401,44 @@ namespace EpicGames.BuildGraph
 					}
 
 					// Check if the property already exists. If it does, we don't need to register it as an option.
-					string? ExistingValue;
-					if (TryGetPropertyValue(Name, out ExistingValue) && ExistingValue != null)
+					string? existingValue;
+					if (TryGetPropertyValue(name, out existingValue) && existingValue != null)
 					{
 						// If there's a restriction on this definition, check it matches
-						string Restrict = ReadAttribute(Element, "Restrict");
-						if (!String.IsNullOrEmpty(Restrict) && !Regex.IsMatch(ExistingValue, "^" + Restrict + "$", RegexOptions.IgnoreCase))
+						string restrict = ReadAttribute(element, "Restrict");
+						if (!String.IsNullOrEmpty(restrict) && !Regex.IsMatch(existingValue, "^" + restrict + "$", RegexOptions.IgnoreCase))
 						{
-							LogError(Element, "'{0} is already set to '{1}', which does not match the given restriction ('{2}')", Name, ExistingValue, Restrict);
+							LogError(element, "'{0} is already set to '{1}', which does not match the given restriction ('{2}')", name, existingValue, restrict);
 						}
 					}
 					else
 					{
 						// Create a new option object to store the settings
-						string Description = ReadAttribute(Element, "Description");
-						string DefaultValue = ReadAttribute(Element, "DefaultValue");
-						BgOption Option = new BgOption(Name, Description, DefaultValue);
-						Graph.Options.Add(Option);
+						string description = ReadAttribute(element, "Description");
+						string defaultValue = ReadAttribute(element, "DefaultValue");
+						BgOption option = new BgOption(name, description, defaultValue);
+						_graph.Options.Add(option);
 
 						// Get the value of this property
-						string? Value;
-						if (!Arguments.TryGetValue(Name, out Value))
+						string? value;
+						if (!_arguments.TryGetValue(name, out value))
 						{
-							Value = Option.DefaultValue;
+							value = option.DefaultValue;
 						}
-						SetPropertyValue(Element, Name, Value);
+						SetPropertyValue(element, name, value);
 
 						// If there's a restriction on it, check it's valid
-						string Restrict = ReadAttribute(Element, "Restrict");
-						if (!String.IsNullOrEmpty(Restrict))
+						string restrict = ReadAttribute(element, "Restrict");
+						if (!String.IsNullOrEmpty(restrict))
 						{
-							string Pattern = "^(" + Restrict + ")$";
-							if (!Regex.IsMatch(Value, Pattern, RegexOptions.IgnoreCase))
+							string pattern = "^(" + restrict + ")$";
+							if (!Regex.IsMatch(value, pattern, RegexOptions.IgnoreCase))
 							{
-								LogError(Element, "'{0}' is not a valid value for '{1}' (required: '{2}')", Value, Name, Restrict);
+								LogError(element, "'{0}' is not a valid value for '{1}' (required: '{2}')", value, name, restrict);
 							}
-							if (Option.DefaultValue != Value && !Regex.IsMatch(Option.DefaultValue, Pattern, RegexOptions.IgnoreCase))
+							if (option.DefaultValue != value && !Regex.IsMatch(option.DefaultValue, pattern, RegexOptions.IgnoreCase))
 							{
-								LogError(Element, "Default value '{0}' is not valid for '{1}' (required: '{2}')", Option.DefaultValue, Name, Restrict);
+								LogError(element, "Default value '{0}' is not valid for '{1}' (required: '{2}')", option.DefaultValue, name, restrict);
 							}
 						}
 					}
@@ -1451,34 +1449,34 @@ namespace EpicGames.BuildGraph
 		/// <summary>
 		/// Reads a property assignment.
 		/// </summary>
-		/// <param name="Element">Xml element to read the definition from</param>
-		protected override async Task ReadPropertyAsync(BgScriptElement Element)
+		/// <param name="element">Xml element to read the definition from</param>
+		protected override async Task ReadPropertyAsync(BgScriptElement element)
 		{
-			if (await EvaluateConditionAsync(Element))
+			if (await EvaluateConditionAsync(element))
 			{
-				string Name = ReadAttribute(Element, "Name");
-				if (ValidateName(Element, Name))
+				string name = ReadAttribute(element, "Name");
+				if (ValidateName(element, name))
 				{
-					string Value = ReadAttribute(Element, "Value");
-					if (Element.HasChildNodes)
+					string value = ReadAttribute(element, "Value");
+					if (element.HasChildNodes)
 					{
 						// Read the element content, and append each line to the value as a semicolon delimited list
-						StringBuilder Builder = new StringBuilder(Value);
-						foreach (string Line in Element.InnerText.Split('\n'))
+						StringBuilder builder = new StringBuilder(value);
+						foreach (string line in element.InnerText.Split('\n'))
 						{
-							string TrimLine = ExpandProperties(Element, Line.Trim());
-							if (TrimLine.Length > 0)
+							string trimLine = ExpandProperties(element, line.Trim());
+							if (trimLine.Length > 0)
 							{
-								if (Builder.Length > 0)
+								if (builder.Length > 0)
 								{
-									Builder.Append(";");
+									builder.Append(";");
 								}
-								Builder.Append(TrimLine);
+								builder.Append(trimLine);
 							}
 						}
-						Value = Builder.ToString();
+						value = builder.ToString();
 					}
-					SetPropertyValue(Element, Name, Value);
+					SetPropertyValue(element, name, value);
 				}
 			}
 		}
@@ -1486,50 +1484,50 @@ namespace EpicGames.BuildGraph
 		/// <summary>
 		/// Reads a Regex assignment.
 		/// </summary>
-		/// <param name="Element">Xml element to read the definition from</param>
-		protected override async Task ReadRegexAsync(BgScriptElement Element)
+		/// <param name="element">Xml element to read the definition from</param>
+		protected override async Task ReadRegexAsync(BgScriptElement element)
 		{
-			if (await EvaluateConditionAsync(Element))
+			if (await EvaluateConditionAsync(element))
 			{
 				// Get the pattern
-				string RegexString = ReadAttribute(Element, "Pattern");
+				string regexString = ReadAttribute(element, "Pattern");
 
 				// Make sure its a valid regex.
-				Regex? RegexValue = ParseRegex(Element, RegexString);
-				if (RegexValue != null)
+				Regex? regexValue = ParseRegex(element, regexString);
+				if (regexValue != null)
 				{
 					// read the names in 
-					string[] CaptureNames = ReadListAttribute(Element, "Capture");
+					string[] captureNames = ReadListAttribute(element, "Capture");
 
 					// get number of groups we passed in
-					int[] GroupNumbers = RegexValue.GetGroupNumbers();
+					int[] groupNumbers = regexValue.GetGroupNumbers();
 
 					// make sure the number of property names is the same as the number of match groups
 					// this includes the entire string match group as [0], so don't count that one.
-					if (CaptureNames.Length != GroupNumbers.Count() - 1)
+					if (captureNames.Length != groupNumbers.Count() - 1)
 					{
-						LogError(Element, "MatchGroup count: {0} does not match the number of names specified: {1}", GroupNumbers.Count() - 1, CaptureNames.Length);
+						LogError(element, "MatchGroup count: {0} does not match the number of names specified: {1}", groupNumbers.Count() - 1, captureNames.Length);
 					}
 					else
 					{
 						// apply the regex to the value
-						string Input = ReadAttribute(Element, "Input");
-						Match Match = RegexValue.Match(Input);
+						string input = ReadAttribute(element, "Input");
+						Match match = regexValue.Match(input);
 
-						bool Optional = await BgCondition.EvaluateAsync(ReadAttribute(Element, "Optional"), Context);
-						if (!Match.Success)
+						bool optional = await BgCondition.EvaluateAsync(ReadAttribute(element, "Optional"), Context);
+						if (!match.Success)
 						{
-							if (!Optional)
+							if (!optional)
 							{
-								LogError(Element, "Regex {0} did not find a match against input string {1}", RegexString, Input);
+								LogError(element, "Regex {0} did not find a match against input string {1}", regexString, input);
 							}
 						}
 						else
 						{
 							// assign each property to the group it matches, skip over [0]
-							for (int MatchIdx = 1; MatchIdx < GroupNumbers.Count(); MatchIdx++)
+							for (int matchIdx = 1; matchIdx < groupNumbers.Count(); matchIdx++)
 							{
-								SetPropertyValue(Element, CaptureNames[MatchIdx - 1], Match.Groups[MatchIdx].Value);
+								SetPropertyValue(element, captureNames[matchIdx - 1], match.Groups[matchIdx].Value);
 							}
 						}
 					}
@@ -1540,16 +1538,16 @@ namespace EpicGames.BuildGraph
 		/// <summary>
 		/// Reads a property assignment from an environment variable.
 		/// </summary>
-		/// <param name="Element">Xml element to read the definition from</param>
-		protected override async Task ReadEnvVarAsync(BgScriptElement Element)
+		/// <param name="element">Xml element to read the definition from</param>
+		protected override async Task ReadEnvVarAsync(BgScriptElement element)
 		{
-			if (await EvaluateConditionAsync(Element))
+			if (await EvaluateConditionAsync(element))
 			{
-				string Name = ReadAttribute(Element, "Name");
-				if (ValidateName(Element, Name))
+				string name = ReadAttribute(element, "Name");
+				if (ValidateName(element, name))
 				{
-					string Value = Environment.GetEnvironmentVariable(Name) ?? "";
-					SetPropertyValue(Element, Name, Value);
+					string value = Environment.GetEnvironmentVariable(name) ?? "";
+					SetPropertyValue(element, name, value);
 				}
 			}
 		}
@@ -1557,102 +1555,102 @@ namespace EpicGames.BuildGraph
 		/// <summary>
 		/// Reads the definition for an agent.
 		/// </summary>
-		/// <param name="Element">Xml element to read the definition from</param>
-		protected override async Task ReadAgentAsync(BgScriptElement Element)
+		/// <param name="element">Xml element to read the definition from</param>
+		protected override async Task ReadAgentAsync(BgScriptElement element)
 		{
-			string? Name;
-			if (await EvaluateConditionAsync(Element) && TryReadObjectName(Element, out Name))
+			string? name;
+			if (await EvaluateConditionAsync(element) && TryReadObjectName(element, out name))
 			{
 				// Read the valid agent types. This may be omitted if we're continuing an existing agent.
-				string[] Types = ReadListAttribute(Element, "Type");
+				string[] types = ReadListAttribute(element, "Type");
 
 				// Create the agent object, or continue an existing one
-				BgAgent? Agent;
-				if (Graph.NameToAgent.TryGetValue(Name, out Agent))
+				BgAgent? agent;
+				if (_graph.NameToAgent.TryGetValue(name, out agent))
 				{
-					if (Types.Length > 0 && Agent.PossibleTypes.Length > 0)
+					if (types.Length > 0 && agent.PossibleTypes.Length > 0)
 					{
-						if (Types.Length != Agent.PossibleTypes.Length || !Types.SequenceEqual(Agent.PossibleTypes, StringComparer.InvariantCultureIgnoreCase))
+						if (types.Length != agent.PossibleTypes.Length || !types.SequenceEqual(agent.PossibleTypes, StringComparer.InvariantCultureIgnoreCase))
 						{
-							LogError(Element, "Agent types ({0}) were different than previous agent definition with types ({1}). Must either be empty or match exactly.", string.Join(",", Types), string.Join(",", Agent.PossibleTypes));
+							LogError(element, "Agent types ({0}) were different than previous agent definition with types ({1}). Must either be empty or match exactly.", string.Join(",", types), string.Join(",", agent.PossibleTypes));
 						}
 					}
 				}
 				else
 				{
-					if (Types.Length == 0)
+					if (types.Length == 0)
 					{
-						LogError(Element, "Missing type for agent '{0}'", Name);
+						LogError(element, "Missing type for agent '{0}'", name);
 					}
-					Agent = new BgAgent(Name, Types);
-					Graph.NameToAgent.Add(Name, Agent);
-					Graph.Agents.Add(Agent);
+					agent = new BgAgent(name, types);
+					_graph.NameToAgent.Add(name, agent);
+					_graph.Agents.Add(agent);
 				}
 
 				// Process all the child elements.
-				EnclosingAgent = Agent;
-				await ReadAgentBodyAsync(Element);
-				EnclosingAgent = null;
+				_enclosingAgent = agent;
+				await ReadAgentBodyAsync(element);
+				_enclosingAgent = null;
 			}
 		}
 
 		/// <summary>
 		/// Reads the definition for an aggregate
 		/// </summary>
-		/// <param name="Element">Xml element to read the definition from</param>
-		protected override async Task ReadAggregateAsync(BgScriptElement Element)
+		/// <param name="element">Xml element to read the definition from</param>
+		protected override async Task ReadAggregateAsync(BgScriptElement element)
 		{
-			string? Name;
-			if (await EvaluateConditionAsync(Element) && TryReadObjectName(Element, out Name) && CheckNameIsUnique(Element, Name))
+			string? name;
+			if (await EvaluateConditionAsync(element) && TryReadObjectName(element, out name) && CheckNameIsUnique(element, name))
 			{
-				string[] RequiredNames = ReadListAttribute(Element, "Requires");
+				string[] requiredNames = ReadListAttribute(element, "Requires");
 
-				BgAggregate NewAggregate = new BgAggregate(Name);
-				foreach (BgNode ReferencedNode in ResolveReferences(Element, RequiredNames))
+				BgAggregate newAggregate = new BgAggregate(name);
+				foreach (BgNode referencedNode in ResolveReferences(element, requiredNames))
 				{
-					NewAggregate.RequiredNodes.Add(ReferencedNode);
+					newAggregate.RequiredNodes.Add(referencedNode);
 				}
-				Graph.NameToAggregate[Name] = NewAggregate;
+				_graph.NameToAggregate[name] = newAggregate;
 
-				string LabelCategoryName = ReadAttribute(Element, "Label");
-				if (!String.IsNullOrEmpty(LabelCategoryName))
+				string labelCategoryName = ReadAttribute(element, "Label");
+				if (!String.IsNullOrEmpty(labelCategoryName))
 				{
-					BgLabel Label;
+					BgLabel label;
 
 					// Create the label
-					int SlashIdx = LabelCategoryName.IndexOf('/');
-					if (SlashIdx != -1)
+					int slashIdx = labelCategoryName.IndexOf('/');
+					if (slashIdx != -1)
 					{
-						Label = new BgLabel(LabelCategoryName.Substring(SlashIdx + 1), LabelCategoryName.Substring(0, SlashIdx), null, null, BgLabelChange.Current);
+						label = new BgLabel(labelCategoryName.Substring(slashIdx + 1), labelCategoryName.Substring(0, slashIdx), null, null, BgLabelChange.Current);
 					}
 					else
 					{
-						Label = new BgLabel(LabelCategoryName, "Other", null, null, BgLabelChange.Current);
+						label = new BgLabel(labelCategoryName, "Other", null, null, BgLabelChange.Current);
 					}
 
 					// Find all the included nodes
-					foreach (BgNode RequiredNode in NewAggregate.RequiredNodes)
+					foreach (BgNode requiredNode in newAggregate.RequiredNodes)
 					{
-						Label.RequiredNodes.Add(RequiredNode);
-						Label.IncludedNodes.Add(RequiredNode);
-						Label.IncludedNodes.UnionWith(RequiredNode.OrderDependencies);
+						label.RequiredNodes.Add(requiredNode);
+						label.IncludedNodes.Add(requiredNode);
+						label.IncludedNodes.UnionWith(requiredNode.OrderDependencies);
 					}
 
-					string[] IncludedNames = ReadListAttribute(Element, "Include");
-					foreach (BgNode IncludedNode in ResolveReferences(Element, IncludedNames))
+					string[] includedNames = ReadListAttribute(element, "Include");
+					foreach (BgNode includedNode in ResolveReferences(element, includedNames))
 					{
-						Label.IncludedNodes.Add(IncludedNode);
-						Label.IncludedNodes.UnionWith(IncludedNode.OrderDependencies);
+						label.IncludedNodes.Add(includedNode);
+						label.IncludedNodes.UnionWith(includedNode.OrderDependencies);
 					}
 
-					string[] ExcludedNames = ReadListAttribute(Element, "Exclude");
-					foreach (BgNode ExcludedNode in ResolveReferences(Element, ExcludedNames))
+					string[] excludedNames = ReadListAttribute(element, "Exclude");
+					foreach (BgNode excludedNode in ResolveReferences(element, excludedNames))
 					{
-						Label.IncludedNodes.Remove(ExcludedNode);
-						Label.IncludedNodes.ExceptWith(ExcludedNode.OrderDependencies);
+						label.IncludedNodes.Remove(excludedNode);
+						label.IncludedNodes.ExceptWith(excludedNode.OrderDependencies);
 					}
 
-					Graph.Labels.Add(Label);
+					_graph.Labels.Add(label);
 				}
 			}
 		}
@@ -1660,203 +1658,203 @@ namespace EpicGames.BuildGraph
 		/// <summary>
 		/// Reads the definition for a report
 		/// </summary>
-		/// <param name="Element">Xml element to read the definition from</param>
-		protected override async Task ReadReportAsync(BgScriptElement Element)
+		/// <param name="element">Xml element to read the definition from</param>
+		protected override async Task ReadReportAsync(BgScriptElement element)
 		{
-			string? Name;
-			if (await EvaluateConditionAsync(Element) && TryReadObjectName(Element, out Name) && CheckNameIsUnique(Element, Name))
+			string? name;
+			if (await EvaluateConditionAsync(element) && TryReadObjectName(element, out name) && CheckNameIsUnique(element, name))
 			{
-				string[] RequiredNames = ReadListAttribute(Element, "Requires");
+				string[] requiredNames = ReadListAttribute(element, "Requires");
 
-				BgReport NewReport = new BgReport(Name);
-				foreach (BgNode ReferencedNode in ResolveReferences(Element, RequiredNames))
+				BgReport newReport = new BgReport(name);
+				foreach (BgNode referencedNode in ResolveReferences(element, requiredNames))
 				{
-					NewReport.Nodes.Add(ReferencedNode);
-					NewReport.Nodes.UnionWith(ReferencedNode.OrderDependencies);
+					newReport.Nodes.Add(referencedNode);
+					newReport.Nodes.UnionWith(referencedNode.OrderDependencies);
 				}
-				Graph.NameToReport.Add(Name, NewReport);
+				_graph.NameToReport.Add(name, newReport);
 			}
 		}
 
 		/// <summary>
 		/// Reads the definition for a badge
 		/// </summary>
-		/// <param name="Element">Xml element to read the definition from</param>
-		protected override async Task ReadBadgeAsync(BgScriptElement Element)
+		/// <param name="element">Xml element to read the definition from</param>
+		protected override async Task ReadBadgeAsync(BgScriptElement element)
 		{
-			string? Name;
-			if (await EvaluateConditionAsync(Element) && TryReadObjectName(Element, out Name))
+			string? name;
+			if (await EvaluateConditionAsync(element) && TryReadObjectName(element, out name))
 			{
-				string[] RequiredNames = ReadListAttribute(Element, "Requires");
-				string[] TargetNames = ReadListAttribute(Element, "Targets");
-				string Project = ReadAttribute(Element, "Project");
-				int Change = ReadIntegerAttribute(Element, "Change", 0);
+				string[] requiredNames = ReadListAttribute(element, "Requires");
+				string[] targetNames = ReadListAttribute(element, "Targets");
+				string project = ReadAttribute(element, "Project");
+				int change = ReadIntegerAttribute(element, "Change", 0);
 
-				BgBadge NewBadge = new BgBadge(Name, Project, Change);
-				foreach (BgNode ReferencedNode in ResolveReferences(Element, RequiredNames))
+				BgBadge newBadge = new BgBadge(name, project, change);
+				foreach (BgNode referencedNode in ResolveReferences(element, requiredNames))
 				{
-					NewBadge.Nodes.Add(ReferencedNode);
+					newBadge.Nodes.Add(referencedNode);
 				}
-				foreach (BgNode ReferencedNode in ResolveReferences(Element, TargetNames))
+				foreach (BgNode referencedNode in ResolveReferences(element, targetNames))
 				{
-					NewBadge.Nodes.Add(ReferencedNode);
-					NewBadge.Nodes.UnionWith(ReferencedNode.OrderDependencies);
+					newBadge.Nodes.Add(referencedNode);
+					newBadge.Nodes.UnionWith(referencedNode.OrderDependencies);
 				}
-				Graph.Badges.Add(NewBadge);
+				_graph.Badges.Add(newBadge);
 			}
 		}
 
 		/// <summary>
 		/// Reads the definition for a label
 		/// </summary>
-		/// <param name="Element">Xml element to read the definition from</param>
-		protected override async Task ReadLabelAsync(BgScriptElement Element)
+		/// <param name="element">Xml element to read the definition from</param>
+		protected override async Task ReadLabelAsync(BgScriptElement element)
 		{
-			if (await EvaluateConditionAsync(Element))
+			if (await EvaluateConditionAsync(element))
 			{
-				string Name = ReadAttribute(Element, "Name");
-				if (!String.IsNullOrEmpty(Name))
+				string name = ReadAttribute(element, "Name");
+				if (!String.IsNullOrEmpty(name))
 				{
-					ValidateName(Element, Name);
+					ValidateName(element, name);
 				}
 
-				string Category = ReadAttribute(Element, "Category");
+				string category = ReadAttribute(element, "Category");
 
-				string[] RequiredNames = ReadListAttribute(Element, "Requires");
-				string[] IncludedNames = ReadListAttribute(Element, "Include");
-				string[] ExcludedNames = ReadListAttribute(Element, "Exclude");
+				string[] requiredNames = ReadListAttribute(element, "Requires");
+				string[] includedNames = ReadListAttribute(element, "Include");
+				string[] excludedNames = ReadListAttribute(element, "Exclude");
 
-				string UgsBadge = ReadAttribute(Element, "UgsBadge");
-				string UgsProject = ReadAttribute(Element, "UgsProject");
+				string ugsBadge = ReadAttribute(element, "UgsBadge");
+				string ugsProject = ReadAttribute(element, "UgsProject");
 
-				BgLabelChange Change = ReadEnumAttribute<BgLabelChange>(Element, "Change", BgLabelChange.Current);
+				BgLabelChange change = ReadEnumAttribute<BgLabelChange>(element, "Change", BgLabelChange.Current);
 
-				BgLabel NewLabel = new BgLabel(Name, Category, UgsBadge, UgsProject, Change);
-				foreach (BgNode ReferencedNode in ResolveReferences(Element, RequiredNames))
+				BgLabel newLabel = new BgLabel(name, category, ugsBadge, ugsProject, change);
+				foreach (BgNode referencedNode in ResolveReferences(element, requiredNames))
 				{
-					NewLabel.RequiredNodes.Add(ReferencedNode);
-					NewLabel.IncludedNodes.Add(ReferencedNode);
-					NewLabel.IncludedNodes.UnionWith(ReferencedNode.OrderDependencies);
+					newLabel.RequiredNodes.Add(referencedNode);
+					newLabel.IncludedNodes.Add(referencedNode);
+					newLabel.IncludedNodes.UnionWith(referencedNode.OrderDependencies);
 				}
-				foreach (BgNode IncludedNode in ResolveReferences(Element, IncludedNames))
+				foreach (BgNode includedNode in ResolveReferences(element, includedNames))
 				{
-					NewLabel.IncludedNodes.Add(IncludedNode);
-					NewLabel.IncludedNodes.UnionWith(IncludedNode.OrderDependencies);
+					newLabel.IncludedNodes.Add(includedNode);
+					newLabel.IncludedNodes.UnionWith(includedNode.OrderDependencies);
 				}
-				foreach (BgNode ExcludedNode in ResolveReferences(Element, ExcludedNames))
+				foreach (BgNode excludedNode in ResolveReferences(element, excludedNames))
 				{
-					NewLabel.IncludedNodes.Remove(ExcludedNode);
-					NewLabel.IncludedNodes.ExceptWith(ExcludedNode.OrderDependencies);
+					newLabel.IncludedNodes.Remove(excludedNode);
+					newLabel.IncludedNodes.ExceptWith(excludedNode.OrderDependencies);
 				}
-				Graph.Labels.Add(NewLabel);
+				_graph.Labels.Add(newLabel);
 			}
 		}
 
 		/// <summary>
 		/// Reads the definition for a node, and adds it to the given agent
 		/// </summary>
-		/// <param name="Element">Xml element to read the definition from</param>
-		protected override async Task ReadNodeAsync(BgScriptElement Element)
+		/// <param name="element">Xml element to read the definition from</param>
+		protected override async Task ReadNodeAsync(BgScriptElement element)
 		{
-			string? Name;
-			if (await EvaluateConditionAsync(Element) && TryReadObjectName(Element, out Name))
+			string? name;
+			if (await EvaluateConditionAsync(element) && TryReadObjectName(element, out name))
 			{
-				string[] RequiresNames = ReadListAttribute(Element, "Requires");
-				string[] ProducesNames = ReadListAttribute(Element, "Produces");
-				string[] AfterNames = ReadListAttribute(Element, "After");
-				string[] TokenFileNames = ReadListAttribute(Element, "Token");
-				bool bRunEarly = ReadBooleanAttribute(Element, "RunEarly", false);
-				bool bNotifyOnWarnings = ReadBooleanAttribute(Element, "NotifyOnWarnings", true);
+				string[] requiresNames = ReadListAttribute(element, "Requires");
+				string[] producesNames = ReadListAttribute(element, "Produces");
+				string[] afterNames = ReadListAttribute(element, "After");
+				string[] tokenFileNames = ReadListAttribute(element, "Token");
+				bool bRunEarly = ReadBooleanAttribute(element, "RunEarly", false);
+				bool bNotifyOnWarnings = ReadBooleanAttribute(element, "NotifyOnWarnings", true);
 
 				// Resolve all the inputs we depend on
-				HashSet<BgNodeOutput> Inputs = ResolveInputReferences(Element, RequiresNames);
+				HashSet<BgNodeOutput> inputs = ResolveInputReferences(element, requiresNames);
 
 				// Gather up all the input dependencies, and check they're all upstream of the current node
-				HashSet<BgNode> InputDependencies = new HashSet<BgNode>();
-				foreach (BgNode InputDependency in Inputs.Select(x => x.ProducingNode).Distinct())
+				HashSet<BgNode> inputDependencies = new HashSet<BgNode>();
+				foreach (BgNode inputDependency in inputs.Select(x => x.ProducingNode).Distinct())
 				{
-					InputDependencies.Add(InputDependency);
+					inputDependencies.Add(inputDependency);
 				}
 
 				// Remove all the lock names from the list of required names
-				HashSet<FileReference> RequiredTokens = new HashSet<FileReference>(TokenFileNames.Select(x => new FileReference(x)));
+				HashSet<FileReference> requiredTokens = new HashSet<FileReference>(tokenFileNames.Select(x => new FileReference(x)));
 
 				// Recursively include all their dependencies too
-				foreach (BgNode InputDependency in InputDependencies.ToArray())
+				foreach (BgNode inputDependency in inputDependencies.ToArray())
 				{
-					RequiredTokens.UnionWith(InputDependency.RequiredTokens);
-					InputDependencies.UnionWith(InputDependency.InputDependencies);
+					requiredTokens.UnionWith(inputDependency.RequiredTokens);
+					inputDependencies.UnionWith(inputDependency.InputDependencies);
 				}
 
 				// Validate all the outputs
-				List<string> ValidOutputNames = new List<string>();
-				foreach (string ProducesName in ProducesNames)
+				List<string> validOutputNames = new List<string>();
+				foreach (string producesName in producesNames)
 				{
-					BgNodeOutput? ExistingOutput;
-					if (Graph.TagNameToNodeOutput.TryGetValue(ProducesName, out ExistingOutput))
+					BgNodeOutput? existingOutput;
+					if (_graph.TagNameToNodeOutput.TryGetValue(producesName, out existingOutput))
 					{
-						LogError(Element, "Output tag '{0}' is already generated by node '{1}'", ProducesName, ExistingOutput.ProducingNode.Name);
+						LogError(element, "Output tag '{0}' is already generated by node '{1}'", producesName, existingOutput.ProducingNode.Name);
 					}
-					else if (!ProducesName.StartsWith("#"))
+					else if (!producesName.StartsWith("#"))
 					{
-						LogError(Element, "Output tag names must begin with a '#' character ('{0}')", ProducesName);
+						LogError(element, "Output tag names must begin with a '#' character ('{0}')", producesName);
 					}
 					else
 					{
-						ValidOutputNames.Add(ProducesName);
+						validOutputNames.Add(producesName);
 					}
 				}
 
 				// Gather up all the order dependencies
-				HashSet<BgNode> OrderDependencies = new HashSet<BgNode>(InputDependencies);
-				OrderDependencies.UnionWith(ResolveReferences(Element, AfterNames));
+				HashSet<BgNode> orderDependencies = new HashSet<BgNode>(inputDependencies);
+				orderDependencies.UnionWith(ResolveReferences(element, afterNames));
 
 				// Recursively include all their order dependencies too
-				foreach (BgNode OrderDependency in OrderDependencies.ToArray())
+				foreach (BgNode orderDependency in orderDependencies.ToArray())
 				{
-					OrderDependencies.UnionWith(OrderDependency.OrderDependencies);
+					orderDependencies.UnionWith(orderDependency.OrderDependencies);
 				}
 
 				// Check that we're not dependent on anything completing that is declared after the initial declaration of this agent.
-				int AgentIdx = Graph.Agents.IndexOf(EnclosingAgent!);
-				for (int Idx = AgentIdx + 1; Idx < Graph.Agents.Count; Idx++)
+				int agentIdx = _graph.Agents.IndexOf(_enclosingAgent!);
+				for (int idx = agentIdx + 1; idx < _graph.Agents.Count; idx++)
 				{
-					foreach (BgNode Node in Graph.Agents[Idx].Nodes.Where(x => OrderDependencies.Contains(x)))
+					foreach (BgNode node in _graph.Agents[idx].Nodes.Where(x => orderDependencies.Contains(x)))
 					{
-						LogError(Element, "Node '{0}' has a dependency on '{1}', which was declared after the initial definition of '{2}'.", Name, Node.Name, EnclosingAgent!.Name);
+						LogError(element, "Node '{0}' has a dependency on '{1}', which was declared after the initial definition of '{2}'.", name, node.Name, _enclosingAgent!.Name);
 					}
 				}
 
 				// Construct and register the node
-				if (CheckNameIsUnique(Element, Name))
+				if (CheckNameIsUnique(element, name))
 				{
 					// Add it to the node lookup
-					BgNode NewNode = new BgNode(Name, Inputs.ToArray(), ValidOutputNames.ToArray(), InputDependencies.ToArray(), OrderDependencies.ToArray(), RequiredTokens.ToArray());
-					NewNode.bRunEarly = bRunEarly;
-					NewNode.bNotifyOnWarnings = bNotifyOnWarnings;
-					Graph.NameToNode.Add(Name, NewNode);
+					BgNode newNode = new BgNode(name, inputs.ToArray(), validOutputNames.ToArray(), inputDependencies.ToArray(), orderDependencies.ToArray(), requiredTokens.ToArray());
+					newNode.BRunEarly = bRunEarly;
+					newNode.BNotifyOnWarnings = bNotifyOnWarnings;
+					_graph.NameToNode.Add(name, newNode);
 
 					// Register all the output tags in the global name table.
-					foreach (BgNodeOutput Output in NewNode.Outputs)
+					foreach (BgNodeOutput output in newNode.Outputs)
 					{
-						BgNodeOutput? ExistingOutput;
-						if (Graph.TagNameToNodeOutput.TryGetValue(Output.TagName, out ExistingOutput))
+						BgNodeOutput? existingOutput;
+						if (_graph.TagNameToNodeOutput.TryGetValue(output.TagName, out existingOutput))
 						{
-							LogError(Element, "Node '{0}' already has an output called '{1}'", ExistingOutput.ProducingNode.Name, Output.TagName);
+							LogError(element, "Node '{0}' already has an output called '{1}'", existingOutput.ProducingNode.Name, output.TagName);
 						}
 						else
 						{
-							Graph.TagNameToNodeOutput.Add(Output.TagName, Output);
+							_graph.TagNameToNodeOutput.Add(output.TagName, output);
 						}
 					}
 
 					// Add all the tasks
-					EnclosingNode = NewNode;
-					await ReadNodeBodyAsync(Element);
-					EnclosingNode = null;
+					_enclosingNode = newNode;
+					await ReadNodeBodyAsync(element);
+					_enclosingNode = null;
 
 					// Add it to the current agent
-					EnclosingAgent!.Nodes.Add(NewNode);
+					_enclosingAgent!.Nodes.Add(newNode);
 				}
 			}
 		}
@@ -1864,15 +1862,15 @@ namespace EpicGames.BuildGraph
 		/// <summary>
 		/// Reads a "Switch" element 
 		/// </summary>
-		/// <param name="Element">Xml element to read the definition from</param>
-		/// <param name="ReadContentsAsync">Delegate to read the contents of the element, if the condition evaluates to true</param>
-		protected override async Task ReadSwitchAsync(BgScriptElement Element, Func<BgScriptElement, Task> ReadContentsAsync)
+		/// <param name="element">Xml element to read the definition from</param>
+		/// <param name="readContentsAsync">Delegate to read the contents of the element, if the condition evaluates to true</param>
+		protected override async Task ReadSwitchAsync(BgScriptElement element, Func<BgScriptElement, Task> readContentsAsync)
 		{
-			foreach (BgScriptElement ChildElement in Element.ChildNodes.OfType<BgScriptElement>())
+			foreach (BgScriptElement childElement in element.ChildNodes.OfType<BgScriptElement>())
 			{
-				if (ChildElement.Name == "Default" || await EvaluateConditionAsync(ChildElement))
+				if (childElement.Name == "Default" || await EvaluateConditionAsync(childElement))
 				{
-					await ReadContentsAsync(ChildElement);
+					await readContentsAsync(childElement);
 				}
 			}
 		}
@@ -1880,117 +1878,117 @@ namespace EpicGames.BuildGraph
 		/// <summary>
 		/// Reads a task definition from the given element, and add it to the given list
 		/// </summary>
-		/// <param name="Element">Xml element to read the definition from</param>
-		protected override async Task ReadTaskAsync(BgScriptElement Element)
+		/// <param name="element">Xml element to read the definition from</param>
+		protected override async Task ReadTaskAsync(BgScriptElement element)
 		{
 			// If we're running a single node and this element's parent isn't the single node to run, ignore the error and return.
-			if (!string.IsNullOrWhiteSpace(SingleNodeName) && EnclosingNode!.Name != SingleNodeName)
+			if (!string.IsNullOrWhiteSpace(_singleNodeName) && _enclosingNode!.Name != _singleNodeName)
 			{
 				return;
 			}
 
-			if (await EvaluateConditionAsync(Element))
+			if (await EvaluateConditionAsync(element))
 			{
-				BgTask Info = new BgTask(Element.Location, Element.Name);
-				foreach (XmlAttribute? Attribute in Element.Attributes)
+				BgTask info = new BgTask(element.Location, element.Name);
+				foreach (XmlAttribute? attribute in element.Attributes)
 				{
-					if (String.Compare(Attribute!.Name, "If", StringComparison.InvariantCultureIgnoreCase) != 0)
+					if (String.Compare(attribute!.Name, "If", StringComparison.InvariantCultureIgnoreCase) != 0)
 					{
-						string ExpandedValue = ExpandProperties(Element, Attribute.Value);
-						Info.Arguments.Add(Attribute.Name, ExpandedValue);
+						string expandedValue = ExpandProperties(element, attribute.Value);
+						info.Arguments.Add(attribute.Name, expandedValue);
 					}
 				}
-				EnclosingNode!.Tasks.Add(Info);
+				_enclosingNode!.Tasks.Add(info);
 			}
 		}
 
 		/// <summary>
 		/// Reads the definition for an email notifier
 		/// </summary>
-		/// <param name="Element">Xml element to read the definition from</param>
-		protected override async Task ReadNotifierAsync(BgScriptElement Element)
+		/// <param name="element">Xml element to read the definition from</param>
+		protected override async Task ReadNotifierAsync(BgScriptElement element)
 		{
-			if (await EvaluateConditionAsync(Element))
+			if (await EvaluateConditionAsync(element))
 			{
-				string[] TargetNames = ReadListAttribute(Element, "Targets");
-				string[] ExceptNames = ReadListAttribute(Element, "Except");
-				string[] IndividualNodeNames = ReadListAttribute(Element, "Nodes");
-				string[] ReportNames = ReadListAttribute(Element, "Reports");
-				string[] Users = ReadListAttribute(Element, "Users");
-				string[] Submitters = ReadListAttribute(Element, "Submitters");
-				bool? bWarnings = Element.HasAttribute("Warnings") ? (bool?)ReadBooleanAttribute(Element, "Warnings", true) : null;
-				bool bAbsolute = Element.HasAttribute("Absolute") ? ReadBooleanAttribute(Element, "Absolute", true) : false;
+				string[] targetNames = ReadListAttribute(element, "Targets");
+				string[] exceptNames = ReadListAttribute(element, "Except");
+				string[] individualNodeNames = ReadListAttribute(element, "Nodes");
+				string[] reportNames = ReadListAttribute(element, "Reports");
+				string[] users = ReadListAttribute(element, "Users");
+				string[] submitters = ReadListAttribute(element, "Submitters");
+				bool? bWarnings = element.HasAttribute("Warnings") ? (bool?)ReadBooleanAttribute(element, "Warnings", true) : null;
+				bool bAbsolute = element.HasAttribute("Absolute") ? ReadBooleanAttribute(element, "Absolute", true) : false;
 
 				// Find the list of targets which are included, and recurse through all their dependencies
-				HashSet<BgNode> Nodes = new HashSet<BgNode>();
-				if (TargetNames != null)
+				HashSet<BgNode> nodes = new HashSet<BgNode>();
+				if (targetNames != null)
 				{
-					HashSet<BgNode> TargetNodes = ResolveReferences(Element, TargetNames);
-					foreach (BgNode Node in TargetNodes)
+					HashSet<BgNode> targetNodes = ResolveReferences(element, targetNames);
+					foreach (BgNode node in targetNodes)
 					{
-						Nodes.Add(Node);
-						Nodes.UnionWith(Node.InputDependencies);
+						nodes.Add(node);
+						nodes.UnionWith(node.InputDependencies);
 					}
 				}
 
 				// Add all the individually referenced nodes
-				if (IndividualNodeNames != null)
+				if (individualNodeNames != null)
 				{
-					HashSet<BgNode> IndividualNodes = ResolveReferences(Element, IndividualNodeNames);
-					Nodes.UnionWith(IndividualNodes);
+					HashSet<BgNode> individualNodes = ResolveReferences(element, individualNodeNames);
+					nodes.UnionWith(individualNodes);
 				}
 
 				// Exclude all the exceptions
-				if (ExceptNames != null)
+				if (exceptNames != null)
 				{
-					HashSet<BgNode> ExceptNodes = ResolveReferences(Element, ExceptNames);
-					Nodes.ExceptWith(ExceptNodes);
+					HashSet<BgNode> exceptNodes = ResolveReferences(element, exceptNames);
+					nodes.ExceptWith(exceptNodes);
 				}
 
 				// Update all the referenced nodes with the settings
-				foreach (BgNode Node in Nodes)
+				foreach (BgNode node in nodes)
 				{
-					if (Users != null)
+					if (users != null)
 					{
 						if (bAbsolute)
 						{
-							Node.NotifyUsers = new HashSet<string>(Users);
+							node.NotifyUsers = new HashSet<string>(users);
 						}
 						else
 						{
-							Node.NotifyUsers.UnionWith(Users);
+							node.NotifyUsers.UnionWith(users);
 						}
 					}
-					if (Submitters != null)
+					if (submitters != null)
 					{
 						if (bAbsolute)
 						{
-							Node.NotifySubmitters = new HashSet<string>(Submitters);
+							node.NotifySubmitters = new HashSet<string>(submitters);
 						}
 						else
 						{
-							Node.NotifySubmitters.UnionWith(Submitters);
+							node.NotifySubmitters.UnionWith(submitters);
 						}
 					}
 					if (bWarnings.HasValue)
 					{
-						Node.bNotifyOnWarnings = bWarnings.Value;
+						node.BNotifyOnWarnings = bWarnings.Value;
 					}
 				}
 
 				// Add the users to the list of reports
-				if (ReportNames != null)
+				if (reportNames != null)
 				{
-					foreach (string ReportName in ReportNames)
+					foreach (string reportName in reportNames)
 					{
-						BgReport? Report;
-						if (Graph.NameToReport.TryGetValue(ReportName, out Report))
+						BgReport? report;
+						if (_graph.NameToReport.TryGetValue(reportName, out report))
 						{
-							Report.NotifyUsers.UnionWith(Users);
+							report.NotifyUsers.UnionWith(users);
 						}
 						else
 						{
-							LogError(Element, "Report '{0}' has not been defined", ReportName);
+							LogError(element, "Report '{0}' has not been defined", reportName);
 						}
 					}
 				}
@@ -2000,31 +1998,31 @@ namespace EpicGames.BuildGraph
 		/// <summary>
 		/// Reads a warning from the given element, evaluates the condition on it, and writes it to the log if the condition passes.
 		/// </summary>
-		/// <param name="Element">Xml element to read the definition from</param>
-		/// <param name="EventType">The diagnostic event type</param>
-		protected override async Task ReadDiagnosticAsync(BgScriptElement Element, LogEventType EventType)
+		/// <param name="element">Xml element to read the definition from</param>
+		/// <param name="eventType">The diagnostic event type</param>
+		protected override async Task ReadDiagnosticAsync(BgScriptElement element, LogEventType eventType)
 		{
-			if (await EvaluateConditionAsync(Element))
+			if (await EvaluateConditionAsync(element))
 			{
-				string Message = ReadAttribute(Element, "Message");
+				string message = ReadAttribute(element, "Message");
 
-				BgGraphDiagnostic Diagnostic = new BgGraphDiagnostic(Element.Location, EventType, Message, EnclosingNode, EnclosingAgent);
-				Graph.Diagnostics.Add(Diagnostic);
+				BgGraphDiagnostic diagnostic = new BgGraphDiagnostic(element.Location, eventType, message, _enclosingNode, _enclosingAgent);
+				_graph.Diagnostics.Add(diagnostic);
 			}
 		}
 
 		/// <summary>
 		/// Checks that the given name does not already used to refer to a node, and print an error if it is.
 		/// </summary>
-		/// <param name="Element">Xml element to read from</param>
-		/// <param name="Name">Name of the alias</param>
+		/// <param name="element">Xml element to read from</param>
+		/// <param name="name">Name of the alias</param>
 		/// <returns>True if the name was registered correctly, false otherwise.</returns>
-		bool CheckNameIsUnique(BgScriptElement Element, string Name)
+		bool CheckNameIsUnique(BgScriptElement element, string name)
 		{
 			// Get the nodes that it maps to
-			if (Graph.ContainsName(Name))
+			if (_graph.ContainsName(name))
 			{
-				LogError(Element, "'{0}' is already defined; cannot add a second time", Name);
+				LogError(element, "'{0}' is already defined; cannot add a second time", name);
 				return false;
 			}
 			return true;
@@ -2033,57 +2031,57 @@ namespace EpicGames.BuildGraph
 		/// <summary>
 		/// Resolve a list of references to a set of nodes
 		/// </summary>
-		/// <param name="Element">Element used to locate any errors</param>
-		/// <param name="ReferenceNames">Sequence of names to look up</param>
+		/// <param name="element">Element used to locate any errors</param>
+		/// <param name="referenceNames">Sequence of names to look up</param>
 		/// <returns>Hashset of all the nodes included by the given names</returns>
-		HashSet<BgNode> ResolveReferences(BgScriptElement Element, IEnumerable<string> ReferenceNames)
+		HashSet<BgNode> ResolveReferences(BgScriptElement element, IEnumerable<string> referenceNames)
 		{
-			HashSet<BgNode> Nodes = new HashSet<BgNode>();
-			foreach (string ReferenceName in ReferenceNames)
+			HashSet<BgNode> nodes = new HashSet<BgNode>();
+			foreach (string referenceName in referenceNames)
 			{
-				BgNode[]? OtherNodes;
-				if (Graph.TryResolveReference(ReferenceName, out OtherNodes))
+				BgNode[]? otherNodes;
+				if (_graph.TryResolveReference(referenceName, out otherNodes))
 				{
-					Nodes.UnionWith(OtherNodes);
+					nodes.UnionWith(otherNodes);
 				}
-				else if (!ReferenceName.StartsWith("#") && Graph.TagNameToNodeOutput.ContainsKey("#" + ReferenceName))
+				else if (!referenceName.StartsWith("#") && _graph.TagNameToNodeOutput.ContainsKey("#" + referenceName))
 				{
-					LogError(Element, "Reference to '{0}' cannot be resolved; did you mean '#{0}'?", ReferenceName);
+					LogError(element, "Reference to '{0}' cannot be resolved; did you mean '#{0}'?", referenceName);
 				}
 				else
 				{
-					LogError(Element, "Reference to '{0}' cannot be resolved; check it has been defined.", ReferenceName);
+					LogError(element, "Reference to '{0}' cannot be resolved; check it has been defined.", referenceName);
 				}
 			}
-			return Nodes;
+			return nodes;
 		}
 
 		/// <summary>
 		/// Resolve a list of references to a set of nodes
 		/// </summary>
-		/// <param name="Element">Element used to locate any errors</param>
-		/// <param name="ReferenceNames">Sequence of names to look up</param>
+		/// <param name="element">Element used to locate any errors</param>
+		/// <param name="referenceNames">Sequence of names to look up</param>
 		/// <returns>Set of all the nodes included by the given names</returns>
-		HashSet<BgNodeOutput> ResolveInputReferences(BgScriptElement Element, IEnumerable<string> ReferenceNames)
+		HashSet<BgNodeOutput> ResolveInputReferences(BgScriptElement element, IEnumerable<string> referenceNames)
 		{
-			HashSet<BgNodeOutput> Inputs = new HashSet<BgNodeOutput>();
-			foreach (string ReferenceName in ReferenceNames)
+			HashSet<BgNodeOutput> inputs = new HashSet<BgNodeOutput>();
+			foreach (string referenceName in referenceNames)
 			{
-				BgNodeOutput[]? ReferenceInputs;
-				if (Graph.TryResolveInputReference(ReferenceName, out ReferenceInputs))
+				BgNodeOutput[]? referenceInputs;
+				if (_graph.TryResolveInputReference(referenceName, out referenceInputs))
 				{
-					Inputs.UnionWith(ReferenceInputs);
+					inputs.UnionWith(referenceInputs);
 				}
-				else if (!ReferenceName.StartsWith("#") && Graph.TagNameToNodeOutput.ContainsKey("#" + ReferenceName))
+				else if (!referenceName.StartsWith("#") && _graph.TagNameToNodeOutput.ContainsKey("#" + referenceName))
 				{
-					LogError(Element, "Reference to '{0}' cannot be resolved; did you mean '#{0}'?", ReferenceName);
+					LogError(element, "Reference to '{0}' cannot be resolved; did you mean '#{0}'?", referenceName);
 				}
 				else
 				{
-					LogError(Element, "Reference to '{0}' cannot be resolved; check it has been defined.", ReferenceName);
+					LogError(element, "Reference to '{0}' cannot be resolved; check it has been defined.", referenceName);
 				}
 			}
-			return Inputs;
+			return inputs;
 		}
 	}
 }
\ No newline at end of file
