Reformat EpicGames.BuildGraph according to style guidelines.

diff --git a/Engine/Source/Programs/Shared/EpicGames.BuildGraph/BgTask.cs b/Engine/Source/Programs/Shared/EpicGames.BuildGraph/BgTask.cs
--- a/Engine/Source/Programs/Shared/EpicGames.BuildGraph/BgTask.cs
+++ b/Engine/Source/Programs/Shared/EpicGames.BuildGraph/BgTask.cs
@@ -1,10 +1,6 @@
 // Copyright Epic Games, Inc. All Rights Reserved.
 
-using EpicGames.Core;
-using Microsoft.Extensions.Logging;
-using System;
 using System.Collections.Generic;
-using System.Text;
 using System.Xml;
 
 namespace EpicGames.BuildGraph
@@ -32,24 +28,24 @@ namespace EpicGames.BuildGraph
 		/// <summary>
 		/// Constructor
 		/// </summary>
-		public BgTask(BgScriptLocation Location, string Name)
+		public BgTask(BgScriptLocation location, string name)
 		{
-			this.Location = Location;
-			this.Name = Name;
+			Location = location;
+			Name = name;
 		}
 
 		/// <summary>
 		/// Write to an xml file
 		/// </summary>
-		/// <param name="Writer"></param>
-		public void Write(XmlWriter Writer)
+		/// <param name="writer"></param>
+		public void Write(XmlWriter writer)
 		{
-			Writer.WriteStartElement(Name);
-			foreach (KeyValuePair<string, string> Argument in Arguments)
+			writer.WriteStartElement(Name);
+			foreach (KeyValuePair<string, string> argument in Arguments)
 			{
-				Writer.WriteAttributeString(Argument.Key, Argument.Value);
+				writer.WriteAttributeString(argument.Key, argument.Value);
 			}
-			Writer.WriteEndElement();
+			writer.WriteEndElement();
 		}
 	}
 }
diff --git a/Engine/Source/Programs/Shared/EpicGames.BuildGraph/Expressions/BgBool.cs b/Engine/Source/Programs/Shared/EpicGames.BuildGraph/Expressions/BgBool.cs
--- a/Engine/Source/Programs/Shared/EpicGames.BuildGraph/Expressions/BgBool.cs
+++ b/Engine/Source/Programs/Shared/EpicGames.BuildGraph/Expressions/BgBool.cs
@@ -1,17 +1,6 @@
 // Copyright Epic Games, Inc. All Rights Reserved.
 
 using System;
-using System.Collections.Generic;
-using System.ComponentModel;
-using System.Diagnostics;
-using System.IO;
-using System.Linq;
-using System.Reflection;
-using System.Text;
-using System.Xml;
-using EpicGames.Core;
-using Microsoft.Extensions.Logging;
-using System.Threading.Tasks;
 
 namespace EpicGames.BuildGraph.Expressions
 {
@@ -32,36 +21,36 @@ namespace EpicGames.BuildGraph.Expressions
 		public static BgBool True { get; } = true;
 
 		/// <inheritdoc/>
-		public BgBool IfThen(BgBool Condition, BgBool ValueIfTrue) => new BgBoolChooseExpr(Condition, ValueIfTrue, this);
+		public BgBool IfThen(BgBool condition, BgBool valueIfTrue) => new BgBoolChooseExpr(condition, valueIfTrue, this);
 
 		/// <summary>
 		/// Implict conversion operator from a boolean literal
 		/// </summary>
-		public static implicit operator BgBool(bool Value)
+		public static implicit operator BgBool(bool value)
 		{
-			return new BgBoolConstantExpr(Value);
+			return new BgBoolConstantExpr(value);
 		}
 
 		/// <inheritdoc/>
-		public static BgBool operator !(BgBool Inner) => new BgBoolNotExpr(Inner);
+		public static BgBool operator !(BgBool inner) => new BgBoolNotExpr(inner);
 
 		/// <inheritdoc/>
-		public static BgBool operator &(BgBool Lhs, BgBool Rhs) => new BgBoolAndExpr(Lhs, Rhs);
+		public static BgBool operator &(BgBool lhs, BgBool rhs) => new BgBoolAndExpr(lhs, rhs);
 
 		/// <inheritdoc/>
-		public static BgBool operator |(BgBool Lhs, BgBool Rhs) => new BgBoolOrExpr(Lhs, Rhs);
+		public static BgBool operator |(BgBool lhs, BgBool rhs) => new BgBoolOrExpr(lhs, rhs);
 
 		/// <inheritdoc/>
-		public static BgBool operator ==(BgBool Lhs, BgBool Rhs) => new BgBoolEqExpr(Lhs, Rhs);
+		public static BgBool operator ==(BgBool lhs, BgBool rhs) => new BgBoolEqExpr(lhs, rhs);
 
 		/// <inheritdoc/>
-		public static BgBool operator !=(BgBool Lhs, BgBool Rhs) => new BgBoolNotExpr(new BgBoolEqExpr(Lhs, Rhs));
+		public static BgBool operator !=(BgBool lhs, BgBool rhs) => new BgBoolNotExpr(new BgBoolEqExpr(lhs, rhs));
 
 		/// <inheritdoc/>
-		object IBgExpr.Compute(BgExprContext Context) => Compute(Context);
+		object IBgExpr.Compute(BgExprContext context) => Compute(context);
 
 		/// <inheritdoc/>
-		public abstract bool Compute(BgExprContext Context);
+		public abstract bool Compute(BgExprContext context);
 
 		/// <inheritdoc/>
 		public sealed override bool Equals(object? obj) => throw new InvalidOperationException();
@@ -79,13 +68,13 @@ namespace EpicGames.BuildGraph.Expressions
 	class BgBoolType : BgTypeBase<BgBool>
 	{
 		/// <inheritdoc/>
-		public override BgBool DeserializeArgument(string Text) => bool.Parse(Text);
+		public override BgBool DeserializeArgument(string text) => bool.Parse(text);
 
 		/// <inheritdoc/>
-		public override string SerializeArgument(BgBool Value, BgExprContext Context) => Value.Compute(Context).ToString();
+		public override string SerializeArgument(BgBool value, BgExprContext context) => value.Compute(context).ToString();
 
 		/// <inheritdoc/>
-		public override BgBool CreateConstant(object Value) => new BgBoolConstantExpr((bool)Value);
+		public override BgBool CreateConstant(object value) => new BgBoolConstantExpr((bool)value);
 
 		/// <inheritdoc/>
 		public override IBgExprVariable<BgBool> CreateVariable() => new BgBoolVariableExpr();
@@ -97,12 +86,12 @@ namespace EpicGames.BuildGraph.Expressions
 	{
 		public BgBool Inner { get; }
 
-		public BgBoolNotExpr(BgBool Inner)
+		public BgBoolNotExpr(BgBool inner)
 		{
-			this.Inner = Inner;
+			Inner = inner;
 		}
 
-		public override bool Compute(BgExprContext Context) => !Inner.Compute(Context);
+		public override bool Compute(BgExprContext context) => !Inner.Compute(context);
 	}
 
 	class BgBoolEqExpr : BgBool
@@ -110,13 +99,13 @@ namespace EpicGames.BuildGraph.Expressions
 		public BgBool Lhs { get; }
 		public BgBool Rhs { get; }
 
-		public BgBoolEqExpr(BgBool Lhs, BgBool Rhs)
+		public BgBoolEqExpr(BgBool lhs, BgBool rhs)
 		{
-			this.Lhs = Lhs;
-			this.Rhs = Rhs;
+			Lhs = lhs;
+			Rhs = rhs;
 		}
 
-		public override bool Compute(BgExprContext Context) => Lhs.Compute(Context) == Rhs.Compute(Context);
+		public override bool Compute(BgExprContext context) => Lhs.Compute(context) == Rhs.Compute(context);
 	}
 
 	class BgBoolAndExpr : BgBool
@@ -124,13 +113,13 @@ namespace EpicGames.BuildGraph.Expressions
 		public BgBool Lhs { get; }
 		public BgBool Rhs { get; }
 
-		public BgBoolAndExpr(BgBool Lhs, BgBool Rhs)
+		public BgBoolAndExpr(BgBool lhs, BgBool rhs)
 		{
-			this.Lhs = Lhs;
-			this.Rhs = Rhs;
+			Lhs = lhs;
+			Rhs = rhs;
 		}
 
-		public override bool Compute(BgExprContext Context) => Lhs.Compute(Context) && Rhs.Compute(Context);
+		public override bool Compute(BgExprContext context) => Lhs.Compute(context) && Rhs.Compute(context);
 	}
 
 	class BgBoolOrExpr : BgBool
@@ -138,60 +127,60 @@ namespace EpicGames.BuildGraph.Expressions
 		public BgBool Lhs { get; }
 		public BgBool Rhs { get; }
 
-		public BgBoolOrExpr(BgBool Lhs, BgBool Rhs)
+		public BgBoolOrExpr(BgBool lhs, BgBool rhs)
 		{
-			this.Lhs = Lhs;
-			this.Rhs = Rhs;
+			Lhs = lhs;
+			Rhs = rhs;
 		}
 
-		public override bool Compute(BgExprContext Context) => Lhs.Compute(Context) || Rhs.Compute(Context);
+		public override bool Compute(BgExprContext context) => Lhs.Compute(context) || Rhs.Compute(context);
 	}
 
 	class BgBoolChooseExpr : BgBool
 	{
-		public BgBool Condition;
-		public BgBool ValueIfTrue;
-		public BgBool ValueIfFalse;
+		public BgBool Condition { get; }
+		public BgBool ValueIfTrue { get; }
+		public BgBool ValueIfFalse { get; }
 
-		public BgBoolChooseExpr(BgBool Condition, BgBool ValueIfTrue, BgBool ValueIfFalse)
+		public BgBoolChooseExpr(BgBool condition, BgBool valueIfTrue, BgBool valueIfFalse)
 		{
-			this.Condition = Condition;
-			this.ValueIfTrue = ValueIfTrue;
-			this.ValueIfFalse = ValueIfFalse;
+			Condition = condition;
+			ValueIfTrue = valueIfTrue;
+			ValueIfFalse = valueIfFalse;
 		}
 
-		public override bool Compute(BgExprContext Context) => Condition.Compute(Context) ? ValueIfTrue.Compute(Context) : ValueIfFalse.Compute(Context);
+		public override bool Compute(BgExprContext context) => Condition.Compute(context) ? ValueIfTrue.Compute(context) : ValueIfFalse.Compute(context);
 	}
 
 	class BgBoolVariableExpr : BgBool, IBgExprVariable<BgBool>
 	{
 		public BgBool Value { get; set; } = BgBool.False;
 
-		public override bool Compute(BgExprContext Context) => Value.Compute(Context);
+		public override bool Compute(BgExprContext context) => Value.Compute(context);
 	}
 
 	class BgBoolConstantExpr : BgBool
 	{
 		public bool Value { get; }
 
-		public BgBoolConstantExpr(bool Value)
+		public BgBoolConstantExpr(bool value)
 		{
-			this.Value = Value;
+			Value = value;
 		}
 
-		public override bool Compute(BgExprContext Context) => Value;
+		public override bool Compute(BgExprContext context) => Value;
 	}
 
 	class BgBoolFormatExpr : BgString
 	{
-		BgBool Value;
+		public BgBool Value { get; }
 
-		public BgBoolFormatExpr(BgBool Value)
+		public BgBoolFormatExpr(BgBool value)
 		{
-			this.Value = Value;
+			Value = value;
 		}
 
-		public override string Compute(BgExprContext Context) => Value.Compute(Context).ToString();
+		public override string Compute(BgExprContext context) => Value.Compute(context).ToString();
 	}
 
 	#endregion
diff --git a/Engine/Source/Programs/Shared/EpicGames.BuildGraph/Expressions/BgEnum.cs b/Engine/Source/Programs/Shared/EpicGames.BuildGraph/Expressions/BgEnum.cs
--- a/Engine/Source/Programs/Shared/EpicGames.BuildGraph/Expressions/BgEnum.cs
+++ b/Engine/Source/Programs/Shared/EpicGames.BuildGraph/Expressions/BgEnum.cs
@@ -1,21 +1,6 @@
 // Copyright Epic Games, Inc. All Rights Reserved.
 
-using EpicGames.BuildGraph;
 using System;
-using System.Collections.Generic;
-using System.ComponentModel;
-using System.Diagnostics;
-using System.IO;
-using System.Linq;
-using System.Reflection;
-using System.Text;
-using System.Xml;
-using EpicGames.Core;
-using Microsoft.Extensions.Logging;
-using OpenTracing;
-using OpenTracing.Util;
-using System.Threading.Tasks;
-using System.Text.RegularExpressions;
 
 namespace EpicGames.BuildGraph.Expressions
 {
@@ -42,27 +27,27 @@ namespace EpicGames.BuildGraph.Expressions
 		/// <summary>
 		/// Implicit conversion from a regular enum type
 		/// </summary>
-		public static implicit operator BgEnum<TEnum>(TEnum Value)
+		public static implicit operator BgEnum<TEnum>(TEnum value)
 		{
-			return new BgEnumConstantExpr<TEnum>(Value);
+			return new BgEnumConstantExpr<TEnum>(value);
 		}
 
 		/// <summary>
 		/// Explicit conversion from a string value
 		/// </summary>
-		public static explicit operator BgEnum<TEnum>(BgString Value)
+		public static explicit operator BgEnum<TEnum>(BgString value)
 		{
-			return new BgEnumParseExpr<TEnum>(Value);
+			return new BgEnumParseExpr<TEnum>(value);
 		}
 
 		/// <inheritdoc/>
-		public BgEnum<TEnum> IfThen(BgBool Condition, BgEnum<TEnum> ValueIfTrue) => new BgEnumChooseExpr<TEnum>(Condition, ValueIfTrue, this);
+		public BgEnum<TEnum> IfThen(BgBool condition, BgEnum<TEnum> valueIfTrue) => new BgEnumChooseExpr<TEnum>(condition, valueIfTrue, this);
 
 		/// <inheritdoc/>
-		object IBgExpr.Compute(BgExprContext Context) => Compute(Context);
+		object IBgExpr.Compute(BgExprContext context) => Compute(context);
 
 		/// <inheritdoc/>
-		public abstract TEnum Compute(BgExprContext Context);
+		public abstract TEnum Compute(BgExprContext context);
 
 		/// <inheritdoc/>
 		public BgString ToBgString() => new BgEnumFormatExpr<TEnum>(this);
@@ -74,13 +59,13 @@ namespace EpicGames.BuildGraph.Expressions
 	class BgEnumType<TEnum> : BgTypeBase<BgEnum<TEnum>> where TEnum : struct
 	{
 		/// <inheritdoc/>
-		public override BgEnum<TEnum> DeserializeArgument(string Text) => Enum.Parse<TEnum>(Text);
+		public override BgEnum<TEnum> DeserializeArgument(string text) => Enum.Parse<TEnum>(text);
 
 		/// <inheritdoc/>
-		public override string SerializeArgument(BgEnum<TEnum> Value, BgExprContext Context) => Value.Compute(Context).ToString() ?? String.Empty;
+		public override string SerializeArgument(BgEnum<TEnum> value, BgExprContext context) => value.Compute(context).ToString() ?? String.Empty;
 
 		/// <inheritdoc/>
-		public override BgEnum<TEnum> CreateConstant(object Value) => new BgEnumConstantExpr<TEnum>((TEnum)Value);
+		public override BgEnum<TEnum> CreateConstant(object value) => new BgEnumConstantExpr<TEnum>((TEnum)value);
 
 		/// <inheritdoc/>
 		public override IBgExprVariable<BgEnum<TEnum>> CreateVariable() => throw new NotImplementedException();
@@ -92,12 +77,12 @@ namespace EpicGames.BuildGraph.Expressions
 	{
 		public BgString Value { get; }
 
-		public BgEnumParseExpr(BgString Value)
+		public BgEnumParseExpr(BgString value)
 		{
-			this.Value = Value;
+			Value = value;
 		}
 
-		public override TEnum Compute(BgExprContext Context) => Enum.Parse<TEnum>(Value.Compute(Context));
+		public override TEnum Compute(BgExprContext context) => Enum.Parse<TEnum>(Value.Compute(context));
 
 		public override string ToString()
 		{
@@ -109,12 +94,12 @@ namespace EpicGames.BuildGraph.Expressions
 	{
 		public TEnum Value { get; }
 
-		public BgEnumConstantExpr(TEnum Value)
+		public BgEnumConstantExpr(TEnum value)
 		{
-			this.Value = Value;
+			Value = value;
 		}
 
-		public override TEnum Compute(BgExprContext Context) => Value;
+		public override TEnum Compute(BgExprContext context) => Value;
 
 		public override string ToString()
 		{
@@ -128,14 +113,14 @@ namespace EpicGames.BuildGraph.Expressions
 		public BgEnum<TEnum> ValueIfTrue { get; }
 		public BgEnum<TEnum> ValueIfFalse { get; }
 
-		public BgEnumChooseExpr(BgBool Condition, BgEnum<TEnum> ValueIfTrue, BgEnum<TEnum> ValueIfFalse)
+		public BgEnumChooseExpr(BgBool condition, BgEnum<TEnum> valueIfTrue, BgEnum<TEnum> valueIfFalse)
 		{
-			this.Condition = Condition;
-			this.ValueIfTrue = ValueIfTrue;
-			this.ValueIfFalse = ValueIfFalse;
+			Condition = condition;
+			ValueIfTrue = valueIfTrue;
+			ValueIfFalse = valueIfFalse;
 		}
 
-		public override TEnum Compute(BgExprContext Context) => Condition.Compute(Context) ? ValueIfTrue.Compute(Context) : ValueIfFalse.Compute(Context);
+		public override TEnum Compute(BgExprContext context) => Condition.Compute(context) ? ValueIfTrue.Compute(context) : ValueIfFalse.Compute(context);
 
 		public override string ToString()
 		{
@@ -147,12 +132,12 @@ namespace EpicGames.BuildGraph.Expressions
 	{
 		public BgEnum<TEnum> Value { get; }
 
-		public BgEnumFormatExpr(BgEnum<TEnum> Value)
+		public BgEnumFormatExpr(BgEnum<TEnum> value)
 		{
-			this.Value = Value;
+			Value = value;
 		}
 
-		public override string Compute(BgExprContext Context) => Value.Compute(Context).ToString() ?? String.Empty;
+		public override string Compute(BgExprContext context) => Value.Compute(context).ToString() ?? String.Empty;
 
 		public override string ToString()
 		{
@@ -164,7 +149,7 @@ namespace EpicGames.BuildGraph.Expressions
 	{
 		public BgEnum<TEnum> Value { get; set; } = null!;
 
-		public override TEnum Compute(BgExprContext Context) => Value.Compute(Context);
+		public override TEnum Compute(BgExprContext context) => Value.Compute(context);
 	}
 
 	#endregion
diff --git a/Engine/Source/Programs/Shared/EpicGames.BuildGraph/Expressions/BgFileSet.cs b/Engine/Source/Programs/Shared/EpicGames.BuildGraph/Expressions/BgFileSet.cs
--- a/Engine/Source/Programs/Shared/EpicGames.BuildGraph/Expressions/BgFileSet.cs
+++ b/Engine/Source/Programs/Shared/EpicGames.BuildGraph/Expressions/BgFileSet.cs
@@ -1,10 +1,7 @@
 // Copyright Epic Games, Inc. All Rights Reserved.
 
-using EpicGames.Core;
 using System;
-using System.Collections.Generic;
-using System.Linq;
-using System.Text;
+using EpicGames.Core;
 
 namespace EpicGames.BuildGraph.Expressions
 {
@@ -15,31 +12,31 @@ namespace EpicGames.BuildGraph.Expressions
 	public abstract class BgFileSet : IBgExpr<BgFileSet>
 	{
 		/// <inheritdoc/>
-		public BgFileSet IfThen(BgBool Condition, BgFileSet ValueIfTrue) => new BgFileSetIfThenExpr(Condition, ValueIfTrue, this);
+		public BgFileSet IfThen(BgBool condition, BgFileSet valueIfTrue) => new BgFileSetIfThenExpr(condition, valueIfTrue, this);
 
 		/// <inheritdoc/>
-		object IBgExpr.Compute(BgExprContext Context) => Compute(Context);
+		object IBgExpr.Compute(BgExprContext context) => Compute(context);
 
 		/// <inheritdoc/>
-		public abstract BgFileSet Compute(BgExprContext Context);
+		public abstract BgFileSet Compute(BgExprContext context);
 
 		/// <summary>
 		/// Gets the tag name for this fileset
 		/// </summary>
 		/// <returns></returns>
-		public string ComputeTag(BgExprContext Context) => ((BgFileSetTagExpr)Compute(Context)).Name;
+		public string ComputeTag(BgExprContext context) => ((BgFileSetTagExpr)Compute(context)).Name;
 
 		/// <summary>
 		/// Gets the tag name for this fileset
 		/// </summary>
 		/// <returns></returns>
-		public FileSet ComputeValue(BgExprContext Context) => ((BgFileSetValueExpr)Compute(Context)).FileSet;
+		public FileSet ComputeValue(BgExprContext context) => ((BgFileSetValueExpr)Compute(context)).FileSet;
 
 		/// <summary>
 		/// Implicit conversion from a file set to a functional file set
 		/// </summary>
-		/// <param name="FileSet"></param>
-		public static implicit operator BgFileSet(FileSet FileSet) => new BgFileSetValueExpr(FileSet);
+		/// <param name="fileSet"></param>
+		public static implicit operator BgFileSet(FileSet fileSet) => new BgFileSetValueExpr(fileSet);
 
 		/// <inheritdoc/>
 		public BgString ToBgString() => new BgFileSetToStringExpr(this);
@@ -51,13 +48,13 @@ namespace EpicGames.BuildGraph.Expressions
 	class BgFileSetType : BgTypeBase<BgFileSet>
 	{
 		/// <inheritdoc/>
-		public override BgFileSet DeserializeArgument(string Text) => new BgFileSetTagExpr(Text);
+		public override BgFileSet DeserializeArgument(string text) => new BgFileSetTagExpr(text);
 
 		/// <inheritdoc/>
-		public override string SerializeArgument(BgFileSet Value, BgExprContext Context) => ((BgFileSetTagExpr)Value.Compute(Context)).Name;
+		public override string SerializeArgument(BgFileSet value, BgExprContext context) => ((BgFileSetTagExpr)value.Compute(context)).Name;
 
 		/// <inheritdoc/>
-		public override BgFileSet CreateConstant(object Value) => new BgFileSetTagExpr(((BgFileSetTagExpr)Value).Name);
+		public override BgFileSet CreateConstant(object value) => new BgFileSetTagExpr(((BgFileSetTagExpr)value).Name);
 
 		/// <inheritdoc/>
 		public override IBgExprVariable<BgFileSet> CreateVariable() => throw new NotImplementedException();
@@ -67,35 +64,35 @@ namespace EpicGames.BuildGraph.Expressions
 
 	class BgFileSetIfThenExpr : BgFileSet
 	{
-		public BgBool Condition;
-		public BgFileSet ValueIfTrue;
-		public BgFileSet ValueIfFalse;
+		public BgBool Condition { get; }
+		public BgFileSet ValueIfTrue { get; }
+		public BgFileSet ValueIfFalse { get; }
 
-		public BgFileSetIfThenExpr(BgBool Condition, BgFileSet ValueIfTrue, BgFileSet ValueIfFalse)
+		public BgFileSetIfThenExpr(BgBool condition, BgFileSet valueIfTrue, BgFileSet valueIfFalse)
 		{
-			this.Condition = Condition;
-			this.ValueIfTrue = ValueIfTrue;
-			this.ValueIfFalse = ValueIfFalse;
+			Condition = condition;
+			ValueIfTrue = valueIfTrue;
+			ValueIfFalse = valueIfFalse;
 		}
 
-		public override BgFileSet Compute(BgExprContext Context) => Condition.Compute(Context) ? ValueIfTrue.Compute(Context) : ValueIfFalse.Compute(Context);
+		public override BgFileSet Compute(BgExprContext context) => Condition.Compute(context) ? ValueIfTrue.Compute(context) : ValueIfFalse.Compute(context);
 	}
 
 	class BgFileSetToStringExpr : BgString
 	{
-		BgFileSet FileSet;
+		public BgFileSet FileSet { get; }
 
-		internal BgFileSetToStringExpr(BgFileSet Token)
+		internal BgFileSetToStringExpr(BgFileSet token)
 		{
-			this.FileSet = Token;
+			FileSet = token;
 		}
 
-		public override string Compute(BgExprContext Context)
+		public override string Compute(BgExprContext context)
 		{
-			BgFileSet FileSetValue = FileSet.Compute(Context);
-			if (FileSetValue is BgFileSetTagExpr TagExpr)
+			BgFileSet fileSetValue = FileSet.Compute(context);
+			if (fileSetValue is BgFileSetTagExpr tagExpr)
 			{
-				return TagExpr.Name;
+				return tagExpr.Name;
 			}
 			else
 			{
@@ -108,43 +105,43 @@ namespace EpicGames.BuildGraph.Expressions
 	{
 		public string Name { get; set; }
 
-		public BgFileSetTagExpr(string Name)
+		public BgFileSetTagExpr(string name)
 		{
-			this.Name = Name;
+			Name = name;
 		}
 
-		public override BgFileSet Compute(BgExprContext Context) => this;
+		public override BgFileSet Compute(BgExprContext context) => this;
 	}
 
 	class BgFileSetTagFromStringExpr : BgFileSet
 	{
 		public BgString Name { get; }
 
-		public BgFileSetTagFromStringExpr(BgString Name)
+		public BgFileSetTagFromStringExpr(BgString name)
 		{
-			this.Name = Name;
+			Name = name;
 		}
 
-		public override BgFileSet Compute(BgExprContext Context) => new BgFileSetTagExpr(Name.Compute(Context));
+		public override BgFileSet Compute(BgExprContext context) => new BgFileSetTagExpr(Name.Compute(context));
 	}
 
 	class BgFileSetVariableExpr : BgFileSet, IBgExprVariable<BgFileSet>
 	{
 		public BgFileSet Value { get; set; } = new BgFileSetValueExpr(FileSet.Empty);
 
-		public override BgFileSet Compute(BgExprContext Context) => Value.Compute(Context);
+		public override BgFileSet Compute(BgExprContext context) => Value.Compute(context);
 	}
 
 	class BgFileSetValueExpr : BgFileSet
 	{
 		public FileSet FileSet { get; }
 
-		public BgFileSetValueExpr(FileSet FileSet)
+		public BgFileSetValueExpr(FileSet fileSet)
 		{
-			this.FileSet = FileSet;
+			FileSet = fileSet;
 		}
 
-		public override BgFileSet Compute(BgExprContext Context) => this;
+		public override BgFileSet Compute(BgExprContext context) => this;
 	}
 
 	#endregion
diff --git a/Engine/Source/Programs/Shared/EpicGames.BuildGraph/Expressions/BgInt.cs b/Engine/Source/Programs/Shared/EpicGames.BuildGraph/Expressions/BgInt.cs
--- a/Engine/Source/Programs/Shared/EpicGames.BuildGraph/Expressions/BgInt.cs
+++ b/Engine/Source/Programs/Shared/EpicGames.BuildGraph/Expressions/BgInt.cs
@@ -1,20 +1,6 @@
 // Copyright Epic Games, Inc. All Rights Reserved.
 
-using EpicGames.BuildGraph;
 using System;
-using System.Collections.Generic;
-using System.ComponentModel;
-using System.Diagnostics;
-using System.IO;
-using System.Linq;
-using System.Reflection;
-using System.Text;
-using System.Xml;
-using EpicGames.Core;
-using Microsoft.Extensions.Logging;
-using OpenTracing;
-using OpenTracing.Util;
-using System.Threading.Tasks;
 
 namespace EpicGames.BuildGraph.Expressions
 {
@@ -27,43 +13,43 @@ namespace EpicGames.BuildGraph.Expressions
 		/// <summary>
 		/// Implicit conversion from an integer value
 		/// </summary>
-		public static implicit operator BgInt(int Value)
+		public static implicit operator BgInt(int value)
 		{
-			return new BgIntConstantExpr(Value);
+			return new BgIntConstantExpr(value);
 		}
 
 		/// <inheritdoc/>
-		public static BgInt operator +(BgInt Lhs, BgInt Rhs) => new BgIntAddExpr(Lhs, Rhs);
+		public static BgInt operator +(BgInt lhs, BgInt rhs) => new BgIntAddExpr(lhs, rhs);
 
 		/// <inheritdoc/>
-		public static BgInt operator -(BgInt Lhs, BgInt Rhs) => new BgIntAddExpr(Lhs, new BgIntNegateExpr(Rhs));
+		public static BgInt operator -(BgInt lhs, BgInt rhs) => new BgIntAddExpr(lhs, new BgIntNegateExpr(rhs));
 
 		/// <inheritdoc/>
-		public static BgBool operator <(BgInt Lhs, BgInt Rhs) => new BgIntLtExpr(Lhs, Rhs);
+		public static BgBool operator <(BgInt lhs, BgInt rhs) => new BgIntLtExpr(lhs, rhs);
 
 		/// <inheritdoc/>
-		public static BgBool operator >(BgInt Lhs, BgInt Rhs) => new BgIntGtExpr(Lhs, Rhs);
+		public static BgBool operator >(BgInt lhs, BgInt rhs) => new BgIntGtExpr(lhs, rhs);
 
 		/// <inheritdoc/>
-		public static BgBool operator ==(BgInt Lhs, BgInt Rhs) => new BgIntEqExpr(Lhs, Rhs);
+		public static BgBool operator ==(BgInt lhs, BgInt rhs) => new BgIntEqExpr(lhs, rhs);
 
 		/// <inheritdoc/>
-		public static BgBool operator !=(BgInt Lhs, BgInt Rhs) => new BgBoolNotExpr(new BgIntEqExpr(Lhs, Rhs));
+		public static BgBool operator !=(BgInt lhs, BgInt rhs) => new BgBoolNotExpr(new BgIntEqExpr(lhs, rhs));
 
 		/// <inheritdoc/>
-		public static BgBool operator <=(BgInt Lhs, BgInt Rhs) => new BgBoolNotExpr(new BgIntGtExpr(Lhs, Rhs));
+		public static BgBool operator <=(BgInt lhs, BgInt rhs) => new BgBoolNotExpr(new BgIntGtExpr(lhs, rhs));
 
 		/// <inheritdoc/>
-		public static BgBool operator >=(BgInt Lhs, BgInt Rhs) => new BgBoolNotExpr(new BgIntLtExpr(Lhs, Rhs));
+		public static BgBool operator >=(BgInt lhs, BgInt rhs) => new BgBoolNotExpr(new BgIntLtExpr(lhs, rhs));
 
 		/// <inheritdoc/>
-		object IBgExpr.Compute(BgExprContext Context) => Compute(Context);
+		object IBgExpr.Compute(BgExprContext context) => Compute(context);
 
 		/// <inheritdoc/>
-		public abstract int Compute(BgExprContext Context);
+		public abstract int Compute(BgExprContext context);
 
 		/// <inheritdoc/>
-		public BgInt IfThen(BgBool Condition, BgInt ValueIfTrue) => new BgIntChooseExpr(Condition, ValueIfTrue, this);
+		public BgInt IfThen(BgBool condition, BgInt valueIfTrue) => new BgIntChooseExpr(condition, valueIfTrue, this);
 
 		/// <inheritdoc/>
 		public override bool Equals(object? obj) => throw new InvalidOperationException();
@@ -81,13 +67,13 @@ namespace EpicGames.BuildGraph.Expressions
 	class BgIntTraits : BgTypeBase<BgInt>
 	{
 		/// <inheritdoc/>
-		public override BgInt DeserializeArgument(string Value) => (BgInt)int.Parse(Value);
+		public override BgInt DeserializeArgument(string value) => (BgInt)int.Parse(value);
 
 		/// <inheritdoc/>
-		public override string SerializeArgument(BgInt Value, BgExprContext Context) => Value.Compute(Context).ToString();
+		public override string SerializeArgument(BgInt value, BgExprContext context) => value.Compute(context).ToString();
 
 		/// <inheritdoc/>
-		public override BgInt CreateConstant(object Value) => new BgIntConstantExpr((int)Value);
+		public override BgInt CreateConstant(object value) => new BgIntConstantExpr((int)value);
 
 		/// <inheritdoc/>
 		public override IBgExprVariable<BgInt> CreateVariable() => new BgIntVariableExpr();
@@ -100,13 +86,13 @@ namespace EpicGames.BuildGraph.Expressions
 		public BgInt Lhs { get; }
 		public BgInt Rhs { get; }
 
-		public BgIntEqExpr(BgInt Lhs, BgInt Rhs)
+		public BgIntEqExpr(BgInt lhs, BgInt rhs)
 		{
-			this.Lhs = Lhs;
-			this.Rhs = Rhs;
+			Lhs = lhs;
+			Rhs = rhs;
 		}
 
-		public override bool Compute(BgExprContext Context) => Lhs.Compute(Context) == Rhs.Compute(Context);
+		public override bool Compute(BgExprContext context) => Lhs.Compute(context) == Rhs.Compute(context);
 	}
 
 	class BgIntLtExpr : BgBool
@@ -114,13 +100,13 @@ namespace EpicGames.BuildGraph.Expressions
 		public BgInt Lhs { get; }
 		public BgInt Rhs { get; }
 
-		public BgIntLtExpr(BgInt Lhs, BgInt Rhs)
+		public BgIntLtExpr(BgInt lhs, BgInt rhs)
 		{
-			this.Lhs = Lhs;
-			this.Rhs = Rhs;
+			Lhs = lhs;
+			Rhs = rhs;
 		}
 
-		public override bool Compute(BgExprContext Context) => Lhs.Compute(Context) < Rhs.Compute(Context);
+		public override bool Compute(BgExprContext context) => Lhs.Compute(context) < Rhs.Compute(context);
 	}
 
 	class BgIntGtExpr : BgBool
@@ -128,13 +114,13 @@ namespace EpicGames.BuildGraph.Expressions
 		public BgInt Lhs { get; }
 		public BgInt Rhs { get; }
 
-		public BgIntGtExpr(BgInt Lhs, BgInt Rhs)
+		public BgIntGtExpr(BgInt lhs, BgInt rhs)
 		{
-			this.Lhs = Lhs;
-			this.Rhs = Rhs;
+			Lhs = lhs;
+			Rhs = rhs;
 		}
 
-		public override bool Compute(BgExprContext Context) => Lhs.Compute(Context) > Rhs.Compute(Context);
+		public override bool Compute(BgExprContext context) => Lhs.Compute(context) > Rhs.Compute(context);
 	}
 
 	class BgIntAddExpr : BgInt
@@ -142,73 +128,73 @@ namespace EpicGames.BuildGraph.Expressions
 		public BgInt Lhs { get; }
 		public BgInt Rhs { get; }
 
-		public BgIntAddExpr(BgInt Lhs, BgInt Rhs)
+		public BgIntAddExpr(BgInt lhs, BgInt rhs)
 		{
-			this.Lhs = Lhs;
-			this.Rhs = Rhs;
+			Lhs = lhs;
+			Rhs = rhs;
 		}
 
-		public override int Compute(BgExprContext Context) => Lhs.Compute(Context) + Rhs.Compute(Context);
+		public override int Compute(BgExprContext context) => Lhs.Compute(context) + Rhs.Compute(context);
 	}
 
 	class BgIntNegateExpr : BgInt
 	{
 		public BgInt Inner { get; }
 
-		public BgIntNegateExpr(BgInt Inner)
+		public BgIntNegateExpr(BgInt inner)
 		{
-			this.Inner = Inner;
+			Inner = inner;
 		}
 
-		public override int Compute(BgExprContext Context) => -Inner.Compute(Context);
+		public override int Compute(BgExprContext context) => -Inner.Compute(context);
 	}
 
 	class BgIntChooseExpr : BgInt
 	{
-		public BgBool Condition;
-		public BgInt ValueIfTrue;
-		public BgInt ValueIfFalse;
+		public BgBool Condition { get; }
+		public BgInt ValueIfTrue { get; }
+		public BgInt ValueIfFalse { get; }
 
-		public BgIntChooseExpr(BgBool Condition, BgInt ValueIfTrue, BgInt ValueIfFalse)
+		public BgIntChooseExpr(BgBool condition, BgInt valueIfTrue, BgInt valueIfFalse)
 		{
-			this.Condition = Condition;
-			this.ValueIfTrue = ValueIfTrue;
-			this.ValueIfFalse = ValueIfFalse;
+			Condition = condition;
+			ValueIfTrue = valueIfTrue;
+			ValueIfFalse = valueIfFalse;
 		}
 
-		public override int Compute(BgExprContext Context) => Condition.Compute(Context) ? ValueIfTrue.Compute(Context) : ValueIfFalse.Compute(Context);
+		public override int Compute(BgExprContext context) => Condition.Compute(context) ? ValueIfTrue.Compute(context) : ValueIfFalse.Compute(context);
 	}
 
 	class BgIntConstantExpr : BgInt
 	{
 		public int Value { get; }
 
-		public BgIntConstantExpr(int Value)
+		public BgIntConstantExpr(int value)
 		{
-			this.Value = Value;
+			Value = value;
 		}
 
-		public override int Compute(BgExprContext Context) => Value;
+		public override int Compute(BgExprContext context) => Value;
 	}
 
 	class BgIntVariableExpr : BgInt, IBgExprVariable<BgInt>
 	{
 		public BgInt Value { get; set; } = 0;
 
-		public override int Compute(BgExprContext Context) => Value.Compute(Context);
+		public override int Compute(BgExprContext context) => Value.Compute(context);
 	}
 
 	class BgIntFormatExpr : BgString
 	{
-		BgInt Value;
+		public BgInt Value { get; }
 
-		public BgIntFormatExpr(BgInt Value)
+		public BgIntFormatExpr(BgInt value)
 		{
-			this.Value = Value;
+			Value = value;
 		}
 
-		public override string Compute(BgExprContext Context) => Value.Compute(Context).ToString();
+		public override string Compute(BgExprContext context) => Value.Compute(context).ToString();
 	}
 
-#endregion
+	#endregion
 }
diff --git a/Engine/Source/Programs/Shared/EpicGames.BuildGraph/Expressions/BgList.cs b/Engine/Source/Programs/Shared/EpicGames.BuildGraph/Expressions/BgList.cs
--- a/Engine/Source/Programs/Shared/EpicGames.BuildGraph/Expressions/BgList.cs
+++ b/Engine/Source/Programs/Shared/EpicGames.BuildGraph/Expressions/BgList.cs
@@ -1,21 +1,8 @@
 // Copyright Epic Games, Inc. All Rights Reserved.
 
-using EpicGames.BuildGraph;
 using System;
 using System.Collections.Generic;
-using System.ComponentModel;
-using System.Diagnostics;
-using System.IO;
 using System.Linq;
-using System.Reflection;
-using System.Text;
-using System.Xml;
-using EpicGames.Core;
-using Microsoft.Extensions.Logging;
-using OpenTracing;
-using OpenTracing.Util;
-using System.Threading.Tasks;
-using System.Diagnostics.Contracts;
 
 namespace EpicGames.BuildGraph.Expressions
 {
@@ -47,69 +34,69 @@ namespace EpicGames.BuildGraph.Expressions
 		/// <summary>
 		/// Implicit conversion operator from a single value
 		/// </summary>
-		public static implicit operator BgList<T>(T Value) => new BgListConstantExpr<T>(new[] { Value });
+		public static implicit operator BgList<T>(T value) => new BgListConstantExpr<T>(new[] { value });
 
 		/// <summary>
 		/// Implicit conversion operator from an array of values
 		/// </summary>
-		public static implicit operator BgList<T>(T[] Value) => new BgListConstantExpr<T>(Value);
+		public static implicit operator BgList<T>(T[] value) => new BgListConstantExpr<T>(value);
 
 		/// <summary>
 		/// Implicit conversion operator from a list of values
 		/// </summary>
-		public static implicit operator BgList<T>(List<T> Value) => new BgListConstantExpr<T>(Value);
+		public static implicit operator BgList<T>(List<T> value) => new BgListConstantExpr<T>(value);
 
 		/// <summary>
 		/// Crates a list from an array of values
 		/// </summary>
-		/// <param name="Items">Sequence to construct from</param>
+		/// <param name="items">Sequence to construct from</param>
 		/// <returns></returns>
-		public static BgList<T> Create(IEnumerable<T> Items) => new BgListConstantExpr<T>(Items);
+		public static BgList<T> Create(IEnumerable<T> items) => new BgListConstantExpr<T>(items);
 
 		/// <summary>
 		/// Crates a list from an array of values
 		/// </summary>
-		/// <param name="Items">Sequence to construct from</param>
+		/// <param name="items">Sequence to construct from</param>
 		/// <returns></returns>
-		public static BgList<T> Create(params T[] Items) => new BgListConstantExpr<T>(Items);
+		public static BgList<T> Create(params T[] items) => new BgListConstantExpr<T>(items);
 
 		/// <summary>
 		/// Concatenates mutiple lists together
 		/// </summary>
-		/// <param name="Sources"></param>
+		/// <param name="sources"></param>
 		/// <returns></returns>
-		public static BgList<T> Concat(params BgList<T>[] Sources) => new BgListConcatExpr<T>(Sources);
+		public static BgList<T> Concat(params BgList<T>[] sources) => new BgListConcatExpr<T>(sources);
 
 		/// <summary>
 		/// Adds items to the end of the list, returning the new list
 		/// </summary>
-		/// <param name="Items">Items to add</param>
+		/// <param name="items">Items to add</param>
 		/// <returns>New list containing the given items</returns>
-		public BgList<T> Add(params T[] Items) => new BgListConcatExpr<T>(this, Items);
+		public BgList<T> Add(params T[] items) => new BgListConcatExpr<T>(this, items);
 
 		/// <inheritdoc cref="Add(T[])"/>
-		public BgList<T> Add(params BgList<T>[] Items) => new BgListConcatExpr<T>(this, Items);
+		public BgList<T> Add(params BgList<T>[] items) => new BgListConcatExpr<T>(this, items);
 
 		/// <inheritdoc cref="Add(T[])"/>
-		public BgList<T> AddIf(BgBool Condition, params T[] Items) => IfThen(Condition, Add(Items));
+		public BgList<T> AddIf(BgBool condition, params T[] items) => IfThen(condition, Add(items));
 
 		/// <inheritdoc cref="Add(T[])"/>
-		public BgList<T> AddIf(BgBool Condition, params BgList<T>[] Items) => IfThen(Condition, Add(Items));
+		public BgList<T> AddIf(BgBool condition, params BgList<T>[] items) => IfThen(condition, Add(items));
 
 		/// <summary>
 		/// Removes the given items from this list
 		/// </summary>
-		/// <param name="Items">Items to remove</param>
+		/// <param name="items">Items to remove</param>
 		/// <returns>New list without the given items</returns>
-		public BgList<T> Except(BgList<T> Items) => new BgListExceptExpr<T>(this, Items);
+		public BgList<T> Except(BgList<T> items) => new BgListExceptExpr<T>(this, items);
 
 		/// <summary>
 		/// Removes the given items from this list
 		/// </summary>
-		/// <param name="Condition">Condition to remove the items</param>
-		/// <param name="Items">Items to remove</param>
+		/// <param name="condition">Condition to remove the items</param>
+		/// <param name="items">Items to remove</param>
 		/// <returns>New list without the given items</returns>
-		public BgList<T> ExceptIf(BgBool Condition, BgList<T> Items) => IfThen(Condition, Except(Items));
+		public BgList<T> ExceptIf(BgBool condition, BgList<T> items) => IfThen(condition, Except(items));
 
 		/// <summary>
 		/// Removes any duplicate items from the list. The first item in the list is retained in its original order.
@@ -118,22 +105,22 @@ namespace EpicGames.BuildGraph.Expressions
 		public BgList<T> Distinct() => new BgListDistinctExpr<T>(this);
 
 		/// <inheritdoc cref="Enumerable.Select{TSource, TResult}(IEnumerable{TSource}, Func{TSource, TResult})"/>
-		public BgList<TResult> Select<TResult>(Func<T, TResult> Function) where TResult : IBgExpr<TResult> => new BgListSelectExpr<T, TResult>(this, Function);
+		public BgList<TResult> Select<TResult>(Func<T, TResult> function) where TResult : IBgExpr<TResult> => new BgListSelectExpr<T, TResult>(this, function);
 
 		/// <inheritdoc cref="Enumerable.Where{TSource}(IEnumerable{TSource}, Func{TSource, bool})"/>
-		public BgList<T> Where(Func<T, BgBool> Predicate) => new BgListWhereExpr<T>(this, Predicate);
+		public BgList<T> Where(Func<T, BgBool> predicate) => new BgListWhereExpr<T>(this, predicate);
 
 		/// <inheritdoc cref="Enumerable.Contains{TSource}(IEnumerable{TSource}, TSource)"/>
-		public BgBool Contains(T Item) => new BgListContainsExpr<T>(this, Item);
+		public BgBool Contains(T item) => new BgListContainsExpr<T>(this, item);
 
 		/// <inheritdoc/>
-		object IBgExpr.Compute(BgExprContext Context) => GetEnumerable(Context);
+		object IBgExpr.Compute(BgExprContext context) => GetEnumerable(context);
 
 		/// <inheritdoc/>
-		public abstract IEnumerable<T> GetEnumerable(BgExprContext Context);
+		public abstract IEnumerable<T> GetEnumerable(BgExprContext context);
 
 		/// <inheritdoc/>
-		public BgList<T> IfThen(BgBool Condition, BgList<T> ValueIfTrue) => new BgListChooseExpr<T>(Condition, ValueIfTrue, this);
+		public BgList<T> IfThen(BgBool condition, BgList<T> valueIfTrue) => new BgListChooseExpr<T>(condition, valueIfTrue, this);
 
 		/// <inheritdoc/>
 		public BgString ToBgString() => new BgListFormatExpr<T>(this);
@@ -145,23 +132,23 @@ namespace EpicGames.BuildGraph.Expressions
 	class BgListType<T> : BgTypeBase<BgList<T>> where T : IBgExpr<T>
 	{
 		/// <inheritdoc/>
-		public override BgList<T> DeserializeArgument(string Text)
+		public override BgList<T> DeserializeArgument(string text)
 		{
-			IBgType<T> Converter = BgType.Get<T>();
-			return BgList<T>.Create(Text.Split(';').Select(x => Converter.DeserializeArgument(x)));
+			IBgType<T> converter = BgType.Get<T>();
+			return BgList<T>.Create(text.Split(';').Select(x => converter.DeserializeArgument(x)));
 		}
 
 		/// <inheritdoc/>
-		public override string SerializeArgument(BgList<T> Value, BgExprContext Context)
+		public override string SerializeArgument(BgList<T> value, BgExprContext context)
 		{
-			IBgType<T> Converter = BgType.Get<T>();
-			return String.Join(";", Value.GetEnumerable(Context).Select(x => Converter.SerializeArgument(x, Context)));
+			IBgType<T> converter = BgType.Get<T>();
+			return String.Join(";", value.GetEnumerable(context).Select(x => converter.SerializeArgument(x, context)));
 		}
 
 		/// <inheritdoc/>
-		public override BgList<T> CreateConstant(object Value)
+		public override BgList<T> CreateConstant(object value)
 		{
-			return new BgListConstantExpr<T>((IEnumerable<T>)Value);
+			return new BgListConstantExpr<T>((IEnumerable<T>)value);
 		}
 
 		/// <inheritdoc/>
@@ -176,19 +163,19 @@ namespace EpicGames.BuildGraph.Expressions
 	/// </summary>
 	public static class BgListExtensions
 	{
-		internal static List<string> Compute(this BgList<BgString> List, BgExprContext Context)
+		internal static List<string> Compute(this BgList<BgString> list, BgExprContext context)
 		{
-			return List.GetEnumerable(Context).Select(x => x.Compute(Context)).ToList();
+			return list.GetEnumerable(context).Select(x => x.Compute(context)).ToList();
 		}
 
-		internal static List<BgFileSet> Compute(this BgList<BgFileSet> List, BgExprContext Context)
+		internal static List<BgFileSet> Compute(this BgList<BgFileSet> list, BgExprContext context)
 		{
-			return List.GetEnumerable(Context).Select(x => x.Compute(Context)).ToList();
+			return list.GetEnumerable(context).Select(x => x.Compute(context)).ToList();
 		}
 
-		internal static List<string> ComputeTags(this BgList<BgFileSet> List, BgExprContext Context)
+		internal static List<string> ComputeTags(this BgList<BgFileSet> list, BgExprContext context)
 		{
-			return List.GetEnumerable(Context).Select(x => x.ComputeTag(Context)).ToList();
+			return list.GetEnumerable(context).Select(x => x.ComputeTag(context)).ToList();
 		}
 	}
 
@@ -196,93 +183,91 @@ namespace EpicGames.BuildGraph.Expressions
 
 	class BgListConcatExpr<T> : BgList<T> where T : IBgExpr<T>
 	{
-		public BgList<T>[] Sources;
+		public BgList<T>[] Sources { get; }
 
-		public BgListConcatExpr(params BgList<T>[] Sources)
+		public BgListConcatExpr(params BgList<T>[] sources)
 		{
-			this.Sources = Sources;
+			Sources = sources;
 		}
 
-		public BgListConcatExpr(BgList<T> FirstSource, BgList<T>[] OtherSources)
+		public BgListConcatExpr(BgList<T> firstSource, BgList<T>[] otherSources)
 		{
-			Sources = new BgList<T>[OtherSources.Length + 1];
-			Sources[0] = FirstSource;
-			Array.Copy(OtherSources, 0, Sources, 1, OtherSources.Length);
+			Sources = new BgList<T>[otherSources.Length + 1];
+			Sources[0] = firstSource;
+			Array.Copy(otherSources, 0, Sources, 1, otherSources.Length);
 		}
 
-		public override IEnumerable<T> GetEnumerable(BgExprContext Context)
+		public override IEnumerable<T> GetEnumerable(BgExprContext context)
 		{
-			return Sources.SelectMany(x => x.GetEnumerable(Context));
+			return Sources.SelectMany(x => x.GetEnumerable(context));
 		}
 	}
 
 	class BgListChooseExpr<T> : BgList<T> where T : IBgExpr<T>
 	{
-		public BgBool Condition;
-		public BgList<T> ValueIfTrue;
-		public BgList<T> ValueIfFalse;
+		public BgBool Condition { get; }
+		public BgList<T> ValueIfTrue { get; }
+		public BgList<T> ValueIfFalse { get; }
 
-		public BgListChooseExpr(BgBool Condition, BgList<T> ValueIfTrue, BgList<T> ValueIfFalse)
+		public BgListChooseExpr(BgBool condition, BgList<T> valueIfTrue, BgList<T> valueIfFalse)
 		{
-			this.Condition = Condition;
-			this.ValueIfTrue = ValueIfTrue;
-			this.ValueIfFalse = ValueIfFalse;
+			Condition = condition;
+			ValueIfTrue = valueIfTrue;
+			ValueIfFalse = valueIfFalse;
 		}
 
-		public override IEnumerable<T> GetEnumerable(BgExprContext Context) => Condition.Compute(Context) ? ValueIfTrue.GetEnumerable(Context) : ValueIfFalse.GetEnumerable(Context);
+		public override IEnumerable<T> GetEnumerable(BgExprContext context) => Condition.Compute(context) ? ValueIfTrue.GetEnumerable(context) : ValueIfFalse.GetEnumerable(context);
 	}
 
 	class BgListConstantExpr<T> : BgList<T> where T : IBgExpr<T>
 	{
 		public IEnumerable<T> Value { get; }
 
-		public BgListConstantExpr(IEnumerable<T> Value)
+		public BgListConstantExpr(IEnumerable<T> value)
 		{
-			this.Value = Value;
+			Value = value;
 		}
 
-		public override IEnumerable<T> GetEnumerable(BgExprContext Context) => Value;
+		public override IEnumerable<T> GetEnumerable(BgExprContext context) => Value;
 	}
 
 	class BgListFormatExpr<T> : BgString where T : IBgExpr<T>
 	{
-		public BgList<T> Inner;
+		public BgList<T> Inner { get; }
 
-		public BgListFormatExpr(BgList<T> Inner)
+		public BgListFormatExpr(BgList<T> inner)
 		{
-			this.Inner = Inner;
+			Inner = inner;
 		}
 
-		public override string Compute(BgExprContext Context) => String.Join(";", Inner.GetEnumerable(Context).Select(x => x.ToBgString().Compute(Context)));
+		public override string Compute(BgExprContext context) => String.Join(";", Inner.GetEnumerable(context).Select(x => x.ToBgString().Compute(context)));
 	}
 
 	class BgListVariableExpr<T> : BgList<T>, IBgExprVariable<BgList<T>> where T : IBgExpr<T>
 	{
 		public BgList<T> Value { get; set; } = BgList<T>.Empty;
 
-		public override IEnumerable<T> GetEnumerable(BgExprContext Context) => Value.GetEnumerable(Context);
+		public override IEnumerable<T> GetEnumerable(BgExprContext context) => Value.GetEnumerable(context);
 	}
 
 	class BgListDistinctExpr<T> : BgList<T> where T : IBgExpr<T>
 	{
-		BgList<T> Source;
+		public BgList<T> Source { get; }
 
-		public BgListDistinctExpr(BgList<T> Source)
+		public BgListDistinctExpr(BgList<T> source)
 		{
-			this.Source = Source;
+			Source = source;
 		}
 
-		public override IEnumerable<T> GetEnumerable(BgExprContext Context)
+		public override IEnumerable<T> GetEnumerable(BgExprContext context)
 		{
-			IBgType<T> Type = BgType.Get<T>();
-
-			HashSet<object> Values = new HashSet<object>();
-			foreach (T Item in Source.GetEnumerable(Context))
+			HashSet<object> values = new HashSet<object>();
+			foreach (T item in Source.GetEnumerable(context))
 			{
-				object ItemValue = Item.Compute(Context);
-				if (Values.Add(ItemValue))
+				object itemValue = item.Compute(context);
+				if (values.Add(itemValue))
 				{
-					yield return Item;
+					yield return item;
 				}
 			}
 		}
@@ -290,25 +275,25 @@ namespace EpicGames.BuildGraph.Expressions
 
 	class BgFunc<TIn, TOut> where TIn : IBgExpr<TIn> where TOut : IBgExpr<TOut>
 	{
-		IBgExprVariable<TIn> ArgVar;
-		IBgExpr<TOut> FuncExpr;
+		public IBgExprVariable<TIn> ArgVar { get; }
+		public IBgExpr<TOut> FuncExpr { get; }
 
-		public BgFunc(IBgExprVariable<TIn> ArgVar, IBgExpr<TOut> FuncExpr)
+		public BgFunc(IBgExprVariable<TIn> argVar, IBgExpr<TOut> funcExpr)
 		{
-			this.ArgVar = ArgVar;
-			this.FuncExpr = FuncExpr;
+			ArgVar = argVar;
+			FuncExpr = funcExpr;
 		}
 
-		public static implicit operator BgFunc<TIn, TOut>(Func<TIn, TOut> Function)
+		public static implicit operator BgFunc<TIn, TOut>(Func<TIn, TOut> function)
 		{
-			IBgExprVariable<TIn> ArgVar = BgType.Get<TIn>().CreateVariable();
-			return new BgFunc<TIn, TOut>(ArgVar, Function((TIn)ArgVar));
+			IBgExprVariable<TIn> argVar = BgType.Get<TIn>().CreateVariable();
+			return new BgFunc<TIn, TOut>(argVar, function((TIn)argVar));
 		}
 
-		public object Compute(BgExprContext Context, TIn Argument)
+		public object Compute(BgExprContext context, TIn argument)
 		{
-			ArgVar.Value = Argument;
-			return FuncExpr.Compute(Context);
+			ArgVar.Value = argument;
+			return FuncExpr.Compute(context);
 		}
 
 		public BgString ToBgString()
@@ -319,61 +304,61 @@ namespace EpicGames.BuildGraph.Expressions
 
 	class BgListSelectExpr<TIn, TOut> : BgList<TOut> where TIn : IBgExpr<TIn> where TOut : IBgExpr<TOut>
 	{
-		BgList<TIn> Source;
-		BgFunc<TIn, TOut> Function;
+		public BgList<TIn> Source { get; }
+		public BgFunc<TIn, TOut> Function { get; }
 
-		public BgListSelectExpr(BgList<TIn> Source, BgFunc<TIn, TOut> Function)
+		public BgListSelectExpr(BgList<TIn> source, BgFunc<TIn, TOut> function)
 		{
-			this.Source = Source;
-			this.Function = Function;
+			Source = source;
+			Function = function;
 		}
 
-		public override IEnumerable<TOut> GetEnumerable(BgExprContext Context)
+		public override IEnumerable<TOut> GetEnumerable(BgExprContext context)
 		{
-			return Source.GetEnumerable(Context).Select(x => (TOut)Function.Compute(Context, x));
+			return Source.GetEnumerable(context).Select(x => (TOut)Function.Compute(context, x));
 		}
 	}
 
 	class BgListWhereExpr<T> : BgList<T> where T : IBgExpr<T>
 	{
-		BgList<T> Source;
-		BgFunc<T, BgBool> Predicate;
+		public BgList<T> Source { get; }
+		public BgFunc<T, BgBool> Predicate { get; }
 
-		public BgListWhereExpr(BgList<T> Source, BgFunc<T, BgBool> Predicate)
+		public BgListWhereExpr(BgList<T> source, BgFunc<T, BgBool> predicate)
 		{
-			this.Source = Source;
-			this.Predicate = Predicate;
+			Source = source;
+			Predicate = predicate;
 		}
 
-		public override IEnumerable<T> GetEnumerable(BgExprContext Context)
+		public override IEnumerable<T> GetEnumerable(BgExprContext context)
 		{
-			return Source.GetEnumerable(Context).Where(x => (bool)Predicate.Compute(Context, x));
+			return Source.GetEnumerable(context).Where(x => (bool)Predicate.Compute(context, x));
 		}
 	}
 
 	class BgListExceptExpr<T> : BgList<T> where T : IBgExpr<T>
 	{
-		BgList<T> SourceList;
-		BgList<T> ExceptList;
+		public BgList<T> SourceList { get; }
+		public BgList<T> ExceptList { get; }
 
-		public BgListExceptExpr(BgList<T> Source, BgList<T> Except)
+		public BgListExceptExpr(BgList<T> source, BgList<T> except)
 		{
-			this.SourceList = Source;
-			this.ExceptList = Except;
+			SourceList = source;
+			ExceptList = except;
 		}
 
-		public override IEnumerable<T> GetEnumerable(BgExprContext Context)
+		public override IEnumerable<T> GetEnumerable(BgExprContext context)
 		{
-			HashSet<object> ExceptValues = new HashSet<object>();
-			ExceptValues.UnionWith(ExceptList.GetEnumerable(Context).Select(x => x.Compute(Context)));
+			HashSet<object> exceptValues = new HashSet<object>();
+			exceptValues.UnionWith(ExceptList.GetEnumerable(context).Select(x => x.Compute(context)));
 
-			IBgType<T> Type = BgType.Get<T>();
-			foreach (T SourceItem in SourceList.GetEnumerable(Context))
+			IBgType<T> type = BgType.Get<T>();
+			foreach (T sourceItem in SourceList.GetEnumerable(context))
 			{
-				object SourceValue = SourceItem.Compute(Context);
-				if (!ExceptValues.Contains(SourceValue))
+				object sourceValue = sourceItem.Compute(context);
+				if (!exceptValues.Contains(sourceValue))
 				{
-					yield return Type.CreateConstant(SourceValue);
+					yield return type.CreateConstant(sourceValue);
 				}
 			}
 		}
@@ -381,19 +366,19 @@ namespace EpicGames.BuildGraph.Expressions
 
 	class BgListContainsExpr<T> : BgBool where T : IBgExpr<T>
 	{
-		BgList<T> Source;
-		T Item;
+		public BgList<T> Source { get; }
+		public T Item { get; }
 
-		public BgListContainsExpr(BgList<T> Source, T Item)
+		public BgListContainsExpr(BgList<T> source, T item)
 		{
-			this.Source = Source;
-			this.Item = Item;
+			Source = source;
+			Item = item;
 		}
 
-		public override bool Compute(BgExprContext Context)
+		public override bool Compute(BgExprContext context)
 		{
-			object Value = Item.Compute(Context);
-			return Source.GetEnumerable(Context).Any(x => x.Compute(Context).Equals(Value));
+			object value = Item.Compute(context);
+			return Source.GetEnumerable(context).Any(x => x.Compute(context).Equals(value));
 		}
 	}
 
diff --git a/Engine/Source/Programs/Shared/EpicGames.BuildGraph/Expressions/BgString.cs b/Engine/Source/Programs/Shared/EpicGames.BuildGraph/Expressions/BgString.cs
--- a/Engine/Source/Programs/Shared/EpicGames.BuildGraph/Expressions/BgString.cs
+++ b/Engine/Source/Programs/Shared/EpicGames.BuildGraph/Expressions/BgString.cs
@@ -1,20 +1,8 @@
 // Copyright Epic Games, Inc. All Rights Reserved.
 
-using EpicGames.BuildGraph;
 using System;
 using System.Collections.Generic;
-using System.ComponentModel;
-using System.Diagnostics;
-using System.IO;
 using System.Linq;
-using System.Reflection;
-using System.Text;
-using System.Xml;
-using EpicGames.Core;
-using Microsoft.Extensions.Logging;
-using OpenTracing;
-using OpenTracing.Util;
-using System.Threading.Tasks;
 using System.Text.RegularExpressions;
 
 namespace EpicGames.BuildGraph.Expressions
@@ -33,56 +21,56 @@ namespace EpicGames.BuildGraph.Expressions
 		/// <summary>
 		/// Implicit conversion from a regular string type
 		/// </summary>
-		public static implicit operator BgString(string Value)
+		public static implicit operator BgString(string value)
 		{
-			return new BgStringConstantExpr(Value);
+			return new BgStringConstantExpr(value);
 		}
 
-		/// <inheritdoc cref="String.Equals(string?, string?, StringComparison)"/>
-		public static BgBool Equals(BgString Lhs, BgString Rhs, StringComparison Comparison = StringComparison.CurrentCulture) => Compare(Lhs, Rhs, Comparison) == 0;
+		/// <inheritdoc cref="String.Equals(String?, String?, StringComparison)"/>
+		public static BgBool Equals(BgString lhs, BgString rhs, StringComparison comparison = StringComparison.CurrentCulture) => Compare(lhs, rhs, comparison) == 0;
 
-		/// <inheritdoc cref="String.Compare(string?, string?, StringComparison)"/>
-		public static BgInt Compare(BgString Lhs, BgString Rhs, StringComparison Comparison = StringComparison.CurrentCulture) => new BgStringCompareExpr(Lhs, Rhs, Comparison);
+		/// <inheritdoc cref="String.Compare(String?, String?, StringComparison)"/>
+		public static BgInt Compare(BgString lhs, BgString rhs, StringComparison comparison = StringComparison.CurrentCulture) => new BgStringCompareExpr(lhs, rhs, comparison);
 
-		/// <inheritdoc cref="String.Join{T}(string?, IEnumerable{T})"/>
-		public static BgString Join(BgString Separator, BgList<BgString> Values) => new BgStringJoinExpr(Separator, Values);
+		/// <inheritdoc cref="String.Join{T}(String?, IEnumerable{T})"/>
+		public static BgString Join(BgString separator, BgList<BgString> values) => new BgStringJoinExpr(separator, values);
 
-		/// <inheritdoc cref="String.Format(string, object?[])"/>
-		public static BgString Format(string Format, params object?[] Args) => new BgStringFormatExpr(Format, Args);
+		/// <inheritdoc cref="String.Format(String, Object?[])"/>
+		public static BgString Format(string format, params object?[] args) => new BgStringFormatExpr(format, args);
 
 		/// <inheritdoc/>
-		public static BgString operator +(BgString Lhs, BgString Rhs) => new BgStringConcatExpr(Lhs, Rhs);
+		public static BgString operator +(BgString lhs, BgString rhs) => new BgStringConcatExpr(lhs, rhs);
 
 		/// <inheritdoc/>
-		public static BgBool operator ==(BgString Lhs, BgString Rhs) => Equals(Lhs, Rhs);
+		public static BgBool operator ==(BgString lhs, BgString rhs) => Equals(lhs, rhs);
 
 		/// <inheritdoc/>
-		public static BgBool operator !=(BgString Lhs, BgString Rhs) => !Equals(Lhs, Rhs);
+		public static BgBool operator !=(BgString lhs, BgString rhs) => !Equals(lhs, rhs);
 
 		/// <summary>
 		/// Appens another string to this one
 		/// </summary>
-		public BgString Append(BgString Other) => this + Other;
+		public BgString Append(BgString other) => this + other;
 
 		/// <summary>
 		/// Appens another string to this one if a condition is true
 		/// </summary>
-		public BgString AppendIf(BgBool Condition, BgString Other) => IfThen(Condition, this + Other);
+		public BgString AppendIf(BgBool condition, BgString other) => IfThen(condition, this + other);
 
 		/// <inheritdoc/>
-		public BgString IfThen(BgBool Condition, BgString ValueIfTrue) => new BgStringChooseExpr(Condition, ValueIfTrue, this);
+		public BgString IfThen(BgBool condition, BgString valueIfTrue) => new BgStringChooseExpr(condition, valueIfTrue, this);
 
 		/// <inheritdoc/>
-		public BgBool IsMatch(BgString Input, string Pattern) => new BgStringIsMatchExpr(Input, Pattern);
+		public BgBool IsMatch(BgString input, string pattern) => new BgStringIsMatchExpr(input, pattern);
 
 		/// <inheritdoc/>
-		public BgString Replace(BgString Input, string Pattern, string Replace) => new BgStringReplaceExpr(Input, Pattern, Replace);
+		public BgString Replace(BgString input, string pattern, string replace) => new BgStringReplaceExpr(input, pattern, replace);
 
 		/// <inheritdoc/>
-		object IBgExpr.Compute(BgExprContext Context) => Compute(Context);
+		object IBgExpr.Compute(BgExprContext context) => Compute(context);
 
 		/// <inheritdoc/>
-		public abstract string Compute(BgExprContext Context);
+		public abstract string Compute(BgExprContext context);
 
 		/// <inheritdoc/>
 		public override bool Equals(object? obj) => throw new InvalidOperationException();
@@ -100,13 +88,13 @@ namespace EpicGames.BuildGraph.Expressions
 	class BgStringTraits : BgTypeBase<BgString>
 	{
 		/// <inheritdoc/>
-		public override BgString DeserializeArgument(string Text) => Text;
+		public override BgString DeserializeArgument(string text) => text;
 
 		/// <inheritdoc/>
-		public override string SerializeArgument(BgString Value, BgExprContext Context) => Value.Compute(Context);
+		public override string SerializeArgument(BgString value, BgExprContext context) => value.Compute(context);
 
 		/// <inheritdoc/>
-		public override BgString CreateConstant(object Value) => new BgStringConstantExpr((string)Value);
+		public override BgString CreateConstant(object value) => new BgStringConstantExpr((string)value);
 
 		/// <inheritdoc/>
 		public override IBgExprVariable<BgString> CreateVariable() => new BgStringVariableExpr();
@@ -120,14 +108,14 @@ namespace EpicGames.BuildGraph.Expressions
 		public BgString Rhs { get; }
 		public StringComparison Comparison { get; }
 
-		public BgStringCompareExpr(BgString Lhs, BgString Rhs, StringComparison Comparison)
+		public BgStringCompareExpr(BgString lhs, BgString rhs, StringComparison comparison)
 		{
-			this.Lhs = Lhs;
-			this.Rhs = Rhs;
-			this.Comparison = Comparison;
+			Lhs = lhs;
+			Rhs = rhs;
+			Comparison = comparison;
 		}
 
-		public override int Compute(BgExprContext Context) => String.Compare(Lhs.Compute(Context), Rhs.Compute(Context), Comparison);
+		public override int Compute(BgExprContext context) => String.Compare(Lhs.Compute(context), Rhs.Compute(context), Comparison);
 
 		public override string ToString() => $"Compare({Lhs}, {Rhs})";
 	}
@@ -137,13 +125,13 @@ namespace EpicGames.BuildGraph.Expressions
 		public BgString Lhs { get; }
 		public BgString Rhs { get; }
 
-		public BgStringConcatExpr(BgString Lhs, BgString Rhs)
+		public BgStringConcatExpr(BgString lhs, BgString rhs)
 		{
-			this.Lhs = Lhs;
-			this.Rhs = Rhs;
+			Lhs = lhs;
+			Rhs = rhs;
 		}
 
-		public override string Compute(BgExprContext Context) => Lhs.Compute(Context) + Rhs.Compute(Context);
+		public override string Compute(BgExprContext context) => Lhs.Compute(context) + Rhs.Compute(context);
 
 		public override string ToString() => $"Concat({Lhs}, {Rhs})";
 	}
@@ -153,13 +141,13 @@ namespace EpicGames.BuildGraph.Expressions
 		public BgString Separator { get; }
 		public BgList<BgString> Values { get; }
 
-		public BgStringJoinExpr(BgString Separator, BgList<BgString> Values)
+		public BgStringJoinExpr(BgString separator, BgList<BgString> values)
 		{
-			this.Separator = Separator;
-			this.Values = Values;
+			Separator = separator;
+			Values = values;
 		}
 
-		public override string Compute(BgExprContext Context) => String.Join(Separator.Compute(Context), Values.Compute(Context));
+		public override string Compute(BgExprContext context) => String.Join(Separator.Compute(context), Values.Compute(context));
 
 		public override string ToString() => $"Join({Separator}, {Values})";
 	}
@@ -169,13 +157,13 @@ namespace EpicGames.BuildGraph.Expressions
 		public BgString Input { get; }
 		public string Pattern { get; }
 
-		public BgStringIsMatchExpr(BgString Input, string Pattern)
+		public BgStringIsMatchExpr(BgString input, string pattern)
 		{
-			this.Input = Input;
-			this.Pattern = Pattern;
+			Input = input;
+			Pattern = pattern;
 		}
 
-		public override bool Compute(BgExprContext Context) => Regex.IsMatch(Input.Compute(Context), Pattern);
+		public override bool Compute(BgExprContext context) => Regex.IsMatch(Input.Compute(context), Pattern);
 	}
 
 	class BgStringReplaceExpr : BgString
@@ -184,14 +172,14 @@ namespace EpicGames.BuildGraph.Expressions
 		public string Pattern { get; }
 		public string Replacement { get; }
 
-		public BgStringReplaceExpr(BgString Input, string Pattern, string Replacement)
+		public BgStringReplaceExpr(BgString input, string pattern, string replacement)
 		{
-			this.Input = Input;
-			this.Pattern = Pattern;
-			this.Replacement = Replacement;
+			Input = input;
+			Pattern = pattern;
+			Replacement = replacement;
 		}
 
-		public override string Compute(BgExprContext Context) => Regex.Replace(Input.Compute(Context), Pattern, Replacement);
+		public override string Compute(BgExprContext context) => Regex.Replace(Input.Compute(context), Pattern, Replacement);
 	}
 
 	class BgStringFormatExpr : BgString
@@ -199,47 +187,47 @@ namespace EpicGames.BuildGraph.Expressions
 		new string Format { get; }
 		BgString[] Arguments { get; }
 
-		public BgStringFormatExpr(string Format, object?[] Args)
+		public BgStringFormatExpr(string format, object?[] args)
 		{
-			this.Format = Format;
+			Format = format;
 
-			Arguments = new BgString[Args.Length];
-			for (int Idx = 0; Idx < Args.Length; Idx++)
+			Arguments = new BgString[args.Length];
+			for (int idx = 0; idx < args.Length; idx++)
 			{
-				if (Args[Idx] is IBgExpr Expr)
+				if (args[idx] is IBgExpr expr)
 				{
-					Arguments[Idx] = Expr.ToBgString();
+					Arguments[idx] = expr.ToBgString();
 				}
-				else if (Args[Idx] != null)
+				else if (args[idx] != null)
 				{
-					Arguments[Idx] = Args[Idx]?.ToString() ?? String.Empty;
+					Arguments[idx] = args[idx]?.ToString() ?? String.Empty;
 				}
 				else
 				{
-					Arguments[Idx] = String.Empty;
+					Arguments[idx] = String.Empty;
 				}
 			}
 		}
 
-		public override string Compute(BgExprContext Context) => String.Format(Format, Arguments.Select(x => x.Compute(Context)).ToArray());
+		public override string Compute(BgExprContext context) => String.Format(Format, Arguments.Select(x => x.Compute(context)).ToArray());
 
 		public override string ToString() => $"Format({String.Join(", ", new[] { Format }.Concat(Arguments.Select(x => x.ToString())))})";
 	}
 
 	class BgStringChooseExpr : BgString
 	{
-		public BgBool Condition;
-		public BgString ValueIfTrue;
-		public BgString ValueIfFalse;
+		public BgBool Condition { get; }
+		public BgString ValueIfTrue { get; }
+		public BgString ValueIfFalse { get; }
 
-		public BgStringChooseExpr(BgBool Condition, BgString ValueIfTrue, BgString ValueIfFalse)
+		public BgStringChooseExpr(BgBool condition, BgString valueIfTrue, BgString valueIfFalse)
 		{
-			this.Condition = Condition;
-			this.ValueIfTrue = ValueIfTrue;
-			this.ValueIfFalse = ValueIfFalse;
+			Condition = condition;
+			ValueIfTrue = valueIfTrue;
+			ValueIfFalse = valueIfFalse;
 		}
 
-		public override string Compute(BgExprContext Context) => Condition.Compute(Context) ? ValueIfTrue.Compute(Context) : ValueIfFalse.Compute(Context);
+		public override string Compute(BgExprContext context) => Condition.Compute(context) ? ValueIfTrue.Compute(context) : ValueIfFalse.Compute(context);
 	}
 
 	class BgStringVariableExpr : BgString, IBgExprVariable<BgString>
@@ -248,23 +236,23 @@ namespace EpicGames.BuildGraph.Expressions
 
 		public BgString Variable => this;
 
-		public override string Compute(BgExprContext Context) => Value.Compute(Context);
+		public override string Compute(BgExprContext context) => Value.Compute(context);
 	}
 
 	class BgStringConstantExpr : BgString
 	{
 		public string Value { get; }
 
-		public BgStringConstantExpr(string Value)
+		public BgStringConstantExpr(string value)
 		{
-			this.Value = Value;
+			Value = value;
 		}
 
-		public override string Compute(BgExprContext Context) => Value;
+		public override string Compute(BgExprContext context) => Value;
 
-		public override bool Equals(object? Obj)
+		public override bool Equals(object? obj)
 		{
-			return Obj is string ObjStr && String.Equals(Value, ObjStr, StringComparison.OrdinalIgnoreCase);
+			return obj is string objStr && String.Equals(Value, objStr, StringComparison.OrdinalIgnoreCase);
 		}
 
 		public override int GetHashCode()
diff --git a/Engine/Source/Programs/Shared/EpicGames.BuildGraph/Expressions/IBgExpr.cs b/Engine/Source/Programs/Shared/EpicGames.BuildGraph/Expressions/IBgExpr.cs
--- a/Engine/Source/Programs/Shared/EpicGames.BuildGraph/Expressions/IBgExpr.cs
+++ b/Engine/Source/Programs/Shared/EpicGames.BuildGraph/Expressions/IBgExpr.cs
@@ -1,11 +1,8 @@
 // Copyright Epic Games, Inc. All Rights Reserved.
 
+using System.Collections.Generic;
 using EpicGames.BuildGraph.Expressions;
 using EpicGames.Core;
-using System;
-using System.Collections.Generic;
-using System.Reflection;
-using System.Text;
 
 namespace EpicGames.BuildGraph
 {
@@ -17,12 +14,12 @@ namespace EpicGames.BuildGraph
 		/// <summary>
 		/// Command line options for the graph
 		/// </summary>
-		public Dictionary<string, string> Options = new Dictionary<string, string>();
+		public Dictionary<string, string> Options { get; set; } = new Dictionary<string, string>();
 
 		/// <summary>
 		/// Map from tag name to fileset
 		/// </summary>
-		public Dictionary<string, FileSet> TagNameToFileSet = new Dictionary<string, FileSet>();
+		public Dictionary<string, FileSet> TagNameToFileSet { get; set; } = new Dictionary<string, FileSet>();
 	}
 
 	/// <summary>
@@ -33,8 +30,8 @@ namespace EpicGames.BuildGraph
 		/// <summary>
 		/// Compute the value of this expression
 		/// </summary>
-		/// <param name="Context"></param>
-		object Compute(BgExprContext Context);
+		/// <param name="context"></param>
+		object Compute(BgExprContext context);
 
 		/// <summary>
 		/// Helper method to create an expression which formats this expression as a string
@@ -50,9 +47,9 @@ namespace EpicGames.BuildGraph
 		/// <summary>
 		/// Chooses between the current value and another value based on a condition
 		/// </summary>
-		/// <param name="Condition">The condition to evaluate</param>
-		/// <param name="ValueIfTrue">Value to return if the condition is true</param>
-		TExpr IfThen(BgBool Condition, TExpr ValueIfTrue);
+		/// <param name="condition">The condition to evaluate</param>
+		/// <param name="valueIfTrue">Value to return if the condition is true</param>
+		TExpr IfThen(BgBool condition, TExpr valueIfTrue);
 	}
 
 	/// <summary>
diff --git a/Engine/Source/Programs/Shared/EpicGames.BuildGraph/Expressions/IBgType.cs b/Engine/Source/Programs/Shared/EpicGames.BuildGraph/Expressions/IBgType.cs
--- a/Engine/Source/Programs/Shared/EpicGames.BuildGraph/Expressions/IBgType.cs
+++ b/Engine/Source/Programs/Shared/EpicGames.BuildGraph/Expressions/IBgType.cs
@@ -2,9 +2,7 @@
 
 using System;
 using System.Collections.Concurrent;
-using System.Collections.Generic;
 using System.Reflection;
-using System.Text;
 
 namespace EpicGames.BuildGraph.Expressions
 {
@@ -16,17 +14,17 @@ namespace EpicGames.BuildGraph.Expressions
 		/// <summary>
 		/// Serializes the value of an expression to a string
 		/// </summary>
-		/// <param name="Value">The value to serialize</param>
-		/// <param name="Context">Context for evaluating the expression</param>
+		/// <param name="value">The value to serialize</param>
+		/// <param name="context">Context for evaluating the expression</param>
 		/// <returns></returns>
-		string SerializeArgument(object Value, BgExprContext Context);
+		string SerializeArgument(object value, BgExprContext context);
 
 		/// <summary>
 		/// Constructs an expression from a string value
 		/// </summary>
-		/// <param name="Text">The serialized value</param>
+		/// <param name="text">The serialized value</param>
 		/// <returns>Instance of the expression</returns>
-		object DeserializeArgument(string Text);
+		object DeserializeArgument(string text);
 	}
 
 	/// <summary>
@@ -37,24 +35,24 @@ namespace EpicGames.BuildGraph.Expressions
 		/// <summary>
 		/// Serializes the value of an expression to an argument string
 		/// </summary>
-		/// <param name="Value">The value to serialize</param>
-		/// <param name="Context">Context for evaluating the expression</param>
+		/// <param name="value">The value to serialize</param>
+		/// <param name="context">Context for evaluating the expression</param>
 		/// <returns></returns>
-		string SerializeArgument(T Value, BgExprContext Context);
+		string SerializeArgument(T value, BgExprContext context);
 
 		/// <summary>
 		/// Constructs an expression from an argument value
 		/// </summary>
-		/// <param name="Text">The serialized value</param>
+		/// <param name="text">The serialized value</param>
 		/// <returns>Instance of the expression</returns>
-		new T DeserializeArgument(string Text);
+		new T DeserializeArgument(string text);
 
 		/// <summary>
 		/// Creates a value of the expression type
 		/// </summary>
-		/// <param name="Value">Value to wrap</param>
+		/// <param name="value">Value to wrap</param>
 		/// <returns></returns>
-		T CreateConstant(object Value);
+		T CreateConstant(object value);
 
 		/// <summary>
 		/// Creates a variable of the given type
@@ -69,19 +67,19 @@ namespace EpicGames.BuildGraph.Expressions
 	public abstract class BgTypeBase<T> : IBgType<T> where T : IBgExpr<T>
 	{
 		/// <inheritdoc/>
-		object IBgType.DeserializeArgument(string Text) => DeserializeArgument(Text);
+		object IBgType.DeserializeArgument(string text) => DeserializeArgument(text);
 
 		/// <inheritdoc/>
-		public abstract T DeserializeArgument(string Text);
+		public abstract T DeserializeArgument(string text);
 
 		/// <inheritdoc/>
-		string IBgType.SerializeArgument(object Value, BgExprContext Context) => SerializeArgument((T)Value, Context);
+		string IBgType.SerializeArgument(object value, BgExprContext context) => SerializeArgument((T)value, context);
 
 		/// <inheritdoc/>
-		public abstract string SerializeArgument(T Value, BgExprContext Context);
+		public abstract string SerializeArgument(T value, BgExprContext context);
 
 		/// <inheritdoc/>
-		public abstract T CreateConstant(object Value);
+		public abstract T CreateConstant(object value);
 
 		/// <inheritdoc/>
 		public abstract IBgExprVariable<T> CreateVariable();
@@ -100,10 +98,10 @@ namespace EpicGames.BuildGraph.Expressions
 		/// <summary>
 		/// Constructor
 		/// </summary>
-		/// <param name="Type"></param>
-		public BgTypeAttribute(Type Type)
+		/// <param name="type"></param>
+		public BgTypeAttribute(Type type)
 		{
-			this.Type = Type;
+			Type = type;
 		}
 	}
 
@@ -115,7 +113,7 @@ namespace EpicGames.BuildGraph.Expressions
 		/// <summary>
 		/// Cache of traits implementations for each type
 		/// </summary>
-		static ConcurrentDictionary<Type, IBgType> TypeToTraits = new ConcurrentDictionary<Type, IBgType>();
+		static readonly ConcurrentDictionary<Type, IBgType> s_typeToTraits = new ConcurrentDictionary<Type, IBgType>();
 
 		/// <summary>
 		/// Cache of traits for a static type
@@ -128,60 +126,60 @@ namespace EpicGames.BuildGraph.Expressions
 		/// <summary>
 		/// Create a converter instance for the given type
 		/// </summary>
-		/// <param name="Type"></param>
+		/// <param name="type"></param>
 		/// <returns></returns>
-		public static IBgType Get(Type Type)
+		public static IBgType Get(Type type)
 		{
-			IBgType? Traits = GetInner(Type);
-			if(Traits == null)
+			IBgType? traits = GetInner(type);
+			if (traits == null)
 			{
-				throw new ArgumentException($"Missing converter attribute on type {Type.Name}");
+				throw new ArgumentException($"Missing converter attribute on type {type.Name}");
 			}
-			return Traits;
+			return traits;
 		}
 
 		/// <summary>
 		/// Helper method for getting the type traits from a hierarchy of types
 		/// </summary>
-		/// <param name="Type"></param>
+		/// <param name="type"></param>
 		/// <returns></returns>
-		static IBgType? GetInner(Type Type)
+		static IBgType? GetInner(Type type)
 		{
-			IBgType? Traits;
-			if (!TypeToTraits.TryGetValue(Type, out Traits))
+			IBgType? traits;
+			if (!s_typeToTraits.TryGetValue(type, out traits))
 			{
-				BgTypeAttribute? ConverterAttr = Type.GetCustomAttribute<BgTypeAttribute>(false);
-				if (ConverterAttr == null)
+				BgTypeAttribute? converterAttr = type.GetCustomAttribute<BgTypeAttribute>(false);
+				if (converterAttr == null)
 				{
-					if (Type.BaseType == null)
+					if (type.BaseType == null)
 					{
-						Traits = null;
+						traits = null;
 					}
 					else
 					{
-						Traits = GetInner(Type.BaseType);
+						traits = GetInner(type.BaseType);
 					}
 				}
 				else
 				{
-					Type ConverterType = ConverterAttr.Type;
-					if (ConverterType.IsGenericType)
+					Type converterType = converterAttr.Type;
+					if (converterType.IsGenericType)
 					{
-						ConverterType = ConverterType.MakeGenericType(Type.GetGenericArguments());
+						converterType = converterType.MakeGenericType(type.GetGenericArguments());
 					}
 
-					IBgType NewTraits = (IBgType)Activator.CreateInstance(ConverterType)!;
-					if (TypeToTraits.TryAdd(Type, NewTraits))
+					IBgType newTraits = (IBgType)Activator.CreateInstance(converterType)!;
+					if (s_typeToTraits.TryAdd(type, newTraits))
 					{
-						Traits = NewTraits;
+						traits = newTraits;
 					}
 					else
 					{
-						Traits = TypeToTraits[Type];
+						traits = s_typeToTraits[type];
 					}
 				}
 			}
-			return Traits;
+			return traits;
 		}
 
 		/// <summary>
