Updated UHT in UBT with changes to getter/setter lookup.
Fixed comment issue at the end of function declarations.
Fixed nested include file issue.
Fixed issue with RigVM visible arguments.
Updated NoExportTypes.h to mark two structures as core types.
Fixed type check for bit fields in Getter/Setters.

#rnx

diff --git a/Engine/Source/Programs/Shared/EpicGames.UHT/Parsers/UhtFunctionParser.cs b/Engine/Source/Programs/Shared/EpicGames.UHT/Parsers/UhtFunctionParser.cs
--- a/Engine/Source/Programs/Shared/EpicGames.UHT/Parsers/UhtFunctionParser.cs
+++ b/Engine/Source/Programs/Shared/EpicGames.UHT/Parsers/UhtFunctionParser.cs
@@ -463,7 +463,7 @@ namespace EpicGames.UHT.Parsers
 					}
 
 					// Record the tokens so we can detect this function as a declaration later (i.e. RPC)
-					using (UhtTokenRecorder TokenRecorder = new UhtTokenRecorder(ParentScope))
+					using (UhtTokenRecorder TokenRecorder = new UhtTokenRecorder(ParentScope, Function))
 					{
 
 						if (TopScope.TokenReader.TryOptional("virtual"))
@@ -564,6 +564,10 @@ namespace EpicGames.UHT.Parsers
 						// Try parsing metadata for the function
 						SpecifierParser.ParseFieldMetaData(Function.EngineName);
 
+						// COMPATIBILITY-TODO - Try to pull any comment following the declaration
+						TopScope.TokenReader.PeekToken();
+						TopScope.TokenReader.CommitPendingComments();
+
 						TopScope.AddFormattedCommentsAsTooltipMetaData();
 
 						// 'final' and 'override' can appear in any order before an optional '= 0' pure virtual specifier
diff --git a/Engine/Source/Programs/Shared/EpicGames.UHT/Parsers/UhtHeaderFileParser.cs b/Engine/Source/Programs/Shared/EpicGames.UHT/Parsers/UhtHeaderFileParser.cs
--- a/Engine/Source/Programs/Shared/EpicGames.UHT/Parsers/UhtHeaderFileParser.cs
+++ b/Engine/Source/Programs/Shared/EpicGames.UHT/Parsers/UhtHeaderFileParser.cs
@@ -874,7 +874,7 @@ namespace EpicGames.UHT.Parsers
 			{
 				return true;
 			}
-			return (GetCurrentCompositeCompilerDirective() & (UhtCompilerDirective.NotCPPBlock | UhtCompilerDirective.OneBlock | UhtCompilerDirective.WithEditor | UhtCompilerDirective.WithEditorOnlyData | UhtCompilerDirective.WithHotReload)) != 0;
+			return !GetCurrentCompositeCompilerDirective().HasAnyFlags(UhtCompilerDirective.CPPBlock | UhtCompilerDirective.ZeroBlock | UhtCompilerDirective.Unrecognized);
 		}
 
 		/// <summary>
diff --git a/Engine/Source/Programs/Shared/EpicGames.UHT/Parsers/UhtParsingScope.cs b/Engine/Source/Programs/Shared/EpicGames.UHT/Parsers/UhtParsingScope.cs
--- a/Engine/Source/Programs/Shared/EpicGames.UHT/Parsers/UhtParsingScope.cs
+++ b/Engine/Source/Programs/Shared/EpicGames.UHT/Parsers/UhtParsingScope.cs
@@ -483,6 +483,7 @@ namespace EpicGames.UHT.Parsers
 	{
 		private readonly UhtCompilerDirective CompilerDirective;
 		private readonly UhtParsingScope Scope;
+		private readonly UhtFunction? Function;
 		private bool bFlushed;
 
 		/// <summary>
@@ -494,6 +495,7 @@ namespace EpicGames.UHT.Parsers
 		{
 			this.Scope = Scope;
 			this.CompilerDirective = this.Scope.HeaderParser.GetCurrentCompositeCompilerDirective();
+			this.Function = null;
 			this.bFlushed = false;
 			if (this.Scope.ScopeType is UhtClass)
 			{
@@ -506,10 +508,12 @@ namespace EpicGames.UHT.Parsers
 		/// Create a new recorder
 		/// </summary>
 		/// <param name="Scope">Scope being parsed</param>
-		public UhtTokenRecorder(UhtParsingScope Scope)
+		/// <param name="Function">Function associated with the recorder</param>
+		public UhtTokenRecorder(UhtParsingScope Scope, UhtFunction Function)
 		{
 			this.Scope = Scope;
 			this.CompilerDirective = this.Scope.HeaderParser.GetCurrentCompositeCompilerDirective();
+			this.Function = Function;
 			this.bFlushed = false;
 			if (this.Scope.ScopeType is UhtClass)
 			{
@@ -536,7 +540,7 @@ namespace EpicGames.UHT.Parsers
 				bFlushed = true;
 				if (this.Scope.ScopeType is UhtClass Class)
 				{
-					Class.AddDeclaration(this.CompilerDirective, this.Scope.TokenReader.RecordedTokens);
+					Class.AddDeclaration(this.CompilerDirective, this.Scope.TokenReader.RecordedTokens, this.Function);
 					this.Scope.TokenReader.DisableRecording();
 					return true;
 				}
diff --git a/Engine/Source/Programs/Shared/EpicGames.UHT/Types/UhtClass.cs b/Engine/Source/Programs/Shared/EpicGames.UHT/Types/UhtClass.cs
--- a/Engine/Source/Programs/Shared/EpicGames.UHT/Types/UhtClass.cs
+++ b/Engine/Source/Programs/Shared/EpicGames.UHT/Types/UhtClass.cs
@@ -234,6 +234,11 @@ namespace EpicGames.UHT.Types
 		/// Collection of tokens parsed in the declaration
 		/// </summary>
 		public UhtToken[] Tokens;
+
+		/// <summary>
+		/// If this declaration is part of a UFUNCTION, this will be set
+		/// </summary>
+		public UhtFunction? Function;
 	}
 
 	/// <summary>
@@ -455,13 +460,14 @@ namespace EpicGames.UHT.Types
 		/// </summary>
 		/// <param name="CompilerDirectives">Currently active compiler directives</param>
 		/// <param name="Tokens">List of declaration tokens</param>
-		public void AddDeclaration(UhtCompilerDirective CompilerDirectives, List<UhtToken> Tokens)
+		/// <param name="Function">If parsed as part of a UFUNCTION, this will reference it</param>
+		public void AddDeclaration(UhtCompilerDirective CompilerDirectives, List<UhtToken> Tokens, UhtFunction? Function)
 		{
 			if (this.DeclarationsInternal == null)
 			{
 				this.DeclarationsInternal = new List<UhtDeclaration>();
 			}
-			this.DeclarationsInternal.Add(new UhtDeclaration { CompilerDirectives = CompilerDirectives, Tokens = Tokens.ToArray() });
+			this.DeclarationsInternal.Add(new UhtDeclaration { CompilerDirectives = CompilerDirectives, Tokens = Tokens.ToArray(), Function = Function });
 		}
 
 		/// <summary>
@@ -573,7 +579,6 @@ namespace EpicGames.UHT.Types
 					break;
 
 				case UhtResolvePhase.Final:
-					bool bAutoGetterSetter = UhtConfig.Instance.bAllowAutomaticSettersAndGetters;
 					Dictionary<string, List<GetterSetterToResolve>>? GSToResolve = null;
 					foreach (UhtType Child in this.Children)
 					{
@@ -596,10 +601,6 @@ namespace EpicGames.UHT.Types
 								{
 									GSToResolve = AddGetterSetter(GSToResolve, Property.Getter == null ? $"Get{Property.SourceName}" : Property.Getter, Property, false);
 								}
-								else if (bAutoGetterSetter)
-								{
-									GSToResolve = AddGetterSetter(GSToResolve, $"Get{Property.SourceName}", Property, false);
-								}
 							}
 
 							if (!Property.PropertyExportFlags.HasAnyFlags(UhtPropertyExportFlags.SetterSpecifiedNone))
@@ -608,10 +609,6 @@ namespace EpicGames.UHT.Types
 								{
 									GSToResolve = AddGetterSetter(GSToResolve, Property.Setter == null ? $"Set{Property.SourceName}" : Property.Setter, Property, true);
 								}
-								else if (bAutoGetterSetter)
-								{
-									GSToResolve = AddGetterSetter(GSToResolve, $"Set{Property.SourceName}", Property, true);
-								}
 							}
 						}
 						else if (Child is UhtFunction Function)
@@ -876,23 +873,35 @@ namespace EpicGames.UHT.Types
 
 			// Validate the #if blocks.  This should really be done during validation
 			string FuncType = bSetter ? "setter" : "getter";
+			UhtPropertyExportFlags TestFlag = bSetter ? UhtPropertyExportFlags.SetterSpecified : UhtPropertyExportFlags.GetterSpecified;
 			if (Declaration.CompilerDirectives.HasAnyFlags(UhtCompilerDirective.WithEditor))
 			{
-				if (Property.PropertyExportFlags.HasAnyFlags(UhtPropertyExportFlags.GetterSpecified))
+				if (Property.PropertyExportFlags.HasAnyFlags(TestFlag))
 				{
-					Property.LogError(Tokens[0].InputLine, $"Property {Property.SourceName} {FuncType} function {Tokens[NameIndex].Value.ToString()} "
+					Property.LogError(Tokens[0].InputLine, $"Property {Property.SourceName} {FuncType} function {Tokens[NameIndex].Value} "
 						+ "cannot be declared within WITH_EDITOR block. Use WITH_EDITORONLY_DATA instead.");
 				}
 				return false;
 			}
 			if (Declaration.CompilerDirectives.HasAnyFlags(UhtCompilerDirective.WithEditorOnlyData) && !Property.PropertyFlags.HasAnyFlags(EPropertyFlags.EditorOnly))
 			{
-				if (Property.PropertyExportFlags.HasAnyFlags(UhtPropertyExportFlags.GetterSpecified))
+				if (Property.PropertyExportFlags.HasAnyFlags(TestFlag))
 				{
-					Property.LogError(Tokens[0].InputLine, $"Property {Property.SourceName} is not editor-only but its {FuncType} function '{Tokens[NameIndex].Value.ToString()}' is");
+					Property.LogError(Tokens[0].InputLine, $"Property {Property.SourceName} is not editor-only but its {FuncType} function '{Tokens[NameIndex].Value}' is");
 				}
 				return false;
 			}
+			if (Declaration.Function != null)
+			{
+				if (!Declaration.Function.FunctionFlags.HasAnyFlags(EFunctionFlags.Native))
+				{
+					Property.LogError(Tokens[0].InputLine, $"Property {Property.SourceName} {FuncType} function '{Tokens[NameIndex].Value}' has to be native");
+				}
+				if (Declaration.Function.FunctionFlags.HasAnyFlags(EFunctionFlags.Net | EFunctionFlags.Event))
+				{
+					Property.LogError(Tokens[0].InputLine, $"Property {Property.SourceName} {FuncType} function '{Tokens[NameIndex].Value}' cannot be a net or an event");
+				}
+			}
 			return true;
 		}
 
diff --git a/Engine/Source/Programs/Shared/EpicGames.UHT/Types/UhtScriptStruct.cs b/Engine/Source/Programs/Shared/EpicGames.UHT/Types/UhtScriptStruct.cs
--- a/Engine/Source/Programs/Shared/EpicGames.UHT/Types/UhtScriptStruct.cs
+++ b/Engine/Source/Programs/Shared/EpicGames.UHT/Types/UhtScriptStruct.cs
@@ -196,8 +196,8 @@ namespace EpicGames.UHT.Types
 
 			if (Property.MetaData.ContainsKey("Visible"))
 			{
-				this.ParameterFlags |= UhtRigVMParameterFlags.Constant;
-				this.ParameterFlags &= (UhtRigVMParameterFlags.Input | UhtRigVMParameterFlags.Output);
+				this.ParameterFlags |= UhtRigVMParameterFlags.Constant | UhtRigVMParameterFlags.Input;
+				this.ParameterFlags &= ~UhtRigVMParameterFlags.Output;
 			}
 
 			if (this.bEditorOnly)
diff --git a/Engine/Source/Programs/Shared/EpicGames.UHT/Utils/UhtConfig.cs b/Engine/Source/Programs/Shared/EpicGames.UHT/Utils/UhtConfig.cs
--- a/Engine/Source/Programs/Shared/EpicGames.UHT/Utils/UhtConfig.cs
+++ b/Engine/Source/Programs/Shared/EpicGames.UHT/Utils/UhtConfig.cs
@@ -59,11 +59,6 @@ namespace EpicGames.UHT.Utils
 		/// </summary>
 		public UhtPointerMemberBehavior NonEngineObjectPtrMemberBehavior { get; }
 
-		/// <summary>
-		/// If true setters and getters will be automatically (without specifying their function names on a property) parsed and generated if a function with matching signature is found
-		/// </summary>
-		public bool bAllowAutomaticSettersAndGetters { get; }
-
 		/// <summary>
 		/// If the token references a remapped identifier, update the value in the token 
 		/// </summary>
diff --git a/Engine/Source/Programs/UnrealBuildTool/Modes/UnrealHeaderToolMode.cs b/Engine/Source/Programs/UnrealBuildTool/Modes/UnrealHeaderToolMode.cs
--- a/Engine/Source/Programs/UnrealBuildTool/Modes/UnrealHeaderToolMode.cs
+++ b/Engine/Source/Programs/UnrealBuildTool/Modes/UnrealHeaderToolMode.cs
@@ -79,11 +79,6 @@ namespace UnrealBuildTool.Modes
 		/// </summary>
 		private readonly UhtPointerMemberBehavior NonEngineObjectPtrMemberBehaviorInternal = UhtPointerMemberBehavior.AllowSilently;
 
-		/// <summary>
-		/// If true setters and getters will be automatically (without specifying their function names on a property) parsed and generated if a function with matching signature is found
-		/// </summary>
-		private readonly bool bAllowAutomaticSettersAndGettersInternal;
-
 		#region IUhtConfig Implementation
 		/// <inheritdoc/>
 		public EGeneratedCodeVersion DefaultGeneratedCodeVersion => this.DefaultGeneratedCodeVersionInternal;
@@ -100,9 +95,6 @@ namespace UnrealBuildTool.Modes
 		/// <inheritdoc/>
 		public UhtPointerMemberBehavior NonEngineObjectPtrMemberBehavior => this.NonEngineObjectPtrMemberBehaviorInternal;
 
-		/// <inheritdoc/>
-		public bool bAllowAutomaticSettersAndGetters => this.bAllowAutomaticSettersAndGettersInternal;
-
 		/// <inheritdoc/>
 		public void RedirectTypeIdentifier(ref UhtToken Token)
 		{
@@ -189,7 +181,6 @@ namespace UnrealBuildTool.Modes
 			DirectoryReference ConfigDirectory = DirectoryReference.Combine(Unreal.EngineDirectory, "Programs/UnrealBuildTool");
 			this.Ini = ConfigCache.ReadHierarchy(ConfigHierarchyType.Engine, ConfigDirectory, BuildHostPlatform.Current.Platform, "", Args.GetRawArray());
 
-			this.bAllowAutomaticSettersAndGettersInternal = GetBoolean("UnrealHeaderTool", "AutomaticSettersAndGetters", false);
 			this.TypeRedirectMap = GetRedirectsStringView("UnrealHeaderTool", "TypeRedirects", "OldType", "NewType");
 			this.MetaDataRedirectMap = GetRedirectsString("CoreUObject.Metadata", "MetadataRedirects", "OldKey", "NewKey");
 			this.StructsWithNoPrefix = GetHashSet("UnrealHeaderTool", "StructsWithNoPrefix", StringViewComparer.Ordinal);
diff --git a/Engine/Source/Runtime/CoreUObject/Public/UObject/NoExportTypes.h b/Engine/Source/Runtime/CoreUObject/Public/UObject/NoExportTypes.h
--- a/Engine/Source/Runtime/CoreUObject/Public/UObject/NoExportTypes.h
+++ b/Engine/Source/Runtime/CoreUObject/Public/UObject/NoExportTypes.h
@@ -839,7 +839,7 @@ struct FPackedRGBA16N
  * Screen coordinates.
  * @note The full C++ class is located here: Engine\Source\Runtime\Core\Public\Math\IntPoint.h
  */
-USTRUCT(immutable, noexport, BlueprintType, IsAlwaysAccessible, HasDefaults, HasNoOpConstructor)
+USTRUCT(immutable, noexport, BlueprintType, IsAlwaysAccessible, HasDefaults, HasNoOpConstructor, IsCoreType)
 struct FIntPoint
 {
 	UPROPERTY(EditAnywhere, BlueprintReadWrite, Category=IntPoint, SaveGame)
@@ -854,7 +854,7 @@ struct FIntPoint
  * An integer vector in 3D space.
  * @note The full C++ class is located here: Engine\Source\Runtime\Core\Public\Math\IntVector.h
  */
-USTRUCT(immutable, noexport, BlueprintType, IsAlwaysAccessible, HasDefaults, HasNoOpConstructor)
+USTRUCT(immutable, noexport, BlueprintType, IsAlwaysAccessible, HasDefaults, HasNoOpConstructor, IsCoreType)
 struct FIntVector
 {
 	UPROPERTY(EditAnywhere, BlueprintReadWrite, Category=IntVector, SaveGame)
