Horde: Fix build errors after BuildGraph refactor.

diff --git a/Engine/Source/Programs/Horde/Horde.Build/Utilities/BuildGraphExtensions.cs b/Engine/Source/Programs/Horde/Horde.Build/Utilities/BuildGraphExtensions.cs
--- a/Engine/Source/Programs/Horde/Horde.Build/Utilities/BuildGraphExtensions.cs
+++ b/Engine/Source/Programs/Horde/Horde.Build/Utilities/BuildGraphExtensions.cs
@@ -170,8 +170,8 @@ namespace Horde.Build.Utilities
 				List<NewNode> newNodes = new List<NewNode>();
 				foreach (BgNode node in agent.Nodes)
 				{
-					const bool BAllowRetry = true;
-					NewNode newNode = new NewNode(node.Name, node.InputDependencies.Select(x => x.Name).ToList(), node.OrderDependencies.Select(x => x.Name).ToList(), HordeCommon.Priority.Normal, BAllowRetry, node.bRunEarly, node.bNotifyOnWarnings, new Dictionary<string, string>(), new Dictionary<string, string>());
+					const bool allowRetry = true;
+					NewNode newNode = new NewNode(node.Name, node.InputDependencies.Select(x => x.Name).ToList(), node.OrderDependencies.Select(x => x.Name).ToList(), HordeCommon.Priority.Normal, allowRetry, node.RunEarly, node.NotifyOnWarnings, new Dictionary<string, string>(), new Dictionary<string, string>());
 					newNodes.Add(newNode);
 				}
 
diff --git a/Engine/Source/Programs/Shared/EpicGames.BuildGraph/BgGraph.cs b/Engine/Source/Programs/Shared/EpicGames.BuildGraph/BgGraph.cs
--- a/Engine/Source/Programs/Shared/EpicGames.BuildGraph/BgGraph.cs
+++ b/Engine/Source/Programs/Shared/EpicGames.BuildGraph/BgGraph.cs
@@ -411,11 +411,11 @@ namespace EpicGames.BuildGraph
 							jsonWriter.WriteObjectStart();
 							jsonWriter.WriteValue("Name", node.Name);
 							jsonWriter.WriteValue("DependsOn", String.Join(";", node.GetDirectOrderDependencies().Where(x => nodesToExecute.Contains(x))));
-							jsonWriter.WriteValue("RunEarly", node.BRunEarly);
+							jsonWriter.WriteValue("RunEarly", node.RunEarly);
 							jsonWriter.WriteObjectStart("Notify");
 							jsonWriter.WriteValue("Default", String.Join(";", node.NotifyUsers));
 							jsonWriter.WriteValue("Submitters", String.Join(";", node.NotifySubmitters));
-							jsonWriter.WriteValue("Warnings", node.BNotifyOnWarnings);
+							jsonWriter.WriteValue("Warnings", node.NotifyOnWarnings);
 							jsonWriter.WriteObjectEnd();
 							jsonWriter.WriteObjectEnd();
 						}
@@ -510,8 +510,8 @@ namespace EpicGames.BuildGraph
 					{
 						jsonWriter.WriteObjectStart();
 						jsonWriter.WriteValue("Name", node.Name);
-						jsonWriter.WriteValue("RunEarly", node.BRunEarly);
-						jsonWriter.WriteValue("Warnings", node.BNotifyOnWarnings);
+						jsonWriter.WriteValue("RunEarly", node.RunEarly);
+						jsonWriter.WriteValue("Warnings", node.NotifyOnWarnings);
 
 						jsonWriter.WriteArrayStart("InputDependencies");
 						foreach (string inputDependency in node.GetDirectInputDependencies().Select(x => x.Name))
@@ -674,7 +674,7 @@ namespace EpicGames.BuildGraph
 				logger.LogInformation("        Agent: {0} ({1})", agent.Name, String.Join(";", agent.PossibleTypes));
 				foreach (BgNode node in agent.Nodes)
 				{
-					logger.LogInformation("            Node: {0}{1}", node.Name, completedNodes.Contains(node) ? " (completed)" : node.BRunEarly ? " (early)" : "");
+					logger.LogInformation("            Node: {0}{1}", node.Name, completedNodes.Contains(node) ? " (completed)" : node.RunEarly ? " (early)" : "");
 					if (printOptions.HasFlag(GraphPrintOptions.ShowDependencies))
 					{
 						HashSet<BgNode> inputDependencies = new HashSet<BgNode>(node.GetDirectInputDependencies());
@@ -690,7 +690,7 @@ namespace EpicGames.BuildGraph
 					}
 					if (printOptions.HasFlag(GraphPrintOptions.ShowNotifications))
 					{
-						string label = node.BNotifyOnWarnings ? "warnings" : "errors";
+						string label = node.NotifyOnWarnings ? "warnings" : "errors";
 						foreach (string user in node.NotifyUsers)
 						{
 							logger.LogInformation("                {0}> {1}", label, user);
diff --git a/Engine/Source/Programs/Shared/EpicGames.BuildGraph/BgNode.cs b/Engine/Source/Programs/Shared/EpicGames.BuildGraph/BgNode.cs
--- a/Engine/Source/Programs/Shared/EpicGames.BuildGraph/BgNode.cs
+++ b/Engine/Source/Programs/Shared/EpicGames.BuildGraph/BgNode.cs
@@ -97,12 +97,12 @@ namespace EpicGames.BuildGraph
 		/// <summary>
 		/// Whether to start this node as soon as its dependencies are satisfied, rather than waiting for all of its agent's dependencies to be met.
 		/// </summary>
-		public bool BRunEarly { get; set; } = false;
+		public bool RunEarly { get; set; } = false;
 
 		/// <summary>
 		/// Whether to ignore warnings produced by this node
 		/// </summary>
-		public bool BNotifyOnWarnings { get; set; } = true;
+		public bool NotifyOnWarnings { get; set; } = true;
 
 		/// <summary>
 		/// Constructor
@@ -191,14 +191,14 @@ namespace EpicGames.BuildGraph
 				writer.WriteAttributeString("After", String.Join(";", afterNames));
 			}
 
-			if (!BNotifyOnWarnings)
+			if (!NotifyOnWarnings)
 			{
-				writer.WriteAttributeString("NotifyOnWarnings", BNotifyOnWarnings.ToString());
+				writer.WriteAttributeString("NotifyOnWarnings", NotifyOnWarnings.ToString());
 			}
 
-			if (BRunEarly)
+			if (RunEarly)
 			{
-				writer.WriteAttributeString("RunEarly", BRunEarly.ToString());
+				writer.WriteAttributeString("RunEarly", RunEarly.ToString());
 			}
 
 			foreach (BgTask task in Tasks)
diff --git a/Engine/Source/Programs/Shared/EpicGames.BuildGraph/BgNodeSpec.cs b/Engine/Source/Programs/Shared/EpicGames.BuildGraph/BgNodeSpec.cs
--- a/Engine/Source/Programs/Shared/EpicGames.BuildGraph/BgNodeSpec.cs
+++ b/Engine/Source/Programs/Shared/EpicGames.BuildGraph/BgNodeSpec.cs
@@ -380,7 +380,7 @@ namespace EpicGames.BuildGraph
 			bool runEarly = CanRunEarly.Compute(context);
 
 			BgNode node = new BgNode(name, inputDependencies, outputNames, inputNodes, afterNodes, Array.Empty<FileReference>());
-			node.BRunEarly = runEarly;
+			node.RunEarly = runEarly;
 
 			BgScriptLocation location = new BgScriptLocation("(unknown)", "(unknown)", 1);
 
diff --git a/Engine/Source/Programs/Shared/EpicGames.BuildGraph/BgScriptReader.cs b/Engine/Source/Programs/Shared/EpicGames.BuildGraph/BgScriptReader.cs
--- a/Engine/Source/Programs/Shared/EpicGames.BuildGraph/BgScriptReader.cs
+++ b/Engine/Source/Programs/Shared/EpicGames.BuildGraph/BgScriptReader.cs
@@ -1830,8 +1830,8 @@ namespace EpicGames.BuildGraph
 				{
 					// Add it to the node lookup
 					BgNode newNode = new BgNode(name, inputs.ToArray(), validOutputNames.ToArray(), inputDependencies.ToArray(), orderDependencies.ToArray(), requiredTokens.ToArray());
-					newNode.BRunEarly = bRunEarly;
-					newNode.BNotifyOnWarnings = bNotifyOnWarnings;
+					newNode.RunEarly = bRunEarly;
+					newNode.NotifyOnWarnings = bNotifyOnWarnings;
 					_graph.NameToNode.Add(name, newNode);
 
 					// Register all the output tags in the global name table.
@@ -1972,7 +1972,7 @@ namespace EpicGames.BuildGraph
 					}
 					if (bWarnings.HasValue)
 					{
-						node.BNotifyOnWarnings = bWarnings.Value;
+						node.NotifyOnWarnings = bWarnings.Value;
 					}
 				}
 
