Fix test failure.

diff --git a/Engine/Source/Programs/Shared/EpicGames.Serialization.Tests/SerializerTests.cs b/Engine/Source/Programs/Shared/EpicGames.Serialization.Tests/SerializerTests.cs
--- a/Engine/Source/Programs/Shared/EpicGames.Serialization.Tests/SerializerTests.cs
+++ b/Engine/Source/Programs/Shared/EpicGames.Serialization.Tests/SerializerTests.cs
@@ -46,7 +46,7 @@ namespace EpicGames.Serialization.Tests
 			ReadOnlyMemory<byte> mem2 = cbObj2.GetView();
 
 			Assert.AreEqual(obj1, obj2);
-			Assert.AreEqual(mem1, mem2);
+			Assert.IsTrue(mem1.Span.SequenceEqual(mem2.Span));
 		}
 
 		[TestMethod]
