Fixed issue where UHT wasn't properly renaming the output files.
Added number of written files as an informational line.

#rnx

diff --git a/Engine/Source/Programs/Shared/EpicGames.UHT/Utils/UhtFileManager.cs b/Engine/Source/Programs/Shared/EpicGames.UHT/Utils/UhtFileManager.cs
--- a/Engine/Source/Programs/Shared/EpicGames.UHT/Utils/UhtFileManager.cs
+++ b/Engine/Source/Programs/Shared/EpicGames.UHT/Utils/UhtFileManager.cs
@@ -156,7 +156,7 @@ namespace EpicGames.UHT.Utils
 		{
 			try
 			{
-				File.Move(OldFilePath, OldFilePath, true);
+				File.Move(OldFilePath, NewFilePath, true);
 				return true;
 			}
 			catch (Exception)
diff --git a/Engine/Source/Programs/Shared/EpicGames.UHT/Utils/UhtSession.cs b/Engine/Source/Programs/Shared/EpicGames.UHT/Utils/UhtSession.cs
--- a/Engine/Source/Programs/Shared/EpicGames.UHT/Utils/UhtSession.cs
+++ b/Engine/Source/Programs/Shared/EpicGames.UHT/Utils/UhtSession.cs
@@ -2364,6 +2364,7 @@ namespace EpicGames.UHT.Utils
 		#region Exporting
 		private void StepExport()
 		{
+			long TotalWrittenFiles = 0;
 			Try(null, () =>
 			{
 				Log.Logger.LogTrace("Step - Exports");
@@ -2383,6 +2384,13 @@ namespace EpicGames.UHT.Utils
 						Log.Logger.LogTrace($"       Running exporter {Exporter.Name}");
 						UhtExportFactory Factory = new UhtExportFactory(this, Exporter, Options);
 						Factory.Run();
+						foreach (var Output in Factory.Outputs)
+						{
+							if (Output.bSaved)
+							{
+								TotalWrittenFiles++;
+							}
+						}
 					}
 					else
 					{
@@ -2390,6 +2398,8 @@ namespace EpicGames.UHT.Utils
 					}
 				}
 			});
+
+			Log.Logger.LogInformation($"Total of {TotalWrittenFiles} written");
 		}
 		#endregion
 	}
