Changed the collection of parsing tables and configuration to be session specific and not a singleton.  This is required for plugin support.

#rnx

diff --git a/Engine/Source/Programs/Shared/EpicGames.UHT/EpicGames.UHT.csproj b/Engine/Source/Programs/Shared/EpicGames.UHT/EpicGames.UHT.csproj
--- a/Engine/Source/Programs/Shared/EpicGames.UHT/EpicGames.UHT.csproj
+++ b/Engine/Source/Programs/Shared/EpicGames.UHT/EpicGames.UHT.csproj
@@ -11,6 +11,7 @@
     <AllowUnsafeBlocks>true</AllowUnsafeBlocks>
     <GenerateDocumentationFile>True</GenerateDocumentationFile>
     <DebugType>pdbonly</DebugType>
+    <EnforceCodeStyleInBuild>true</EnforceCodeStyleInBuild>
   </PropertyGroup>
   <PropertyGroup Condition=" '$(Configuration)|$(Platform)' == 'Debug|AnyCPU' ">
     <DocumentationFile></DocumentationFile>
diff --git a/Engine/Source/Programs/Shared/EpicGames.UHT/Parsers/UhtClassParser.cs b/Engine/Source/Programs/Shared/EpicGames.UHT/Parsers/UhtClassParser.cs
--- a/Engine/Source/Programs/Shared/EpicGames.UHT/Parsers/UhtClassParser.cs
+++ b/Engine/Source/Programs/Shared/EpicGames.UHT/Parsers/UhtClassParser.cs
@@ -17,9 +17,6 @@ namespace EpicGames.UHT.Parsers
 	[UnrealHeaderTool]
 	public class UhtClassParser : UhtClassBaseParser
 	{
-		private static UhtKeywordTable KeywordTable = UhtKeywordTables.Instance.Get(UhtTableNames.Class);
-		private static UhtSpecifierTable SpecifierTable = UhtSpecifierTables.Instance.Get(UhtTableNames.Class);
-
 		/// <summary>
 		/// Engine class flags removed
 		/// </summary>
@@ -392,7 +389,7 @@ namespace EpicGames.UHT.Parsers
 				TopScope.AccessSpecifier = UhtAccessSpecifier.Public;
 			}
 
-			UhtParserHelpers.ParseCompileVersionDeclaration(TopScope.TokenReader, Class);
+			UhtParserHelpers.ParseCompileVersionDeclaration(TopScope.TokenReader, TopScope.Session.Config!, Class);
 
 			Class.GeneratedBodyLineNumber = TopScope.TokenReader.InputLine;
 
@@ -404,7 +401,7 @@ namespace EpicGames.UHT.Parsers
 		private static UhtParseResult ParseUClass(UhtParsingScope ParentScope, ref UhtToken Token)
 		{
 			UhtClassParser Class = new UhtClassParser(ParentScope.ScopeType, Token.InputLine);
-			using (var TopScope = new UhtParsingScope(ParentScope, Class, UhtClassParser.KeywordTable, UhtAccessSpecifier.Private))
+			using (var TopScope = new UhtParsingScope(ParentScope, Class, ParentScope.Session.GetKeywordTable(UhtTableNames.Class), UhtAccessSpecifier.Private))
 			{
 				const string ScopeName = "class";
 
@@ -412,7 +409,7 @@ namespace EpicGames.UHT.Parsers
 				{
 					// Parse the specifiers
 					UhtSpecifierContext SpecifierContext = new UhtSpecifierContext(TopScope, TopScope.TokenReader, Class.MetaData);
-					UhtSpecifierParser Specifiers = TopScope.HeaderParser.GetSpecifierParser(SpecifierContext, ScopeName, UhtClassParser.SpecifierTable);
+					UhtSpecifierParser Specifiers = TopScope.HeaderParser.GetSpecifierParser(SpecifierContext, ScopeName, ParentScope.Session.GetSpecifierTable(UhtTableNames.Class));
 					Specifiers.ParseSpecifiers();
 					Class.PrologLineNumber = TopScope.TokenReader.InputLine;
 					Class.ClassFlags |= EClassFlags.Native;
@@ -446,7 +443,7 @@ namespace EpicGames.UHT.Parsers
 					TopScope.TokenReader.Optional("final");
 
 					// Parse the inheritance
-					UhtParserHelpers.ParseInheritance(TopScope.TokenReader, out Class.SuperIdentifier, out Class.BaseIdentifiers);
+					UhtParserHelpers.ParseInheritance(TopScope.TokenReader, TopScope.Session.Config!, out Class.SuperIdentifier, out Class.BaseIdentifiers);
 
 					if (APIMacroToken)
 					{
diff --git a/Engine/Source/Programs/Shared/EpicGames.UHT/Parsers/UhtEnumParser.cs b/Engine/Source/Programs/Shared/EpicGames.UHT/Parsers/UhtEnumParser.cs
--- a/Engine/Source/Programs/Shared/EpicGames.UHT/Parsers/UhtEnumParser.cs
+++ b/Engine/Source/Programs/Shared/EpicGames.UHT/Parsers/UhtEnumParser.cs
@@ -14,9 +14,6 @@ namespace EpicGames.UHT.Parsers
 	[UnrealHeaderTool]
 	public class UhtEnumParser : UhtEnum
 	{
-		private static UhtKeywordTable KeywordTable = UhtKeywordTables.Instance.Get(UhtTableNames.Enum);
-		private static UhtSpecifierTable SpecifierTable = UhtSpecifierTables.Instance.Get(UhtTableNames.Enum);
-
 		/// <summary>
 		/// Construct a new enum parser
 		/// </summary>
@@ -37,7 +34,7 @@ namespace EpicGames.UHT.Parsers
 		private static UhtParseResult ParseUEnum(UhtParsingScope ParentScope, UhtToken KeywordToken)
 		{
 			UhtEnumParser Enum = new UhtEnumParser(ParentScope.ScopeType, KeywordToken.InputLine);
-			using (var TopScope = new UhtParsingScope(ParentScope, Enum, KeywordTable, UhtAccessSpecifier.Public))
+			using (var TopScope = new UhtParsingScope(ParentScope, Enum, ParentScope.Session.GetKeywordTable(UhtTableNames.Enum), UhtAccessSpecifier.Public))
 			{
 				const string ScopeName = "UENUM";
 
@@ -46,7 +43,7 @@ namespace EpicGames.UHT.Parsers
 
 					// Parse the specifiers
 					UhtSpecifierContext SpecifierContext = new UhtSpecifierContext(TopScope, TopScope.TokenReader, Enum.MetaData);
-					UhtSpecifierParser Specifiers = TopScope.HeaderParser.GetSpecifierParser(SpecifierContext, ScopeName, SpecifierTable);
+					UhtSpecifierParser Specifiers = TopScope.HeaderParser.GetSpecifierParser(SpecifierContext, ScopeName, ParentScope.Session.GetSpecifierTable(UhtTableNames.Enum));
 					Specifiers.ParseSpecifiers();
 
 					// Read the name and the CPP type
diff --git a/Engine/Source/Programs/Shared/EpicGames.UHT/Parsers/UhtFunctionParser.cs b/Engine/Source/Programs/Shared/EpicGames.UHT/Parsers/UhtFunctionParser.cs
--- a/Engine/Source/Programs/Shared/EpicGames.UHT/Parsers/UhtFunctionParser.cs
+++ b/Engine/Source/Programs/Shared/EpicGames.UHT/Parsers/UhtFunctionParser.cs
@@ -96,9 +96,6 @@ namespace EpicGames.UHT.Parsers
 	[UnrealHeaderTool]
 	public class UhtFunctionParser : UhtFunction
 	{
-		private static UhtKeywordTable KeywordTable = UhtKeywordTables.Instance.Get(UhtTableNames.Function);
-		private static UhtSpecifierTable SpecifierTable = UhtSpecifierTables.Instance.Get(UhtTableNames.Function);
-
 		/// <summary>
 		/// Construct a new function parser
 		/// </summary>
@@ -212,7 +209,7 @@ namespace EpicGames.UHT.Parsers
 		{
 			UhtFunctionParser Function = new UhtFunctionParser(ParentScope.ScopeType, Token.InputLine);
 
-			using (var TopScope = new UhtParsingScope(ParentScope, Function, KeywordTable, UhtAccessSpecifier.Public))
+			using (var TopScope = new UhtParsingScope(ParentScope, Function, ParentScope.Session.GetKeywordTable(UhtTableNames.Function), UhtAccessSpecifier.Public))
 			{
 				const string ScopeName = "delegate declaration";
 
@@ -221,7 +218,7 @@ namespace EpicGames.UHT.Parsers
 					TopScope.AddModuleRelativePathToMetaData();
 
 					UhtSpecifierContext SpecifierContext = new UhtSpecifierContext(TopScope, TopScope.TokenReader, Function.MetaData);
-					UhtSpecifierParser Specifiers = TopScope.HeaderParser.GetSpecifierParser(SpecifierContext, ScopeName, SpecifierTable);
+					UhtSpecifierParser Specifiers = TopScope.HeaderParser.GetSpecifierParser(SpecifierContext, ScopeName, ParentScope.Session.GetSpecifierTable(UhtTableNames.Function));
 
 					// If this is a UDELEGATE, parse the specifiers
 					StringView DelegateMacro = new StringView();
@@ -250,14 +247,14 @@ namespace EpicGames.UHT.Parsers
 					bool bIsSparse = DelegateMacro.Span.Contains("_SPARSE".AsSpan(), StringComparison.Ordinal);
 
 					// Determine the parameter count
-					int FoundParamIndex = UhtConfig.Instance.FindDelegateParameterCount(DelegateMacro);
+					int FoundParamIndex = TopScope.Session.Config!.FindDelegateParameterCount(DelegateMacro);
 
 					// Try reconstructing the string to make sure it matches our expectations
 					string ExpectedOriginalString = string.Format("DECLARE_DYNAMIC{0}{1}_DELEGATE{2}{3}{4}",
 						bIsMulticast ? "_MULTICAST" : "",
 						bIsSparse ? "_SPARSE" : "",
 						bHasReturnValue ? "_RetVal" : "",
-						UhtConfig.Instance.GetDelegateParameterCountString(FoundParamIndex),
+						TopScope.Session.Config!.GetDelegateParameterCountString(FoundParamIndex),
 						bDeclaredConst ? "_Const" : "");
 					if (DelegateMacro != ExpectedOriginalString)
 					{
@@ -382,7 +379,7 @@ namespace EpicGames.UHT.Parsers
 		{
 			UhtFunctionParser Function = new UhtFunctionParser(ParentScope.ScopeType, Token.InputLine);
 
-			using (var TopScope = new UhtParsingScope(ParentScope, Function, KeywordTable, UhtAccessSpecifier.Public))
+			using (var TopScope = new UhtParsingScope(ParentScope, Function, ParentScope.Session.GetKeywordTable(UhtTableNames.Function), UhtAccessSpecifier.Public))
 			{
 				UhtParsingScope OuterClassScope = TopScope.CurrentClassScope;
 				UhtClass OuterClass = (UhtClass)OuterClassScope.ScopeType;
@@ -393,7 +390,7 @@ namespace EpicGames.UHT.Parsers
 					TopScope.AddModuleRelativePathToMetaData();
 
 					UhtSpecifierContext SpecifierContext = new UhtSpecifierContext(TopScope, TopScope.TokenReader, Function.MetaData);
-					UhtSpecifierParser SpecifierParser = TopScope.HeaderParser.GetSpecifierParser(SpecifierContext, ScopeName, SpecifierTable);
+					UhtSpecifierParser SpecifierParser = TopScope.HeaderParser.GetSpecifierParser(SpecifierContext, ScopeName, ParentScope.Session.GetSpecifierTable(UhtTableNames.Function));
 					SpecifierParser.ParseSpecifiers();
 
 					if (!OuterClass.ClassFlags.HasAnyFlags(EClassFlags.Native))
diff --git a/Engine/Source/Programs/Shared/EpicGames.UHT/Parsers/UhtHeaderFileParser.cs b/Engine/Source/Programs/Shared/EpicGames.UHT/Parsers/UhtHeaderFileParser.cs
--- a/Engine/Source/Programs/Shared/EpicGames.UHT/Parsers/UhtHeaderFileParser.cs
+++ b/Engine/Source/Programs/Shared/EpicGames.UHT/Parsers/UhtHeaderFileParser.cs
@@ -170,8 +170,6 @@ namespace EpicGames.UHT.Parsers
 	/// </summary>
 	public class UhtHeaderFileParser : IUhtTokenPreprocessor
 	{
-		private static UhtKeywordTable KeywordTable = UhtKeywordTables.Instance.Get(UhtTableNames.Global);
-
 		private struct CompilerDirective
 		{
 			public UhtCompilerDirective Element;
@@ -237,7 +235,7 @@ namespace EpicGames.UHT.Parsers
 		public static UhtHeaderFileParser Parse(UhtHeaderFile HeaderFile)
 		{
 			UhtHeaderFileParser HeaderParser = new UhtHeaderFileParser(HeaderFile);
-			using (UhtParsingScope TopScope = new UhtParsingScope(HeaderParser, HeaderParser.HeaderFile, UhtHeaderFileParser.KeywordTable))
+			using (UhtParsingScope TopScope = new UhtParsingScope(HeaderParser, HeaderParser.HeaderFile, HeaderFile.Session.GetKeywordTable(UhtTableNames.Global)))
 			{
 				HeaderParser.ParseStatements();
 
@@ -476,7 +474,7 @@ namespace EpicGames.UHT.Parsers
 
 			if (ParseResult == UhtParseResult.Unhandled && bLogUnhandledKeywords)
 			{
-				UhtKeywordTables.Instance.LogUnhandledError(TopScope.TokenReader, Token);
+				TopScope.Session.LogUnhandledKeywordError(TopScope.TokenReader, Token);
 			}
 
 			if (ParseResult == UhtParseResult.Unhandled || ParseResult == UhtParseResult.Invalid)
diff --git a/Engine/Source/Programs/Shared/EpicGames.UHT/Parsers/UhtInterfaceClassParser.cs b/Engine/Source/Programs/Shared/EpicGames.UHT/Parsers/UhtInterfaceClassParser.cs
--- a/Engine/Source/Programs/Shared/EpicGames.UHT/Parsers/UhtInterfaceClassParser.cs
+++ b/Engine/Source/Programs/Shared/EpicGames.UHT/Parsers/UhtInterfaceClassParser.cs
@@ -15,9 +15,6 @@ namespace EpicGames.UHT.Parsers
 	[UnrealHeaderTool]
 	public class UhtInterfaceClassParser : UhtClassBaseParser
 	{
-		private static UhtKeywordTable KeywordTable = UhtKeywordTables.Instance.Get(UhtTableNames.Interface);
-		private static UhtSpecifierTable SpecifierTable = UhtSpecifierTables.Instance.Get(UhtTableNames.Interface);
-
 		/// <summary>
 		/// Construct a new interface class parser
 		/// </summary>
@@ -91,7 +88,7 @@ namespace EpicGames.UHT.Parsers
 		{
 			UhtClass Class = (UhtClass)TopScope.ScopeType;
 
-			UhtParserHelpers.ParseCompileVersionDeclaration(TopScope.TokenReader, Class);
+			UhtParserHelpers.ParseCompileVersionDeclaration(TopScope.TokenReader, TopScope.Session.Config!, Class);
 
 			Class.GeneratedBodyAccessSpecifier = TopScope.AccessSpecifier;
 			Class.GeneratedBodyLineNumber = TopScope.TokenReader.InputLine;
@@ -109,7 +106,7 @@ namespace EpicGames.UHT.Parsers
 			UhtInterfaceClassParser Class = new UhtInterfaceClassParser(ParentScope.ScopeType, Token.InputLine);
 			Class.ClassType = UhtClassType.Interface;
 
-			using (var TopScope = new UhtParsingScope(ParentScope, Class, KeywordTable, UhtAccessSpecifier.Private))
+			using (var TopScope = new UhtParsingScope(ParentScope, Class, ParentScope.Session.GetKeywordTable(UhtTableNames.Interface), UhtAccessSpecifier.Private))
 			{
 				const string ScopeName = "interface";
 
@@ -117,7 +114,7 @@ namespace EpicGames.UHT.Parsers
 				{
 					// Parse the specifiers
 					UhtSpecifierContext SpecifierContext = new UhtSpecifierContext(TopScope, TopScope.TokenReader, Class.MetaData);
-					UhtSpecifierParser Specifiers = TopScope.HeaderParser.GetSpecifierParser(SpecifierContext, ScopeName, SpecifierTable);
+					UhtSpecifierParser Specifiers = TopScope.HeaderParser.GetSpecifierParser(SpecifierContext, ScopeName, ParentScope.Session.GetSpecifierTable(UhtTableNames.Interface));
 					Specifiers.ParseSpecifiers();
 					Class.PrologLineNumber = TopScope.TokenReader.InputLine;
 					Class.ClassFlags |= EClassFlags.Native | EClassFlags.Interface | EClassFlags.Abstract;
diff --git a/Engine/Source/Programs/Shared/EpicGames.UHT/Parsers/UhtNativeInterfaceClassParser.cs b/Engine/Source/Programs/Shared/EpicGames.UHT/Parsers/UhtNativeInterfaceClassParser.cs
--- a/Engine/Source/Programs/Shared/EpicGames.UHT/Parsers/UhtNativeInterfaceClassParser.cs
+++ b/Engine/Source/Programs/Shared/EpicGames.UHT/Parsers/UhtNativeInterfaceClassParser.cs
@@ -15,9 +15,6 @@ namespace EpicGames.UHT.Parsers
 	[UnrealHeaderTool]
 	public class UhtParserNativeInterfaceClass : UhtClassBaseParser
 	{
-		private static UhtKeywordTable KeywordTable = UhtKeywordTables.Instance.Get(UhtTableNames.NativeInterface);
-		private static UhtSpecifierTable SpecifierTable = UhtSpecifierTables.Instance.Get(UhtTableNames.NativeInterface);
-
 		/// <summary>
 		/// Construct a new native interface
 		/// </summary>
@@ -91,7 +88,7 @@ namespace EpicGames.UHT.Parsers
 		{
 			UhtClass Class = (UhtClass)TopScope.ScopeType;
 
-			UhtParserHelpers.ParseCompileVersionDeclaration(TopScope.TokenReader, Class);
+			UhtParserHelpers.ParseCompileVersionDeclaration(TopScope.TokenReader, TopScope.Session.Config!, Class);
 
 			Class.GeneratedBodyAccessSpecifier = TopScope.AccessSpecifier;
 			Class.GeneratedBodyLineNumber = TopScope.TokenReader.InputLine;
@@ -173,7 +170,7 @@ namespace EpicGames.UHT.Parsers
 
 			Class.SuperIdentifier = SuperName;
 
-			using (var TopScope = new UhtParsingScope(ParentScope, Class, KeywordTable, UhtAccessSpecifier.Private))
+			using (var TopScope = new UhtParsingScope(ParentScope, Class, ParentScope.Session.GetKeywordTable(UhtTableNames.NativeInterface), UhtAccessSpecifier.Private))
 			{
 				const string ScopeName = "native interface";
 
diff --git a/Engine/Source/Programs/Shared/EpicGames.UHT/Parsers/UhtParserHelpers.cs b/Engine/Source/Programs/Shared/EpicGames.UHT/Parsers/UhtParserHelpers.cs
--- a/Engine/Source/Programs/Shared/EpicGames.UHT/Parsers/UhtParserHelpers.cs
+++ b/Engine/Source/Programs/Shared/EpicGames.UHT/Parsers/UhtParserHelpers.cs
@@ -19,16 +19,17 @@ namespace EpicGames.UHT.Parsers
 		/// Parse the inheritance 
 		/// </summary>
 		/// <param name="TokenReader">Token reader</param>
+		/// <param name="Config">Configuration</param>
 		/// <param name="SuperIdentifier">Output super identifier</param>
 		/// <param name="BaseIdentifiers">Output base identifiers</param>
-		public static void ParseInheritance(IUhtTokenReader TokenReader, out UhtToken SuperIdentifier, out List<UhtToken[]>? BaseIdentifiers)
+		public static void ParseInheritance(IUhtTokenReader TokenReader, IUhtConfig Config, out UhtToken SuperIdentifier, out List<UhtToken[]>? BaseIdentifiers)
 		{
 			UhtToken SuperIdentifierTemp = new UhtToken();
 			List<UhtToken[]>? BaseIdentifiersTemp = null;
 			TokenReader.OptionalInheritance(
 				(ref UhtToken Identifier) =>
 				{
-					UhtConfig.Instance.RedirectTypeIdentifier(ref Identifier);
+					Config.RedirectTypeIdentifier(ref Identifier);
 					SuperIdentifierTemp = Identifier;
 				},
 				(UhtTokenList Identifier) =>
@@ -47,15 +48,16 @@ namespace EpicGames.UHT.Parsers
 		/// Parse compiler version declaration
 		/// </summary>
 		/// <param name="TokenReader">Token reader</param>
+		/// <param name="Config">Configuration</param>
 		/// <param name="Struct">Struct being parsed</param>
-		public static void ParseCompileVersionDeclaration(IUhtTokenReader TokenReader, UhtStruct Struct)
+		public static void ParseCompileVersionDeclaration(IUhtTokenReader TokenReader, IUhtConfig Config, UhtStruct Struct)
 		{
 
 			// Fetch the default generation code version. If supplied, then package code version overrides the default.
 			EGeneratedCodeVersion Version = Struct.Package.Module.GeneratedCodeVersion;
 			if (Version == EGeneratedCodeVersion.None)
 			{
-				Version = UhtConfig.Instance.DefaultGeneratedCodeVersion;
+				Version = Config.DefaultGeneratedCodeVersion;
 			}
 
 			// Fetch the code version from header file
diff --git a/Engine/Source/Programs/Shared/EpicGames.UHT/Parsers/UhtParsingScope.cs b/Engine/Source/Programs/Shared/EpicGames.UHT/Parsers/UhtParsingScope.cs
--- a/Engine/Source/Programs/Shared/EpicGames.UHT/Parsers/UhtParsingScope.cs
+++ b/Engine/Source/Programs/Shared/EpicGames.UHT/Parsers/UhtParsingScope.cs
@@ -47,6 +47,11 @@ namespace EpicGames.UHT.Parsers
 		/// </summary>
 		public UhtAccessSpecifier AccessSpecifier = UhtAccessSpecifier.Public;
 
+		/// <summary>
+		/// Current session
+		/// </summary>
+		public UhtSession Session => this.ScopeType.Session;
+
 		/// <summary>
 		/// Return the current class scope being compiled
 		/// </summary>
diff --git a/Engine/Source/Programs/Shared/EpicGames.UHT/Parsers/UhtPropertyParser.cs b/Engine/Source/Programs/Shared/EpicGames.UHT/Parsers/UhtPropertyParser.cs
--- a/Engine/Source/Programs/Shared/EpicGames.UHT/Parsers/UhtPropertyParser.cs
+++ b/Engine/Source/Programs/Shared/EpicGames.UHT/Parsers/UhtPropertyParser.cs
@@ -547,7 +547,7 @@ namespace EpicGames.UHT.Parsers
 		/// <returns></returns>
 		public static UhtProperty? ResolveProperty(UhtPropertyResolvePhase ResolvePhase, UhtPropertySettings PropertySettings, ReadOnlyMemory<char> Data, ReadOnlyMemory<UhtToken> TypeTokens)
 		{
-			UhtPropertyTypeTable PropertyTypeTable = UhtPropertyTypeTable.Instance;
+			UhtSession Session = PropertySettings.Outer.Session;
 			IUhtTokenReader ReplayReader = UhtTokenReplayReader.GetThreadInstance(PropertySettings.Outer, Data, TypeTokens, UhtTokenType.EndOfType);
 
 			// Loop through the tokens until we find a known property type or the start of a template argument list
@@ -564,7 +564,7 @@ namespace EpicGames.UHT.Parsers
 				}
 				else
 				{
-					if (PropertyTypeTable.TryGet(TypeTokens.Span[Index].Value, out PropertyType))
+					if (Session.TryGetPropertyType(TypeTokens.Span[Index].Value, out PropertyType))
 					{
 						if (ResolvePhase == UhtPropertyResolvePhase.Resolving || PropertyType.Options.HasAnyFlags(UhtPropertyTypeOptions.Immediate))
 						{
@@ -577,7 +577,7 @@ namespace EpicGames.UHT.Parsers
 			// Try the default processor.  We only use the default processor when trying to resolve something post parsing phase.
 			if (ResolvePhase == UhtPropertyResolvePhase.Resolving)
 			{
-				return ResolvePropertyType(ResolvePhase, PropertySettings, ReplayReader, PropertyTypeTable.Default, new UhtToken(), false);
+				return ResolvePropertyType(ResolvePhase, PropertySettings, ReplayReader, Session.DefaultPropertyType, new UhtToken(), false);
 			}
 			return null;
 		}
@@ -632,7 +632,7 @@ namespace EpicGames.UHT.Parsers
 		/// <returns>Parsed property</returns>
 		public static UhtProperty? ParseTemplateParam(UhtPropertyResolvePhase ResolvePhase, UhtPropertySettings ParentPropertySettings, StringView ParamName, IUhtTokenReader TokenReader)
 		{
-			UhtPropertyTypeTable PropertyTypeTable = UhtPropertyTypeTable.Instance;
+			UhtSession Session = ParentPropertySettings.Outer.Session;
 
 			// Save the token reader state.  We need this to restore back to the start when invoking the resolve methods.
 			using (var SaveState = new UhtTokenSaveState(TokenReader))
@@ -649,7 +649,7 @@ namespace EpicGames.UHT.Parsers
 						break;
 					}
 
-					if (PropertyTypeTable.TryGet(Token.Value, out PropertyType))
+					if (Session.TryGetPropertyType(Token.Value, out PropertyType))
 					{
 						if (ResolvePhase == UhtPropertyResolvePhase.Resolving || PropertyType.Options.HasAnyFlags(UhtPropertyTypeOptions.Immediate))
 						{
@@ -663,7 +663,7 @@ namespace EpicGames.UHT.Parsers
 				if (ResolvePhase == UhtPropertyResolvePhase.Resolving)
 				{
 					SaveState.RestoreState();
-					return ResolvePropertyType(ResolvePhase, PropertySettings, TokenReader, PropertyTypeTable.Default, new UhtToken(), true);
+					return ResolvePropertyType(ResolvePhase, PropertySettings, TokenReader, Session.DefaultPropertyType, new UhtToken(), true);
 				}
 				return null;
 			}
@@ -775,7 +775,7 @@ namespace EpicGames.UHT.Parsers
 			bool bIsParamList = this.SpecifierContext.PropertySettings.PropertyCategory != UhtPropertyCategory.Member && this.TokenReader.TryOptional("UPARAM");
 
 			UhtSpecifierParser Specifiers = this.TopScope.HeaderParser.GetSpecifierParser(this.SpecifierContext, "Variable", 
-				bIsParamList ? UhtPropertyArgumentSpecifiers.SpecifierTable : UhtPropertyMemberSpecifiers.SpecifierTable);
+				bIsParamList ? this.TopScope.Session.GetSpecifierTable(UhtTableNames.PropertyArgument) : this.TopScope.Session.GetSpecifierTable(UhtTableNames.PropertyMember));
 			if (DeclarationStyle == UhtParsePropertyDeclarationStyle.UPROPERTY || bIsParamList)
 			{
 				Specifiers.ParseSpecifiers();
@@ -1112,7 +1112,7 @@ namespace EpicGames.UHT.Parsers
 			// Try gathering metadata for member fields
 			if (NewProperty.PropertyCategory == UhtPropertyCategory.Member)
 			{
-				UhtSpecifierParser Specifiers = this.TopScope.HeaderParser.GetSpecifierParser(this.SpecifierContext, NewProperty.SourceName, UhtPropertyMemberSpecifiers.SpecifierTable);
+				UhtSpecifierParser Specifiers = this.TopScope.HeaderParser.GetSpecifierParser(this.SpecifierContext, NewProperty.SourceName, this.TopScope.Session.GetSpecifierTable(UhtTableNames.PropertyMember));
 				Specifiers.ParseFieldMetaData(NewProperty.SourceName);
 				this.TopScope.AddFormattedCommentsAsTooltipMetaData(NewProperty);
 			}
@@ -1231,6 +1231,7 @@ namespace EpicGames.UHT.Parsers
 		[UhtPropertyType(Options = UhtPropertyTypeOptions.Default)]
 		private static UhtProperty? DefaultProperty(UhtPropertyResolvePhase ResolvePhase, UhtPropertySettings PropertySettings, IUhtTokenReader TokenReader, UhtToken MatchedToken)
 		{
+			UhtSession Session = PropertySettings.Outer.Session;
 			int TypeStartPos = TokenReader.PeekToken().InputStartPos;
 
 			if (TokenReader.TryOptional("const"))
@@ -1305,8 +1306,8 @@ namespace EpicGames.UHT.Parsers
 				TokenReader.Require('*');
 
 				// Optionally emit messages about native pointer members and swallow trailing 'const' after pointer properties
-				UhtObjectPropertyBase.ConditionalLogPointerUsage(PropertySettings, UhtConfig.Instance.EngineNativePointerMemberBehavior,
-					UhtConfig.Instance.NonEngineNativePointerMemberBehavior, "Native pointer", TokenReader, TypeStartPos, "TObjectPtr");
+				UhtObjectPropertyBase.ConditionalLogPointerUsage(PropertySettings, Session.Config!.EngineNativePointerMemberBehavior,
+					Session.Config!.NonEngineNativePointerMemberBehavior, "Native pointer", TokenReader, TypeStartPos, "TObjectPtr");
 
 				if (PropertySettings.PropertyCategory == UhtPropertyCategory.Member)
 				{
@@ -1319,10 +1320,10 @@ namespace EpicGames.UHT.Parsers
 				{
 					return new UhtInterfaceProperty(PropertySettings, Class);
 				}
-				else if (Class.IsChildOf(Class.Session.UClass))
+				else if (Class.IsChildOf(Session.UClass))
 				{
 					// UObject here specifies that there is no limiter
-					return new UhtClassProperty(PropertySettings, Class, Class.Session.UObject);
+					return new UhtClassProperty(PropertySettings, Class, Session.UObject);
 				}
 				else
 				{
diff --git a/Engine/Source/Programs/Shared/EpicGames.UHT/Parsers/UhtScriptStructParser.cs b/Engine/Source/Programs/Shared/EpicGames.UHT/Parsers/UhtScriptStructParser.cs
--- a/Engine/Source/Programs/Shared/EpicGames.UHT/Parsers/UhtScriptStructParser.cs
+++ b/Engine/Source/Programs/Shared/EpicGames.UHT/Parsers/UhtScriptStructParser.cs
@@ -16,9 +16,6 @@ namespace EpicGames.UHT.Parsers
 	[UnrealHeaderTool]
 	public class UhtScriptStructParser : UhtScriptStruct
 	{
-		private static UhtKeywordTable KeywordTable = UhtKeywordTables.Instance.Get(UhtTableNames.ScriptStruct);
-		private static UhtSpecifierTable SpecifierTable = UhtSpecifierTables.Instance.Get(UhtTableNames.ScriptStruct);
-
 		/// <summary>
 		/// Super identifier
 		/// </summary>
@@ -97,7 +94,7 @@ namespace EpicGames.UHT.Parsers
 
 			ScriptStruct.MacroDeclaredLineNumber = TopScope.TokenReader.InputLine;
 
-			UhtParserHelpers.ParseCompileVersionDeclaration(TopScope.TokenReader, ScriptStruct);
+			UhtParserHelpers.ParseCompileVersionDeclaration(TopScope.TokenReader, TopScope.Session.Config!, ScriptStruct);
 
 			TopScope.TokenReader.Optional(';');
 			return UhtParseResult.Handled;
@@ -114,7 +111,7 @@ namespace EpicGames.UHT.Parsers
 		private static UhtParseResult ParseUScriptStruct(UhtParsingScope ParentScope, UhtToken KeywordToken)
 		{
 			UhtScriptStructParser ScriptStruct = new UhtScriptStructParser(ParentScope.ScopeType, KeywordToken.InputLine);
-			using (var TopScope = new UhtParsingScope(ParentScope, ScriptStruct, KeywordTable, UhtAccessSpecifier.Public))
+			using (var TopScope = new UhtParsingScope(ParentScope, ScriptStruct, ParentScope.Session.GetKeywordTable(UhtTableNames.ScriptStruct), UhtAccessSpecifier.Public))
 			{
 				const string ScopeName = "struct";
 
@@ -122,7 +119,7 @@ namespace EpicGames.UHT.Parsers
 				{
 					// Parse the specifiers
 					UhtSpecifierContext SpecifierContext = new UhtSpecifierContext(TopScope, TopScope.TokenReader, ScriptStruct.MetaData);
-					UhtSpecifierParser Specifiers = TopScope.HeaderParser.GetSpecifierParser(SpecifierContext, ScopeName, SpecifierTable);
+					UhtSpecifierParser Specifiers = TopScope.HeaderParser.GetSpecifierParser(SpecifierContext, ScopeName, ParentScope.Session.GetSpecifierTable(UhtTableNames.ScriptStruct));
 					Specifiers.ParseSpecifiers();
 
 					// Consume the struct specifier
@@ -155,7 +152,7 @@ namespace EpicGames.UHT.Parsers
 					}
 
 					// Parse the inheritance
-					UhtParserHelpers.ParseInheritance(TopScope.TokenReader, out ScriptStruct.SuperIdentifier, out ScriptStruct.BaseIdentifiers);
+					UhtParserHelpers.ParseInheritance(TopScope.TokenReader, TopScope.Session.Config!, out ScriptStruct.SuperIdentifier, out ScriptStruct.BaseIdentifiers);
 
 					// Add the comments here for compatibility with old UHT
 					//COMPATIBILITY-TODO - Move this back to where the AddModuleRelativePathToMetaData is called.
diff --git a/Engine/Source/Programs/Shared/EpicGames.UHT/Specifiers/UhtDefaultSpecifiers.cs b/Engine/Source/Programs/Shared/EpicGames.UHT/Specifiers/UhtDefaultSpecifiers.cs
--- a/Engine/Source/Programs/Shared/EpicGames.UHT/Specifiers/UhtDefaultSpecifiers.cs
+++ b/Engine/Source/Programs/Shared/EpicGames.UHT/Specifiers/UhtDefaultSpecifiers.cs
@@ -228,7 +228,7 @@ namespace EpicGames.UHT.Parsers
 				}
 			}
 
-			if (!UhtConfig.Instance.IsValidUnits(Value))
+			if (!Type.Session.Config!.IsValidUnits(Value))
 			{
 				Type.LogError($"Unrecognized units '{Value}' specified for '{Type.FullName}'");
 			}
diff --git a/Engine/Source/Programs/Shared/EpicGames.UHT/Specifiers/UhtPropertyArgumentSpecifiers.cs b/Engine/Source/Programs/Shared/EpicGames.UHT/Specifiers/UhtPropertyArgumentSpecifiers.cs
--- a/Engine/Source/Programs/Shared/EpicGames.UHT/Specifiers/UhtPropertyArgumentSpecifiers.cs
+++ b/Engine/Source/Programs/Shared/EpicGames.UHT/Specifiers/UhtPropertyArgumentSpecifiers.cs
@@ -14,17 +14,6 @@ namespace EpicGames.UHT.Parsers
 	[UnrealHeaderTool]
 	public static class UhtPropertyArgumentSpecifiers
 	{
-
-		/// <summary>
-		/// Table for property argument specifiers
-		/// </summary>
-		public static readonly UhtSpecifierTable SpecifierTable = UhtSpecifierTables.Instance.Get(UhtTableNames.PropertyArgument);
-
-		/// <summary>
-		/// Table for property argument specifier validators
-		/// </summary>
-		public static readonly UhtSpecifierValidatorTable SpecifierValidatorTable = UhtSpecifierValidatorTables.Instance.Get(UhtTableNames.PropertyArgument);
-
 		#region Argument Property Specifiers
 		[UhtSpecifier(Extends = UhtTableNames.PropertyArgument, ValueType = UhtSpecifierValueType.Legacy)]
 		private static void ConstSpecifier(UhtSpecifierContext SpecifierContext)
diff --git a/Engine/Source/Programs/Shared/EpicGames.UHT/Specifiers/UhtPropertyMemberSpecifiers.cs b/Engine/Source/Programs/Shared/EpicGames.UHT/Specifiers/UhtPropertyMemberSpecifiers.cs
--- a/Engine/Source/Programs/Shared/EpicGames.UHT/Specifiers/UhtPropertyMemberSpecifiers.cs
+++ b/Engine/Source/Programs/Shared/EpicGames.UHT/Specifiers/UhtPropertyMemberSpecifiers.cs
@@ -15,17 +15,6 @@ namespace EpicGames.UHT.Parsers
 	[UnrealHeaderTool]
 	public static class UhtPropertyMemberSpecifiers
 	{
-
-		/// <summary>
-		/// Property member specifier table
-		/// </summary>
-		public static UhtSpecifierTable SpecifierTable = UhtSpecifierTables.Instance.Get(UhtTableNames.PropertyMember);
-
-		/// <summary>
-		/// Property member specifier validator table
-		/// </summary>
-		public static UhtSpecifierValidatorTable SpecifierValidatorTable = UhtSpecifierValidatorTables.Instance.Get(UhtTableNames.PropertyMember);
-
 		[UhtSpecifier(Extends = UhtTableNames.PropertyMember, ValueType = UhtSpecifierValueType.Legacy)]
 		private static void EditAnywhereSpecifier(UhtSpecifierContext SpecifierContext)
 		{
diff --git a/Engine/Source/Programs/Shared/EpicGames.UHT/Tables/UhtEngineClassTable.cs b/Engine/Source/Programs/Shared/EpicGames.UHT/Tables/UhtEngineClassTable.cs
--- a/Engine/Source/Programs/Shared/EpicGames.UHT/Tables/UhtEngineClassTable.cs
+++ b/Engine/Source/Programs/Shared/EpicGames.UHT/Tables/UhtEngineClassTable.cs
@@ -47,11 +47,6 @@ namespace EpicGames.UHT.Tables
 	/// </summary>
 	public class UhtEngineClassTable
 	{
-		/// <summary>
-		/// Global instance of the engine class table
-		/// </summary>
-		public static UhtEngineClassTable Instance = new UhtEngineClassTable();
-
 		/// <summary>
 		/// Internal mapping from engine class name to information
 		/// </summary>
diff --git a/Engine/Source/Programs/Shared/EpicGames.UHT/Tables/UhtExporterTable.cs b/Engine/Source/Programs/Shared/EpicGames.UHT/Tables/UhtExporterTable.cs
--- a/Engine/Source/Programs/Shared/EpicGames.UHT/Tables/UhtExporterTable.cs
+++ b/Engine/Source/Programs/Shared/EpicGames.UHT/Tables/UhtExporterTable.cs
@@ -161,11 +161,6 @@ namespace EpicGames.UHT.Tables
 	public class UhtExporterTable : IEnumerable<UhtExporter>
 	{
 
-		/// <summary>
-		/// Global instance of the table
-		/// </summary>
-		public static UhtExporterTable Instance = new UhtExporterTable();
-
 		private Dictionary<string, UhtExporter> ExporterValues = new Dictionary<string, UhtExporter>(StringComparer.OrdinalIgnoreCase);
 
 		/// <summary>
diff --git a/Engine/Source/Programs/Shared/EpicGames.UHT/Tables/UhtKeywordTable.cs b/Engine/Source/Programs/Shared/EpicGames.UHT/Tables/UhtKeywordTable.cs
--- a/Engine/Source/Programs/Shared/EpicGames.UHT/Tables/UhtKeywordTable.cs
+++ b/Engine/Source/Programs/Shared/EpicGames.UHT/Tables/UhtKeywordTable.cs
@@ -185,11 +185,6 @@ namespace EpicGames.UHT.Tables
 	public class UhtKeywordTables : UhtLookupTables<UhtKeywordTable>
 	{
 
-		/// <summary>
-		/// Global instance of the tables
-		/// </summary>
-		public static UhtKeywordTables Instance = new UhtKeywordTables();
-
 		/// <summary>
 		/// Construct the keyword tables
 		/// </summary>
diff --git a/Engine/Source/Programs/Shared/EpicGames.UHT/Tables/UhtLocTextDefaultValueTable.cs b/Engine/Source/Programs/Shared/EpicGames.UHT/Tables/UhtLocTextDefaultValueTable.cs
--- a/Engine/Source/Programs/Shared/EpicGames.UHT/Tables/UhtLocTextDefaultValueTable.cs
+++ b/Engine/Source/Programs/Shared/EpicGames.UHT/Tables/UhtLocTextDefaultValueTable.cs
@@ -52,11 +52,6 @@ namespace EpicGames.UHT.Tables
 	public class UhtLocTextDefaultValueTable
 	{
 
-		/// <summary>
-		/// Global instance of the table
-		/// </summary>
-		public static UhtLocTextDefaultValueTable Instance = new UhtLocTextDefaultValueTable();
-
 		private Dictionary<StringView, UhtLocTextDefaultValue> LocTextDefaultValues = new Dictionary<StringView, UhtLocTextDefaultValue>();
 
 		/// <summary>
diff --git a/Engine/Source/Programs/Shared/EpicGames.UHT/Tables/UhtPropertyTypeTable.cs b/Engine/Source/Programs/Shared/EpicGames.UHT/Tables/UhtPropertyTypeTable.cs
--- a/Engine/Source/Programs/Shared/EpicGames.UHT/Tables/UhtPropertyTypeTable.cs
+++ b/Engine/Source/Programs/Shared/EpicGames.UHT/Tables/UhtPropertyTypeTable.cs
@@ -154,12 +154,6 @@ namespace EpicGames.UHT.Tables
 	/// </summary>
 	public class UhtPropertyTypeTable
 	{
-
-		/// <summary>
-		/// Global instance of the property type table
-		/// </summary>
-		public static UhtPropertyTypeTable Instance = new UhtPropertyTypeTable();
-
 		private Dictionary<StringView, UhtPropertyType> CaseSensitive = new Dictionary<StringView, UhtPropertyType>();
 		private Dictionary<StringView, UhtPropertyType> CaseInsensitive = new Dictionary<StringView, UhtPropertyType>(StringViewComparer.OrdinalIgnoreCase);
 		private UhtPropertyType? DefaultInternal = null;
diff --git a/Engine/Source/Programs/Shared/EpicGames.UHT/Tables/UhtSpecifierTable.cs b/Engine/Source/Programs/Shared/EpicGames.UHT/Tables/UhtSpecifierTable.cs
--- a/Engine/Source/Programs/Shared/EpicGames.UHT/Tables/UhtSpecifierTable.cs
+++ b/Engine/Source/Programs/Shared/EpicGames.UHT/Tables/UhtSpecifierTable.cs
@@ -537,11 +537,6 @@ namespace EpicGames.UHT.Tables
 	public class UhtSpecifierTables : UhtLookupTables<UhtSpecifierTable>
 	{
 
-		/// <summary>
-		/// Returns the global instance for all specifiers
-		/// </summary>
-		public static UhtSpecifierTables Instance = new UhtSpecifierTables();
-
 		/// <summary>
 		/// Construct the specifier table
 		/// </summary>
diff --git a/Engine/Source/Programs/Shared/EpicGames.UHT/Tables/UhtSpecifierValidatorTable.cs b/Engine/Source/Programs/Shared/EpicGames.UHT/Tables/UhtSpecifierValidatorTable.cs
--- a/Engine/Source/Programs/Shared/EpicGames.UHT/Tables/UhtSpecifierValidatorTable.cs
+++ b/Engine/Source/Programs/Shared/EpicGames.UHT/Tables/UhtSpecifierValidatorTable.cs
@@ -94,11 +94,6 @@ namespace EpicGames.UHT.Tables
 	public class UhtSpecifierValidatorTables : UhtLookupTables<UhtSpecifierValidatorTable>
 	{
 
-		/// <summary>
-		/// Global instance of the validator tables
-		/// </summary>
-		public static UhtSpecifierValidatorTables Instance = new UhtSpecifierValidatorTables();
-
 		/// <summary>
 		/// Construct the validator tables
 		/// </summary>
diff --git a/Engine/Source/Programs/Shared/EpicGames.UHT/Tables/UhtStructDefaultValueTable.cs b/Engine/Source/Programs/Shared/EpicGames.UHT/Tables/UhtStructDefaultValueTable.cs
--- a/Engine/Source/Programs/Shared/EpicGames.UHT/Tables/UhtStructDefaultValueTable.cs
+++ b/Engine/Source/Programs/Shared/EpicGames.UHT/Tables/UhtStructDefaultValueTable.cs
@@ -116,12 +116,6 @@ namespace EpicGames.UHT.Tables
 	/// </summary>
 	public class UhtStructDefaultValueTable
 	{
-
-		/// <summary>
-		/// Global instance of the structure default value sanitizers
-		/// </summary>
-		public static UhtStructDefaultValueTable Instance = new UhtStructDefaultValueTable();
-
 		private Dictionary<StringView, UhtStructDefaultValue> StructDefaultValues = new Dictionary<StringView, UhtStructDefaultValue>();
 		private UhtStructDefaultValue? DefaultInternal = null;
 
diff --git a/Engine/Source/Programs/Shared/EpicGames.UHT/Types/Properties/UhtFieldPathProperty.cs b/Engine/Source/Programs/Shared/EpicGames.UHT/Types/Properties/UhtFieldPathProperty.cs
--- a/Engine/Source/Programs/Shared/EpicGames.UHT/Types/Properties/UhtFieldPathProperty.cs
+++ b/Engine/Source/Programs/Shared/EpicGames.UHT/Types/Properties/UhtFieldPathProperty.cs
@@ -123,7 +123,7 @@ namespace EpicGames.UHT.Types
 					.Require('>');
 
 				StringView FieldClassName = new StringView(Identifier.Value, 1);
-				if (!UhtEngineClassTable.Instance.IsValidPropertyTypeName(FieldClassName))
+				if (!PropertySettings.Outer.Session.IsValidPropertyTypeName(FieldClassName))
 				{
 					throw new UhtException($"Undefined property type: {Identifier.Value}");
 				}
diff --git a/Engine/Source/Programs/Shared/EpicGames.UHT/Types/Properties/UhtObjectPropertyBase.cs b/Engine/Source/Programs/Shared/EpicGames.UHT/Types/Properties/UhtObjectPropertyBase.cs
--- a/Engine/Source/Programs/Shared/EpicGames.UHT/Types/Properties/UhtObjectPropertyBase.cs
+++ b/Engine/Source/Programs/Shared/EpicGames.UHT/Types/Properties/UhtObjectPropertyBase.cs
@@ -152,6 +152,7 @@ namespace EpicGames.UHT.Types
 		/// <returns>Referenced class</returns>
 		public static UhtClass? ParseTemplateObject(UhtPropertySettings PropertySettings, IUhtTokenReader TokenReader, UhtToken MatchedToken, bool bReturnUInterface)
 		{
+			UhtSession Session = PropertySettings.Outer.Session;
 			if (TokenReader.TryOptional("const"))
 			{
 				PropertySettings.MetaData.Add(UhtNames.NativeConst, "");
@@ -174,7 +175,7 @@ namespace EpicGames.UHT.Types
 			{
 				PropertySettings.MetaData.Add(UhtNames.NativeConstTemplateArg, "");
 			}
-			UhtConfig.Instance.RedirectTypeIdentifier(ref Identifier);
+			Session.Config!.RedirectTypeIdentifier(ref Identifier);
 			UhtClass? Return = PropertySettings.Outer.FindType(UhtFindOptions.SourceName | UhtFindOptions.Class, ref Identifier, TokenReader) as UhtClass;
 			if (Return != null && Return.AlternateObject != null && bReturnUInterface)
 			{
@@ -192,6 +193,7 @@ namespace EpicGames.UHT.Types
 		/// <returns>Referenced class</returns>
 		public static UhtClass? ParseTemplateClass(UhtPropertySettings PropertySettings, IUhtTokenReader TokenReader, UhtToken MatchedToken)
 		{
+			UhtSession Session = PropertySettings.Outer.Session;
 			UhtToken Identifier = new UhtToken();
 
 			if (TokenReader.TryOptional("const"))
@@ -212,7 +214,7 @@ namespace EpicGames.UHT.Types
 				.RequireIdentifier((ref UhtToken Token) => { Identifier = Token; })
 				.Require('>');
 
-			UhtConfig.Instance.RedirectTypeIdentifier(ref Identifier);
+			Session.Config!.RedirectTypeIdentifier(ref Identifier);
 			UhtClass? Return = PropertySettings.Outer.FindType(UhtFindOptions.SourceName | UhtFindOptions.Class, ref Identifier, TokenReader) as UhtClass;
 			if (Return != null && Return.AlternateObject != null)
 			{
diff --git a/Engine/Source/Programs/Shared/EpicGames.UHT/Types/Properties/UhtObjectPtrProperty.cs b/Engine/Source/Programs/Shared/EpicGames.UHT/Types/Properties/UhtObjectPtrProperty.cs
--- a/Engine/Source/Programs/Shared/EpicGames.UHT/Types/Properties/UhtObjectPtrProperty.cs
+++ b/Engine/Source/Programs/Shared/EpicGames.UHT/Types/Properties/UhtObjectPtrProperty.cs
@@ -96,6 +96,7 @@ namespace EpicGames.UHT.Types
 		[UhtPropertyType(Keyword = "TObjectPtr")]
 		private static UhtProperty? ObjectPtrProperty(UhtPropertyResolvePhase ResolvePhase, UhtPropertySettings PropertySettings, IUhtTokenReader TokenReader, UhtToken MatchedToken)
 		{
+			UhtSession Session = PropertySettings.Outer.Session;
 			int TypeStartPos = TokenReader.PeekToken().InputStartPos;
 
 			UhtClass? PropertyClass = ParseTemplateObject(PropertySettings, TokenReader, MatchedToken, true);
@@ -104,7 +105,8 @@ namespace EpicGames.UHT.Types
 				return null;
 			}
 
-			ConditionalLogPointerUsage(PropertySettings, UhtConfig.Instance.EngineObjectPtrMemberBehavior, UhtConfig.Instance.NonEngineObjectPtrMemberBehavior, "ObjectPtr", TokenReader, TypeStartPos, null);
+			ConditionalLogPointerUsage(PropertySettings, Session.Config!.EngineObjectPtrMemberBehavior, 
+				Session.Config!.NonEngineObjectPtrMemberBehavior, "ObjectPtr", TokenReader, TypeStartPos, null);
 
 			if (PropertyClass.IsChildOf(PropertyClass.Session.UClass))
 			{
diff --git a/Engine/Source/Programs/Shared/EpicGames.UHT/Types/Properties/UhtStructProperty.cs b/Engine/Source/Programs/Shared/EpicGames.UHT/Types/Properties/UhtStructProperty.cs
--- a/Engine/Source/Programs/Shared/EpicGames.UHT/Types/Properties/UhtStructProperty.cs
+++ b/Engine/Source/Programs/Shared/EpicGames.UHT/Types/Properties/UhtStructProperty.cs
@@ -158,9 +158,9 @@ namespace EpicGames.UHT.Types
 		public override bool SanitizeDefaultValue(IUhtTokenReader DefaultValueReader, StringBuilder InnerDefaultValue)
 		{
 			UhtStructDefaultValue StructDefaultValue;
-			if (!UhtStructDefaultValueTable.Instance.TryGet(this.ScriptStruct.SourceName, out StructDefaultValue))
+			if (!this.Session.TryGetStructDefaultValue(this.ScriptStruct.SourceName, out StructDefaultValue))
 			{
-				StructDefaultValue = UhtStructDefaultValueTable.Instance.Default;
+				StructDefaultValue = this.Session.DefaultStructDefaultValue;
 			}
 			return StructDefaultValue.Delegate(this, DefaultValueReader, InnerDefaultValue);
 		}
diff --git a/Engine/Source/Programs/Shared/EpicGames.UHT/Types/Properties/UhtTextProperty.cs b/Engine/Source/Programs/Shared/EpicGames.UHT/Types/Properties/UhtTextProperty.cs
--- a/Engine/Source/Programs/Shared/EpicGames.UHT/Types/Properties/UhtTextProperty.cs
+++ b/Engine/Source/Programs/Shared/EpicGames.UHT/Types/Properties/UhtTextProperty.cs
@@ -395,7 +395,7 @@ namespace EpicGames.UHT.Types
 			{
 				case UhtTokenType.Identifier:
 					UhtLocTextDefaultValue LocTextDefaultValue;
-					if (UhtLocTextDefaultValueTable.Instance.TryGet(Token.Value, out LocTextDefaultValue))
+					if (Property.Session.TryGetLocTextDefaultValue(Token.Value, out LocTextDefaultValue))
 					{
 						if (LocTextDefaultValue.Delegate(Property, DefaultValueReader, ref Token, InnerDefaultValue))
 						{
diff --git a/Engine/Source/Programs/Shared/EpicGames.UHT/Types/UhtClass.cs b/Engine/Source/Programs/Shared/EpicGames.UHT/Types/UhtClass.cs
--- a/Engine/Source/Programs/Shared/EpicGames.UHT/Types/UhtClass.cs
+++ b/Engine/Source/Programs/Shared/EpicGames.UHT/Types/UhtClass.cs
@@ -274,9 +274,6 @@ namespace EpicGames.UHT.Types
 	[UhtEngineClass(Name = "Class")]
 	public class UhtClass : UhtStruct
 	{
-		private static UhtSpecifierValidatorTable ClassSpecifierValidatorTable = UhtSpecifierValidatorTables.Instance.Get(UhtTableNames.Class);
-		private static UhtSpecifierValidatorTable InterfaceSpecifierValidatorTable = UhtSpecifierValidatorTables.Instance.Get(UhtTableNames.Interface);
-		private static UhtSpecifierValidatorTable NativeInterfaceSpecifierValidatorTable = UhtSpecifierValidatorTables.Instance.Get(UhtTableNames.NativeInterface);
 		private List<UhtDeclaration>? DeclarationsInternal = null;
 
 		/// <summary>
@@ -392,11 +389,11 @@ namespace EpicGames.UHT.Types
 				switch (this.ClassType)
 				{
 					case UhtClassType.Class:
-						return UhtClass.ClassSpecifierValidatorTable;
+						return this.Session.GetSpecifierValidatorTable(UhtTableNames.Class);
 					case UhtClassType.Interface:
-						return UhtClass.InterfaceSpecifierValidatorTable;
+						return this.Session.GetSpecifierValidatorTable(UhtTableNames.Interface);
 					case UhtClassType.NativeInterface:
-						return UhtClass.NativeInterfaceSpecifierValidatorTable;
+						return this.Session.GetSpecifierValidatorTable(UhtTableNames.NativeInterface);
 					default:
 						throw new UhtIceException("Unknown class type");
 				}
diff --git a/Engine/Source/Programs/Shared/EpicGames.UHT/Types/UhtEnum.cs b/Engine/Source/Programs/Shared/EpicGames.UHT/Types/UhtEnum.cs
--- a/Engine/Source/Programs/Shared/EpicGames.UHT/Types/UhtEnum.cs
+++ b/Engine/Source/Programs/Shared/EpicGames.UHT/Types/UhtEnum.cs
@@ -105,8 +105,6 @@ namespace EpicGames.UHT.Types
 	[UhtEngineClass(Name = "Enum")]
 	public class UhtEnum : UhtField, IUhtMetaDataKeyConversion
 	{
-		private static UhtSpecifierValidatorTable EnumSpecifierValidatorTable = UhtSpecifierValidatorTables.Instance.Get(UhtTableNames.Enum);
-
 		/// <summary>
 		/// Engine enumeration flags
 		/// </summary>
@@ -149,7 +147,7 @@ namespace EpicGames.UHT.Types
 
 		///<inheritdoc/>
 		[JsonIgnore]
-		protected override UhtSpecifierValidatorTable? SpecifierValidatorTable { get => UhtEnum.EnumSpecifierValidatorTable; }
+		protected override UhtSpecifierValidatorTable? SpecifierValidatorTable { get => this.Session.GetSpecifierValidatorTable(UhtTableNames.Enum); }
 
 		/// <summary>
 		/// Construct a new enumeration
diff --git a/Engine/Source/Programs/Shared/EpicGames.UHT/Types/UhtFunction.cs b/Engine/Source/Programs/Shared/EpicGames.UHT/Types/UhtFunction.cs
--- a/Engine/Source/Programs/Shared/EpicGames.UHT/Types/UhtFunction.cs
+++ b/Engine/Source/Programs/Shared/EpicGames.UHT/Types/UhtFunction.cs
@@ -196,7 +196,6 @@ namespace EpicGames.UHT.Types
 	/// </summary>
 	public class UhtFunction : UhtStruct
 	{
-		private static UhtSpecifierValidatorTable FunctionSpecifierValidatorTable = UhtSpecifierValidatorTables.Instance.Get(UhtTableNames.Function);
 		private string? StrippedFunctionNameInternal = null;
 
 		/// <summary>
@@ -287,7 +286,7 @@ namespace EpicGames.UHT.Types
 
 		///<inheritdoc/>
 		[JsonIgnore]
-		protected override UhtSpecifierValidatorTable? SpecifierValidatorTable { get => UhtFunction.FunctionSpecifierValidatorTable; }
+		protected override UhtSpecifierValidatorTable? SpecifierValidatorTable { get => this.Session.GetSpecifierValidatorTable(UhtTableNames.Function); }
 
 		/// <summary>
 		/// Identifier for an RPC call to a platform service
diff --git a/Engine/Source/Programs/Shared/EpicGames.UHT/Types/UhtMetaData.cs b/Engine/Source/Programs/Shared/EpicGames.UHT/Types/UhtMetaData.cs
--- a/Engine/Source/Programs/Shared/EpicGames.UHT/Types/UhtMetaData.cs
+++ b/Engine/Source/Programs/Shared/EpicGames.UHT/Types/UhtMetaData.cs
@@ -129,7 +129,7 @@ namespace EpicGames.UHT.Types
 		/// <summary>
 		/// Empty collection of meta data
 		/// </summary>
-		public static UhtMetaData Empty = new UhtMetaData(null);
+		public static UhtMetaData Empty = new UhtMetaData(null, null);
 
 		/// <summary>
 		/// The meta data of the outer object for the type that owns this meta data
@@ -140,6 +140,11 @@ namespace EpicGames.UHT.Types
 		/// Message site associated with the meta data.  That in combination with the line number is used to generate errors
 		/// </summary>
 		public IUhtMessageSite? MessageSite = null;
+
+		/// <summary>
+		/// Configuration object used to redirect meta data
+		/// </summary>
+		public IUhtConfig? Config = null;
 		
 		/// <summary>
 		/// Source code line number where the meta data is declared
@@ -166,9 +171,11 @@ namespace EpicGames.UHT.Types
 		/// Construct new meta data
 		/// </summary>
 		/// <param name="MessageSite">Message site for generating errors</param>
-		public UhtMetaData(IUhtMessageSite? MessageSite)
+		/// <param name="Config">Configuration for redirects</param>
+		public UhtMetaData(IUhtMessageSite? MessageSite, IUhtConfig? Config)
 		{
 			this.MessageSite = MessageSite;
+			this.Config = Config;
 		}
 
 		/// <summary>
@@ -414,11 +421,14 @@ namespace EpicGames.UHT.Types
 		public void Add(string Name, int NameIndex, string Value)
 		{
 			string RemappedName;
-			if (UhtConfig.Instance.RedirectMetaDataKey(Name, out RemappedName))
+			if (this.Config != null)
 			{
-				if (this.MessageSite != null)
+				if (this.Config.RedirectMetaDataKey(Name, out RemappedName))
 				{
-					this.MessageSite.LogWarning(this.LineNumber, $"Remapping old metadata key '{Name}' to new key '{RemappedName}', please update the declaration.");
+					if (this.MessageSite != null)
+					{
+						this.MessageSite.LogWarning(this.LineNumber, $"Remapping old metadata key '{Name}' to new key '{RemappedName}', please update the declaration.");
+					}
 				}
 			}
 			GetDictionary()[new UhtMetaDataKey(Name, NameIndex)] = Value;
diff --git a/Engine/Source/Programs/Shared/EpicGames.UHT/Types/UhtProperty.cs b/Engine/Source/Programs/Shared/EpicGames.UHT/Types/UhtProperty.cs
--- a/Engine/Source/Programs/Shared/EpicGames.UHT/Types/UhtProperty.cs
+++ b/Engine/Source/Programs/Shared/EpicGames.UHT/Types/UhtProperty.cs
@@ -709,7 +709,7 @@ namespace EpicGames.UHT.Types
 		{
 			this.SourceName = SourceName;
 			this.EngineName = SourceName;
-			this.MetaData = new UhtMetaData(MessageSite);
+			this.MetaData = new UhtMetaData(MessageSite, ParentPropertySettings.Outer.Session.Config);
 			this.Outer = ParentPropertySettings.Outer;
 			this.LineNumber = ParentPropertySettings.LineNumber;
 			this.PropertyCategory = ParentPropertySettings.PropertyCategory;
@@ -738,7 +738,7 @@ namespace EpicGames.UHT.Types
 		{
 			this.SourceName = string.Empty;
 			this.EngineName = string.Empty;
-			this.MetaData = new UhtMetaData(Outer);
+			this.MetaData = new UhtMetaData(Outer, Outer.Session.Config);
 			this.Outer = Outer;
 			this.LineNumber = LineNumber;
 			this.PropertyCategory = PropertyCategory;
@@ -941,11 +941,11 @@ namespace EpicGames.UHT.Types
 				switch (this.PropertyCategory)
 				{
 					case UhtPropertyCategory.Member:
-						return UhtPropertyMemberSpecifiers.SpecifierValidatorTable;
+						return this.Session.GetSpecifierValidatorTable(UhtTableNames.PropertyMember);
 					case UhtPropertyCategory.RegularParameter:
 					case UhtPropertyCategory.ReplicatedParameter:
 					case UhtPropertyCategory.Return:
-						return UhtPropertyArgumentSpecifiers.SpecifierValidatorTable;
+						return this.Session.GetSpecifierValidatorTable(UhtTableNames.PropertyArgument);
 					default:
 						throw new UhtIceException("Unknown property category type");
 				}
diff --git a/Engine/Source/Programs/Shared/EpicGames.UHT/Types/UhtScriptStruct.cs b/Engine/Source/Programs/Shared/EpicGames.UHT/Types/UhtScriptStruct.cs
--- a/Engine/Source/Programs/Shared/EpicGames.UHT/Types/UhtScriptStruct.cs
+++ b/Engine/Source/Programs/Shared/EpicGames.UHT/Types/UhtScriptStruct.cs
@@ -538,8 +538,6 @@ namespace EpicGames.UHT.Types
 	[UhtEngineClass(Name = "ScriptStruct")]
 	public class UhtScriptStruct : UhtStruct
 	{
-		private static UhtSpecifierValidatorTable ScriptStructSpecifierValidatorTable = UhtSpecifierValidatorTables.Instance.Get(UhtTableNames.ScriptStruct);
-
 		/// <summary>
 		/// Script struct engine flags
 		/// </summary>
@@ -599,13 +597,13 @@ namespace EpicGames.UHT.Types
 		{
 			get
 			{
-				return UhtConfig.Instance.IsStructWithTPrefix(this.EngineName) ? "T" : "F";
+				return this.Session.Config!.IsStructWithTPrefix(this.EngineName) ? "T" : "F";
 			}
 		}
 
 		///<inheritdoc/>
 		[JsonIgnore]
-		protected override UhtSpecifierValidatorTable? SpecifierValidatorTable { get => UhtScriptStruct.ScriptStructSpecifierValidatorTable; }
+		protected override UhtSpecifierValidatorTable? SpecifierValidatorTable { get => this.Session.GetSpecifierValidatorTable(UhtTableNames.ScriptStruct); }
 
 		/// <summary>
 		/// Super struct
diff --git a/Engine/Source/Programs/Shared/EpicGames.UHT/Types/UhtType.cs b/Engine/Source/Programs/Shared/EpicGames.UHT/Types/UhtType.cs
--- a/Engine/Source/Programs/Shared/EpicGames.UHT/Types/UhtType.cs
+++ b/Engine/Source/Programs/Shared/EpicGames.UHT/Types/UhtType.cs
@@ -823,7 +823,7 @@ namespace EpicGames.UHT.Types
 			this.SourceName = string.Empty;
 			this.LineNumber = 1;
 			this.TypeIndex = Session.GetNextTypeIndex();
-			this.MetaData = new UhtMetaData(this);
+			this.MetaData = new UhtMetaData(this, this.Session.Config);
 		}
 
 		/// <summary>
@@ -839,8 +839,9 @@ namespace EpicGames.UHT.Types
 			this.SourceName = string.Empty;
 			this.LineNumber = LineNumber;
 			this.TypeIndex = this.Session.GetNextTypeIndex();
-			this.MetaData = MetaData ?? new UhtMetaData(this);
+			this.MetaData = MetaData ?? new UhtMetaData(this, this.Session.Config);
 			this.MetaData.MessageSite = this; // Make sure the site is correct
+			this.MetaData.Config = this.Session.Config;
 		}
 
 		/// <summary>
diff --git a/Engine/Source/Programs/Shared/EpicGames.UHT/Utils/UhtAttributeScanner.cs b/Engine/Source/Programs/Shared/EpicGames.UHT/Utils/UhtAttributeScanner.cs
deleted file mode 100644
diff --git a/Engine/Source/Programs/Shared/EpicGames.UHT/Utils/UhtConfig.cs b/Engine/Source/Programs/Shared/EpicGames.UHT/Utils/UhtConfig.cs
--- a/Engine/Source/Programs/Shared/EpicGames.UHT/Utils/UhtConfig.cs
+++ b/Engine/Source/Programs/Shared/EpicGames.UHT/Utils/UhtConfig.cs
@@ -108,32 +108,4 @@ namespace EpicGames.UHT.Utils
 		/// <returns>True if the exporter is enabled, false if not</returns>
 		public bool IsExporterEnabled(string Name);
 	}
-
-	/// <summary>
-	/// Helper class to house the IUhtConfig instance
-	/// </summary>
-	public static class UhtConfig
-	{
-		private static IUhtConfig? InstanceInternal = null;
-
-		/// <summary>
-		/// Get/Set the instance of IUhtConfig
-		/// </summary>
-		public static IUhtConfig Instance 
-		{ 
-			get
-			{
-				if (UhtConfig.InstanceInternal == null)
-				{
-					throw new UhtIceException("UhtConfig instance has not been set");
-				}
-				return UhtConfig.InstanceInternal;
-			}
-
-			set
-			{
-				UhtConfig.InstanceInternal = value;
-			}
-		}
-	}
 }
diff --git a/Engine/Source/Programs/Shared/EpicGames.UHT/Utils/UhtSession.cs b/Engine/Source/Programs/Shared/EpicGames.UHT/Utils/UhtSession.cs
--- a/Engine/Source/Programs/Shared/EpicGames.UHT/Utils/UhtSession.cs
+++ b/Engine/Source/Programs/Shared/EpicGames.UHT/Utils/UhtSession.cs
@@ -857,6 +857,16 @@ namespace EpicGames.UHT.Utils
 		/// harness to generate more stable console output.
 		/// </summary>
 		public bool bCacheMessages = false;
+
+		/// <summary>
+		/// Collection of UHT tables
+		/// </summary>
+		public UhtTables? Tables;
+
+		/// <summary>
+		/// Configuration for the session
+		/// </summary>
+		public IUhtConfig? Config;
 		#endregion
 
 		/// <summary>
@@ -1079,6 +1089,11 @@ namespace EpicGames.UHT.Utils
 		/// </summary>
 		public int ObjectTypeCount => this.ObjectTypeCountInternal.Count;
 
+		/// <summary>
+		/// Return the collection of exporters
+		/// </summary>
+		public UhtExporterTable ExporterTable => this.Tables!.ExporterTable;
+
 		/// <summary>
 		/// Enable/Disable an exporter.  This overrides the default state of the exporter.
 		/// </summary>
@@ -1089,6 +1104,99 @@ namespace EpicGames.UHT.Utils
 			this.ExporterStates[Name] = Enabled;
 		}
 
+		/// <summary>
+		/// Return the keyword table for the given table name
+		/// </summary>
+		/// <param name="TableName">Name of the table</param>
+		/// <returns>The requested table</returns>
+		public UhtKeywordTable GetKeywordTable(string TableName)
+		{
+			return this.Tables!.KeywordTables.Get(TableName);
+		}
+
+		/// <summary>
+		/// Return the specifier table for the given table name
+		/// </summary>
+		/// <param name="TableName">Name of the table</param>
+		/// <returns>The requested table</returns>
+		public UhtSpecifierTable GetSpecifierTable(string TableName)
+		{
+			return this.Tables!.SpecifierTables.Get(TableName);
+		}
+
+		/// <summary>
+		/// Return the specifier validator table for the given table name
+		/// </summary>
+		/// <param name="TableName">Name of the table</param>
+		/// <returns>The requested table</returns>
+		public UhtSpecifierValidatorTable GetSpecifierValidatorTable(string TableName)
+		{
+			return this.Tables!.SpecifierValidatorTables.Get(TableName);
+		}
+
+		/// <summary>
+		/// Generate an error for the given unhandled keyword
+		/// </summary>
+		/// <param name="TokenReader">Token reader</param>
+		/// <param name="Token">Unhandled token</param>
+		public void LogUnhandledKeywordError(IUhtTokenReader TokenReader, UhtToken Token)
+		{
+			this.Tables!.KeywordTables.LogUnhandledError(TokenReader, Token);
+		}
+
+		/// <summary>
+		/// Test to see if the given class name is a property
+		/// </summary>
+		/// <param name="Name">Name of the class without the prefix</param>
+		/// <returns>True if the class name is a property.  False if the class name isn't a property or isn't an engine class.</returns>
+		public bool IsValidPropertyTypeName(StringView Name)
+		{
+			return this.Tables!.EngineClassTable.IsValidPropertyTypeName(Name);
+		}
+
+		/// <summary>
+		/// Return the loc text default value associated with the given name
+		/// </summary>
+		/// <param name="Name"></param>
+		/// <param name="LocTextDefaultValue">Loc text default value handler</param>
+		/// <returns></returns>
+		public bool TryGetLocTextDefaultValue(StringView Name, out UhtLocTextDefaultValue LocTextDefaultValue)
+		{
+			return this.Tables!.LocTextDefaultValueTable.TryGet(Name, out LocTextDefaultValue);
+		}
+
+		/// <summary>
+		/// Return the default processor
+		/// </summary>
+		public UhtPropertyType DefaultPropertyType => this.Tables!.PropertyTypeTable.Default;
+
+		/// <summary>
+		/// Return the property type associated with the given name
+		/// </summary>
+		/// <param name="Name"></param>
+		/// <param name="PropertyType">Property type if matched</param>
+		/// <returns></returns>
+		public bool TryGetPropertyType(StringView Name, out UhtPropertyType PropertyType)
+		{
+			return this.Tables!.PropertyTypeTable.TryGet(Name, out PropertyType);
+		}
+
+		/// <summary>
+		/// Fetch the default sanitizer
+		/// </summary>
+		public UhtStructDefaultValue DefaultStructDefaultValue => this.Tables!.StructDefaultValueTable.Default;
+
+		/// <summary>
+		/// Return the structure default value associated with the given name
+		/// </summary>
+		/// <param name="Name"></param>
+		/// <param name="StructDefaultValue">Structure default value handler</param>
+		/// <returns></returns>
+		public bool TryGetStructDefaultValue(StringView Name, out UhtStructDefaultValue StructDefaultValue)
+		{
+			return this.Tables!.StructDefaultValueTable.TryGet(Name, out StructDefaultValue);
+		}
+
 		/// <summary>
 		/// Run UHT on the given manifest.  Use the bHasError property to see if process was successful.
 		/// </summary>
@@ -1102,11 +1210,19 @@ namespace EpicGames.UHT.Utils
 				return;
 			}
 
-			// Initialize
-			Try(null, () => {
-				_ = UhtConfig.Instance;
-				UhtAttributeScanner.Scan();
-			});
+			if (this.Config == null)
+			{
+				Interlocked.Increment(ref this.ErrorCountInternal);
+				Log.Logger.LogError("No configuration supplied, aborting.");
+				return;
+			}
+
+			if (this.Tables == null)
+			{
+				Interlocked.Increment(ref this.ErrorCountInternal);
+				Log.Logger.LogError("No parsing tables supplied, aborting.");
+				return;
+			}
 
 			switch (this.ReferenceMode)
 			{
@@ -2370,12 +2486,12 @@ namespace EpicGames.UHT.Utils
 				Log.Logger.LogTrace("Step - Exports");
 				UhtExportOptions Options = UhtExportOptions.None;
 
-				foreach (UhtExporter Exporter in UhtExporterTable.Instance)
+				foreach (UhtExporter Exporter in this.ExporterTable)
 				{
 					bool Run = false;
 					if (!this.ExporterStates.TryGetValue(Exporter.Name, out Run))
 					{
-						Run = UhtConfig.Instance.IsExporterEnabled(Exporter.Name) || 
+						Run = Config!.IsExporterEnabled(Exporter.Name) || 
 							(Exporter.Options.HasAnyFlags(UhtExporterOptions.Default) && !this.bNoDefaultExporters);
 					}
 
diff --git a/Engine/Source/Programs/Shared/EpicGames.UHT/Utils/UhtTables.cs b/Engine/Source/Programs/Shared/EpicGames.UHT/Utils/UhtTables.cs
--- a/Engine/Source/Programs/Shared/EpicGames.UHT/Utils/UhtTables.cs
+++ b/Engine/Source/Programs/Shared/EpicGames.UHT/Utils/UhtTables.cs
@@ -6,6 +6,7 @@ using System;
 using System.Collections.Generic;
 using System.Diagnostics.CodeAnalysis;
 using System.Reflection;
+using UnrealBuildBase;
 
 namespace EpicGames.UHT.Utils
 {
@@ -339,7 +340,6 @@ namespace EpicGames.UHT.Utils
 	/// <summary>
 	/// Bootstraps the standard UHT tables
 	/// </summary>
-	[UnrealHeaderTool(InitMethod= "InitStandardTables")]
 	public static class UhtStandardTables
 	{
 
@@ -363,42 +363,248 @@ namespace EpicGames.UHT.Utils
 		/// <summary>
 		/// Create all of the standard scope tables.
 		/// </summary>
-		public static void InitStandardTables()
+		/// <param name="Tables">UHT tables</param>
+		public static void InitStandardTables(UhtTables Tables)
 		{
-			CreateTables(UhtTableNames.Default, "Default", EUhtCreateTablesFlags.Specifiers | EUhtCreateTablesFlags.Keyword, null, true);
-			CreateTables(UhtTableNames.Global, "Global", EUhtCreateTablesFlags.Keyword, UhtTableNames.Default, false);
-			CreateTables(UhtTableNames.PropertyArgument, "Argument/Return", EUhtCreateTablesFlags.Specifiers, UhtTableNames.Default, false);
-			CreateTables(UhtTableNames.PropertyMember, "Member", EUhtCreateTablesFlags.Specifiers, UhtTableNames.Default, false);
-			CreateTables(UhtTableNames.Object, "Object", EUhtCreateTablesFlags.Specifiers | EUhtCreateTablesFlags.Keyword, UhtTableNames.Default, true);
-			CreateTables(UhtTableNames.Field, "Field", EUhtCreateTablesFlags.Specifiers | EUhtCreateTablesFlags.Keyword, UhtTableNames.Object, true);
-			CreateTables(UhtTableNames.Enum, "Enum", EUhtCreateTablesFlags.Specifiers | EUhtCreateTablesFlags.Keyword, UhtTableNames.Field, false);
-			CreateTables(UhtTableNames.Struct, "Struct", EUhtCreateTablesFlags.Specifiers | EUhtCreateTablesFlags.Keyword, UhtTableNames.Field, true);
-			CreateTables(UhtTableNames.ClassBase, "ClassBase", EUhtCreateTablesFlags.Specifiers | EUhtCreateTablesFlags.Keyword, UhtTableNames.Struct, true);
-			CreateTables(UhtTableNames.Class, "Class", EUhtCreateTablesFlags.Specifiers | EUhtCreateTablesFlags.Keyword, UhtTableNames.ClassBase, false);
-			CreateTables(UhtTableNames.Interface, "Interface", EUhtCreateTablesFlags.Specifiers | EUhtCreateTablesFlags.Keyword, UhtTableNames.ClassBase, false);
-			CreateTables(UhtTableNames.NativeInterface, "IInterface", EUhtCreateTablesFlags.Specifiers | EUhtCreateTablesFlags.Keyword, UhtTableNames.ClassBase, false);
-			CreateTables(UhtTableNames.Function, "Function", EUhtCreateTablesFlags.Specifiers | EUhtCreateTablesFlags.Keyword, UhtTableNames.Struct, false);
-			CreateTables(UhtTableNames.ScriptStruct, "Struct", EUhtCreateTablesFlags.Specifiers | EUhtCreateTablesFlags.Keyword, UhtTableNames.Struct, false);
+			CreateTables(Tables, UhtTableNames.Default, "Default", EUhtCreateTablesFlags.Specifiers | EUhtCreateTablesFlags.Keyword, null, true);
+			CreateTables(Tables, UhtTableNames.Global, "Global", EUhtCreateTablesFlags.Keyword, UhtTableNames.Default, false);
+			CreateTables(Tables, UhtTableNames.PropertyArgument, "Argument/Return", EUhtCreateTablesFlags.Specifiers, UhtTableNames.Default, false);
+			CreateTables(Tables, UhtTableNames.PropertyMember, "Member", EUhtCreateTablesFlags.Specifiers, UhtTableNames.Default, false);
+			CreateTables(Tables, UhtTableNames.Object, "Object", EUhtCreateTablesFlags.Specifiers | EUhtCreateTablesFlags.Keyword, UhtTableNames.Default, true);
+			CreateTables(Tables, UhtTableNames.Field, "Field", EUhtCreateTablesFlags.Specifiers | EUhtCreateTablesFlags.Keyword, UhtTableNames.Object, true);
+			CreateTables(Tables, UhtTableNames.Enum, "Enum", EUhtCreateTablesFlags.Specifiers | EUhtCreateTablesFlags.Keyword, UhtTableNames.Field, false);
+			CreateTables(Tables, UhtTableNames.Struct, "Struct", EUhtCreateTablesFlags.Specifiers | EUhtCreateTablesFlags.Keyword, UhtTableNames.Field, true);
+			CreateTables(Tables, UhtTableNames.ClassBase, "ClassBase", EUhtCreateTablesFlags.Specifiers | EUhtCreateTablesFlags.Keyword, UhtTableNames.Struct, true);
+			CreateTables(Tables, UhtTableNames.Class, "Class", EUhtCreateTablesFlags.Specifiers | EUhtCreateTablesFlags.Keyword, UhtTableNames.ClassBase, false);
+			CreateTables(Tables, UhtTableNames.Interface, "Interface", EUhtCreateTablesFlags.Specifiers | EUhtCreateTablesFlags.Keyword, UhtTableNames.ClassBase, false);
+			CreateTables(Tables, UhtTableNames.NativeInterface, "IInterface", EUhtCreateTablesFlags.Specifiers | EUhtCreateTablesFlags.Keyword, UhtTableNames.ClassBase, false);
+			CreateTables(Tables, UhtTableNames.Function, "Function", EUhtCreateTablesFlags.Specifiers | EUhtCreateTablesFlags.Keyword, UhtTableNames.Struct, false);
+			CreateTables(Tables, UhtTableNames.ScriptStruct, "Struct", EUhtCreateTablesFlags.Specifiers | EUhtCreateTablesFlags.Keyword, UhtTableNames.Struct, false);
 		}
 
 		/// <summary>
 		/// Creates a series of tables given the supplied setup
 		/// </summary>
+		/// <param name="Tables">UHT tables</param>
 		/// <param name="TableName">Name of the table.</param>
 		/// <param name="TableUserName">Name presented to the user via error messages.</param>
 		/// <param name="CreateTables">Types of tables to be created.</param>
 		/// <param name="ParentTableName">Name of the parent table or null for none.</param>
 		/// <param name="bInternal">If true, this table will not be included in any error messages.</param>
-		public static void CreateTables(string TableName, string TableUserName, EUhtCreateTablesFlags CreateTables, string? ParentTableName, bool bInternal = false)
+		public static void CreateTables(UhtTables Tables, string TableName, string TableUserName, 
+			EUhtCreateTablesFlags CreateTables, string? ParentTableName, bool bInternal = false)
 		{
 			if (CreateTables.HasFlag(EUhtCreateTablesFlags.Keyword))
 			{
-				UhtKeywordTables.Instance.Create(TableName, TableUserName, ParentTableName, bInternal);
+				Tables.KeywordTables.Create(TableName, TableUserName, ParentTableName, bInternal);
 			}
 			if (CreateTables.HasFlag(EUhtCreateTablesFlags.Specifiers))
 			{
-				UhtSpecifierTables.Instance.Create(TableName, TableUserName, ParentTableName, bInternal);
-				UhtSpecifierValidatorTables.Instance.Create(TableName, TableUserName, ParentTableName, bInternal);
+				Tables.SpecifierTables.Create(TableName, TableUserName, ParentTableName, bInternal);
+				Tables.SpecifierValidatorTables.Create(TableName, TableUserName, ParentTableName, bInternal);
+			}
+		}
+	}
+
+	/// <summary>
+	/// This attribute must be applied to any class that contains other UHT attributes.
+	/// </summary>
+	[AttributeUsage(AttributeTargets.Class)]
+	public class UnrealHeaderToolAttribute : Attribute
+	{
+
+		/// <summary>
+		/// If specified, this method will be invoked once during the scan for attributes.
+		/// It can be used to perform some one time initialization.
+		/// </summary>
+		public string InitMethod = string.Empty;
+	}
+
+	/// <summary>
+	/// UnrealHeaderTool avoids hard coding any table contents by using attributes to add entries to the tables.
+	/// 
+	/// There are two basic styles of tables in UHT.  
+	/// 
+	/// The first style is just a simple association of a string and an attribute.  For example, 
+	/// the engine class table is just a collection of all the engine class names supported by the engine.
+	/// 
+	/// The second style is a table of tables.  Depending on the context (i.e. is a "class" or "function" 
+	/// being parsed), attributes will "extend" a given table adding an entry to that table and every
+	/// table the derives from that table.  For example, the Keywords table will add "private" to the 
+	/// "ClassBase" table.  Since the "Class", "Interface", and "NativeInterface" tables derive from
+	/// "ClassBase", all of those tables will contain the keyword "private".
+	/// 
+	/// See UhtTables.cs for a list of table names and how they relate to each other.
+	/// 
+	/// Tables are implemented in the following source files:
+	/// 
+	///		UhtEngineClassTable.cs - Collection of all the engine class names.
+	///		UhtKeywordTable.cs - Collection of the C++ keywords that UHT understands.
+	///		UhtLocTextDefaultValueTable.cs - Collection of loctext default value parsing
+	///		UhtPropertyTypeTable.cs - Collection of the property type keywords.
+	///		UhtSpecifierTable.cs - Collection of the known specifiers
+	///		UhtSpecifierValidatorTable.cs - Collection of the specifier validators
+	///		UhtStructDefaultValueTable.cs - Collection of structure default value parsing
+	/// </summary>
+	public class UhtTables
+	{
+		/// <summary>
+		/// Collection of specifier tables
+		/// </summary>
+		public readonly UhtSpecifierTables SpecifierTables = new UhtSpecifierTables();
+
+		/// <summary>
+		/// Collection of specifier validator tables
+		/// </summary>
+		public readonly UhtSpecifierValidatorTables SpecifierValidatorTables = new UhtSpecifierValidatorTables();
+
+		/// <summary>
+		/// Collection of keyword tables
+		/// </summary>
+		public readonly UhtKeywordTables KeywordTables = new UhtKeywordTables();
+
+		/// <summary>
+		/// Collection of property types
+		/// </summary>
+		public readonly UhtPropertyTypeTable PropertyTypeTable = new UhtPropertyTypeTable();
+
+		/// <summary>
+		/// Collection of structure default values
+		/// </summary>
+		public readonly UhtStructDefaultValueTable StructDefaultValueTable = new UhtStructDefaultValueTable();
+
+		/// <summary>
+		/// Collection of engine class types
+		/// </summary>
+		public readonly UhtEngineClassTable EngineClassTable = new UhtEngineClassTable();
+
+		/// <summary>
+		/// Collection of exporters
+		/// </summary>
+		public readonly UhtExporterTable ExporterTable = new UhtExporterTable();
+
+		/// <summary>
+		/// Collection loc text default values
+		/// </summary>
+		public readonly UhtLocTextDefaultValueTable LocTextDefaultValueTable = new UhtLocTextDefaultValueTable();
+
+		/// <summary>
+		/// Construct a new table collection
+		/// </summary>
+		public UhtTables()
+		{
+			Initialize();
+		}
+
+		/// <summary>
+		/// Perform one time initialization of the system
+		/// </summary>
+		private void Initialize()
+		{
+			UhtStandardTables.InitStandardTables(this);
+
+			IEnumerable<Assembly> LoadedAssemblies = AppDomain.CurrentDomain.GetAssemblies();
+
+			// Collect all of the types from the loaded assemblies
+			foreach (Assembly LoadedAssembly in LoadedAssemblies)
+			{
+				Type[] Types = LoadedAssembly.SafeGetLoadedTypes();
+				foreach (Type Type in Types)
+				{
+					CheckForAttributes(Type);
+				}
+			}
+
+			PerformPostInitialization();
+		}
+
+		/// <summary>
+		/// For the given type, look for any table related attributes 
+		/// </summary>
+		/// <param name="Type">Type in question</param>
+		private void CheckForAttributes(Type Type)
+		{
+			if (Type.IsClass)
+			{
+
+				// Loop through the attributes
+				foreach (Attribute ClassAttribute in Type.GetCustomAttributes(false))
+				{
+					if (ClassAttribute is UnrealHeaderToolAttribute ParserAttribute)
+					{
+						HandleUnrealHeaderToolAttribute(Type, ParserAttribute);
+					}
+					else if (ClassAttribute is UhtEngineClassAttribute EngineClassAttribute)
+					{
+						this.EngineClassTable.OnEngineClassAttribute(Type, EngineClassAttribute);
+					}
+				}
+			}
+		}
+
+		private void PerformPostInitialization()
+		{
+			this.KeywordTables.Merge();
+			this.SpecifierTables.Merge();
+			this.SpecifierValidatorTables.Merge();
+
+			// Invoke this to generate an exception if there is no default
+			_ = this.PropertyTypeTable.Default;
+			_ = this.StructDefaultValueTable.Default;
+		}
+
+		private void HandleUnrealHeaderToolAttribute(Type Type, UnrealHeaderToolAttribute ParserAttribute)
+		{
+			if (!string.IsNullOrEmpty(ParserAttribute.InitMethod))
+			{
+				MethodInfo? InitInfo = Type.GetMethod(ParserAttribute.InitMethod, BindingFlags.Static | BindingFlags.Public | BindingFlags.NonPublic);
+				if (InitInfo == null)
+				{
+					throw new Exception($"Unable to find UnrealHeaderTool attribute InitMethod {ParserAttribute.InitMethod}");
+				}
+				InitInfo.Invoke(null, new Object[] { });
+			}
+
+			// Scan the methods for things we are interested in
+			foreach (MethodInfo MethodInfo in Type.GetMethods(BindingFlags.Static | BindingFlags.Public | BindingFlags.NonPublic))
+			{
+
+				// Scan for attributes we care about
+				foreach (Attribute MethodAttribute in MethodInfo.GetCustomAttributes())
+				{
+					if (MethodAttribute is UhtSpecifierAttribute SpecifierAttribute)
+					{
+						this.SpecifierTables.OnSpecifierAttribute(Type, MethodInfo, SpecifierAttribute);
+					}
+					else if (MethodAttribute is UhtSpecifierValidatorAttribute SpecifierValidatorAttribute)
+					{
+						this.SpecifierValidatorTables.OnSpecifierValidatorAttribute(Type, MethodInfo, SpecifierValidatorAttribute);
+					}
+					else if (MethodAttribute is UhtKeywordCatchAllAttribute KeywordCatchAllAttribute)
+					{
+						this.KeywordTables.OnKeywordCatchAllAttribute(Type, MethodInfo, KeywordCatchAllAttribute);
+					}
+					else if (MethodAttribute is UhtKeywordAttribute KeywordAttribute)
+					{
+						this.KeywordTables.OnKeywordAttribute(Type, MethodInfo, KeywordAttribute);
+					}
+					else if (MethodAttribute is UhtPropertyTypeAttribute PropertyTypeAttribute)
+					{
+						this.PropertyTypeTable.OnPropertyTypeAttribute(Type, MethodInfo, PropertyTypeAttribute);
+					}
+					else if (MethodAttribute is UhtStructDefaultValueAttribute StructDefaultValueAttribute)
+					{
+						this.StructDefaultValueTable.OnStructDefaultValueAttribute(Type, MethodInfo, StructDefaultValueAttribute);
+					}
+					else if (MethodAttribute is UhtExporterAttribute ExporterAttribute)
+					{
+						this.ExporterTable.OnExporterAttribute(Type, MethodInfo, ExporterAttribute);
+					}
+					else if (MethodAttribute is UhtLocTextDefaultValueAttribute LocTextDefaultValueAttribute)
+					{
+						this.LocTextDefaultValueTable.OnLocTextDefaultValueAttribute(Type, MethodInfo, LocTextDefaultValueAttribute);
+					}
+				}
 			}
 		}
 	}
diff --git a/Engine/Source/Programs/UnrealBuildTool.Tests/UnrealHeaderToolTests.cs b/Engine/Source/Programs/UnrealBuildTool.Tests/UnrealHeaderToolTests.cs
--- a/Engine/Source/Programs/UnrealBuildTool.Tests/UnrealHeaderToolTests.cs
+++ b/Engine/Source/Programs/UnrealBuildTool.Tests/UnrealHeaderToolTests.cs
@@ -18,13 +18,13 @@ namespace UnrealBuildToolTests
 			UhtGlobalOptions Options = new UhtGlobalOptions(CommandLineArguments);
 
 			// Initialize the attributes
-			UhtAttributeScanner.Scan();
+			UhtTables Tables = new UhtTables();
 
 			// Initialize the config
-			UhtConfigImpl.Read(CommandLineArguments);
+			IUhtConfig Config = new UhtConfigImpl(CommandLineArguments);
 
 			// Run the tests
-			Assert.IsTrue(UhtTestHarness.RunTests(Options));
+			Assert.IsTrue(UhtTestHarness.RunTests(Tables, Config, Options));
 		}
 	}
 }
diff --git a/Engine/Source/Programs/UnrealBuildTool/Modes/UnrealHeaderToolMode.cs b/Engine/Source/Programs/UnrealBuildTool/Modes/UnrealHeaderToolMode.cs
--- a/Engine/Source/Programs/UnrealBuildTool/Modes/UnrealHeaderToolMode.cs
+++ b/Engine/Source/Programs/UnrealBuildTool/Modes/UnrealHeaderToolMode.cs
@@ -170,13 +170,8 @@ namespace UnrealBuildTool.Modes
 		/// <summary>
 		/// Read the UHT configuration
 		/// </summary>
-		public static void Read(CommandLineArguments Args)
-		{
-			UhtConfigImpl Instance = new UhtConfigImpl(Args);
-			UhtConfig.Instance = Instance;
-		}
-
-		private UhtConfigImpl(CommandLineArguments Args)
+		/// <param name="Args">Extra command line arguments</param>
+		public UhtConfigImpl(CommandLineArguments Args)
 		{
 			DirectoryReference ConfigDirectory = DirectoryReference.Combine(Unreal.EngineDirectory, "Programs/UnrealBuildTool");
 			this.Ini = ConfigCache.ReadHierarchy(ConfigHierarchyType.Engine, ConfigDirectory, BuildHostPlatform.Current.Platform, "", Args.GetRawArray());
@@ -524,7 +519,7 @@ namespace UnrealBuildTool.Modes
 			public bool External;
 		}
 
-		private static bool RunScriptTest(UhtGlobalOptions Options, string TestDirectory, string TestOutputDirectory, string Script)
+		private static bool RunScriptTest(UhtTables Tables, IUhtConfig Config, UhtGlobalOptions Options, string TestDirectory, string TestOutputDirectory, string Script)
 		{
 			string InPath = Path.Combine(TestDirectory, Script);
 			string OutPath = Path.Combine(TestOutputDirectory, Script);
@@ -532,6 +527,8 @@ namespace UnrealBuildTool.Modes
 			UhtTestFileManager TestFileManager = new UhtTestFileManager(TestDirectory);
 			UhtSession Session = new UhtSession
 			{
+				Tables = Tables,
+				Config = Config,
 				FileManager = TestFileManager,
 				RootDirectory = TestDirectory,
 				bWarningsAsErrors = Options.bWarningsAsErrors,
@@ -814,27 +811,27 @@ namespace UnrealBuildTool.Modes
 			return true;
 		}
 
-		private static bool RunScriptTests(UhtGlobalOptions Options, string TestDirectory, string TestOutputDirectory, List<string> Scripts)
+		private static bool RunScriptTests(UhtTables Tables, IUhtConfig Config, UhtGlobalOptions Options, string TestDirectory, string TestOutputDirectory, List<string> Scripts)
 		{
 			bool bResult = true;
 			foreach (string Script in Scripts)
 			{
-				bResult &= RunScriptTest(Options, TestDirectory, TestOutputDirectory, Script);
+				bResult &= RunScriptTest(Tables, Config, Options, TestDirectory, TestOutputDirectory, Script);
 			}
 			return bResult;
 		}
 
-		private static bool RunDirectoryTests(UhtGlobalOptions Options, string TestDirectory, string TestOutputDirectory, List<string> Directories)
+		private static bool RunDirectoryTests(UhtTables Tables, IUhtConfig Config, UhtGlobalOptions Options, string TestDirectory, string TestOutputDirectory, List<string> Directories)
 		{
 			bool bResult = true;
 			foreach (string Directory in Directories)
 			{
-				bResult &= RunTests(Options, Path.Combine(TestDirectory, Directory), Path.Combine(TestOutputDirectory, Directory));
+				bResult &= RunTests(Tables, Config, Options, Path.Combine(TestDirectory, Directory), Path.Combine(TestOutputDirectory, Directory));
 			}
 			return bResult;
 		}
 
-		private static bool RunTests(UhtGlobalOptions Options, string TestDirectory, string TestOutputDirectory)
+		private static bool RunTests(UhtTables Tables, IUhtConfig Config, UhtGlobalOptions Options, string TestDirectory, string TestOutputDirectory)
 		{
 			// Create output directory
 			Directory.CreateDirectory(TestOutputDirectory);
@@ -862,11 +859,11 @@ namespace UnrealBuildTool.Modes
 
 			return
 				RunManifestTests(Options, TestDirectory, TestOutputDirectory, Manifests) &&
-				RunScriptTests(Options, TestDirectory, TestOutputDirectory, Scripts) &&
-				RunDirectoryTests(Options, TestDirectory, TestOutputDirectory, Directories);
+				RunScriptTests(Tables, Config, Options, TestDirectory, TestOutputDirectory, Scripts) &&
+				RunDirectoryTests(Tables, Config, Options, TestDirectory, TestOutputDirectory, Directories);
 		}
 
-		public static bool RunTests(UhtGlobalOptions Options)
+		public static bool RunTests(UhtTables Tables, IUhtConfig Config, UhtGlobalOptions Options)
 		{
 			DirectoryReference EngineSourceProgramDirectory = DirectoryReference.Combine(Unreal.EngineDirectory, "Source", "Programs");
 			string TestDirectory = FileReference.Combine(EngineSourceProgramDirectory, "UnrealBuildTool.Tests", "UHT").FullName;
@@ -885,7 +882,7 @@ namespace UnrealBuildTool.Modes
 			Log.TraceInformation("Output can be compared in {0}", TestOutputDirectory);
 
 			// Run the tests on the directory
-			return RunTests(Options, TestDirectory, TestOutputDirectory);
+			return RunTests(Tables, Config, Options, TestDirectory, TestOutputDirectory);
 		}
 	}
 
@@ -925,7 +922,9 @@ namespace UnrealBuildTool.Modes
 		/// <summary>
 		/// Print (incomplete) usage information
 		/// </summary>
-		private static void PrintUsage()
+		/// <param name="ExporterTable">Defined exporters</param>
+		/// <param name="Config">Configuration</param>
+		private static void PrintUsage(UhtExporterTable ExporterTable, IUhtConfig Config)
 		{
 			Console.WriteLine("UnrealBuildTool -Mode=UnrealHeaderTool [ProjectFile ManifestFile] -OR [\"-Target...\"] [Options]");
 			Console.WriteLine("");
@@ -942,7 +941,7 @@ namespace UnrealBuildTool.Modes
 				}
 			}
 
-			foreach (UhtExporter Generator in UhtExporterTable.Instance)
+			foreach (UhtExporter Generator in ExporterTable)
 			{
 				LongestPrefix = Generator.Name.Length + 2 > LongestPrefix ? Generator.Name.Length + 2 : LongestPrefix;
 			}
@@ -960,9 +959,9 @@ namespace UnrealBuildTool.Modes
 
 			Console.WriteLine("");
 			Console.WriteLine("Generators: Prefix with 'no' to disable a generator");
-			foreach (UhtExporter Generator in UhtExporterTable.Instance)
+			foreach (UhtExporter Generator in ExporterTable)
 			{
-				string IsDefault = UhtConfig.Instance.IsExporterEnabled(Generator.Name) || Generator.Options.HasAnyFlags(UhtExporterOptions.Default) ? " (Default)" : "";
+				string IsDefault = Config.IsExporterEnabled(Generator.Name) || Generator.Options.HasAnyFlags(UhtExporterOptions.Default) ? " (Default)" : "";
 				Console.WriteLine($"  -{Generator.Name.PadRight(LongestPrefix)} :  {Generator.Description}{IsDefault}");
 			}
 			Console.WriteLine("");
@@ -979,10 +978,10 @@ namespace UnrealBuildTool.Modes
 			{
 
 				// Initialize the attributes
-				UhtAttributeScanner.Scan();
+				UhtTables Tables = new UhtTables();
 
 				// Initialize the config
-				UhtConfigImpl.Read(Arguments);
+				IUhtConfig Config = new UhtConfigImpl(Arguments);
 
 				// Parse the global options
 				UhtGlobalOptions Options = new UhtGlobalOptions(Arguments);
@@ -1001,7 +1000,7 @@ namespace UnrealBuildTool.Modes
 				int RequiredArgCount = TargetArgumentIndex >= 0 || Options.bTest ? 0 : 2;
 				if (Arguments.GetPositionalArgumentCount() != RequiredArgCount || Options.bGetHelp)
 				{
-					PrintUsage();
+					PrintUsage(Tables.ExporterTable, Config);
 					return Options.bGetHelp ? (int)CompilationResult.Succeeded : (int)CompilationResult.OtherCompilationError;
 				}
 
@@ -1035,7 +1034,7 @@ namespace UnrealBuildTool.Modes
 				// If we are running test scripts
 				if (Options.bTest)
 				{
-					return UhtTestHarness.RunTests(Options) ? (int)CompilationResult.Succeeded : (int)CompilationResult.OtherCompilationError;
+					return UhtTestHarness.RunTests(Tables, Config, Options) ? (int)CompilationResult.Succeeded : (int)CompilationResult.OtherCompilationError;
 				}
 
 				string? ProjectPath = null;
@@ -1086,6 +1085,8 @@ namespace UnrealBuildTool.Modes
 
 				UhtSession Session = new UhtSession
 				{
+					Tables = Tables,
+					Config = Config,
 					FileManager = new UhtStdFileManager(),
 					EngineDirectory = Unreal.EngineDirectory.FullName,
 					ProjectDirectory = string.IsNullOrEmpty(ProjectPath) ? null : ProjectPath,
@@ -1108,7 +1109,7 @@ namespace UnrealBuildTool.Modes
 					Session.ReferenceMode = UhtReferenceMode.Verify;
 				}
 
-				foreach(UhtExporter Exporter in UhtExporterTable.Instance)
+				foreach(UhtExporter Exporter in Session.ExporterTable)
 				{
 					if (Arguments.HasOption($"-{Exporter.Name}"))
 					{
