Adjusted UHT table initialization in preperation for plugins.

#rnx

diff --git a/Engine/Source/Programs/Shared/EpicGames.UHT/Utils/UhtTables.cs b/Engine/Source/Programs/Shared/EpicGames.UHT/Utils/UhtTables.cs
--- a/Engine/Source/Programs/Shared/EpicGames.UHT/Utils/UhtTables.cs
+++ b/Engine/Source/Programs/Shared/EpicGames.UHT/Utils/UhtTables.cs
@@ -5,6 +5,7 @@ using EpicGames.UHT.Tables;
 using System;
 using System.Collections.Generic;
 using System.Diagnostics.CodeAnalysis;
+using System.IO;
 using System.Reflection;
 using UnrealBuildBase;
 
@@ -493,29 +494,64 @@ namespace EpicGames.UHT.Utils
 		/// </summary>
 		public UhtTables()
 		{
-			Initialize();
+			UhtStandardTables.InitStandardTables(this);
+			CheckForAttributes(Assembly.GetExecutingAssembly());
+			PerformPostInitialization();
 		}
 
 		/// <summary>
-		/// Perform one time initialization of the system
+		/// Add a collection of plugin assembly file paths
 		/// </summary>
-		private void Initialize()
+		/// <param name="PluginAssembliesFilePaths">Collection of plugins to load</param>
+		public void AddPlugins(IEnumerable<string> PluginAssembliesFilePaths)
 		{
-			UhtStandardTables.InitStandardTables(this);
+			foreach (string AssemblyFilePath in PluginAssembliesFilePaths)
+			{
+				Assembly? Plugin = FindAssemblyByName(Path.GetFileNameWithoutExtension(AssemblyFilePath));
+				if (Plugin == null)
+				{
+					byte[] AssemblyBytes = File.ReadAllBytes(AssemblyFilePath);
+					Plugin = Assembly.Load(AssemblyBytes);
+				}
+				CheckForAttributes(Plugin);
+			}
+			PerformPostInitialization();
+		}
 
-			IEnumerable<Assembly> LoadedAssemblies = AppDomain.CurrentDomain.GetAssemblies();
+		/// <summary>
+		/// Locate the assembly by name
+		/// </summary>
+		/// <param name="Name">Name of the assembly</param>
+		/// <returns>Assembly or null</returns>
+		private Assembly? FindAssemblyByName(string? Name)
+		{
+			if (Name != null)
+			{
+				foreach (Assembly Assembly in AppDomain.CurrentDomain.GetAssemblies())
+				{
+					AssemblyName AssemblyName = Assembly.GetName();
+					if (AssemblyName.Name != null && String.Compare(AssemblyName.Name, Name, true) == 0)
+					{
+						return Assembly;
+					}
+				}
+			}
+			return null;
+		}
 
-			// Collect all of the types from the loaded assemblies
-			foreach (Assembly LoadedAssembly in LoadedAssemblies)
+		/// <summary>
+		/// Check for UHT attributes on all types in the given assembly
+		/// </summary>
+		/// <param name="Assembly">Assembly to scan</param>
+		private void CheckForAttributes(Assembly? Assembly)
+		{
+			if (Assembly != null)
 			{
-				Type[] Types = LoadedAssembly.SafeGetLoadedTypes();
-				foreach (Type Type in Types)
+				foreach (Type Type in Assembly.SafeGetLoadedTypes())
 				{
 					CheckForAttributes(Type);
 				}
 			}
-
-			PerformPostInitialization();
 		}
 
 		/// <summary>
diff --git a/Engine/Source/Programs/UnrealBuildTool.Tests/UnrealHeaderToolTests.cs b/Engine/Source/Programs/UnrealBuildTool.Tests/UnrealHeaderToolTests.cs
--- a/Engine/Source/Programs/UnrealBuildTool.Tests/UnrealHeaderToolTests.cs
+++ b/Engine/Source/Programs/UnrealBuildTool.Tests/UnrealHeaderToolTests.cs
@@ -20,7 +20,7 @@ namespace UnrealBuildToolTests
 			// Initialize the attributes
 			UhtTables Tables = new UhtTables();
 
-			// Initialize the config
+			// Initialize the configuration
 			IUhtConfig Config = new UhtConfigImpl(CommandLineArguments);
 
 			// Run the tests
