Horde: Compact binary serialization improvements.

- Fix serialization of lists, directories and arrays that are set to null.
- Throw an exception when serializing a type that does not have any explicit compact binary field attributes. This is usually an error, caused by serializing a framework type with no matching converter, but can also occur when serializing an empty base class instance. To handle these situations, the class can be marked with [CbObject] to explicitly opt-in to being valid for serialization.

diff --git a/Engine/Source/Programs/Shared/EpicGames.Serialization.Tests/SerializerTests.cs b/Engine/Source/Programs/Shared/EpicGames.Serialization.Tests/SerializerTests.cs
--- a/Engine/Source/Programs/Shared/EpicGames.Serialization.Tests/SerializerTests.cs
+++ b/Engine/Source/Programs/Shared/EpicGames.Serialization.Tests/SerializerTests.cs
@@ -89,5 +89,45 @@ namespace EpicGames.Serialization.Tests
 
 			return obj;
 		}
+
+		class EmptyObject
+		{
+		}
+
+		[CbObject]
+		class EmptyObjectWithAttribute
+		{
+		}
+
+		[TestMethod]
+		public void EmptyObjectSerializationTest()
+		{
+			Assert.ThrowsException<CbEmptyClassException>(() => CbSerializer.Serialize(new EmptyObject()));
+			Assert.ThrowsException<CbEmptyClassException>(() => CbSerializer.Deserialize<EmptyObject>(CbField.Empty));
+
+			EmptyObjectWithAttribute obj2 = new EmptyObjectWithAttribute();
+			CbObject cbObj2 = CbSerializer.Serialize(obj2);
+			CbSerializer.Deserialize<EmptyObjectWithAttribute>(cbObj2);
+		}
+
+		class ObjectWithNullField
+		{
+			[CbField("a")]
+			public int Value { get; set; }
+
+			[CbField("b")]
+			public List<int>? Items { get; set; }
+		}
+
+		[TestMethod]
+		public void NullFieldTest()
+		{
+			ObjectWithNullField input = new ObjectWithNullField { Value = 123 };
+			CbObject obj = CbSerializer.Serialize(input);
+
+			ObjectWithNullField output = CbSerializer.Deserialize<ObjectWithNullField>(obj.AsField());
+			Assert.AreEqual(output.Value, 123);
+			Assert.AreEqual(output.Items, null);
+		}
 	}
 }
diff --git a/Engine/Source/Programs/Shared/EpicGames.Serialization/CbConverter.cs b/Engine/Source/Programs/Shared/EpicGames.Serialization/CbConverter.cs
--- a/Engine/Source/Programs/Shared/EpicGames.Serialization/CbConverter.cs
+++ b/Engine/Source/Programs/Shared/EpicGames.Serialization/CbConverter.cs
@@ -381,7 +381,14 @@ namespace EpicGames.Serialization
 		/// <returns></returns>
 		public static ICbConverter<T> GetConverter<T>()
 		{
-			return CbConverterCache<T>.Instance;
+			try
+			{
+				return CbConverterCache<T>.Instance;
+			}
+			catch (TypeInitializationException ex) when (ex.InnerException != null)
+			{
+				throw ex.InnerException;
+			}
 		}
 	}
 }
diff --git a/Engine/Source/Programs/Shared/EpicGames.Serialization/CbSerializer.cs b/Engine/Source/Programs/Shared/EpicGames.Serialization/CbSerializer.cs
--- a/Engine/Source/Programs/Shared/EpicGames.Serialization/CbSerializer.cs
+++ b/Engine/Source/Programs/Shared/EpicGames.Serialization/CbSerializer.cs
@@ -5,6 +5,15 @@ using System;
 
 namespace EpicGames.Serialization
 {
+	/// <summary>
+	/// Marks a class as supporting serialization to compact-binary, even if it does not have exposed fields. This suppresses errors
+	/// when base class objects are empty.
+	/// </summary>
+	[AttributeUsage(AttributeTargets.Class)]
+	public class CbObjectAttribute : Attribute
+	{
+	}
+
 	/// <summary>
 	/// Attribute used to mark a property that should be serialized to compact binary
 	/// </summary>
@@ -73,6 +82,26 @@ namespace EpicGames.Serialization
 		}
 	}
 
+	/// <summary>
+	/// Exception indicating that a class does not have any fields to serialize
+	/// </summary>
+	public sealed class CbEmptyClassException : CbException
+	{
+		/// <summary>
+		/// Type with missing field annotations
+		/// </summary>
+		public Type Type { get; }
+
+		/// <summary>
+		/// Constructor
+		/// </summary>
+		public CbEmptyClassException(Type type)
+			: base($"{type.Name} does not have any fields marked with a [CbField] attribute. If this is intended, explicitly mark the class with a [CbObject] attribute.")
+		{
+			Type = type;
+		}
+	}
+
 	/// <summary>
 	/// Attribute-driven compact binary serializer
 	/// </summary>
diff --git a/Engine/Source/Programs/Shared/EpicGames.Serialization/Converters/CbArrayConverter.cs b/Engine/Source/Programs/Shared/EpicGames.Serialization/Converters/CbArrayConverter.cs
--- a/Engine/Source/Programs/Shared/EpicGames.Serialization/Converters/CbArrayConverter.cs
+++ b/Engine/Source/Programs/Shared/EpicGames.Serialization/Converters/CbArrayConverter.cs
@@ -15,6 +15,11 @@ namespace EpicGames.Serialization.Converters
 		/// <inheritdoc/>
 		public override T[] Read(CbField field)
 		{
+			if (field.IsNull())
+			{
+				return null!;
+			}
+
 			List<T> list = new List<T>();
 			foreach (CbField elementField in field)
 			{
@@ -26,23 +31,37 @@ namespace EpicGames.Serialization.Converters
 		/// <inheritdoc/>
 		public override void Write(CbWriter writer, T[] list)
 		{
-			writer.BeginArray();
-			foreach (T element in list)
+			if (list == null)
 			{
-				CbSerializer.Serialize<T>(writer, element);
+				writer.WriteNullValue();
+			}
+			else
+			{
+				writer.BeginArray();
+				foreach (T element in list)
+				{
+					CbSerializer.Serialize<T>(writer, element);
+				}
+				writer.EndArray();
 			}
-			writer.EndArray();
 		}
 
 		/// <inheritdoc/>
 		public override void WriteNamed(CbWriter writer, Utf8String name, T[] array)
 		{
-			writer.BeginArray(name);
-			foreach (T element in array)
+			if (array == null)
+			{
+				writer.WriteNull(name);
+			}
+			else
 			{
-				CbSerializer.Serialize<T>(writer, element);
+				writer.BeginArray(name);
+				foreach (T element in array)
+				{
+					CbSerializer.Serialize<T>(writer, element);
+				}
+				writer.EndArray();
 			}
-			writer.EndArray();
 		}
 	}
 	
diff --git a/Engine/Source/Programs/Shared/EpicGames.Serialization/Converters/CbClassConverter.cs b/Engine/Source/Programs/Shared/EpicGames.Serialization/Converters/CbClassConverter.cs
--- a/Engine/Source/Programs/Shared/EpicGames.Serialization/Converters/CbClassConverter.cs
+++ b/Engine/Source/Programs/Shared/EpicGames.Serialization/Converters/CbClassConverter.cs
@@ -436,6 +436,10 @@ namespace EpicGames.Serialization.Converters
 					propertyList.Add((name, property));
 				}
 			}
+			if (propertyList.Count == 0 && type.GetCustomAttribute<CbObjectAttribute>() == null)
+			{
+				throw new CbEmptyClassException(type);
+			}
 			return propertyList.ToArray();
 		}
 
@@ -495,7 +499,14 @@ namespace EpicGames.Serialization.Converters
 			if (type.IsClass)
 			{
 				Type converterType = typeof(CbClassConverter<>).MakeGenericType(type);
-				converter = (ICbConverter?)Activator.CreateInstance(converterType);
+				try
+				{
+					converter = (ICbConverter?)Activator.CreateInstance(converterType);
+				}
+				catch (TargetInvocationException ex) when (ex.InnerException is object)
+				{
+					throw ex.InnerException;
+				}
 			}
 			return converter;
 		}
diff --git a/Engine/Source/Programs/Shared/EpicGames.Serialization/Converters/CbDictionaryConverter.cs b/Engine/Source/Programs/Shared/EpicGames.Serialization/Converters/CbDictionaryConverter.cs
--- a/Engine/Source/Programs/Shared/EpicGames.Serialization/Converters/CbDictionaryConverter.cs
+++ b/Engine/Source/Programs/Shared/EpicGames.Serialization/Converters/CbDictionaryConverter.cs
@@ -14,6 +14,11 @@ namespace EpicGames.Serialization.Converters
 		/// <inheritdoc/>
 		public override Dictionary<TKey, TValue> Read(CbField field)
 		{
+			if (field.IsNull())
+			{
+				return null!;
+			}
+
 			Dictionary<TKey, TValue> dictionary = new Dictionary<TKey, TValue>();
 			foreach (CbField element in field)
 			{
@@ -39,7 +44,11 @@ namespace EpicGames.Serialization.Converters
 		/// <inheritdoc/>
 		public override void Write(CbWriter writer, Dictionary<TKey, TValue> value)
 		{
-			if (value.Count > 0)
+			if (value == null)
+			{
+				writer.WriteNullValue();
+			}
+			else
 			{
 				writer.BeginUniformArray(CbFieldType.Array);
 				foreach (KeyValuePair<TKey, TValue> pair in value)
@@ -56,7 +65,11 @@ namespace EpicGames.Serialization.Converters
 		/// <inheritdoc/>
 		public override void WriteNamed(CbWriter writer, Utf8String name, Dictionary<TKey, TValue> value)
 		{
-			if (value.Count > 0)
+			if (value == null)
+			{
+				writer.WriteNull(name);
+			}
+			else if (value.Count > 0)
 			{
 				writer.BeginUniformArray(name, CbFieldType.Array);
 				foreach (KeyValuePair<TKey, TValue> pair in value)
diff --git a/Engine/Source/Programs/Shared/EpicGames.Serialization/Converters/CbListConverter.cs b/Engine/Source/Programs/Shared/EpicGames.Serialization/Converters/CbListConverter.cs
--- a/Engine/Source/Programs/Shared/EpicGames.Serialization/Converters/CbListConverter.cs
+++ b/Engine/Source/Programs/Shared/EpicGames.Serialization/Converters/CbListConverter.cs
@@ -15,34 +15,55 @@ namespace EpicGames.Serialization.Converters
 		/// <inheritdoc/>
 		public override List<T> Read(CbField field)
 		{
-			List<T> list = new List<T>();
-			foreach (CbField elementField in field)
+			if (field.IsNull())
 			{
-				list.Add(CbSerializer.Deserialize<T>(elementField));
+				return null!;
+			}
+			else
+			{
+				List<T> list = new List<T>();
+				foreach (CbField elementField in field)
+				{
+					list.Add(CbSerializer.Deserialize<T>(elementField));
+				}
+				return list;
 			}
-			return list;
 		}
 
 		/// <inheritdoc/>
 		public override void Write(CbWriter writer, List<T> list)
 		{
-			writer.BeginArray();
-			foreach (T element in list)
+			if (list == null)
+			{
+				writer.WriteNullValue();
+			}
+			else
 			{
-				CbSerializer.Serialize<T>(writer, element);
+				writer.BeginArray();
+				foreach (T element in list)
+				{
+					CbSerializer.Serialize<T>(writer, element);
+				}
+				writer.EndArray();
 			}
-			writer.EndArray();
 		}
 
 		/// <inheritdoc/>
 		public override void WriteNamed(CbWriter writer, Utf8String name, List<T> list)
 		{
-			writer.BeginArray(name);
-			foreach (T element in list)
+			if (list == null)
 			{
-				CbSerializer.Serialize<T>(writer, element);
+				writer.WriteNull(name);
+			}
+			else
+			{
+				writer.BeginArray(name);
+				foreach (T element in list)
+				{
+					CbSerializer.Serialize<T>(writer, element);
+				}
+				writer.EndArray();
 			}
-			writer.EndArray();
 		}
 	}
 
