Horde: Fix incorrectly named enum members, causing syncing to return invalid records.

diff --git a/Engine/Source/Programs/Shared/EpicGames.Perforce.Managed/ManagedWorkspace.cs b/Engine/Source/Programs/Shared/EpicGames.Perforce.Managed/ManagedWorkspace.cs
--- a/Engine/Source/Programs/Shared/EpicGames.Perforce.Managed/ManagedWorkspace.cs
+++ b/Engine/Source/Programs/Shared/EpicGames.Perforce.Managed/ManagedWorkspace.cs
@@ -1310,15 +1310,16 @@ namespace EpicGames.Perforce.Managed
 		/// </summary>
 		class FStatIndexedRecord
 		{
+			// Note: This enum is used for indexing an array of fields, and member names much match P4 field names (including case).
 			enum Field
 			{
-				Code,
-				DepotFile,
-				ClientFile,
-				HeadType,
-				HaveRev,
-				FileSize,
-				Digest
+				code,
+				depotFile,
+				clientFile,
+				headType,
+				haveRev,
+				fileSize,
+				digest
 			}
 
 			public static readonly string[] FieldNames = Enum.GetNames(typeof(Field));
@@ -1326,17 +1327,17 @@ namespace EpicGames.Perforce.Managed
 
 			public PerforceValue[] Values { get; } = new PerforceValue[FieldNames.Length];
 
-			public Utf8String DepotFile => Values[(int)Field.DepotFile].GetString();
+			public Utf8String DepotFile => Values[(int)Field.depotFile].GetString();
 
-			public Utf8String ClientFile => Values[(int)Field.ClientFile].GetString();
+			public Utf8String ClientFile => Values[(int)Field.clientFile].GetString();
 
-			public Utf8String HeadType => Values[(int)Field.HeadType].GetString();
+			public Utf8String HeadType => Values[(int)Field.headType].GetString();
 
-			public Utf8String HaveRev => Values[(int)Field.HaveRev].GetString();
+			public Utf8String HaveRev => Values[(int)Field.haveRev].GetString();
 
-			public long FileSize => Values[(int)Field.FileSize].AsLong();
+			public long FileSize => Values[(int)Field.fileSize].AsLong();
 
-			public Utf8String Digest => Values[(int)Field.Digest].GetString();
+			public Utf8String Digest => Values[(int)Field.digest].GetString();
 		}
 
 		/// <summary>
