Fix some analyzer warnings.

diff --git a/Engine/Source/Programs/Shared/EpicGames.BuildGraph/BgOptionSpec.cs b/Engine/Source/Programs/Shared/EpicGames.BuildGraph/BgOptionSpec.cs
--- a/Engine/Source/Programs/Shared/EpicGames.BuildGraph/BgOptionSpec.cs
+++ b/Engine/Source/Programs/Shared/EpicGames.BuildGraph/BgOptionSpec.cs
@@ -80,7 +80,7 @@ namespace EpicGames.BuildGraph
 			if (context.Options.TryGetValue(Name, out value))
 			{
 				bool boolValue;
-				if (!bool.TryParse(value, out boolValue))
+				if (!Boolean.TryParse(value, out boolValue))
 				{
 					throw new BgOptionValidationException($"Argument for {Name} is not a valid bool ({value})");
 				}
diff --git a/Engine/Source/Programs/Shared/EpicGames.BuildGraph/BgSchema.cs b/Engine/Source/Programs/Shared/EpicGames.BuildGraph/BgSchema.cs
--- a/Engine/Source/Programs/Shared/EpicGames.BuildGraph/BgSchema.cs
+++ b/Engine/Source/Programs/Shared/EpicGames.BuildGraph/BgSchema.cs
@@ -260,9 +260,9 @@ namespace EpicGames.BuildGraph
 		{
 			// Create a lookup from standard types to their qualified names
 			Dictionary<Type, XmlQualifiedName> typeToSchemaTypeName = new Dictionary<Type, XmlQualifiedName>();
-			typeToSchemaTypeName.Add(typeof(String), GetQualifiedTypeName(ScriptSchemaStandardType.BalancedString));
-			typeToSchemaTypeName.Add(typeof(Boolean), GetQualifiedTypeName(ScriptSchemaStandardType.Boolean));
-			typeToSchemaTypeName.Add(typeof(Int32), GetQualifiedTypeName(ScriptSchemaStandardType.Integer));
+			typeToSchemaTypeName.Add(typeof(string), GetQualifiedTypeName(ScriptSchemaStandardType.BalancedString));
+			typeToSchemaTypeName.Add(typeof(bool), GetQualifiedTypeName(ScriptSchemaStandardType.Boolean));
+			typeToSchemaTypeName.Add(typeof(int), GetQualifiedTypeName(ScriptSchemaStandardType.Integer));
 			foreach ((Type type, ScriptSchemaStandardType schemaType) in primitiveTypes)
 			{
 				typeToSchemaTypeName.Add(type, GetQualifiedTypeName(schemaType));
diff --git a/Engine/Source/Programs/Shared/EpicGames.BuildGraph/BgScriptReader.cs b/Engine/Source/Programs/Shared/EpicGames.BuildGraph/BgScriptReader.cs
--- a/Engine/Source/Programs/Shared/EpicGames.BuildGraph/BgScriptReader.cs
+++ b/Engine/Source/Programs/Shared/EpicGames.BuildGraph/BgScriptReader.cs
@@ -542,7 +542,7 @@ namespace EpicGames.BuildGraph
 			{
 				string[] tokens = name.Split(':');
 				name = tokens[0];
-				valueLength = int.Parse(tokens[1]);
+				valueLength = Int32.Parse(tokens[1]);
 			}
 
 			// Check each scope for the property
diff --git a/Engine/Source/Programs/Shared/EpicGames.BuildGraph/Expressions/BgBool.cs b/Engine/Source/Programs/Shared/EpicGames.BuildGraph/Expressions/BgBool.cs
--- a/Engine/Source/Programs/Shared/EpicGames.BuildGraph/Expressions/BgBool.cs
+++ b/Engine/Source/Programs/Shared/EpicGames.BuildGraph/Expressions/BgBool.cs
@@ -68,7 +68,7 @@ namespace EpicGames.BuildGraph.Expressions
 	class BgBoolType : BgTypeBase<BgBool>
 	{
 		/// <inheritdoc/>
-		public override BgBool DeserializeArgument(string text) => bool.Parse(text);
+		public override BgBool DeserializeArgument(string text) => Boolean.Parse(text);
 
 		/// <inheritdoc/>
 		public override string SerializeArgument(BgBool value, BgExprContext context) => value.Compute(context).ToString();
diff --git a/Engine/Source/Programs/Shared/EpicGames.BuildGraph/Expressions/BgInt.cs b/Engine/Source/Programs/Shared/EpicGames.BuildGraph/Expressions/BgInt.cs
--- a/Engine/Source/Programs/Shared/EpicGames.BuildGraph/Expressions/BgInt.cs
+++ b/Engine/Source/Programs/Shared/EpicGames.BuildGraph/Expressions/BgInt.cs
@@ -67,7 +67,7 @@ namespace EpicGames.BuildGraph.Expressions
 	class BgIntTraits : BgTypeBase<BgInt>
 	{
 		/// <inheritdoc/>
-		public override BgInt DeserializeArgument(string value) => (BgInt)int.Parse(value);
+		public override BgInt DeserializeArgument(string value) => (BgInt)Int32.Parse(value);
 
 		/// <inheritdoc/>
 		public override string SerializeArgument(BgInt value, BgExprContext context) => value.Compute(context).ToString();
diff --git a/Engine/Source/Programs/Shared/EpicGames.BuildGraph/Expressions/BgList.cs b/Engine/Source/Programs/Shared/EpicGames.BuildGraph/Expressions/BgList.cs
--- a/Engine/Source/Programs/Shared/EpicGames.BuildGraph/Expressions/BgList.cs
+++ b/Engine/Source/Programs/Shared/EpicGames.BuildGraph/Expressions/BgList.cs
@@ -107,7 +107,7 @@ namespace EpicGames.BuildGraph.Expressions
 		/// <inheritdoc cref="Enumerable.Select{TSource, TResult}(IEnumerable{TSource}, Func{TSource, TResult})"/>
 		public BgList<TResult> Select<TResult>(Func<T, TResult> function) where TResult : IBgExpr<TResult> => new BgListSelectExpr<T, TResult>(this, function);
 
-		/// <inheritdoc cref="Enumerable.Where{TSource}(IEnumerable{TSource}, Func{TSource, bool})"/>
+		/// <inheritdoc cref="Enumerable.Where{TSource}(IEnumerable{TSource}, Func{TSource, Boolean})"/>
 		public BgList<T> Where(Func<T, BgBool> predicate) => new BgListWhereExpr<T>(this, predicate);
 
 		/// <inheritdoc cref="Enumerable.Contains{TSource}(IEnumerable{TSource}, TSource)"/>
diff --git a/Engine/Source/Programs/Shared/EpicGames.Core/MessageTemplate.cs b/Engine/Source/Programs/Shared/EpicGames.Core/MessageTemplate.cs
--- a/Engine/Source/Programs/Shared/EpicGames.Core/MessageTemplate.cs
+++ b/Engine/Source/Programs/Shared/EpicGames.Core/MessageTemplate.cs
@@ -34,7 +34,7 @@ namespace EpicGames.Core
 					foreach ((int offset, int length) in offsets)
 					{
 						ReadOnlySpan<char> name = format.AsSpan(offset, length);
-						if (int.TryParse(name, out int number))
+						if (Int32.TryParse(name, out int number))
 						{
 							outputProperties[name.ToString()] = properties.ElementAtOrDefault(number).Value;
 						}
@@ -203,7 +203,7 @@ namespace EpicGames.Core
 					string name = format.Substring(offsets[idx].Item1, offsets[idx].Item2);
 
 					int number;
-					if (int.TryParse(name, out number))
+					if (Int32.TryParse(name, out number))
 					{
 						if (number >= 0 && number < args.Length)
 						{
@@ -231,7 +231,7 @@ namespace EpicGames.Core
 		public static bool TryGetPropertyValue(ReadOnlySpan<char> name, IEnumerable<KeyValuePair<string, object?>> properties, out object? value)
 		{
 			int number;
-			if (int.TryParse(name, System.Globalization.NumberStyles.Integer, null, out number))
+			if (Int32.TryParse(name, System.Globalization.NumberStyles.Integer, null, out number))
 			{
 				foreach (KeyValuePair<string, object?> property in properties)
 				{
