Horde Storage - Added converter for aspnet ProblemDetails type to compact binary

diff --git a/Engine/Source/Programs/Horde/HordeStorage/Jupiter.Common/BaseStartup.cs b/Engine/Source/Programs/Horde/HordeStorage/Jupiter.Common/BaseStartup.cs
--- a/Engine/Source/Programs/Horde/HordeStorage/Jupiter.Common/BaseStartup.cs
+++ b/Engine/Source/Programs/Horde/HordeStorage/Jupiter.Common/BaseStartup.cs
@@ -56,6 +56,8 @@ namespace Jupiter
         // This method gets called by the runtime. Use this method to add services to the container.
         public void ConfigureServices(IServiceCollection services)
         {
+            CbConvertersAspNet.AddAspnetConverters();
+
             // aws specific settings
             services.AddOptions<AWSCredentialsSettings>().Bind(Configuration.GetSection("AWSCredentials")).ValidateDataAnnotations();
             services.AddDefaultAWSOptions(Configuration.GetAWSOptions());
diff --git a/Engine/Source/Programs/Shared/EpicGames.AspNet/CbConvertersAspnet.cs b/Engine/Source/Programs/Shared/EpicGames.AspNet/CbConvertersAspnet.cs
new file mode 100644
--- /dev/null
+++ b/Engine/Source/Programs/Shared/EpicGames.AspNet/CbConvertersAspnet.cs
@@ -0,0 +1,72 @@
+// Copyright Epic Games, Inc. All Rights Reserved.
+
+using EpicGames.Core;
+using EpicGames.Serialization;
+using Microsoft.AspNetCore.Mvc;
+
+namespace EpicGames.AspNet
+{
+	public static class CbConvertersAspNet
+	{
+		public static void AddAspnetConverters()
+		{
+			CbConverter.TypeToConverter[typeof(ProblemDetails)] = new CbProblemDetailsConverter();
+		}
+	}
+
+	/// <summary>
+	/// Converter for asp.net problem details type
+	/// </summary>
+	/// <typeparam name="T"></typeparam>
+	class CbProblemDetailsConverter : CbConverterBase<ProblemDetails>
+	{
+		/// <inheritdoc/>
+		public override ProblemDetails Read(CbField field)
+		{
+			if (!field.IsObject())
+				throw new CbException($"Error converting field \"{field.Name}\" to ProblemDetails. Expected CbObject.");
+			
+			ProblemDetails result = new ProblemDetails
+			{
+				Title = field["title"].AsString(),
+				Detail = field["detail"].AsString(),
+				Type = field["type"].AsString(),
+				Instance = field["instance"].AsString(),
+				Status = field["status"].AsInt32()
+			};
+			return result;
+		}
+
+		/// <inheritdoc/>
+		public override void Write(CbWriter writer, ProblemDetails problemDetails)
+		{
+			writer.WriteObject(ToCbObject(problemDetails));
+		}
+
+		/// <inheritdoc/>
+		public override void WriteNamed(CbWriter writer, Utf8String name, ProblemDetails problemDetails)
+		{
+			writer.WriteField(name, ToCbObject(problemDetails).AsField());
+		}
+
+		private CbObject ToCbObject(ProblemDetails problemDetails)
+		{
+			CbWriter objectWriter = new CbWriter();
+			objectWriter.BeginObject();
+
+			objectWriter.WriteString("title", problemDetails.Title);
+			if (!string.IsNullOrEmpty(problemDetails.Detail))
+				objectWriter.WriteString("detail", problemDetails.Detail);
+			if (!string.IsNullOrEmpty(problemDetails.Type))
+				objectWriter.WriteString("type", problemDetails.Type);
+			if (!string.IsNullOrEmpty(problemDetails.Instance))
+				objectWriter.WriteString("instance", problemDetails.Instance);
+			if (problemDetails.Status.HasValue)
+				objectWriter.WriteInteger("status", problemDetails.Status.Value);
+			
+			objectWriter.EndObject();
+
+			return objectWriter.ToObject();
+		}
+	}
+}
