Updated the exporter system in UHT to simplify it.
Added some basic support for UBT plugins.

#rnx

diff --git a/Engine/Source/Programs/Shared/EpicGames.Core/UHTTypes.cs b/Engine/Source/Programs/Shared/EpicGames.Core/UHTTypes.cs
--- a/Engine/Source/Programs/Shared/EpicGames.Core/UHTTypes.cs
+++ b/Engine/Source/Programs/Shared/EpicGames.Core/UHTTypes.cs
@@ -121,5 +121,11 @@ namespace EpicGames.Core
 		/// List of modules
 		/// </summary>
 		public List<Module> Modules { get; set; } = new List<Module>();
+
+		/// <summary>
+		/// List of active UHT plugins.  Only used by the C# version of UHT.  This
+		/// list contains plugins from only modules listed above.
+		/// </summary>
+		public List<string> UhtPlugins { get; set; } = new List<string>();
 	}
 }
diff --git a/Engine/Source/Programs/Shared/EpicGames.UHT/Exporters/CodeGen/UhtCodeGenerator.cs b/Engine/Source/Programs/Shared/EpicGames.UHT/Exporters/CodeGen/UhtCodeGenerator.cs
--- a/Engine/Source/Programs/Shared/EpicGames.UHT/Exporters/CodeGen/UhtCodeGenerator.cs
+++ b/Engine/Source/Programs/Shared/EpicGames.UHT/Exporters/CodeGen/UhtCodeGenerator.cs
@@ -55,7 +55,6 @@ namespace EpicGames.UHT.Exporters.CodeGen
 
 		public readonly IUhtExportFactory Factory;
 		public UhtSession Session => Factory.Session;
-		public UhtExportOptions Options => Factory.Options;
 
 		private UhtCodeGenerator(IUhtExportFactory Factory)
 		{
@@ -65,16 +64,6 @@ namespace EpicGames.UHT.Exporters.CodeGen
 			this.PackageInfos = new PackageInfo[this.Factory.Session.PackageTypeCount];
 		}
 
-		private static UhtExportOptions GetAdditionalOptions(UHTManifest.Module Module)
-		{
-			UhtExportOptions Options = UhtExportOptions.None;
-			if (Module.SaveExportedHeaders)
-			{
-				Options |= UhtExportOptions.WriteOutput;
-			}
-			return Options;
-		}
-
 		private void Generate()
 		{
 			List<IUhtExportTask> Prereqs = new List<IUhtExportTask>();
@@ -102,10 +91,6 @@ namespace EpicGames.UHT.Exporters.CodeGen
 				{
 					UhtPackage Package = HeaderFile.Package;
 					UHTManifest.Module Module = Package.Module;
-					UhtExportOptions AdditionalOptions = GetAdditionalOptions(Module);
-
-					string HeaderPath = this.Factory.MakePath(HeaderFile, ".generated.h");
-					string CppPath = this.Factory.MakePath(HeaderFile, ".gen.cpp");
 
 					Prereqs.Clear();
 					foreach (UhtHeaderFile Referenced in HeaderFile.ReferencedHeadersNoLock)
@@ -120,11 +105,11 @@ namespace EpicGames.UHT.Exporters.CodeGen
 						}
 					}
 
-					this.HeaderInfos[HeaderFile.HeaderFileTypeIndex].Task = Factory.CreateTask(HeaderPath, CppPath, Prereqs, AdditionalOptions, 
-						(IUhtExportOutput HeaderOutput, IUhtExportOutput CppOutput) =>
+					this.HeaderInfos[HeaderFile.HeaderFileTypeIndex].Task = Factory.CreateTask(Prereqs,  
+						(IUhtExportTask Task) =>
 						{
-							new UhtHeaderCodeGeneratorHFile(this, Package, HeaderFile).Generate(HeaderOutput);
-							new UhtHeaderCodeGeneratorCppFile(this, Package, HeaderFile).Generate(CppOutput);
+							new UhtHeaderCodeGeneratorHFile(this, Package, HeaderFile).Generate(Task);
+							new UhtHeaderCodeGeneratorCppFile(this, Package, HeaderFile).Generate(Task);
 						});
 				}
 			}
@@ -134,10 +119,6 @@ namespace EpicGames.UHT.Exporters.CodeGen
 			foreach (UhtPackage Package in this.Session.Packages)
 			{
 				UHTManifest.Module Module = Package.Module;
-				UhtExportOptions AdditionalOptions = GetAdditionalOptions(Module);
-
-				string HeaderPath = this.Factory.MakePath(Package, "Classes.h");
-				string CppPath = this.Factory.MakePath(Package, ".init.gen.cpp");
 
 				bool bWriteHeader = false;
 				Prereqs.Clear();
@@ -166,25 +147,16 @@ namespace EpicGames.UHT.Exporters.CodeGen
 					}
 				}
 
-				if (bWriteHeader)
-				{
-					GeneratedPackages.Add(Factory.CreateTask(HeaderPath, CppPath, Prereqs, AdditionalOptions,
-						(IUhtExportOutput HeaderOutput, IUhtExportOutput CppOutput) =>
-						{
-							List<UhtHeaderFile> PackageSortedHeaders = GetSortedHeaderFiles(Package);
-							new UhtPackageCodeGeneratorHFile(this, Package).Generate(HeaderOutput, PackageSortedHeaders);
-							new UhtPackageCodeGeneratorCppFile(this, Package).Generate(CppOutput, PackageSortedHeaders);
-						}));
-				}
-				else
-				{
-					GeneratedPackages.Add(Factory.CreateTask(CppPath, Prereqs, AdditionalOptions, 
-						(IUhtExportOutput Output) =>
+				GeneratedPackages.Add(Factory.CreateTask(Prereqs, 
+					(IUhtExportTask Task) =>
+					{
+						List<UhtHeaderFile> PackageSortedHeaders = GetSortedHeaderFiles(Package);
+						if (bWriteHeader)
 						{
-							List<UhtHeaderFile> PackageSortedHeaders = GetSortedHeaderFiles(Package);
-							new UhtPackageCodeGeneratorCppFile(this, Package).Generate(Output, PackageSortedHeaders);
-						}));
-				}
+							new UhtPackageCodeGeneratorHFile(this, Package).Generate(Task, PackageSortedHeaders);
+						}
+						new UhtPackageCodeGeneratorCppFile(this, Package).Generate(Task, PackageSortedHeaders);
+					}));
 			}
 
 			// Wait for all the packages to complete
diff --git a/Engine/Source/Programs/Shared/EpicGames.UHT/Exporters/CodeGen/UhtHeaderCodeGeneratorCppFile.cs b/Engine/Source/Programs/Shared/EpicGames.UHT/Exporters/CodeGen/UhtHeaderCodeGeneratorCppFile.cs
--- a/Engine/Source/Programs/Shared/EpicGames.UHT/Exporters/CodeGen/UhtHeaderCodeGeneratorCppFile.cs
+++ b/Engine/Source/Programs/Shared/EpicGames.UHT/Exporters/CodeGen/UhtHeaderCodeGeneratorCppFile.cs
@@ -27,8 +27,8 @@ namespace EpicGames.UHT.Exporters.CodeGen
 		/// <summary>
 		/// For a given UE header file, generated the generated H file
 		/// </summary>
-		/// <param name="CppOutput">Output object</param>
-		public void Generate(IUhtExportOutput CppOutput)
+		/// <param name="Task">Requesting task</param>
+		public void Generate(IUhtExportTask Task)
 		{
 			ref UhtCodeGenerator.HeaderInfo HeaderInfo = ref this.HeaderInfos[this.HeaderFile.HeaderFileTypeIndex];
 			using (BorrowStringBuilder Borrower = new BorrowStringBuilder(StringBuilderCache.Big))
@@ -287,10 +287,19 @@ namespace EpicGames.UHT.Exporters.CodeGen
 				int GeneratedBodyEnd = Builder.Length;
 
 				Builder.Append(EnableDeprecationWarnings).Append("\r\n");
-				StringView GeneratedBody = CppOutput.CommitOutput(Builder);
 
-				// Save the hash of the generated body 
-				this.HeaderInfos[this.HeaderFile.HeaderFileTypeIndex].BodyHash = UhtHash.GenenerateTextHash(GeneratedBody.Span.Slice(GeneratedBodyStart, GeneratedBodyEnd - GeneratedBodyStart));
+				using (UhtBorrowBuffer BorrowBuffer = new UhtBorrowBuffer(Builder))
+				{
+					string CppFilePath = Task.Factory.MakePath(HeaderFile, ".gen.cpp");
+					StringView GeneratedBody = new StringView(BorrowBuffer.Buffer.Memory);
+					if (this.SaveExportedHeaders)
+					{
+						Task.CommitOutput(CppFilePath, GeneratedBody);
+					}
+
+					// Save the hash of the generated body 
+					this.HeaderInfos[this.HeaderFile.HeaderFileTypeIndex].BodyHash = UhtHash.GenenerateTextHash(GeneratedBody.Span.Slice(GeneratedBodyStart, GeneratedBodyEnd - GeneratedBodyStart));
+				}
 			}
 		}
 
diff --git a/Engine/Source/Programs/Shared/EpicGames.UHT/Exporters/CodeGen/UhtHeaderCodeGeneratorHFile.cs b/Engine/Source/Programs/Shared/EpicGames.UHT/Exporters/CodeGen/UhtHeaderCodeGeneratorHFile.cs
--- a/Engine/Source/Programs/Shared/EpicGames.UHT/Exporters/CodeGen/UhtHeaderCodeGeneratorHFile.cs
+++ b/Engine/Source/Programs/Shared/EpicGames.UHT/Exporters/CodeGen/UhtHeaderCodeGeneratorHFile.cs
@@ -31,8 +31,8 @@ namespace EpicGames.UHT.Exporters.CodeGen
 		/// <summary>
 		/// For a given UE header file, generated the generated H file
 		/// </summary>
-		/// <param name="HeaderOutput">Output object</param>
-		public void Generate(IUhtExportOutput HeaderOutput)
+		/// <param name="Task">Requesting task</param>
+		public void Generate(IUhtExportTask Task)
 		{
 			ref UhtCodeGenerator.HeaderInfo HeaderInfo = ref this.HeaderInfos[this.HeaderFile.HeaderFileTypeIndex];
 			using (BorrowStringBuilder Borrower = new BorrowStringBuilder(StringBuilderCache.Big))
@@ -110,7 +110,12 @@ namespace EpicGames.UHT.Exporters.CodeGen
 				}
 
 				Builder.Append(EnableDeprecationWarnings).Append("\r\n");
-				HeaderOutput.CommitOutput(Builder);
+
+				if (this.SaveExportedHeaders)
+				{
+					string HeaderFilePath = Task.Factory.MakePath(this.HeaderFile, ".generated.h");
+					Task.CommitOutput(HeaderFilePath, Builder);
+				}
 			}
 		}
 
diff --git a/Engine/Source/Programs/Shared/EpicGames.UHT/Exporters/CodeGen/UhtPackageCodeGenerator.cs b/Engine/Source/Programs/Shared/EpicGames.UHT/Exporters/CodeGen/UhtPackageCodeGenerator.cs
--- a/Engine/Source/Programs/Shared/EpicGames.UHT/Exporters/CodeGen/UhtPackageCodeGenerator.cs
+++ b/Engine/Source/Programs/Shared/EpicGames.UHT/Exporters/CodeGen/UhtPackageCodeGenerator.cs
@@ -25,6 +25,7 @@ namespace EpicGames.UHT.Exporters.CodeGen
 
 		public readonly UhtCodeGenerator CodeGenerator;
 		public readonly UhtPackage Package;
+		public bool SaveExportedHeaders => this.Package.Module.SaveExportedHeaders;
 
 		public Utils.UhtSession Session => this.CodeGenerator.Session;
 		public UhtCodeGenerator.PackageInfo[] PackageInfos => this.CodeGenerator.PackageInfos;
diff --git a/Engine/Source/Programs/Shared/EpicGames.UHT/Exporters/CodeGen/UhtPackageCodeGeneratorCppFile.cs b/Engine/Source/Programs/Shared/EpicGames.UHT/Exporters/CodeGen/UhtPackageCodeGeneratorCppFile.cs
--- a/Engine/Source/Programs/Shared/EpicGames.UHT/Exporters/CodeGen/UhtPackageCodeGeneratorCppFile.cs
+++ b/Engine/Source/Programs/Shared/EpicGames.UHT/Exporters/CodeGen/UhtPackageCodeGeneratorCppFile.cs
@@ -24,9 +24,9 @@ namespace EpicGames.UHT.Exporters.CodeGen
 		/// <summary>
 		/// For a given UE header file, generated the generated cpp file
 		/// </summary>
-		/// <param name="CppOutput">Output object</param>
+		/// <param name="Task">Requesting task</param>
 		/// <param name="PackageSortedHeaders">Sorted list of headers by name, in the package.</param>
-		public void Generate(IUhtExportOutput CppOutput, List<UhtHeaderFile> PackageSortedHeaders)
+		public void Generate(IUhtExportTask Task, List<UhtHeaderFile> PackageSortedHeaders)
 		{
 			using (BorrowStringBuilder Borrower = new BorrowStringBuilder(StringBuilderCache.Big))
 			{
@@ -140,7 +140,12 @@ namespace EpicGames.UHT.Exporters.CodeGen
 					.Append($"0x{BodiesHash:X8}, 0x{DeclarationsHash:X8}").Append("));\r\n");
 
 				Builder.Append(EnableDeprecationWarnings).Append("\r\n");
-				CppOutput.CommitOutput(Builder);
+
+				if (this.SaveExportedHeaders)
+				{
+					string CppFilePath = Task.Factory.MakePath(this.Package, ".init.gen.cpp");
+					Task.CommitOutput(CppFilePath, Builder);
+				}
 			}
 		}
 	}
diff --git a/Engine/Source/Programs/Shared/EpicGames.UHT/Exporters/CodeGen/UhtPackageCodeGeneratorHFile.cs b/Engine/Source/Programs/Shared/EpicGames.UHT/Exporters/CodeGen/UhtPackageCodeGeneratorHFile.cs
--- a/Engine/Source/Programs/Shared/EpicGames.UHT/Exporters/CodeGen/UhtPackageCodeGeneratorHFile.cs
+++ b/Engine/Source/Programs/Shared/EpicGames.UHT/Exporters/CodeGen/UhtPackageCodeGeneratorHFile.cs
@@ -24,9 +24,9 @@ namespace EpicGames.UHT.Exporters.CodeGen
 		/// <summary>
 		/// For a given UE header file, generated the generated H file
 		/// </summary>
-		/// <param name="HeaderOutput">Output object</param>
+		/// <param name="Task">Requesting task</param>
 		/// <param name="PackageSortedHeaders">Sorted list of headers by name of all headers in the package</param>
-		public void Generate(IUhtExportOutput HeaderOutput, List<UhtHeaderFile> PackageSortedHeaders)
+		public void Generate(IUhtExportTask Task, List<UhtHeaderFile> PackageSortedHeaders)
 		{
 			using (BorrowStringBuilder Borrower = new BorrowStringBuilder(StringBuilderCache.Big))
 			{
@@ -60,7 +60,12 @@ namespace EpicGames.UHT.Exporters.CodeGen
 				}
 
 				Builder.Append("\r\n");
-				HeaderOutput.CommitOutput(Builder);
+
+				if (this.SaveExportedHeaders)
+				{
+					string HeaderFilePath = Task.Factory.MakePath(this.Package, "Classes.h");
+					Task.CommitOutput(HeaderFilePath, Builder);
+				}
 			}
 		}
 	}
diff --git a/Engine/Source/Programs/Shared/EpicGames.UHT/Exporters/Json/UhtJsonExporter.cs b/Engine/Source/Programs/Shared/EpicGames.UHT/Exporters/Json/UhtJsonExporter.cs
--- a/Engine/Source/Programs/Shared/EpicGames.UHT/Exporters/Json/UhtJsonExporter.cs
+++ b/Engine/Source/Programs/Shared/EpicGames.UHT/Exporters/Json/UhtJsonExporter.cs
@@ -22,7 +22,6 @@ namespace EpicGames.UHT.Exporters.Json
 
 		public readonly IUhtExportFactory Factory;
 		public UhtSession Session => Factory.Session;
-		public UhtExportOptions Options => Factory.Options;
 
 		private UhtJsonExporter(IUhtExportFactory Factory)
 		{
@@ -36,15 +35,12 @@ namespace EpicGames.UHT.Exporters.Json
 			foreach (UhtPackage Package in this.Session.Packages)
 			{
 				UHTManifest.Module Module = Package.Module;
-				UhtExportOptions AdditionalOptions = UhtExportOptions.WriteOutput;
-
-				string JsonPath = this.Factory.MakePath(Package, ".json");
-
-				GeneratedPackages.Add(Factory.CreateTask(JsonPath, new List<IUhtExportTask>(), AdditionalOptions,
-					(IUhtExportOutput JsonOutput) =>
+				GeneratedPackages.Add(Factory.CreateTask(
+					(IUhtExportTask Task) =>
 					{
+						string JsonPath = Task.Factory.MakePath(Package, ".json");
 						JsonSerializerOptions Options = new JsonSerializerOptions { WriteIndented = true, IgnoreNullValues = true };
-						JsonOutput.CommitOutput(JsonSerializer.Serialize(Package, Options));
+						Task.CommitOutput(JsonPath, JsonSerializer.Serialize(Package, Options));
 					}));
 			}
 
diff --git a/Engine/Source/Programs/Shared/EpicGames.UHT/Types/Properties/UhtObjectPtrProperty.cs b/Engine/Source/Programs/Shared/EpicGames.UHT/Types/Properties/UhtObjectPtrProperty.cs
--- a/Engine/Source/Programs/Shared/EpicGames.UHT/Types/Properties/UhtObjectPtrProperty.cs
+++ b/Engine/Source/Programs/Shared/EpicGames.UHT/Types/Properties/UhtObjectPtrProperty.cs
@@ -46,7 +46,14 @@ namespace EpicGames.UHT.Types
 			switch (TextType)
 			{
 				case UhtPropertyTextType.GetterSetterArg:
-					Builder.Append(this.Class.SourceName).Append("*");
+					if (bIsTemplateArgument)
+					{
+						Builder.Append("TObjectPtr<").Append(this.Class.SourceName).Append(">");
+					}
+					else
+					{
+						Builder.Append(this.Class.SourceName).Append("*");
+					}
 					break;
 
 				case UhtPropertyTextType.FunctionThunkReturn:
diff --git a/Engine/Source/Programs/Shared/EpicGames.UHT/Utils/UhtExport.cs b/Engine/Source/Programs/Shared/EpicGames.UHT/Utils/UhtExport.cs
--- a/Engine/Source/Programs/Shared/EpicGames.UHT/Utils/UhtExport.cs
+++ b/Engine/Source/Programs/Shared/EpicGames.UHT/Utils/UhtExport.cs
@@ -11,130 +11,47 @@ namespace EpicGames.UHT.Utils
 {
 
 	/// <summary>
-	/// Export options.
-	/// </summary>
-	[Flags]
-	public enum UhtExportOptions
-	{
-
-		/// <summary>
-		/// No export options specified
-		/// </summary>
-		None = 0,
-
-		/// <summary>
-		/// If set, output is to be written
-		/// </summary>
-		WriteOutput = 1 << 0,
-	}
-
-	/// <summary>
-	/// Helper methods for testing flags.  These methods perform better than the generic HasFlag which hits
-	/// the GC and stalls.
+	/// Represents a task to export outputs
 	/// </summary>
-	public static class UhtExportOptionsExtensions
+	public interface IUhtExportTask
 	{
-
-		/// <summary>
-		/// Test to see if any of the specified flags are set
-		/// </summary>
-		/// <param name="InFlags">Current flags</param>
-		/// <param name="TestFlags">Flags to test for</param>
-		/// <returns>True if any of the flags are set</returns>
-		public static bool HasAnyFlags(this UhtExportOptions InFlags, UhtExportOptions TestFlags)
-		{
-			return (InFlags & TestFlags) != 0;
-		}
-
 		/// <summary>
-		/// Test to see if all of the specified flags are set
+		/// Factory associated with the task
 		/// </summary>
-		/// <param name="InFlags">Current flags</param>
-		/// <param name="TestFlags">Flags to test for</param>
-		/// <returns>True if all the flags are set</returns>
-		public static bool HasAllFlags(this UhtExportOptions InFlags, UhtExportOptions TestFlags)
-		{
-			return (InFlags & TestFlags) == TestFlags;
-		}
+		public IUhtExportFactory Factory { get; }
 
 		/// <summary>
-		/// Test to see if a specific set of flags have a specific value.
-		/// </summary>
-		/// <param name="InFlags">Current flags</param>
-		/// <param name="TestFlags">Flags to test for</param>
-		/// <param name="MatchFlags">Expected value of the tested flags</param>
-		/// <returns>True if the given flags have a specific value.</returns>
-		public static bool HasExactFlags(this UhtExportOptions InFlags, UhtExportOptions TestFlags, UhtExportOptions MatchFlags)
-		{
-			return (InFlags & TestFlags) == MatchFlags;
-		}
-	}
-
-	/// <summary>
-	/// Represents an output file
-	/// </summary>
-	public interface IUhtExportOutput
-	{
-
-		/// <summary>
-		/// Session associated with the output.
+		/// Session being run
 		/// </summary>
 		public UhtSession Session { get; }
 
 		/// <summary>
-		/// Export options
+		/// Task handle associated with the task for it can be waited on
 		/// </summary>
-		public UhtExportOptions Options { get; set; }
+		public Task? ActionTask { get; }
 
 		/// <summary>
 		/// Commit the contents of the string builder as the output.
+		/// If you have a string builder, use this method so that a 
+		/// temporary buffer can be used.
 		/// </summary>
+		/// <param name="FilePath">Destination file path</param>
 		/// <param name="Builder">Source for the content</param>
-		/// <returns>A string view of the committed output.  Use this instead of the 
-		/// StringBuilder so the builder can be recycled.</returns>
-		public StringView CommitOutput(StringBuilder Builder);
+		public void CommitOutput(string FilePath, StringBuilder Builder);
 
 		/// <summary>
 		/// Commit the value of the string as the output
 		/// </summary>
+		/// <param name="FilePath">Destination file path</param>
 		/// <param name="Output">Output to commit</param>
-		public void CommitOutput(string Output);
+		public void CommitOutput(string FilePath, StringView Output);
 	}
 
 	/// <summary>
-	/// Represents a task to export outputs
+	/// Delegate invoked by a task
 	/// </summary>
-	public interface IUhtExportTask
-	{
-
-		/// <summary>
-		/// Session being run
-		/// </summary>
-		public UhtSession Session { get; }
-
-		/// <summary>
-		/// Export options
-		/// </summary>
-		public UhtExportOptions Options { get; }
-
-		/// <summary>
-		/// Task handle associated with the task for it can be waited on
-		/// </summary>
-		public Task? ActionTask { get; }
-	}
-
-	/// <summary>
-	/// Delegate used to generate a single output from a task
-	/// </summary>
-	/// <param name="Output">The single output</param>
-	public delegate void UhtExportSingleDelegate(IUhtExportOutput Output);
-
-	/// <summary>
-	/// Delegate used to generate two outputs from a task.
-	/// </summary>
-	/// <param name="HeaderOutput">Header output</param>
-	/// <param name="CppOutput">CPP output</param>
-	public delegate void UhtExportPairDelegate(IUhtExportOutput HeaderOutput, IUhtExportOutput CppOutput);
+	/// <param name="Task">Invoking task</param>
+	public delegate void UhtExportTaskDelegate(IUhtExportTask Task);
 
 	/// <summary>
 	/// Factory object used to generate export tasks
@@ -148,30 +65,19 @@ namespace EpicGames.UHT.Utils
 		public UhtSession Session { get; }
 
 		/// <summary>
-		/// Export options
-		/// </summary>
-		public UhtExportOptions Options { get; }
-
-		/// <summary>
-		/// Create a task to export a single file
+		/// Create a task
 		/// </summary>
-		/// <param name="Path">Destination file path</param>
 		/// <param name="Prereqs">Tasks that must be completed prior to this task running</param>
-		/// <param name="AdditionalOptions">Additional options</param>
-		/// <param name="Action">Action to be invoked to generate the output</param>
+		/// <param name="Action">Action to be invoked to generate the output(s)</param>
 		/// <returns>Task interface.</returns>
-		public IUhtExportTask CreateTask(string Path, List<IUhtExportTask> Prereqs, UhtExportOptions AdditionalOptions, UhtExportSingleDelegate Action);
+		public IUhtExportTask CreateTask(List<IUhtExportTask>? Prereqs, UhtExportTaskDelegate Action);
 
 		/// <summary>
-		/// Create a task to export two files
+		/// Create a task
 		/// </summary>
-		/// <param name="HeaderPath">Header destination file path</param>
-		/// <param name="CppPath">Cpp destination file path</param>
-		/// <param name="Prereqs">Tasks that must be completed prior to this task running</param>
-		/// <param name="AdditionalOptions">Additional options</param>
-		/// <param name="Action">Action to be invoked to generate the output</param>
+		/// <param name="Action">Action to be invoked to generate the output(s)</param>
 		/// <returns>Task interface.</returns>
-		public IUhtExportTask CreateTask(string HeaderPath, string CppPath, List<IUhtExportTask> Prereqs, UhtExportOptions AdditionalOptions, UhtExportPairDelegate Action);
+		public IUhtExportTask CreateTask(UhtExportTaskDelegate Action);
 
 		/// <summary>
 		/// Make a path for an output based on the header file name.
@@ -188,5 +94,14 @@ namespace EpicGames.UHT.Utils
 		/// <param name="Suffix">Suffix to be added to the file name.</param>
 		/// <returns>Output file path</returns>
 		public string MakePath(UhtPackage Package, string Suffix);
+
+		/// <summary>
+		/// Make a path for an output based on the package output directory.
+		/// </summary>
+		/// <param name="Package">Destination package</param>
+		/// <param name="FileName">Name of the file</param>
+		/// <param name="Extension">Extension to add to the file</param>
+		/// <returns>Output file path</returns>
+		public string MakePath(UhtPackage Package, string FileName, string Extension);
 	}
 }
diff --git a/Engine/Source/Programs/Shared/EpicGames.UHT/Utils/UhtSession.cs b/Engine/Source/Programs/Shared/EpicGames.UHT/Utils/UhtSession.cs
--- a/Engine/Source/Programs/Shared/EpicGames.UHT/Utils/UhtSession.cs
+++ b/Engine/Source/Programs/Shared/EpicGames.UHT/Utils/UhtSession.cs
@@ -48,180 +48,99 @@ namespace EpicGames.UHT.Utils
 	}
 
 	/// <summary>
-	/// Implementation of the export output interface
+	/// Base implementation of an export task
 	/// </summary>
-	class UhtExportOutput : IUhtExportOutput
+	class UhtExportTask : IUhtExportTask
 	{
 		/// <summary>
-		/// Export task that created the export output
-		/// </summary>
-		private readonly IUhtExportTask ExportTask;
-
-		/// <summary>
-		/// Export options for the output
-		/// </summary>
-		private UhtExportOptions OptionsInternal;
-
-		/// <summary>
-		/// Output buffer associated with the exported output
-		/// </summary>
-		private UhtBuffer? ExportedBuffer = null;
-
-		/// <summary>
-		/// Destination file path of the output
-		/// </summary>
-		public readonly string FilePath;
-
-		/// <summary>
-		/// Destination file path of the temporary copy of the output
+		/// Action to invoke to generate the output
 		/// </summary>
-		public readonly string TempFilePath;
+		private readonly UhtExportTaskDelegate Action;
 
 		/// <summary>
-		/// Exported data
+		/// Factory requesting the export
 		/// </summary>
-		public StringView Exported;
+		protected readonly UhtExportFactory ExportFactory;
 
 		/// <summary>
-		/// True if the exported data was committed
+		/// Action task created to invoke the task
 		/// </summary>
-		public bool bCommitted = false;
+		private Task? ActionTaskInternal = null;
 
 		/// <summary>
-		/// True if the exported data saves (i.e. it doesn't match the existing file)
+		/// Factory associated with the task
 		/// </summary>
-		public bool bSaved = false;
+		public IUhtExportFactory Factory => this.ExportFactory;
 
 		/// <summary>
 		/// UHT session
 		/// </summary>
-		public UhtSession Session => ExportTask.Session;
-
-		/// <summary>
-		/// Options associated with the output
-		/// </summary>
-		public UhtExportOptions Options
-		{
-			get => this.OptionsInternal;
-			set => this.OptionsInternal = value;
-		}
-
-		/// <summary>
-		/// Construct a new instance of an export output
-		/// </summary>
-		/// <param name="ExportTask">Task associated with the output</param>
-		/// <param name="FilePath">Destination file path</param>
-		public UhtExportOutput(IUhtExportTask ExportTask, string FilePath)
-		{
-			this.ExportTask = ExportTask;
-			this.Options = ExportTask.Options;
-			this.FilePath = FilePath;
-			this.TempFilePath = FilePath + ".tmp";
-		}
+		public UhtSession Session => this.ExportFactory.Session;
 
 		/// <summary>
-		/// Commit the output using the given string builder
+		/// Export options associated with the task
 		/// </summary>
-		/// <param name="Builder">Builder containing the output.</param>
-		/// <returns>String view containing the committed data in the output buffer. 
-		/// While supported, this view should not be saved for any significant length of time.</returns>
-
-		public StringView CommitOutput(StringBuilder Builder)
-		{
-			this.bCommitted = true;
-			this.ExportedBuffer = UhtBuffer.Borrow(Builder);
-			this.Exported = new StringView(this.ExportedBuffer.Memory);
-			return this.Exported;
-		}
+		public Task? ActionTask => this.ActionTaskInternal;
 
 		/// <summary>
-		/// Save the given committed output
+		/// Construct a new instance of an export task
 		/// </summary>
-		/// <param name="Output">The output to save</param>
-		public void CommitOutput(string Output)
+		/// <param name="ExportFactory">Factory requesting the task</param>
+		/// <param name="Action">Action to be invoked as part of the export</param>
+		public UhtExportTask(UhtExportFactory ExportFactory, UhtExportTaskDelegate Action)
 		{
-			this.bCommitted = true;
-			this.Exported = Output;
+			this.ExportFactory = ExportFactory;
+			this.Action = Action;
 		}
 
 		/// <summary>
-		/// Reset the contents of the output task so the memory can be freed and output buffers returned to the cache.
+		/// Commit the contents of the string builder as the output.
+		/// If you have a string builder, use this method so that a 
+		/// temporary buffer can be used.
 		/// </summary>
-		public void Reset()
+		/// <param name="FilePath">Destination file path</param>
+		/// <param name="Builder">Source for the content</param>
+		public void CommitOutput(string FilePath, StringBuilder Builder)
 		{
-			if (this.ExportedBuffer != null)
+			using (UhtBorrowBuffer BorrowBuffer = new UhtBorrowBuffer(Builder))
 			{
-				UhtBuffer.Return(this.ExportedBuffer);
-				this.ExportedBuffer = null;
+				string TempFilePath = FilePath + ".tmp";
+				this.ExportFactory.SaveIfChanged(FilePath, TempFilePath, new StringView(BorrowBuffer.Buffer.Memory));
 			}
-			this.Exported = new StringView();
 		}
-	}
-
-	/// <summary>
-	/// Base implementation of an export task
-	/// </summary>
-	abstract class UhtExportBaseTask : IUhtExportTask
-	{
-
-		/// <summary>
-		/// Factory requesting the export
-		/// </summary>
-		protected readonly UhtExportFactory ExportFactory;
 
 		/// <summary>
-		/// Action task created to invoke the task
-		/// </summary>
-		private Task? ActionTaskInternal = null;
-
-		/// <summary>
-		/// Export options associated with the task
+		/// Commit the value of the string as the output
 		/// </summary>
-		private readonly UhtExportOptions OptionsInternal;
-
-		/// <summary>
-		/// UHT session
-		/// </summary>
-		public UhtSession Session => this.ExportFactory.Session;
-
-		/// <summary>
-		/// Export options associated with the task
-		/// </summary>
-		public UhtExportOptions Options => this.OptionsInternal;
-
-		/// <summary>
-		/// Export options associated with the task
-		/// </summary>
-		public Task? ActionTask => this.ActionTaskInternal;
-
-		/// <summary>
-		/// Construct a new instance of an export task
-		/// </summary>
-		/// <param name="ExportFactory">Factory requesting the task</param>
-		/// <param name="AdditionalOptions">Addition options to added to the existing factory options</param>
-		public UhtExportBaseTask(UhtExportFactory ExportFactory, UhtExportOptions AdditionalOptions)
+		/// <param name="FilePath">Destination file path</param>
+		/// <param name="Output">Output to commit</param>
+		public void CommitOutput(string FilePath, StringView Output)
 		{
-			this.ExportFactory = ExportFactory;
-			this.OptionsInternal = this.ExportFactory.Options | AdditionalOptions;
+			string TempFilePath = FilePath + ".tmp";
+			this.ExportFactory.SaveIfChanged(FilePath, TempFilePath, Output);
 		}
 
 		/// <summary>
 		/// Queue the task
 		/// </summary>
 		/// <param name="Prereqs">List of prerequisite tasks that must complete prior to this task being invoked.</param>
-		public void Queue(List<IUhtExportTask> Prereqs)
+		public void Queue(List<IUhtExportTask>? Prereqs)
 		{
 			if (this.Session.bGoWide)
 			{
-				List<Task> PrereqTasks = new List<Task>();
-				foreach (IUhtExportTask Prereq in Prereqs)
+				List<Task>? PrereqTasks = null;
+				if (Prereqs != null)
 				{
-					if (Prereq.ActionTask != null)
+					PrereqTasks = new List<Task>();
+					foreach (IUhtExportTask Prereq in Prereqs)
 					{
-						PrereqTasks.Add(Prereq.ActionTask);
+						if (Prereq.ActionTask != null)
+						{
+							PrereqTasks.Add(Prereq.ActionTask);
+						}
 					}
 				}
-				if (PrereqTasks.Count > 0)
+				if (PrereqTasks != null && PrereqTasks.Count > 0)
 				{
 					this.ActionTaskInternal = Task.Factory.ContinueWhenAll(PrereqTasks.ToArray(), (Task[] Tasks) => { Export(); });
 				}
@@ -239,91 +158,9 @@ namespace EpicGames.UHT.Utils
 		/// <summary>
 		/// Invoked to perform the actual export
 		/// </summary>
-		public abstract void Export();
-	}
-
-	/// <summary>
-	/// Export task that generates a single output
-	/// </summary>
-	class UhtExportSingleTask : UhtExportBaseTask
-	{
-		/// <summary>
-		/// Single output
-		/// </summary>
-		public readonly UhtExportOutput File;
-
-		/// <summary>
-		/// Action to invoke to generate the output
-		/// </summary>
-		private readonly UhtExportSingleDelegate Action;
-
-		/// <summary>
-		/// Create a new instance of a single output task
-		/// </summary>
-		/// <param name="ExportFactory">Requesting factory</param>
-		/// <param name="Path">Destination path of the output</param>
-		/// <param name="AdditionalOptions">Additional options for the output</param>
-		/// <param name="Action">Action to be invoked to generate the output</param>
-		public UhtExportSingleTask(UhtExportFactory ExportFactory, string Path, UhtExportOptions AdditionalOptions, UhtExportSingleDelegate Action) : base(ExportFactory, AdditionalOptions)
-		{
-			this.File = new UhtExportOutput(this, Path);
-			this.Action = Action;
-		}
-
-		/// <summary>
-		/// Invoke the actions and save the output
-		/// </summary>
-		public override void Export()
-		{
-			this.Action(this.File);
-			this.ExportFactory.SaveIfChanged(this.File);
-		}
-	}
-
-	/// <summary>
-	/// Export task that generates a pair of outputs
-	/// </summary>
-	class UhtExportPairTask : UhtExportBaseTask
-	{
-
-		/// <summary>
-		/// First/Header file being generated
-		/// </summary>
-		private readonly UhtExportOutput HeaderFile;
-
-		/// <summary>
-		/// Second/Cpp file being generated
-		/// </summary>
-		private readonly UhtExportOutput CppFile;
-
-		/// <summary>
-		/// Action to be invoked to generate the output
-		/// </summary>
-		private readonly UhtExportPairDelegate Action;
-
-		/// <summary>
-		/// Construct a new instance of the pair output task
-		/// </summary>
-		/// <param name="ExportFactory">Requesting factory</param>
-		/// <param name="HeaderPath">Path of the header file</param>
-		/// <param name="CppPath">Path of the cpp file</param>
-		/// <param name="AdditionalOptions">Additional options to be applied to the outputs</param>
-		/// <param name="Action">Action to invoke to generate the output</param>
-		public UhtExportPairTask(UhtExportFactory ExportFactory, string HeaderPath, string CppPath, UhtExportOptions AdditionalOptions, UhtExportPairDelegate Action) : base(ExportFactory, AdditionalOptions)
+		public void Export()
 		{
-			this.HeaderFile = new UhtExportOutput(this, HeaderPath);
-			this.CppFile = new UhtExportOutput(this, CppPath);
-			this.Action = Action;
-		}
-
-		/// <summary>
-		/// Invoked to export and save the outputs
-		/// </summary>
-		public override void Export()
-		{
-			this.Action(this.HeaderFile, this.CppFile);
-			this.ExportFactory.SaveIfChanged(this.HeaderFile);
-			this.ExportFactory.SaveIfChanged(this.CppFile);
+			this.Action(this);
 		}
 	}
 
@@ -332,17 +169,18 @@ namespace EpicGames.UHT.Utils
 	/// </summary>
 	class UhtExportFactory : IUhtExportFactory
 	{
+		public struct Output
+		{
+			public string FilePath;
+			public string TempFilePath;
+			public bool bSaved;
+		}
 
 		/// <summary>
 		/// UHT session
 		/// </summary>
 		private readonly UhtSession SessionInternal;
 
-		/// <summary>
-		/// Export options 
-		/// </summary>
-		private readonly UhtExportOptions OptionsInternal;
-
 		/// <summary>
 		/// Limiter for the number of files being saved to the reference directory.
 		/// The OS can get swamped on high core systems
@@ -359,11 +197,6 @@ namespace EpicGames.UHT.Utils
 		/// </summary>
 		public UhtSession Session => this.SessionInternal;
 
-		/// <summary>
-		/// Export options 
-		/// </summary>
-		public UhtExportOptions Options => this.OptionsInternal;
-		
 		/// <summary>
 		/// Collection of error from mismatches with the reference files
 		/// </summary>
@@ -372,7 +205,7 @@ namespace EpicGames.UHT.Utils
 		/// <summary>
 		/// List of export outputs
 		/// </summary>
-		public List<UhtExportOutput> Outputs = new List<UhtExportOutput>();
+		public List<Output> Outputs = new List<Output>();
 
 		/// <summary>
 		/// Directory for the reference output
@@ -389,12 +222,10 @@ namespace EpicGames.UHT.Utils
 		/// </summary>
 		/// <param name="Session">UHT session</param>
 		/// <param name="Exporter">Exporter being run</param>
-		/// <param name="Options">Export options</param>
-		public UhtExportFactory(UhtSession Session, UhtExporter Exporter, UhtExportOptions Options)
+		public UhtExportFactory(UhtSession Session, UhtExporter Exporter)
 		{
 			this.Exporter = Exporter;
 			this.SessionInternal = Session;
-			this.OptionsInternal = Options;
 			if (this.Session.ReferenceMode != UhtReferenceMode.None)
 			{
 				this.ReferenceDirectory = Path.Combine(this.Session.ReferenceDirectory, this.Exporter.Name);
@@ -404,34 +235,26 @@ namespace EpicGames.UHT.Utils
 		}
 
 		/// <summary>
-		/// Create an export task that write one output
+		/// Create a task to export two files
 		/// </summary>
-		/// <param name="Path">Destination file path</param>
-		/// <param name="Prereqs">List of required export tasks that must be completed prior to this export</param>
-		/// <param name="AdditionalOptions">Additional export options for the output</param>
-		/// <param name="Action">Action to be invoked to generated the output</param>
-		/// <returns>Created task</returns>
-		public IUhtExportTask CreateTask(string Path, List<IUhtExportTask> Prereqs, UhtExportOptions AdditionalOptions, UhtExportSingleDelegate Action)
+		/// <param name="Prereqs">Tasks that must be completed prior to this task running</param>
+		/// <param name="Action">Action to be invoked to generate the output</param>
+		/// <returns>Task interface.</returns>
+		public IUhtExportTask CreateTask(List<IUhtExportTask>? Prereqs, UhtExportTaskDelegate Action)
 		{
-			UhtExportBaseTask Task = new UhtExportSingleTask(this, Path, AdditionalOptions, Action);
+			UhtExportTask Task = new UhtExportTask(this, Action);
 			Task.Queue(Prereqs);
 			return Task;
 		}
 
 		/// <summary>
-		/// Create an export task that write one output
+		/// Create a task to export two files
 		/// </summary>
-		/// <param name="HeaderPath">Destination header file path</param>
-		/// <param name="CppPath">Destination cpp file path</param>
-		/// <param name="Prereqs">List of required export tasks that must be completed prior to this export</param>
-		/// <param name="AdditionalOptions">Additional export options for the output</param>
-		/// <param name="Action">Action to be invoked to generated the output</param>
-		/// <returns>Created task</returns>
-		public IUhtExportTask CreateTask(string HeaderPath, string CppPath, List<IUhtExportTask> Prereqs, UhtExportOptions AdditionalOptions, UhtExportPairDelegate Action)
+		/// <param name="Action">Action to be invoked to generate the output</param>
+		/// <returns>Task interface.</returns>
+		public IUhtExportTask CreateTask(UhtExportTaskDelegate Action)
 		{
-			UhtExportBaseTask Task = new UhtExportPairTask(this, HeaderPath, CppPath, AdditionalOptions, Action);
-			Task.Queue(Prereqs);
-			return Task;
+			return CreateTask(null, Action);
 		}
 
 		/// <summary>
@@ -458,30 +281,34 @@ namespace EpicGames.UHT.Utils
 			return Path.Combine(Module.OutputDirectory, Package.ShortName.ToString()) + Suffix;
 		}
 
+
 		/// <summary>
-		/// Helper method to test to see if the output has changed.
+		/// Make a path for an output based on the package output directory.
 		/// </summary>
-		/// <param name="Output">The output file</param>
-		internal void SaveIfChanged(UhtExportOutput Output)
+		/// <param name="Package">Destination package</param>
+		/// <param name="FileName">Name of the file</param>
+		/// <param name="Extension">Extension to add to the file</param>
+		/// <returns>Output file path</returns>
+		public string MakePath(UhtPackage Package, string FileName, string Extension)
 		{
+			UHTManifest.Module Module = Package.Module;
+			return Path.Combine(Module.OutputDirectory, FileName) + Extension;
+		}
 
-			if (!Output.Options.HasAnyFlags(UhtExportOptions.WriteOutput) || !Output.bCommitted)
-			{
-				return;
-			}
-
-			// Add this to the list of outputs
-			lock (this.Outputs)
-			{
-				this.Outputs.Add(Output);
-			}
+		/// <summary>
+		/// Helper method to test to see if the output has changed.
+		/// </summary>
+		/// <param name="FilePath">Name of the output file</param>
+		/// <param name="TempFilePath">Name of the temporary file</param>
+		/// <param name="Exported">Exported contents of the file</param>
+		internal void SaveIfChanged(string FilePath, string TempFilePath, StringView Exported)
+		{
 
-			StringView Exported = Output.Exported;
 			ReadOnlySpan<char> ExportedSpan = Exported.Span;
 
 			if (this.Session.ReferenceMode != UhtReferenceMode.None)
 			{
-				string FileName = Path.GetFileName(Output.FilePath);
+				string FileName = Path.GetFileName(FilePath);
 
 				// Writing billions of files to the same directory causes issues.  Use ourselves to throttle reference writes
 				try
@@ -532,7 +359,7 @@ namespace EpicGames.UHT.Utils
 			}
 
 			// Check to see if the contents have changed
-			UhtBuffer? Original = this.Session.ReadSourceToBuffer(Output.FilePath);
+			UhtBuffer? Original = this.Session.ReadSourceToBuffer(FilePath);
 			bool bSave = Original == null;
 			if (Original != null)
 			{
@@ -541,10 +368,10 @@ namespace EpicGames.UHT.Utils
 				{
 					if (this.Session.bFailIfGeneratedCodeChanges)
 					{
-						string ConflictPath = Output.FilePath + ".conflict";
+						string ConflictPath = FilePath + ".conflict";
 						if (!this.Session.WriteSource(ConflictPath, Exported.Span))
 						{
-							new UhtSimpleFileMessageSite(this.Session, Output.FilePath).LogError($"Changes to generated code are not allowed - conflicts written to '{ConflictPath}'");
+							new UhtSimpleFileMessageSite(this.Session, FilePath).LogError($"Changes to generated code are not allowed - conflicts written to '{ConflictPath}'");
 						}
 					}
 					bSave = true;
@@ -555,17 +382,23 @@ namespace EpicGames.UHT.Utils
 			// If changed of the original didn't exist, then save the new version
 			if (bSave && !this.Session.bNoOutput)
 			{
-				Output.bSaved = true;
-				if (!this.Session.WriteSource(Output.TempFilePath, Exported.Span))
+				if (!this.Session.WriteSource(TempFilePath, Exported.Span))
 				{
-					new UhtSimpleFileMessageSite(this.Session, Output.FilePath).LogWarning($"Failed to save export file: '{Output.TempFilePath}'");
+					new UhtSimpleFileMessageSite(this.Session, FilePath).LogWarning($"Failed to save export file: '{TempFilePath}'");
 				}
 			}
+			else
+			{
+				bSave = false;
+			}
 
-			// Reset the output to try and free up some of the temporary buffers
-			Output.Reset();
-		}
-
+			// Add this to the list of outputs
+			lock (this.Outputs)
+			{
+				this.Outputs.Add(new Output { FilePath = FilePath, TempFilePath = TempFilePath, bSaved = bSave });
+			}
+		}       
+		
 		/// <summary>
 		/// Run the output exporter
 		/// </summary>
@@ -581,10 +414,10 @@ namespace EpicGames.UHT.Utils
 
 				// These outputs are used to cull old outputs from the directories
 				Dictionary<string, HashSet<string>> OutputsByDirectory = new Dictionary<string, HashSet<string>>(StringComparer.OrdinalIgnoreCase);
-				List<UhtExportOutput> Saves = new List<UhtExportOutput>();
+				List<UhtExportFactory.Output> Saves = new List<UhtExportFactory.Output>();
 
 				// Collect information about the outputs
-				foreach (UhtExportOutput Output in this.Outputs)
+				foreach (UhtExportFactory.Output Output in this.Outputs)
 				{
 
 					// Add this output to the list of expected outputs by directory
@@ -610,14 +443,14 @@ namespace EpicGames.UHT.Utils
 				// Perform the renames
 				if (this.Session.bGoWide)
 				{
-					Parallel.ForEach(Saves, (UhtExportOutput Output) =>
+					Parallel.ForEach(Saves, (UhtExportFactory.Output Output) =>
 					{
 						RenameSource(Output);
 					});
 				}
 				else
 				{ 
-					foreach (UhtExportOutput Output in Saves)
+					foreach (UhtExportFactory.Output Output in Saves)
 					{
 						RenameSource(Output);
 					}
@@ -650,7 +483,7 @@ namespace EpicGames.UHT.Utils
 		/// If there exists a current final file, it will be replaced.
 		/// </summary>
 		/// <param name="Output">The output file to rename</param>
-		private void RenameSource(UhtExportOutput Output)
+		private void RenameSource(UhtExportFactory.Output Output)
 		{
 			this.Session.RenameSource(Output.TempFilePath, Output.FilePath);
 		}
@@ -787,6 +620,11 @@ namespace EpicGames.UHT.Utils
 		/// </summary>
 		public string? EngineDirectory;
 
+		/// <summary>
+		/// If set, the name of the project file.
+		/// </summary>
+		public string? ProjectFile;
+
 		/// <summary>
 		/// Optional location of the project
 		/// </summary>
@@ -997,6 +835,7 @@ namespace EpicGames.UHT.Utils
 		private Dictionary<string, bool> ExporterStates = new Dictionary<string, bool>(StringComparer.OrdinalIgnoreCase);
 		private Dictionary<string, EnumAndValue> FullEnumValueLookup = new Dictionary<string, EnumAndValue>();
 		private Dictionary<string, UhtEnum> ShortEnumValueLookup = new Dictionary<string, UhtEnum>();
+		private JsonDocument? ProjectJson = null;
 
 		/// <summary>
 		/// The number of errors
@@ -1094,16 +933,6 @@ namespace EpicGames.UHT.Utils
 		/// </summary>
 		public UhtExporterTable ExporterTable => this.Tables!.ExporterTable;
 
-		/// <summary>
-		/// Enable/Disable an exporter.  This overrides the default state of the exporter.
-		/// </summary>
-		/// <param name="Name">Name of the exporter</param>
-		/// <param name="Enabled">If true, the exporter is to be enabled</param>
-		public void SetExporterStatus(string Name, bool Enabled)
-		{
-			this.ExporterStates[Name] = Enabled;
-		}
-
 		/// <summary>
 		/// Return the keyword table for the given table name
 		/// </summary>
@@ -1366,7 +1195,7 @@ namespace EpicGames.UHT.Utils
 		{
 			if (!this.FileManager!.RenameOutput(OldFilePath, NewFilePath))
 			{
-				new UhtSimpleFileMessageSite(this, NewFilePath).LogError($"Failed to rename export file: '{OldFilePath}'");
+				new UhtSimpleFileMessageSite(this, NewFilePath).LogError($"Failed to rename exported file: '{OldFilePath}'");
 			}
 		}
 
@@ -1968,6 +1797,11 @@ namespace EpicGames.UHT.Utils
 				Log.Logger.LogTrace("Step - Read Manifest File");
 
 				this.ManifestFile.Read();
+
+				if (this.Manifest != null && this.Tables != null)
+				{
+					this.Tables.AddPlugins(this.Manifest.UhtPlugins);
+				}
 			});
 		}
 
@@ -2478,13 +2312,82 @@ namespace EpicGames.UHT.Utils
 		#endregion
 
 		#region Exporting
+
+		/// <summary>
+		/// Enable/Disable an exporter.  This overrides the default state of the exporter.
+		/// </summary>
+		/// <param name="Name">Name of the exporter</param>
+		/// <param name="Enabled">If true, the exporter is to be enabled</param>
+		public void SetExporterStatus(string Name, bool Enabled)
+		{
+			this.ExporterStates[Name] = Enabled;
+		}
+
+		/// <summary>
+		/// Test to see if the given exporter plugin is enabled.
+		/// </summary>
+		/// <param name="PluginName">Name of the plugin</param>
+		/// <param name="IncludeTargetCheck">If true, include a target check</param>
+		/// <returns>True if enabled</returns>
+		public bool IsPluginEnabled(string PluginName, bool IncludeTargetCheck)
+		{
+			if (this.ProjectJson == null && this.ProjectDirectory != null && this.ProjectFile != null)
+			{
+				UhtBuffer? Contents = this.ReadSourceToBuffer(this.ProjectFile);
+				if (Contents != null)
+				{
+					this.ProjectJson = JsonDocument.Parse(Contents.Block);
+					UhtBuffer.Return(Contents);
+				}
+			}
+
+			if (this.ProjectJson == null)
+			{
+				return false;
+			}
+
+			JsonObject RootObject = new JsonObject(this.ProjectJson.RootElement);
+			if (RootObject.TryGetObjectArrayField("Plugins", out JsonObject[]? Plugins))
+			{
+				foreach (JsonObject Plugin in Plugins)
+				{
+					if (!Plugin.TryGetStringField("Name", out string? TestPluginName) || string.Compare(PluginName, TestPluginName, StringComparison.OrdinalIgnoreCase) != 0)
+					{
+						continue;
+					}
+					if (!Plugin.TryGetBoolField("Enabled", out bool Enabled) || !Enabled)
+					{
+						return false;
+					}
+					if (IncludeTargetCheck && this.Manifest != null)
+					{
+						if (Plugin.TryGetStringArrayField("TargetAllowList", out string[]? AllowList))
+						{
+							if (AllowList.Contains(this.Manifest.TargetName, StringComparer.OrdinalIgnoreCase))
+							{
+								return true;
+							}
+						}
+						if (Plugin.TryGetStringArrayField("TargetDenyList", out string[]? DenyList))
+						{
+							if (DenyList.Contains(this.Manifest.TargetName, StringComparer.OrdinalIgnoreCase))
+							{
+								return false;
+							}
+						}
+					}
+					return true;
+				}
+			}
+			return false;
+		}
+
 		private void StepExport()
 		{
 			long TotalWrittenFiles = 0;
 			Try(null, () =>
 			{
 				Log.Logger.LogTrace("Step - Exports");
-				UhtExportOptions Options = UhtExportOptions.None;
 
 				foreach (UhtExporter Exporter in this.ExporterTable)
 				{
@@ -2498,7 +2401,7 @@ namespace EpicGames.UHT.Utils
 					if (Run)
 					{
 						Log.Logger.LogTrace($"       Running exporter {Exporter.Name}");
-						UhtExportFactory Factory = new UhtExportFactory(this, Exporter, Options);
+						UhtExportFactory Factory = new UhtExportFactory(this, Exporter);
 						Factory.Run();
 						foreach (var Output in Factory.Outputs)
 						{
diff --git a/Engine/Source/Programs/Shared/EpicGames.UHT/Utils/UhtTables.cs b/Engine/Source/Programs/Shared/EpicGames.UHT/Utils/UhtTables.cs
--- a/Engine/Source/Programs/Shared/EpicGames.UHT/Utils/UhtTables.cs
+++ b/Engine/Source/Programs/Shared/EpicGames.UHT/Utils/UhtTables.cs
@@ -507,15 +507,55 @@ namespace EpicGames.UHT.Utils
 		{
 			foreach (string AssemblyFilePath in PluginAssembliesFilePaths)
 			{
-				Assembly? Plugin = FindAssemblyByName(Path.GetFileNameWithoutExtension(AssemblyFilePath));
-				if (Plugin == null)
+				CheckForAttributes(LoadAssembly(AssemblyFilePath));
+			}
+			PerformPostInitialization();
+		}
+
+		/// <summary>
+		/// Check to see if the assembly is a UHT plugin
+		/// </summary>
+		/// <param name="AssemblyFilePath">Path to the assembly file</param>
+		/// <returns></returns>
+		public static bool IsUhtPlugin(string AssemblyFilePath)
+		{
+			Assembly? Assembly = LoadAssembly(AssemblyFilePath);       
+			if (Assembly != null)
+			{
+				foreach (Type Type in Assembly.SafeGetLoadedTypes())
+				{
+					if (Type.IsClass)
+					{
+						foreach (Attribute ClassAttribute in Type.GetCustomAttributes(false))
+						{
+							if (ClassAttribute is UnrealHeaderToolAttribute || ClassAttribute is UhtEngineClassAttribute)
+							{
+								return true;
+							}
+						}
+					}
+				}
+			}
+			return false;
+		}
+
+		/// <summary>
+		/// Load the given assembly from the file path
+		/// </summary>
+		/// <param name="AssemblyFilePath">Path to the file</param>
+		/// <returns>Assembly if it is already loaded or could be loaded</returns>
+		public static Assembly? LoadAssembly(string AssemblyFilePath)
+		{
+			Assembly? Assembly = FindAssemblyByName(Path.GetFileNameWithoutExtension(AssemblyFilePath));
+			if (Assembly == null)
+			{
+				if (Assembly == null)
 				{
 					byte[] AssemblyBytes = File.ReadAllBytes(AssemblyFilePath);
-					Plugin = Assembly.Load(AssemblyBytes);
+					Assembly = Assembly.Load(AssemblyBytes);
 				}
-				CheckForAttributes(Plugin);
 			}
-			PerformPostInitialization();
+			return Assembly;
 		}
 
 		/// <summary>
@@ -523,7 +563,7 @@ namespace EpicGames.UHT.Utils
 		/// </summary>
 		/// <param name="Name">Name of the assembly</param>
 		/// <returns>Assembly or null</returns>
-		private Assembly? FindAssemblyByName(string? Name)
+		private static Assembly? FindAssemblyByName(string? Name)
 		{
 			if (Name != null)
 			{
diff --git a/Engine/Source/Programs/UnrealBuildTool/Modes/UnrealHeaderToolMode.cs b/Engine/Source/Programs/UnrealBuildTool/Modes/UnrealHeaderToolMode.cs
--- a/Engine/Source/Programs/UnrealBuildTool/Modes/UnrealHeaderToolMode.cs
+++ b/Engine/Source/Programs/UnrealBuildTool/Modes/UnrealHeaderToolMode.cs
@@ -1037,7 +1037,7 @@ namespace UnrealBuildTool.Modes
 					return UhtTestHarness.RunTests(Tables, Config, Options) ? (int)CompilationResult.Succeeded : (int)CompilationResult.OtherCompilationError;
 				}
 
-				string? ProjectPath = null;
+				string? ProjectFile = null;
 				string? ManifestPath = null;
 
 				if (TargetArgumentIndex >= 0)
@@ -1073,22 +1073,25 @@ namespace UnrealBuildTool.Modes
 
 						if (Target.ProjectFile != null)
 						{
-							ProjectPath = Path.GetDirectoryName(Target.ProjectFile.FullName);
+							ProjectFile = Target.ProjectFile.FullName;
 						}
 					}
 				}
 				else
 				{
-					ProjectPath = Path.GetDirectoryName(Arguments.GetPositionalArguments()[0]);
+					ProjectFile = Arguments.GetPositionalArguments()[0];
 					ManifestPath = Arguments.GetPositionalArguments()[1];
 				}
 
+				string? ProjectPath = ProjectFile != null ? Path.GetDirectoryName(ProjectFile) : null;
+
 				UhtSession Session = new UhtSession
 				{
 					Tables = Tables,
 					Config = Config,
 					FileManager = new UhtStdFileManager(),
 					EngineDirectory = Unreal.EngineDirectory.FullName,
+					ProjectFile = ProjectFile,
 					ProjectDirectory = string.IsNullOrEmpty(ProjectPath) ? null : ProjectPath,
 					ReferenceDirectory = FileReference.Combine(EngineProgramSavedDirectory, "UnrealBuildTool", "ReferenceGeneratedCode").FullName,
 					VerifyDirectory = FileReference.Combine(EngineProgramSavedDirectory, "UnrealBuildTool", "VerifyGeneratedCode").FullName,
