More cleanup of the export system.  Removed the UHT specific tasks.

#rnx

diff --git a/Engine/Source/Programs/Shared/EpicGames.UHT/Exporters/CodeGen/UhtCodeGenerator.cs b/Engine/Source/Programs/Shared/EpicGames.UHT/Exporters/CodeGen/UhtCodeGenerator.cs
--- a/Engine/Source/Programs/Shared/EpicGames.UHT/Exporters/CodeGen/UhtCodeGenerator.cs
+++ b/Engine/Source/Programs/Shared/EpicGames.UHT/Exporters/CodeGen/UhtCodeGenerator.cs
@@ -32,7 +32,7 @@ namespace EpicGames.UHT.Exporters.CodeGen
 
 		public struct HeaderInfo
 		{
-			public IUhtExportTask? Task;
+			public Task? Task;
 			public string IncludePath;
 			public string FileId;
 			public uint BodyHash;
@@ -66,7 +66,7 @@ namespace EpicGames.UHT.Exporters.CodeGen
 
 		private void Generate()
 		{
-			List<IUhtExportTask> Prereqs = new List<IUhtExportTask>();
+			List<Task?> Prereqs = new List<Task?>();
 
 			// Perform some startup initialization to compute things we need over and over again
 			if (this.Session.bGoWide)
@@ -97,25 +97,21 @@ namespace EpicGames.UHT.Exporters.CodeGen
 					{
 						if (HeaderFile != Referenced)
 						{
-							IUhtExportTask? GeneratedReferenced = this.HeaderInfos[Referenced.HeaderFileTypeIndex].Task;
-							if (GeneratedReferenced != null)
-							{
-								Prereqs.Add(GeneratedReferenced);
-							}
+							Prereqs.Add(this.HeaderInfos[Referenced.HeaderFileTypeIndex].Task);
 						}
 					}
 
 					this.HeaderInfos[HeaderFile.HeaderFileTypeIndex].Task = Factory.CreateTask(Prereqs,  
-						(IUhtExportTask Task) =>
+						(IUhtExportFactory Factory) =>
 						{
-							new UhtHeaderCodeGeneratorHFile(this, Package, HeaderFile).Generate(Task);
-							new UhtHeaderCodeGeneratorCppFile(this, Package, HeaderFile).Generate(Task);
+							new UhtHeaderCodeGeneratorHFile(this, Package, HeaderFile).Generate(Factory);
+							new UhtHeaderCodeGeneratorCppFile(this, Package, HeaderFile).Generate(Factory);
 						});
 				}
 			}
 
 			// Generate the files for the packages
-			List<IUhtExportTask> GeneratedPackages = new List<IUhtExportTask>(this.Session.PackageTypeCount);
+			List<Task?> GeneratedPackages = new List<Task?>(this.Session.PackageTypeCount);
 			foreach (UhtPackage Package in this.Session.Packages)
 			{
 				UHTManifest.Module Module = Package.Module;
@@ -124,11 +120,7 @@ namespace EpicGames.UHT.Exporters.CodeGen
 				Prereqs.Clear();
 				foreach (UhtHeaderFile HeaderFile in Package.Children)
 				{
-					IUhtExportTask? GeneratedReferenced = this.HeaderInfos[HeaderFile.HeaderFileTypeIndex].Task;
-					if (GeneratedReferenced != null)
-					{
-						Prereqs.Add(GeneratedReferenced);
-					}
+					Prereqs.Add(this.HeaderInfos[HeaderFile.HeaderFileTypeIndex].Task);
 					if (!bWriteHeader)
 					{
 						foreach (UhtType Type in HeaderFile.Children)
@@ -148,24 +140,24 @@ namespace EpicGames.UHT.Exporters.CodeGen
 				}
 
 				GeneratedPackages.Add(Factory.CreateTask(Prereqs, 
-					(IUhtExportTask Task) =>
+					(IUhtExportFactory Factory) =>
 					{
 						List<UhtHeaderFile> PackageSortedHeaders = GetSortedHeaderFiles(Package);
 						if (bWriteHeader)
 						{
-							new UhtPackageCodeGeneratorHFile(this, Package).Generate(Task, PackageSortedHeaders);
+							new UhtPackageCodeGeneratorHFile(this, Package).Generate(Factory, PackageSortedHeaders);
 						}
-						new UhtPackageCodeGeneratorCppFile(this, Package).Generate(Task, PackageSortedHeaders);
+						new UhtPackageCodeGeneratorCppFile(this, Package).Generate(Factory, PackageSortedHeaders);
 					}));
 			}
 
 			// Wait for all the packages to complete
 			List<Task> PackageTasks = new List<Task>(this.Session.PackageTypeCount);
-			foreach (IUhtExportTask Output in GeneratedPackages)
+			foreach (Task? Output in GeneratedPackages)
 			{
-				if (Output.ActionTask != null)
+				if (Output != null)
 				{
-					PackageTasks.Add(Output.ActionTask);
+					PackageTasks.Add(Output);
 				}
 			}
 			Task.WaitAll(PackageTasks.ToArray());
diff --git a/Engine/Source/Programs/Shared/EpicGames.UHT/Exporters/CodeGen/UhtHeaderCodeGeneratorCppFile.cs b/Engine/Source/Programs/Shared/EpicGames.UHT/Exporters/CodeGen/UhtHeaderCodeGeneratorCppFile.cs
--- a/Engine/Source/Programs/Shared/EpicGames.UHT/Exporters/CodeGen/UhtHeaderCodeGeneratorCppFile.cs
+++ b/Engine/Source/Programs/Shared/EpicGames.UHT/Exporters/CodeGen/UhtHeaderCodeGeneratorCppFile.cs
@@ -27,8 +27,8 @@ namespace EpicGames.UHT.Exporters.CodeGen
 		/// <summary>
 		/// For a given UE header file, generated the generated H file
 		/// </summary>
-		/// <param name="Task">Requesting task</param>
-		public void Generate(IUhtExportTask Task)
+		/// <param name="Factory">Requesting factory</param>
+		public void Generate(IUhtExportFactory Factory)
 		{
 			ref UhtCodeGenerator.HeaderInfo HeaderInfo = ref this.HeaderInfos[this.HeaderFile.HeaderFileTypeIndex];
 			using (BorrowStringBuilder Borrower = new BorrowStringBuilder(StringBuilderCache.Big))
@@ -290,11 +290,11 @@ namespace EpicGames.UHT.Exporters.CodeGen
 
 				using (UhtBorrowBuffer BorrowBuffer = new UhtBorrowBuffer(Builder))
 				{
-					string CppFilePath = Task.Factory.MakePath(HeaderFile, ".gen.cpp");
+					string CppFilePath = Factory.MakePath(HeaderFile, ".gen.cpp");
 					StringView GeneratedBody = new StringView(BorrowBuffer.Buffer.Memory);
 					if (this.SaveExportedHeaders)
 					{
-						Task.CommitOutput(CppFilePath, GeneratedBody);
+						Factory.CommitOutput(CppFilePath, GeneratedBody);
 					}
 
 					// Save the hash of the generated body 
diff --git a/Engine/Source/Programs/Shared/EpicGames.UHT/Exporters/CodeGen/UhtHeaderCodeGeneratorHFile.cs b/Engine/Source/Programs/Shared/EpicGames.UHT/Exporters/CodeGen/UhtHeaderCodeGeneratorHFile.cs
--- a/Engine/Source/Programs/Shared/EpicGames.UHT/Exporters/CodeGen/UhtHeaderCodeGeneratorHFile.cs
+++ b/Engine/Source/Programs/Shared/EpicGames.UHT/Exporters/CodeGen/UhtHeaderCodeGeneratorHFile.cs
@@ -31,8 +31,8 @@ namespace EpicGames.UHT.Exporters.CodeGen
 		/// <summary>
 		/// For a given UE header file, generated the generated H file
 		/// </summary>
-		/// <param name="Task">Requesting task</param>
-		public void Generate(IUhtExportTask Task)
+		/// <param name="Factory">Requesting factory</param>
+		public void Generate(IUhtExportFactory Factory)
 		{
 			ref UhtCodeGenerator.HeaderInfo HeaderInfo = ref this.HeaderInfos[this.HeaderFile.HeaderFileTypeIndex];
 			using (BorrowStringBuilder Borrower = new BorrowStringBuilder(StringBuilderCache.Big))
@@ -113,8 +113,8 @@ namespace EpicGames.UHT.Exporters.CodeGen
 
 				if (this.SaveExportedHeaders)
 				{
-					string HeaderFilePath = Task.Factory.MakePath(this.HeaderFile, ".generated.h");
-					Task.CommitOutput(HeaderFilePath, Builder);
+					string HeaderFilePath = Factory.MakePath(this.HeaderFile, ".generated.h");
+					Factory.CommitOutput(HeaderFilePath, Builder);
 				}
 			}
 		}
diff --git a/Engine/Source/Programs/Shared/EpicGames.UHT/Exporters/CodeGen/UhtPackageCodeGeneratorCppFile.cs b/Engine/Source/Programs/Shared/EpicGames.UHT/Exporters/CodeGen/UhtPackageCodeGeneratorCppFile.cs
--- a/Engine/Source/Programs/Shared/EpicGames.UHT/Exporters/CodeGen/UhtPackageCodeGeneratorCppFile.cs
+++ b/Engine/Source/Programs/Shared/EpicGames.UHT/Exporters/CodeGen/UhtPackageCodeGeneratorCppFile.cs
@@ -24,9 +24,9 @@ namespace EpicGames.UHT.Exporters.CodeGen
 		/// <summary>
 		/// For a given UE header file, generated the generated cpp file
 		/// </summary>
-		/// <param name="Task">Requesting task</param>
-		/// <param name="PackageSortedHeaders">Sorted list of headers by name, in the package.</param>
-		public void Generate(IUhtExportTask Task, List<UhtHeaderFile> PackageSortedHeaders)
+		/// <param name="Factory">Requesting factory</param>
+		/// <param name="PackageSortedHeaders">Sorted list of headers by name of all headers in the package</param>
+		public void Generate(IUhtExportFactory Factory, List<UhtHeaderFile> PackageSortedHeaders)
 		{
 			using (BorrowStringBuilder Borrower = new BorrowStringBuilder(StringBuilderCache.Big))
 			{
@@ -143,8 +143,8 @@ namespace EpicGames.UHT.Exporters.CodeGen
 
 				if (this.SaveExportedHeaders)
 				{
-					string CppFilePath = Task.Factory.MakePath(this.Package, ".init.gen.cpp");
-					Task.CommitOutput(CppFilePath, Builder);
+					string CppFilePath = Factory.MakePath(this.Package, ".init.gen.cpp");
+					Factory.CommitOutput(CppFilePath, Builder);
 				}
 			}
 		}
diff --git a/Engine/Source/Programs/Shared/EpicGames.UHT/Exporters/CodeGen/UhtPackageCodeGeneratorHFile.cs b/Engine/Source/Programs/Shared/EpicGames.UHT/Exporters/CodeGen/UhtPackageCodeGeneratorHFile.cs
--- a/Engine/Source/Programs/Shared/EpicGames.UHT/Exporters/CodeGen/UhtPackageCodeGeneratorHFile.cs
+++ b/Engine/Source/Programs/Shared/EpicGames.UHT/Exporters/CodeGen/UhtPackageCodeGeneratorHFile.cs
@@ -24,9 +24,9 @@ namespace EpicGames.UHT.Exporters.CodeGen
 		/// <summary>
 		/// For a given UE header file, generated the generated H file
 		/// </summary>
-		/// <param name="Task">Requesting task</param>
+		/// <param name="Factory">Requesting factory</param>
 		/// <param name="PackageSortedHeaders">Sorted list of headers by name of all headers in the package</param>
-		public void Generate(IUhtExportTask Task, List<UhtHeaderFile> PackageSortedHeaders)
+		public void Generate(IUhtExportFactory Factory, List<UhtHeaderFile> PackageSortedHeaders)
 		{
 			using (BorrowStringBuilder Borrower = new BorrowStringBuilder(StringBuilderCache.Big))
 			{
@@ -63,8 +63,8 @@ namespace EpicGames.UHT.Exporters.CodeGen
 
 				if (this.SaveExportedHeaders)
 				{
-					string HeaderFilePath = Task.Factory.MakePath(this.Package, "Classes.h");
-					Task.CommitOutput(HeaderFilePath, Builder);
+					string HeaderFilePath = Factory.MakePath(this.Package, "Classes.h");
+					Factory.CommitOutput(HeaderFilePath, Builder);
 				}
 			}
 		}
diff --git a/Engine/Source/Programs/Shared/EpicGames.UHT/Exporters/Json/UhtJsonExporter.cs b/Engine/Source/Programs/Shared/EpicGames.UHT/Exporters/Json/UhtJsonExporter.cs
--- a/Engine/Source/Programs/Shared/EpicGames.UHT/Exporters/Json/UhtJsonExporter.cs
+++ b/Engine/Source/Programs/Shared/EpicGames.UHT/Exporters/Json/UhtJsonExporter.cs
@@ -31,26 +31,26 @@ namespace EpicGames.UHT.Exporters.Json
 		private void Export()
 		{
 			// Generate the files for the packages
-			List<IUhtExportTask> GeneratedPackages = new List<IUhtExportTask>(this.Session.PackageTypeCount);
+			List<Task?> GeneratedPackages = new List<Task?>(this.Session.PackageTypeCount);
 			foreach (UhtPackage Package in this.Session.Packages)
 			{
 				UHTManifest.Module Module = Package.Module;
 				GeneratedPackages.Add(Factory.CreateTask(
-					(IUhtExportTask Task) =>
+					(IUhtExportFactory Factory) =>
 					{
-						string JsonPath = Task.Factory.MakePath(Package, ".json");
+						string JsonPath = Factory.MakePath(Package, ".json");
 						JsonSerializerOptions Options = new JsonSerializerOptions { WriteIndented = true, IgnoreNullValues = true };
-						Task.CommitOutput(JsonPath, JsonSerializer.Serialize(Package, Options));
+						Factory.CommitOutput(JsonPath, JsonSerializer.Serialize(Package, Options));
 					}));
 			}
 
 			// Wait for all the packages to complete
 			List<Task> PackageTasks = new List<Task>(this.Session.PackageTypeCount);
-			foreach (IUhtExportTask Output in GeneratedPackages)
+			foreach (Task? Output in GeneratedPackages)
 			{
-				if (Output.ActionTask != null)
+				if (Output != null)
 				{
-					PackageTasks.Add(Output.ActionTask);
+					PackageTasks.Add(Output);
 				}
 			}
 			Task.WaitAll(PackageTasks.ToArray());
diff --git a/Engine/Source/Programs/Shared/EpicGames.UHT/Utils/UhtExport.cs b/Engine/Source/Programs/Shared/EpicGames.UHT/Utils/UhtExport.cs
--- a/Engine/Source/Programs/Shared/EpicGames.UHT/Utils/UhtExport.cs
+++ b/Engine/Source/Programs/Shared/EpicGames.UHT/Utils/UhtExport.cs
@@ -11,24 +11,36 @@ namespace EpicGames.UHT.Utils
 {
 
 	/// <summary>
-	/// Represents a task to export outputs
+	/// Delegate invoked by a factory
 	/// </summary>
-	public interface IUhtExportTask
+	/// <param name="Factory">Invoking factory</param>
+	public delegate void UhtExportTaskDelegate(IUhtExportFactory Factory);
+
+	/// <summary>
+	/// Factory object used to generate export tasks
+	/// </summary>
+	public interface IUhtExportFactory
 	{
+		
 		/// <summary>
-		/// Factory associated with the task
+		/// Session being run
 		/// </summary>
-		public IUhtExportFactory Factory { get; }
+		public UhtSession Session { get; }
 
 		/// <summary>
-		/// Session being run
+		/// Create a task
 		/// </summary>
-		public UhtSession Session { get; }
+		/// <param name="Prereqs">Tasks that must be completed prior to this task running</param>
+		/// <param name="Action">Action to be invoked to generate the output(s)</param>
+		/// <returns>Task object or null if the task was immediately executed.</returns>
+		public Task? CreateTask(List<Task?>? Prereqs, UhtExportTaskDelegate Action);
 
 		/// <summary>
-		/// Task handle associated with the task for it can be waited on
+		/// Create a task
 		/// </summary>
-		public Task? ActionTask { get; }
+		/// <param name="Action">Action to be invoked to generate the output(s)</param>
+		/// <returns>Task object or null if the task was immediately executed.</returns>
+		public Task? CreateTask(UhtExportTaskDelegate Action);
 
 		/// <summary>
 		/// Commit the contents of the string builder as the output.
@@ -45,39 +57,6 @@ namespace EpicGames.UHT.Utils
 		/// <param name="FilePath">Destination file path</param>
 		/// <param name="Output">Output to commit</param>
 		public void CommitOutput(string FilePath, StringView Output);
-	}
-
-	/// <summary>
-	/// Delegate invoked by a task
-	/// </summary>
-	/// <param name="Task">Invoking task</param>
-	public delegate void UhtExportTaskDelegate(IUhtExportTask Task);
-
-	/// <summary>
-	/// Factory object used to generate export tasks
-	/// </summary>
-	public interface IUhtExportFactory
-	{
-		
-		/// <summary>
-		/// Session being run
-		/// </summary>
-		public UhtSession Session { get; }
-
-		/// <summary>
-		/// Create a task
-		/// </summary>
-		/// <param name="Prereqs">Tasks that must be completed prior to this task running</param>
-		/// <param name="Action">Action to be invoked to generate the output(s)</param>
-		/// <returns>Task interface.</returns>
-		public IUhtExportTask CreateTask(List<IUhtExportTask>? Prereqs, UhtExportTaskDelegate Action);
-
-		/// <summary>
-		/// Create a task
-		/// </summary>
-		/// <param name="Action">Action to be invoked to generate the output(s)</param>
-		/// <returns>Task interface.</returns>
-		public IUhtExportTask CreateTask(UhtExportTaskDelegate Action);
 
 		/// <summary>
 		/// Make a path for an output based on the header file name.
diff --git a/Engine/Source/Programs/Shared/EpicGames.UHT/Utils/UhtSession.cs b/Engine/Source/Programs/Shared/EpicGames.UHT/Utils/UhtSession.cs
--- a/Engine/Source/Programs/Shared/EpicGames.UHT/Utils/UhtSession.cs
+++ b/Engine/Source/Programs/Shared/EpicGames.UHT/Utils/UhtSession.cs
@@ -47,123 +47,6 @@ namespace EpicGames.UHT.Utils
 		public StringView Data;
 	}
 
-	/// <summary>
-	/// Base implementation of an export task
-	/// </summary>
-	class UhtExportTask : IUhtExportTask
-	{
-		/// <summary>
-		/// Action to invoke to generate the output
-		/// </summary>
-		private readonly UhtExportTaskDelegate Action;
-
-		/// <summary>
-		/// Factory requesting the export
-		/// </summary>
-		protected readonly UhtExportFactory ExportFactory;
-
-		/// <summary>
-		/// Action task created to invoke the task
-		/// </summary>
-		private Task? ActionTaskInternal = null;
-
-		/// <summary>
-		/// Factory associated with the task
-		/// </summary>
-		public IUhtExportFactory Factory => this.ExportFactory;
-
-		/// <summary>
-		/// UHT session
-		/// </summary>
-		public UhtSession Session => this.ExportFactory.Session;
-
-		/// <summary>
-		/// Export options associated with the task
-		/// </summary>
-		public Task? ActionTask => this.ActionTaskInternal;
-
-		/// <summary>
-		/// Construct a new instance of an export task
-		/// </summary>
-		/// <param name="ExportFactory">Factory requesting the task</param>
-		/// <param name="Action">Action to be invoked as part of the export</param>
-		public UhtExportTask(UhtExportFactory ExportFactory, UhtExportTaskDelegate Action)
-		{
-			this.ExportFactory = ExportFactory;
-			this.Action = Action;
-		}
-
-		/// <summary>
-		/// Commit the contents of the string builder as the output.
-		/// If you have a string builder, use this method so that a 
-		/// temporary buffer can be used.
-		/// </summary>
-		/// <param name="FilePath">Destination file path</param>
-		/// <param name="Builder">Source for the content</param>
-		public void CommitOutput(string FilePath, StringBuilder Builder)
-		{
-			using (UhtBorrowBuffer BorrowBuffer = new UhtBorrowBuffer(Builder))
-			{
-				string TempFilePath = FilePath + ".tmp";
-				this.ExportFactory.SaveIfChanged(FilePath, TempFilePath, new StringView(BorrowBuffer.Buffer.Memory));
-			}
-		}
-
-		/// <summary>
-		/// Commit the value of the string as the output
-		/// </summary>
-		/// <param name="FilePath">Destination file path</param>
-		/// <param name="Output">Output to commit</param>
-		public void CommitOutput(string FilePath, StringView Output)
-		{
-			string TempFilePath = FilePath + ".tmp";
-			this.ExportFactory.SaveIfChanged(FilePath, TempFilePath, Output);
-		}
-
-		/// <summary>
-		/// Queue the task
-		/// </summary>
-		/// <param name="Prereqs">List of prerequisite tasks that must complete prior to this task being invoked.</param>
-		public void Queue(List<IUhtExportTask>? Prereqs)
-		{
-			if (this.Session.bGoWide)
-			{
-				List<Task>? PrereqTasks = null;
-				if (Prereqs != null)
-				{
-					PrereqTasks = new List<Task>();
-					foreach (IUhtExportTask Prereq in Prereqs)
-					{
-						if (Prereq.ActionTask != null)
-						{
-							PrereqTasks.Add(Prereq.ActionTask);
-						}
-					}
-				}
-				if (PrereqTasks != null && PrereqTasks.Count > 0)
-				{
-					this.ActionTaskInternal = Task.Factory.ContinueWhenAll(PrereqTasks.ToArray(), (Task[] Tasks) => { Export(); });
-				}
-				else
-				{
-					this.ActionTaskInternal = Task.Factory.StartNew(() => { Export(); });
-				}
-			}
-			else
-			{
-				Export();
-			}
-		}
-
-		/// <summary>
-		/// Invoked to perform the actual export
-		/// </summary>
-		public void Export()
-		{
-			this.Action(this);
-		}
-	}
-
 	/// <summary>
 	/// Implementation of the export factory
 	/// </summary>
@@ -234,25 +117,66 @@ namespace EpicGames.UHT.Utils
 			}
 		}
 
+		/// <summary>
+		/// Commit the contents of the string builder as the output.
+		/// If you have a string builder, use this method so that a 
+		/// temporary buffer can be used.
+		/// </summary>
+		/// <param name="FilePath">Destination file path</param>
+		/// <param name="Builder">Source for the content</param>
+		public void CommitOutput(string FilePath, StringBuilder Builder)
+		{
+			using (UhtBorrowBuffer BorrowBuffer = new UhtBorrowBuffer(Builder))
+			{
+				string TempFilePath = FilePath + ".tmp";
+				SaveIfChanged(FilePath, TempFilePath, new StringView(BorrowBuffer.Buffer.Memory));
+			}
+		}
+
+		/// <summary>
+		/// Commit the value of the string as the output
+		/// </summary>
+		/// <param name="FilePath">Destination file path</param>
+		/// <param name="Output">Output to commit</param>
+		public void CommitOutput(string FilePath, StringView Output)
+		{
+			string TempFilePath = FilePath + ".tmp";
+			SaveIfChanged(FilePath, TempFilePath, Output);
+		}
+
 		/// <summary>
 		/// Create a task to export two files
 		/// </summary>
 		/// <param name="Prereqs">Tasks that must be completed prior to this task running</param>
 		/// <param name="Action">Action to be invoked to generate the output</param>
-		/// <returns>Task interface.</returns>
-		public IUhtExportTask CreateTask(List<IUhtExportTask>? Prereqs, UhtExportTaskDelegate Action)
+		/// <returns>Task object or null if the task was immediately executed.</returns>
+		public Task? CreateTask(List<Task?>? Prereqs, UhtExportTaskDelegate Action)
 		{
-			UhtExportTask Task = new UhtExportTask(this, Action);
-			Task.Queue(Prereqs);
-			return Task;
+			if (this.Session.bGoWide)
+			{
+				Task[]? PrereqTasks = Prereqs != null ? Prereqs.Where(x => x != null).Cast<Task>().ToArray() : null;
+				if (PrereqTasks != null && PrereqTasks.Length > 0)
+				{
+					return Task.Factory.ContinueWhenAll(PrereqTasks, (Task[] Tasks) => { Action(this); });
+				}
+				else
+				{
+					return Task.Factory.StartNew(() => { Action(this); });
+				}
+			}
+			else
+			{
+				Action(this);
+				return null;
+			}
 		}
 
 		/// <summary>
 		/// Create a task to export two files
 		/// </summary>
 		/// <param name="Action">Action to be invoked to generate the output</param>
-		/// <returns>Task interface.</returns>
-		public IUhtExportTask CreateTask(UhtExportTaskDelegate Action)
+		/// <returns>Task object or null if the task was immediately executed.</returns>
+		public Task? CreateTask(UhtExportTaskDelegate Action)
 		{
 			return CreateTask(null, Action);
 		}
