Added "generic" property type text to provide a pattern similar to older UHT.

#rnx

diff --git a/Engine/Source/Programs/Shared/EpicGames.UHT/Exporters/CodeGen/UhtHeaderCodeGenerator.cs b/Engine/Source/Programs/Shared/EpicGames.UHT/Exporters/CodeGen/UhtHeaderCodeGenerator.cs
--- a/Engine/Source/Programs/Shared/EpicGames.UHT/Exporters/CodeGen/UhtHeaderCodeGenerator.cs
+++ b/Engine/Source/Programs/Shared/EpicGames.UHT/Exporters/CodeGen/UhtHeaderCodeGenerator.cs
@@ -158,14 +158,14 @@ namespace EpicGames.UHT.Exporters.CodeGen
 			{
 				// If the function was marked as 'RequiredAPI', then add the *_API macro prefix.  Note that if the class itself
 				// was marked 'RequiredAPI', this is not needed as C++ will exports all methods automatically.
-				if (TextType != UhtPropertyTextType.EventFunction &&
+				if (TextType != UhtPropertyTextType.EventFunctionArgOrRetVal &&
 					!(OuterClass != null && OuterClass.ClassFlags.HasAnyFlags(EClassFlags.RequiredAPI)) &&
 					ExportFlags.HasAnyFlags(UhtFunctionExportFlags.RequiredAPI))
 				{
 					Builder.Append(this.PackageApi);
 				}
 
-				if (TextType == UhtPropertyTextType.InterfaceFunction)
+				if (TextType == UhtPropertyTextType.InterfaceFunctionArgOrRetVal)
 				{
 					Builder.Append("static ");
 				}
@@ -179,7 +179,7 @@ namespace EpicGames.UHT.Exporters.CodeGen
 					Builder.Append("virtual ");
 				}
 				// this is not an event, the function is not a static function and the function is not marked final
-				else if (TextType != UhtPropertyTextType.EventFunction && !Function.FunctionFlags.HasAnyFlags(EFunctionFlags.Static) && !ExportFlags.HasAnyFlags(UhtFunctionExportFlags.Final))
+				else if (TextType != UhtPropertyTextType.EventFunctionArgOrRetVal && !Function.FunctionFlags.HasAnyFlags(EFunctionFlags.Static) && !ExportFlags.HasAnyFlags(UhtFunctionExportFlags.Final))
 				{
 					Builder.Append("virtual ");
 				}
@@ -217,21 +217,23 @@ namespace EpicGames.UHT.Exporters.CodeGen
 			{
 				switch (TextType)
 				{
-					case UhtPropertyTextType.InterfaceFunction:
+					case UhtPropertyTextType.InterfaceFunctionArgOrRetVal:
 						Builder.Append("Execute_").Append(Function.SourceName);
 						break;
-					case UhtPropertyTextType.EventFunction:
+					case UhtPropertyTextType.EventFunctionArgOrRetVal:
 						Builder.Append(Function.MarshalAndCallName);
 						break;
-					case UhtPropertyTextType.ClassFunction:
+					case UhtPropertyTextType.ClassFunctionArgOrRetVal:
 						Builder.Append(Function.CppImplName);
 						break;
+					default:
+						throw new UhtIceException("Unexpected type text");
 				}
 			}
 
 			AppendParameters(Builder, Function, TextType, ExtraParam, false);
 
-			if (TextType != UhtPropertyTextType.InterfaceFunction)
+			if (TextType != UhtPropertyTextType.InterfaceFunctionArgOrRetVal)
 			{
 				if (!bIsDelegate && Function.FunctionFlags.HasAnyFlags(EFunctionFlags.Const))
 				{
diff --git a/Engine/Source/Programs/Shared/EpicGames.UHT/Exporters/CodeGen/UhtHeaderCodeGeneratorCppFile.cs b/Engine/Source/Programs/Shared/EpicGames.UHT/Exporters/CodeGen/UhtHeaderCodeGeneratorCppFile.cs
--- a/Engine/Source/Programs/Shared/EpicGames.UHT/Exporters/CodeGen/UhtHeaderCodeGeneratorCppFile.cs
+++ b/Engine/Source/Programs/Shared/EpicGames.UHT/Exporters/CodeGen/UhtHeaderCodeGeneratorCppFile.cs
@@ -1054,7 +1054,7 @@ namespace EpicGames.UHT.Exporters.CodeGen
 			{
 				Builder.Append("\tstatic FName NAME_").Append(Function.Outer?.SourceName).Append("_").Append(Function.SourceName).Append(" = FName(TEXT(\"").Append(Function.EngineName).Append("\"));\r\n");
 				string ExtraParameter = Function.FunctionFlags.HasAnyFlags(EFunctionFlags.Const) ? "const UObject* O" : "UObject* O";
-				AppendNativeFunctionHeader(Builder, Function, UhtPropertyTextType.InterfaceFunction, false, null, ExtraParameter, 0, "\r\n");
+				AppendNativeFunctionHeader(Builder, Function, UhtPropertyTextType.InterfaceFunctionArgOrRetVal, false, null, ExtraParameter, 0, "\r\n");
 				Builder.Append("\t{\r\n");
 				Builder.Append("\t\tcheck(O != NULL);\r\n");
 				Builder.Append("\t\tcheck(O->GetClass()->ImplementsInterface(").Append(Class.SourceName).Append("::StaticClass()));\r\n");
@@ -1536,7 +1536,7 @@ namespace EpicGames.UHT.Exporters.CodeGen
 							Builder.Append("\tstatic FName NAME_").Append(Class.SourceName).Append('_').Append(Function.EngineName).Append(" = FName(TEXT(\"").Append(Function.EngineName).Append("\"));\r\n");
 						}
 
-						AppendNativeFunctionHeader(Builder, Function, UhtPropertyTextType.EventFunction, false, null, null, UhtFunctionExportFlags.None, "\r\n");
+						AppendNativeFunctionHeader(Builder, Function, UhtPropertyTextType.EventFunctionArgOrRetVal, false, null, null, UhtFunctionExportFlags.None, "\r\n");
 
 						if (bIsInterfaceClass)
 						{
diff --git a/Engine/Source/Programs/Shared/EpicGames.UHT/Exporters/CodeGen/UhtHeaderCodeGeneratorHFile.cs b/Engine/Source/Programs/Shared/EpicGames.UHT/Exporters/CodeGen/UhtHeaderCodeGeneratorHFile.cs
--- a/Engine/Source/Programs/Shared/EpicGames.UHT/Exporters/CodeGen/UhtHeaderCodeGeneratorHFile.cs
+++ b/Engine/Source/Programs/Shared/EpicGames.UHT/Exporters/CodeGen/UhtHeaderCodeGeneratorHFile.cs
@@ -299,7 +299,7 @@ namespace EpicGames.UHT.Exporters.CodeGen
 
 				AppendEventParameter(Builder, Function, StrippedFunctionName, UhtPropertyTextType.EventParameterMember, true, Tabs, " \\\r\n");
 				Builder.Append("static ");
-				AppendNativeFunctionHeader(Builder, Function, UhtPropertyTextType.EventFunction, true, ExportFunctionName, ExtraParameter, UhtFunctionExportFlags.Inline, " \\\r\n");
+				AppendNativeFunctionHeader(Builder, Function, UhtPropertyTextType.EventFunctionArgOrRetVal, true, ExportFunctionName, ExtraParameter, UhtFunctionExportFlags.Inline, " \\\r\n");
 				AppendEventFunctionPrologue(Builder, Function, StrippedFunctionName, Tabs, " \\\r\n");
 				Builder
 					.Append("\t")
@@ -658,7 +658,7 @@ namespace EpicGames.UHT.Exporters.CodeGen
 					continue;
 				}
 				AppendNetValidateDeclaration(Builder, Function);
-				AppendNativeFunctionHeader(Builder, Function, UhtPropertyTextType.ClassFunction, true, null, null, UhtFunctionExportFlags.None, "; \\\r\n");
+				AppendNativeFunctionHeader(Builder, Function, UhtPropertyTextType.ClassFunctionArgOrRetVal, true, null, null, UhtFunctionExportFlags.None, "; \\\r\n");
 			}
 			return Builder;
 		}
@@ -674,7 +674,7 @@ namespace EpicGames.UHT.Exporters.CodeGen
 				if (!Function.FunctionExportFlags.HasAnyFlags(UhtFunctionExportFlags.ImplFound))
 				{
 					AppendNetValidateDeclaration(Builder, Function);
-					AppendNativeFunctionHeader(Builder, Function, UhtPropertyTextType.ClassFunction, true, null, null, UhtFunctionExportFlags.None, "; \\\r\n");
+					AppendNativeFunctionHeader(Builder, Function, UhtPropertyTextType.ClassFunctionArgOrRetVal, true, null, null, UhtFunctionExportFlags.None, "; \\\r\n");
 				}
 			}
 			return Builder;
@@ -690,7 +690,7 @@ namespace EpicGames.UHT.Exporters.CodeGen
 					Builder.Append("virtual");
 				}
 				Builder.Append(" bool ").Append(Function.CppValidationImplName);
-				AppendParameters(Builder, Function, UhtPropertyTextType.ClassFunction, null, true);
+				AppendParameters(Builder, Function, UhtPropertyTextType.ClassFunctionArgOrRetVal, null, true);
 				Builder.Append("; \\\r\n");
 			}
 			return Builder;
@@ -746,7 +746,7 @@ namespace EpicGames.UHT.Exporters.CodeGen
 						if (!Function.FunctionFlags.HasAnyFlags(EFunctionFlags.NetResponse) &&
 							Function.EngineName != Function.MarshalAndCallName)
 						{
-							AppendNativeFunctionHeader(Builder, Function, UhtPropertyTextType.EventFunction, true, null, null, UhtFunctionExportFlags.None, " \\\r\n");
+							AppendNativeFunctionHeader(Builder, Function, UhtPropertyTextType.EventFunctionArgOrRetVal, true, null, null, UhtFunctionExportFlags.None, " \\\r\n");
 							Builder.Append(" \\\r\n");
 						}
 					}
@@ -1278,7 +1278,7 @@ namespace EpicGames.UHT.Exporters.CodeGen
 
 			foreach (UhtFunction Function in CallbackFunctions)
 			{
-				AppendNativeFunctionHeader(Builder, Function, UhtPropertyTextType.InterfaceFunction, true, null,
+				AppendNativeFunctionHeader(Builder, Function, UhtPropertyTextType.InterfaceFunctionArgOrRetVal, true, null,
 					Function.FunctionFlags.HasAnyFlags(EFunctionFlags.Const) ? ConstExtraArg : ExtraArg, UhtFunctionExportFlags.None,
 					"; \\\r\n");
 			}
diff --git a/Engine/Source/Programs/Shared/EpicGames.UHT/Types/Properties/UhtBooleanProperty.cs b/Engine/Source/Programs/Shared/EpicGames.UHT/Types/Properties/UhtBooleanProperty.cs
--- a/Engine/Source/Programs/Shared/EpicGames.UHT/Types/Properties/UhtBooleanProperty.cs
+++ b/Engine/Source/Programs/Shared/EpicGames.UHT/Types/Properties/UhtBooleanProperty.cs
@@ -37,10 +37,10 @@ namespace EpicGames.UHT.Types
 		{
 			switch (TextType)
 			{
-				case UhtPropertyTextType.UserFacing:
+				case UhtPropertyTextType.Generic:
 				case UhtPropertyTextType.Construction:
 				case UhtPropertyTextType.FunctionThunkParameterArrayType:
-				case UhtPropertyTextType.FunctionThunkReturn:
+				case UhtPropertyTextType.FunctionThunkRetVal:
 				case UhtPropertyTextType.RigVMTemplateArg:
 				case UhtPropertyTextType.ExportMember:
 					Builder.Append(this.CppTypeText);
@@ -48,9 +48,10 @@ namespace EpicGames.UHT.Types
 
 				case UhtPropertyTextType.Sparse:
 				case UhtPropertyTextType.SparseShort:
-				case UhtPropertyTextType.ClassFunction:
-				case UhtPropertyTextType.EventFunction:
-				case UhtPropertyTextType.InterfaceFunction:
+				case UhtPropertyTextType.GenericFunctionArgOrRetVal:
+				case UhtPropertyTextType.ClassFunctionArgOrRetVal:
+				case UhtPropertyTextType.EventFunctionArgOrRetVal:
+				case UhtPropertyTextType.InterfaceFunctionArgOrRetVal:
 				case UhtPropertyTextType.EventParameterMember:
 				case UhtPropertyTextType.EventParameterFunctionMember:
 				case UhtPropertyTextType.GetterSetterArg:
diff --git a/Engine/Source/Programs/Shared/EpicGames.UHT/Types/Properties/UhtClassProperty.cs b/Engine/Source/Programs/Shared/EpicGames.UHT/Types/Properties/UhtClassProperty.cs
--- a/Engine/Source/Programs/Shared/EpicGames.UHT/Types/Properties/UhtClassProperty.cs
+++ b/Engine/Source/Programs/Shared/EpicGames.UHT/Types/Properties/UhtClassProperty.cs
@@ -71,12 +71,13 @@ namespace EpicGames.UHT.Types
 		{
 			switch (TextType)
 			{
-				case UhtPropertyTextType.UserFacing:
+				case UhtPropertyTextType.Generic:
 				case UhtPropertyTextType.Sparse:
 				case UhtPropertyTextType.SparseShort:
-				case UhtPropertyTextType.ClassFunction:
-				case UhtPropertyTextType.EventFunction:
-				case UhtPropertyTextType.InterfaceFunction:
+				case UhtPropertyTextType.GenericFunctionArgOrRetVal:
+				case UhtPropertyTextType.ClassFunctionArgOrRetVal:
+				case UhtPropertyTextType.EventFunctionArgOrRetVal:
+				case UhtPropertyTextType.InterfaceFunctionArgOrRetVal:
 				case UhtPropertyTextType.ExportMember:
 				case UhtPropertyTextType.Construction:
 				case UhtPropertyTextType.FunctionThunkParameterArrayType:
@@ -85,7 +86,7 @@ namespace EpicGames.UHT.Types
 					AppendText(Builder);
 					break;
 
-				case UhtPropertyTextType.FunctionThunkReturn:
+				case UhtPropertyTextType.FunctionThunkRetVal:
 					if (this.PropertyFlags.HasAnyFlags(EPropertyFlags.ConstParm))
 					{
 						Builder.Append("const ");
diff --git a/Engine/Source/Programs/Shared/EpicGames.UHT/Types/Properties/UhtEnumProperty.cs b/Engine/Source/Programs/Shared/EpicGames.UHT/Types/Properties/UhtEnumProperty.cs
--- a/Engine/Source/Programs/Shared/EpicGames.UHT/Types/Properties/UhtEnumProperty.cs
+++ b/Engine/Source/Programs/Shared/EpicGames.UHT/Types/Properties/UhtEnumProperty.cs
@@ -106,9 +106,10 @@ namespace EpicGames.UHT.Types
 		{
 			switch (TextType)
 			{
-				case UhtPropertyTextType.ClassFunction:
-				case UhtPropertyTextType.InterfaceFunction:
-				case UhtPropertyTextType.EventFunction:
+				case UhtPropertyTextType.GenericFunctionArgOrRetVal:
+				case UhtPropertyTextType.ClassFunctionArgOrRetVal:
+				case UhtPropertyTextType.InterfaceFunctionArgOrRetVal:
+				case UhtPropertyTextType.EventFunctionArgOrRetVal:
 					if (Enum.CppForm == UhtEnumCppForm.EnumClass || (!bIsTemplateArgument && (Property.PropertyFlags.HasAnyFlags(EPropertyFlags.ReturnParm) || !Property.PropertyFlags.HasAnyFlags(EPropertyFlags.OutParm))))
 					{
 						Builder.Append(Enum.CppType);
diff --git a/Engine/Source/Programs/Shared/EpicGames.UHT/Types/Properties/UhtObjectProperty.cs b/Engine/Source/Programs/Shared/EpicGames.UHT/Types/Properties/UhtObjectProperty.cs
--- a/Engine/Source/Programs/Shared/EpicGames.UHT/Types/Properties/UhtObjectProperty.cs
+++ b/Engine/Source/Programs/Shared/EpicGames.UHT/Types/Properties/UhtObjectProperty.cs
@@ -62,7 +62,7 @@ namespace EpicGames.UHT.Types
 		{
 			switch (TextType)
 			{
-				case UhtPropertyTextType.FunctionThunkReturn:
+				case UhtPropertyTextType.FunctionThunkRetVal:
 					if (this.PropertyFlags.HasAnyFlags(EPropertyFlags.ConstParm))
 					{
 						Builder.Append("const ");
diff --git a/Engine/Source/Programs/Shared/EpicGames.UHT/Types/Properties/UhtObjectPtrProperty.cs b/Engine/Source/Programs/Shared/EpicGames.UHT/Types/Properties/UhtObjectPtrProperty.cs
--- a/Engine/Source/Programs/Shared/EpicGames.UHT/Types/Properties/UhtObjectPtrProperty.cs
+++ b/Engine/Source/Programs/Shared/EpicGames.UHT/Types/Properties/UhtObjectPtrProperty.cs
@@ -56,7 +56,7 @@ namespace EpicGames.UHT.Types
 					}
 					break;
 
-				case UhtPropertyTextType.FunctionThunkReturn:
+				case UhtPropertyTextType.FunctionThunkRetVal:
 					if (this.PropertyFlags.HasAnyFlags(EPropertyFlags.ConstParm))
 					{
 						Builder.Append("const ");
diff --git a/Engine/Source/Programs/Shared/EpicGames.UHT/Types/UhtFunction.cs b/Engine/Source/Programs/Shared/EpicGames.UHT/Types/UhtFunction.cs
--- a/Engine/Source/Programs/Shared/EpicGames.UHT/Types/UhtFunction.cs
+++ b/Engine/Source/Programs/Shared/EpicGames.UHT/Types/UhtFunction.cs
@@ -803,7 +803,7 @@ namespace EpicGames.UHT.Types
 			UhtProperty? ReturnType = this.ReturnProperty;
 			if (ReturnType != null)
 			{
-				Builder.AppendPropertyText(ReturnType, UhtPropertyTextType.EventFunction);
+				Builder.AppendPropertyText(ReturnType, UhtPropertyTextType.EventFunctionArgOrRetVal);
 				Builder.Append(' ');
 			}
 			else
@@ -821,7 +821,7 @@ namespace EpicGames.UHT.Types
 					Builder.Append(", ");
 				}
 				bFirst = false;
-				Builder.AppendFullDecl(Arg, UhtPropertyTextType.EventFunction, true);
+				Builder.AppendFullDecl(Arg, UhtPropertyTextType.EventFunctionArgOrRetVal, true);
 			}
 
 			Builder.Append(')');
diff --git a/Engine/Source/Programs/Shared/EpicGames.UHT/Types/UhtProperty.cs b/Engine/Source/Programs/Shared/EpicGames.UHT/Types/UhtProperty.cs
--- a/Engine/Source/Programs/Shared/EpicGames.UHT/Types/UhtProperty.cs
+++ b/Engine/Source/Programs/Shared/EpicGames.UHT/Types/UhtProperty.cs
@@ -434,34 +434,39 @@ namespace EpicGames.UHT.Types
 	public enum UhtPropertyTextType
 	{
 		/// <summary>
-		/// Used when the declaration is to be displayed to the user in cases such as error messages.
+		/// Generic type
 		/// </summary>
-		UserFacing,
+		Generic,
 
 		/// <summary>
-		/// Sparse property declaration
+		/// Generic function argument or return value
 		/// </summary>
-		Sparse,
-
-		/// <summary>
-		/// Sparse property short name
-		/// </summary>
-		SparseShort,
+		GenericFunctionArgOrRetVal,
 
 		/// <summary>
 		/// Class function argument or return value
 		/// </summary>
-		ClassFunction,
+		ClassFunctionArgOrRetVal,
 
 		/// <summary>
 		/// Event function argument or return value
 		/// </summary>
-		EventFunction,
+		EventFunctionArgOrRetVal,
 
 		/// <summary>
 		/// Interface function argument or return value
 		/// </summary>
-		InterfaceFunction,
+		InterfaceFunctionArgOrRetVal,
+
+		/// <summary>
+		/// Sparse property declaration
+		/// </summary>
+		Sparse,
+
+		/// <summary>
+		/// Sparse property short name
+		/// </summary>
+		SparseShort,
 
 		/// <summary>
 		/// Class or structure member
@@ -496,7 +501,7 @@ namespace EpicGames.UHT.Types
 		/// <summary>
 		/// Used to get the return type for function thunks
 		/// </summary>
-		FunctionThunkReturn,
+		FunctionThunkRetVal,
 
 		/// <summary>
 		/// Used when generating types for RigVM.  Only invoked for template arguments
@@ -523,9 +528,10 @@ namespace EpicGames.UHT.Types
 		public static bool IsParameter(this UhtPropertyTextType TextType)
 		{
 			return
-				TextType == UhtPropertyTextType.ClassFunction ||
-				TextType == UhtPropertyTextType.EventFunction ||
-				TextType == UhtPropertyTextType.InterfaceFunction;
+				TextType == UhtPropertyTextType.GenericFunctionArgOrRetVal ||
+				TextType == UhtPropertyTextType.ClassFunctionArgOrRetVal ||
+				TextType == UhtPropertyTextType.EventFunctionArgOrRetVal ||
+				TextType == UhtPropertyTextType.InterfaceFunctionArgOrRetVal;
 		}
 	}
 
@@ -1348,7 +1354,7 @@ namespace EpicGames.UHT.Types
 		public string GetUserFacingDecl()
 		{
 			StringBuilder Builder = new StringBuilder();
-			AppendText(Builder, UhtPropertyTextType.UserFacing);
+			AppendText(Builder, UhtPropertyTextType.Generic);
 			return Builder.ToString();
 		}
 
@@ -1994,17 +2000,6 @@ namespace EpicGames.UHT.Types
 			return Property.AppendFullDecl(Builder, TextType, bSkipParameterName);
 		}
 
-		/// <summary>
-		/// Append the user facing declaration.
-		/// </summary>
-		/// <param name="Builder">Destination builder</param>
-		/// <param name="Property">Property in question</param>
-		/// <returns>Destination builder</returns>
-		public static StringBuilder AppendUserFacingDecl(this StringBuilder Builder, UhtProperty Property)
-		{
-			return Property.AppendText(Builder, UhtPropertyTextType.UserFacing);
-		}
-
 		/// <summary>
 		/// Append the function thunk parameter get
 		/// </summary>
@@ -2046,7 +2041,7 @@ namespace EpicGames.UHT.Types
 		/// <returns>Destination builder</returns>
 		public static StringBuilder AppendFunctionThunkReturn(this StringBuilder Builder, UhtProperty Property)
 		{
-			return Property.AppendText(Builder, UhtPropertyTextType.FunctionThunkReturn);
+			return Property.AppendText(Builder, UhtPropertyTextType.FunctionThunkRetVal);
 		}
 
 		/// <summary>
