Horde: Fix comparison of semver objects with identical pre-release fields.

diff --git a/Engine/Source/Programs/Shared/EpicGames.Core.Tests/SemVerTests.cs b/Engine/Source/Programs/Shared/EpicGames.Core.Tests/SemVerTests.cs
--- a/Engine/Source/Programs/Shared/EpicGames.Core.Tests/SemVerTests.cs
+++ b/Engine/Source/Programs/Shared/EpicGames.Core.Tests/SemVerTests.cs
@@ -24,6 +24,10 @@ namespace EpicGames.Core.Tests
 		[TestMethod]
 		public void PrereleaseTests()
 		{
+			Assert.IsTrue(SemVer.Compare(SemVer.Parse("5.1.0-19607491"), SemVer.Parse("5.1.0-19607492")) < 0);
+			Assert.IsTrue(SemVer.Compare(SemVer.Parse("5.1.0-19607492"), SemVer.Parse("5.1.0-19607492")) == 0);
+			Assert.IsTrue(SemVer.Compare(SemVer.Parse("5.1.0-19607493"), SemVer.Parse("5.1.0-19607492")) > 0);
+
 			Assert.IsTrue(SemVer.Parse("1.0.0-123") < SemVer.Parse("1.0.0-124"));
 			Assert.IsTrue(SemVer.Parse("1.0.0-124") > SemVer.Parse("1.0.0-123"));
 
diff --git a/Engine/Source/Programs/Shared/EpicGames.Core/SemVer.cs b/Engine/Source/Programs/Shared/EpicGames.Core/SemVer.cs
--- a/Engine/Source/Programs/Shared/EpicGames.Core/SemVer.cs
+++ b/Engine/Source/Programs/Shared/EpicGames.Core/SemVer.cs
@@ -92,6 +92,14 @@ namespace EpicGames.Core
 		/// <inheritdoc/>
 		public override string ToString() => _text;
 
+		/// <summary>
+		/// Compare two versions
+		/// </summary>
+		/// <param name="lhs">First version to compare</param>
+		/// <param name="rhs">Second version to compare</param>
+		/// <returns>A number indicating the order of the two version strings</returns>
+		public static int Compare(SemVer lhs, SemVer rhs) => Compare(lhs.Span, rhs.Span);
+
 		static int Compare(ReadOnlySpan<char> lhs, ReadOnlySpan<char> rhs)
 		{
 			// Compare major fields
@@ -131,7 +139,7 @@ namespace EpicGames.Core
 			}
 
 			// Otherwise compare each pre-release field
-			for (; ; )
+			while (lhsField.End < lhs.Length && rhsField.End < rhs.Length)
 			{
 				lhsField = GetField(lhs, lhsField.End + 1);
 				rhsField = GetField(rhs, rhsField.End + 1);
@@ -141,13 +149,10 @@ namespace EpicGames.Core
 				{
 					return diff;
 				}
-
-				diff = ((lhsField.End == lhs.Length) ? 0 : 1) - ((rhsField.End == rhs.Length) ? 0 : 1);
-				if (diff != 0)
-				{
-					return diff;
-				}
 			}
+
+			// Give precedence to the longer version string
+			return ((lhsField.End == lhs.Length) ? 0 : 1) - ((rhsField.End == rhs.Length) ? 0 : 1);
 		}
 
 		static Field GetField(ReadOnlySpan<char> text, int pos)
