Added new property type text needed by script exporters.
Added plugin module detection and output path support for script exporters.

#rnx

diff --git a/Engine/Source/Programs/Shared/EpicGames.UHT/Tables/UhtExporterTable.cs b/Engine/Source/Programs/Shared/EpicGames.UHT/Tables/UhtExporterTable.cs
--- a/Engine/Source/Programs/Shared/EpicGames.UHT/Tables/UhtExporterTable.cs
+++ b/Engine/Source/Programs/Shared/EpicGames.UHT/Tables/UhtExporterTable.cs
@@ -92,6 +92,11 @@ namespace EpicGames.UHT.Tables
 		/// </summary>
 		public string Description = string.Empty;
 
+		/// <summary>
+		/// Exporters in plugins need to specify a module name
+		/// </summary>
+		public string ModuleName = string.Empty;
+
 		/// <summary>
 		/// Exporter options
 		/// </summary>
@@ -129,6 +134,11 @@ namespace EpicGames.UHT.Tables
 		/// </summary>
 		public string Description;
 
+		/// <summary>
+		/// Exporters in plugins need to specify a module name
+		/// </summary>
+		public string ModuleName;
+
 		/// <summary>
 		/// Exporter options
 		/// </summary>
@@ -188,10 +198,19 @@ namespace EpicGames.UHT.Tables
 				throw new UhtIceException("An exporter must have a name");
 			}
 
+			if (Assembly.GetExecutingAssembly() != Type.Assembly)
+			{
+				if (string.IsNullOrEmpty(ExporterAttribute.ModuleName))
+				{
+					throw new UhtIceException("An exporter in a UBT plugin must specify a ModuleName");
+				}
+			}
+
 			UhtExporter ExporterValue = new UhtExporter
 			{
 				Name = ExporterAttribute.Name,
 				Description = ExporterAttribute.Description,
+				ModuleName = ExporterAttribute.ModuleName,
 				Delegate = (UhtExporterDelegate)Delegate.CreateDelegate(typeof(UhtExporterDelegate), MethodInfo),
 				Options = ExporterAttribute.Options,
 				CppFilters = ExporterAttribute.CppFilters,
diff --git a/Engine/Source/Programs/Shared/EpicGames.UHT/Types/Properties/UhtBooleanProperty.cs b/Engine/Source/Programs/Shared/EpicGames.UHT/Types/Properties/UhtBooleanProperty.cs
--- a/Engine/Source/Programs/Shared/EpicGames.UHT/Types/Properties/UhtBooleanProperty.cs
+++ b/Engine/Source/Programs/Shared/EpicGames.UHT/Types/Properties/UhtBooleanProperty.cs
@@ -43,12 +43,13 @@ namespace EpicGames.UHT.Types
 				case UhtPropertyTextType.FunctionThunkRetVal:
 				case UhtPropertyTextType.RigVMTemplateArg:
 				case UhtPropertyTextType.ExportMember:
+				case UhtPropertyTextType.GenericFunctionArgOrRetVal:
 					Builder.Append(this.CppTypeText);
 					break;
 
 				case UhtPropertyTextType.Sparse:
 				case UhtPropertyTextType.SparseShort:
-				case UhtPropertyTextType.GenericFunctionArgOrRetVal:
+				case UhtPropertyTextType.GenericFunctionArgOrRetValImpl:
 				case UhtPropertyTextType.ClassFunctionArgOrRetVal:
 				case UhtPropertyTextType.EventFunctionArgOrRetVal:
 				case UhtPropertyTextType.InterfaceFunctionArgOrRetVal:
diff --git a/Engine/Source/Programs/Shared/EpicGames.UHT/Types/Properties/UhtClassProperty.cs b/Engine/Source/Programs/Shared/EpicGames.UHT/Types/Properties/UhtClassProperty.cs
--- a/Engine/Source/Programs/Shared/EpicGames.UHT/Types/Properties/UhtClassProperty.cs
+++ b/Engine/Source/Programs/Shared/EpicGames.UHT/Types/Properties/UhtClassProperty.cs
@@ -75,6 +75,7 @@ namespace EpicGames.UHT.Types
 				case UhtPropertyTextType.Sparse:
 				case UhtPropertyTextType.SparseShort:
 				case UhtPropertyTextType.GenericFunctionArgOrRetVal:
+				case UhtPropertyTextType.GenericFunctionArgOrRetValImpl:
 				case UhtPropertyTextType.ClassFunctionArgOrRetVal:
 				case UhtPropertyTextType.EventFunctionArgOrRetVal:
 				case UhtPropertyTextType.InterfaceFunctionArgOrRetVal:
diff --git a/Engine/Source/Programs/Shared/EpicGames.UHT/Types/Properties/UhtEnumProperty.cs b/Engine/Source/Programs/Shared/EpicGames.UHT/Types/Properties/UhtEnumProperty.cs
--- a/Engine/Source/Programs/Shared/EpicGames.UHT/Types/Properties/UhtEnumProperty.cs
+++ b/Engine/Source/Programs/Shared/EpicGames.UHT/Types/Properties/UhtEnumProperty.cs
@@ -110,6 +110,7 @@ namespace EpicGames.UHT.Types
 				case UhtPropertyTextType.ClassFunctionArgOrRetVal:
 				case UhtPropertyTextType.InterfaceFunctionArgOrRetVal:
 				case UhtPropertyTextType.EventFunctionArgOrRetVal:
+				case UhtPropertyTextType.GenericFunctionArgOrRetValImpl:
 					if (Enum.CppForm == UhtEnumCppForm.EnumClass || (!bIsTemplateArgument && (Property.PropertyFlags.HasAnyFlags(EPropertyFlags.ReturnParm) || !Property.PropertyFlags.HasAnyFlags(EPropertyFlags.OutParm))))
 					{
 						Builder.Append(Enum.CppType);
diff --git a/Engine/Source/Programs/Shared/EpicGames.UHT/Types/UhtProperty.cs b/Engine/Source/Programs/Shared/EpicGames.UHT/Types/UhtProperty.cs
--- a/Engine/Source/Programs/Shared/EpicGames.UHT/Types/UhtProperty.cs
+++ b/Engine/Source/Programs/Shared/EpicGames.UHT/Types/UhtProperty.cs
@@ -443,6 +443,11 @@ namespace EpicGames.UHT.Types
 		/// </summary>
 		GenericFunctionArgOrRetVal,
 
+		/// <summary>
+		/// Generic function argument or return value implementation (specific to booleans and will always return "bool")
+		/// </summary>
+		GenericFunctionArgOrRetValImpl,
+
 		/// <summary>
 		/// Class function argument or return value
 		/// </summary>
@@ -529,6 +534,7 @@ namespace EpicGames.UHT.Types
 		{
 			return
 				TextType == UhtPropertyTextType.GenericFunctionArgOrRetVal ||
+				TextType == UhtPropertyTextType.GenericFunctionArgOrRetValImpl ||
 				TextType == UhtPropertyTextType.ClassFunctionArgOrRetVal ||
 				TextType == UhtPropertyTextType.EventFunctionArgOrRetVal ||
 				TextType == UhtPropertyTextType.InterfaceFunctionArgOrRetVal;
diff --git a/Engine/Source/Programs/Shared/EpicGames.UHT/Utils/UhtExport.cs b/Engine/Source/Programs/Shared/EpicGames.UHT/Utils/UhtExport.cs
--- a/Engine/Source/Programs/Shared/EpicGames.UHT/Utils/UhtExport.cs
+++ b/Engine/Source/Programs/Shared/EpicGames.UHT/Utils/UhtExport.cs
@@ -27,6 +27,11 @@ namespace EpicGames.UHT.Utils
 		/// </summary>
 		public UhtSession Session { get; }
 
+		/// <summary>
+		/// If this exporter is from a plugin, this points to the module of the plugin
+		/// </summary>
+		public UHTManifest.Module? PluginModule { get; }
+
 		/// <summary>
 		/// Create a task
 		/// </summary>
@@ -75,12 +80,11 @@ namespace EpicGames.UHT.Utils
 		public string MakePath(UhtPackage Package, string Suffix);
 
 		/// <summary>
-		/// Make a path for an output based on the package output directory.
+		/// Make a path for the given file name and extension.  This is only valid for plugins.
 		/// </summary>
-		/// <param name="Package">Destination package</param>
 		/// <param name="FileName">Name of the file</param>
 		/// <param name="Extension">Extension to add to the file</param>
 		/// <returns>Output file path</returns>
-		public string MakePath(UhtPackage Package, string FileName, string Extension);
+		public string MakePath(string FileName, string Extension);
 	}
 }
diff --git a/Engine/Source/Programs/Shared/EpicGames.UHT/Utils/UhtSession.cs b/Engine/Source/Programs/Shared/EpicGames.UHT/Utils/UhtSession.cs
--- a/Engine/Source/Programs/Shared/EpicGames.UHT/Utils/UhtSession.cs
+++ b/Engine/Source/Programs/Shared/EpicGames.UHT/Utils/UhtSession.cs
@@ -64,6 +64,11 @@ namespace EpicGames.UHT.Utils
 		/// </summary>
 		private readonly UhtSession SessionInternal;
 
+		/// <summary>
+		/// Module associated with the plugin
+		/// </summary>
+		private readonly UHTManifest.Module? PluginModuleInternal;
+
 		/// <summary>
 		/// Limiter for the number of files being saved to the reference directory.
 		/// The OS can get swamped on high core systems
@@ -80,6 +85,11 @@ namespace EpicGames.UHT.Utils
 		/// </summary>
 		public UhtSession Session => this.SessionInternal;
 
+		/// <summary>
+		/// Plugin module
+		/// </summary>
+		public UHTManifest.Module? PluginModule => this.PluginModuleInternal;
+
 		/// <summary>
 		/// Collection of error from mismatches with the reference files
 		/// </summary>
@@ -104,11 +114,13 @@ namespace EpicGames.UHT.Utils
 		/// Create a new instance of the export factory
 		/// </summary>
 		/// <param name="Session">UHT session</param>
+		/// <param name="PluginModule">Plugin module</param>
 		/// <param name="Exporter">Exporter being run</param>
-		public UhtExportFactory(UhtSession Session, UhtExporter Exporter)
+		public UhtExportFactory(UhtSession Session, UHTManifest.Module? PluginModule, UhtExporter Exporter)
 		{
-			this.Exporter = Exporter;
 			this.SessionInternal = Session;
+			this.PluginModuleInternal = PluginModule;
+			this.Exporter = Exporter;
 			if (this.Session.ReferenceMode != UhtReferenceMode.None)
 			{
 				this.ReferenceDirectory = Path.Combine(this.Session.ReferenceDirectory, this.Exporter.Name);
@@ -189,8 +201,7 @@ namespace EpicGames.UHT.Utils
 		/// <returns>Resulting file name</returns>
 		public string MakePath(UhtHeaderFile HeaderFile, string Suffix)
 		{
-			UHTManifest.Module Module = HeaderFile.Package.Module;
-			return Path.Combine(Module.OutputDirectory, HeaderFile.FileNameWithoutExtension) + Suffix;
+			return MakePath(HeaderFile.Package.Module, HeaderFile.FileNameWithoutExtension, Suffix);
 		}
 
 		/// <summary>
@@ -201,22 +212,32 @@ namespace EpicGames.UHT.Utils
 		/// <returns>Resulting file name</returns>
 		public string MakePath(UhtPackage Package, string Suffix)
 		{
-			UHTManifest.Module Module = Package.Module;
-			return Path.Combine(Module.OutputDirectory, Package.ShortName.ToString()) + Suffix;
+			return MakePath(Package.Module, Package.ShortName, Suffix);
 		}
 
 
 		/// <summary>
 		/// Make a path for an output based on the package output directory.
 		/// </summary>
-		/// <param name="Package">Destination package</param>
 		/// <param name="FileName">Name of the file</param>
 		/// <param name="Extension">Extension to add to the file</param>
 		/// <returns>Output file path</returns>
-		public string MakePath(UhtPackage Package, string FileName, string Extension)
+		public string MakePath(string FileName, string Extension)
 		{
-			UHTManifest.Module Module = Package.Module;
-			return Path.Combine(Module.OutputDirectory, FileName) + Extension;
+			if (this.PluginModule == null)
+			{
+				throw new UhtIceException("MakePath with just a filename and extension can not be called from non-plugin exporters");
+			}
+			return MakePath(this.PluginModule, FileName, Extension);
+		}
+
+		private string MakePath(UHTManifest.Module Module, string FileName, string Suffix)
+		{
+			if (PluginModule != null)
+			{
+				Module = PluginModule;
+			}
+			return Path.Combine(Module.OutputDirectory, FileName) + Suffix;
 		}
 
 		/// <summary>
@@ -2260,7 +2281,7 @@ namespace EpicGames.UHT.Utils
 				UhtBuffer? Contents = this.ReadSourceToBuffer(this.ProjectFile);
 				if (Contents != null)
 				{
-					this.ProjectJson = JsonDocument.Parse(Contents.Block);
+					this.ProjectJson = JsonDocument.Parse(Contents.Memory);
 					UhtBuffer.Return(Contents);
 				}
 			}
@@ -2322,10 +2343,28 @@ namespace EpicGames.UHT.Utils
 							(Exporter.Options.HasAnyFlags(UhtExporterOptions.Default) && !this.bNoDefaultExporters);
 					}
 
+					UHTManifest.Module? PluginModule = null;
+					if (!string.IsNullOrEmpty(Exporter.ModuleName))
+					{
+						foreach (UHTManifest.Module Module in this.Manifest!.Modules)
+						{
+							if (string.Compare(Module.Name, Exporter.ModuleName, StringComparison.OrdinalIgnoreCase) == 0)
+							{
+								PluginModule = Module;
+								break;
+							}
+						}
+						if (PluginModule == null)
+						{
+							Log.Logger.LogWarning($"Exporter \"{Exporter.Name}\" skipped because module \"{Exporter.ModuleName}\" was not found in manifest");
+							continue;
+						}
+					}
+
 					if (Run)
 					{
 						Log.Logger.LogTrace($"       Running exporter {Exporter.Name}");
-						UhtExportFactory Factory = new UhtExportFactory(this, Exporter);
+						UhtExportFactory Factory = new UhtExportFactory(this, PluginModule, Exporter);
 						Factory.Run();
 						foreach (var Output in Factory.Outputs)
 						{
diff --git a/Engine/Source/Programs/Shared/EpicGames.UHT/Utils/UhtTables.cs b/Engine/Source/Programs/Shared/EpicGames.UHT/Utils/UhtTables.cs
--- a/Engine/Source/Programs/Shared/EpicGames.UHT/Utils/UhtTables.cs
+++ b/Engine/Source/Programs/Shared/EpicGames.UHT/Utils/UhtTables.cs
@@ -551,8 +551,7 @@ namespace EpicGames.UHT.Utils
 			{
 				if (Assembly == null)
 				{
-					byte[] AssemblyBytes = File.ReadAllBytes(AssemblyFilePath);
-					Assembly = Assembly.Load(AssemblyBytes);
+					Assembly = Assembly.LoadFile(AssemblyFilePath);
 				}
 			}
 			return Assembly;
