Renamed UhtBooleanProperty to UhtBoolProperty and eliminated the subclasses to reduce confusion with engine types.

#rnx

diff --git a/Engine/Source/Programs/Shared/EpicGames.UHT/Exporters/CodeGen/UhtHeaderCodeGenerator.cs b/Engine/Source/Programs/Shared/EpicGames.UHT/Exporters/CodeGen/UhtHeaderCodeGenerator.cs
--- a/Engine/Source/Programs/Shared/EpicGames.UHT/Exporters/CodeGen/UhtHeaderCodeGenerator.cs
+++ b/Engine/Source/Programs/Shared/EpicGames.UHT/Exporters/CodeGen/UhtHeaderCodeGenerator.cs
@@ -298,7 +298,7 @@ namespace EpicGames.UHT.Exporters.CodeGen
 				else
 				{
 					Builder.AppendTabs(Tabs).Append("Parms.").Append(Property.SourceName).Append("=").Append(Property.SourceName);
-					if (Property is UhtBooleanProperty)
+					if (Property is UhtBoolProperty)
 					{
 						Builder.Append(" ? true : false");
 					}
@@ -346,7 +346,7 @@ namespace EpicGames.UHT.Exporters.CodeGen
 			if (ReturnProperty != null)
 			{
 				// Make sure uint32 -> bool is supported
-				if (ReturnProperty is UhtBooleanProperty)
+				if (ReturnProperty is UhtBoolProperty)
 				{
 					Builder
 						.AppendTabs(Tabs)
diff --git a/Engine/Source/Programs/Shared/EpicGames.UHT/Exporters/CodeGen/UhtHeaderCodeGeneratorHFile.cs b/Engine/Source/Programs/Shared/EpicGames.UHT/Exporters/CodeGen/UhtHeaderCodeGeneratorHFile.cs
--- a/Engine/Source/Programs/Shared/EpicGames.UHT/Exporters/CodeGen/UhtHeaderCodeGeneratorHFile.cs
+++ b/Engine/Source/Programs/Shared/EpicGames.UHT/Exporters/CodeGen/UhtHeaderCodeGeneratorHFile.cs
@@ -552,7 +552,7 @@ namespace EpicGames.UHT.Exporters.CodeGen
 								{
 									string PropertyName = SparseProperty.SourceName;
 									string CleanPropertyName = PropertyName;
-									if (SparseProperty is UhtBooleanProperty && PropertyName.StartsWith("b"))
+									if (SparseProperty is UhtBoolProperty && PropertyName.StartsWith("b"))
 									{
 										CleanPropertyName = PropertyName.Substring(1);
 									}
diff --git a/Engine/Source/Programs/Shared/EpicGames.UHT/Types/Properties/UhtBoolProperty.cs b/Engine/Source/Programs/Shared/EpicGames.UHT/Types/Properties/UhtBoolProperty.cs
new file mode 100644
--- /dev/null
+++ b/Engine/Source/Programs/Shared/EpicGames.UHT/Types/Properties/UhtBoolProperty.cs
@@ -0,0 +1,276 @@
+// Copyright Epic Games, Inc. All Rights Reserved.
+
+using EpicGames.Core;
+using EpicGames.UHT.Tables;
+using EpicGames.UHT.Tokenizer;
+using EpicGames.UHT.Utils;
+using System.Text;
+
+namespace EpicGames.UHT.Types
+{
+
+	/// <summary>
+	/// Type of boolean
+	/// </summary>
+	public enum UhtBoolType
+	{
+
+		/// <summary>
+		/// Native bool
+		/// </summary>
+		Native,
+
+		/// <summary>
+		/// Used for all bitmask uint booleans
+		/// </summary>
+		UInt8,
+
+		/// <summary>
+		/// Currently unused
+		/// </summary>
+		UInt16,
+
+		/// <summary>
+		/// Currently unused
+		/// </summary>
+		UInt32,
+
+		/// <summary>
+		/// Currently unused
+		/// </summary>
+		UInt64,
+	}
+
+	/// <summary>
+	/// Represents the FBoolProperty engine type
+	/// </summary>
+	[UnrealHeaderTool]
+	[UhtEngineClass(Name = "BoolProperty", IsProperty = true)]
+	public class UhtBoolProperty : UhtProperty
+	{
+		/// <inheritdoc/>
+		public override string EngineClassName { get => "BoolProperty"; }
+
+		/// <inheritdoc/>
+		protected override string CppTypeText
+		{
+			get
+			{
+				switch (this.BoolType)
+				{
+					case UhtBoolType.Native:
+						return "bool";
+					case UhtBoolType.UInt8:
+						return "uint8";
+					case UhtBoolType.UInt16:
+						return "uint16";
+					case UhtBoolType.UInt32:
+						return "uint32";
+					case UhtBoolType.UInt64:
+						return "uint64";
+					default:
+						throw new UhtIceException("Unexpected boolean type");
+				}
+			}
+		}
+
+		/// <inheritdoc/>
+		protected override string PGetMacroText
+		{
+			get
+			{
+				switch (this.BoolType)
+				{
+					case UhtBoolType.Native:
+						return "UBOOL";
+					case UhtBoolType.UInt8:
+						return "UBOOL8";
+					case UhtBoolType.UInt16:
+						return "UBOOL16";
+					case UhtBoolType.UInt32:
+						return "UBOOL32";
+					case UhtBoolType.UInt64:
+						return "UBOOL64";
+					default:
+						throw new UhtIceException("Unexpected boolean type");
+				}
+			}
+		}
+
+		/// <summary>
+		/// If true, the boolean is a native bool and not a UBOOL
+		/// </summary>
+		protected bool bIsNativeBool { get => this.BoolType == UhtBoolType.Native; }
+
+		/// <summary>
+		/// Type of the boolean
+		/// </summary>
+		public readonly UhtBoolType BoolType;
+
+		/// <summary>
+		/// Construct a new boolean property
+		/// </summary>
+		/// <param name="PropertySettings">Property settings</param>
+		/// <param name="BoolType">Type of the boolean</param>
+		public UhtBoolProperty(UhtPropertySettings PropertySettings, UhtBoolType BoolType) : base(PropertySettings)
+		{
+			this.PropertyCaps |= UhtPropertyCaps.RequiresNullConstructorArg | UhtPropertyCaps.IsParameterSupportedByBlueprint | UhtPropertyCaps.IsMemberSupportedByBlueprint;
+			if (BoolType == UhtBoolType.Native || BoolType == UhtBoolType.UInt8)
+			{
+				this.PropertyCaps |= UhtPropertyCaps.CanExposeOnSpawn;
+			}
+			this.BoolType = BoolType;
+		}
+
+		/// <inheritdoc/>
+		public override StringBuilder AppendText(StringBuilder Builder, UhtPropertyTextType TextType, bool bIsTemplateArgument = false)
+		{
+			switch (TextType)
+			{
+				case UhtPropertyTextType.Generic:
+				case UhtPropertyTextType.Construction:
+				case UhtPropertyTextType.FunctionThunkParameterArrayType:
+				case UhtPropertyTextType.FunctionThunkRetVal:
+				case UhtPropertyTextType.RigVMTemplateArg:
+				case UhtPropertyTextType.ExportMember:
+				case UhtPropertyTextType.GenericFunctionArgOrRetVal:
+					Builder.Append(this.CppTypeText);
+					break;
+
+				case UhtPropertyTextType.Sparse:
+				case UhtPropertyTextType.SparseShort:
+				case UhtPropertyTextType.GenericFunctionArgOrRetValImpl:
+				case UhtPropertyTextType.ClassFunctionArgOrRetVal:
+				case UhtPropertyTextType.EventFunctionArgOrRetVal:
+				case UhtPropertyTextType.InterfaceFunctionArgOrRetVal:
+				case UhtPropertyTextType.EventParameterMember:
+				case UhtPropertyTextType.EventParameterFunctionMember:
+				case UhtPropertyTextType.GetterSetterArg:
+					Builder.Append("bool");
+					break;
+			}
+			return Builder;
+		}
+
+		/// <inheritdoc/>
+		public override StringBuilder AppendMemberDecl(StringBuilder Builder, IUhtPropertyMemberContext Context, string Name, string NameSuffix, int Tabs)
+		{
+			Builder.AppendMetaDataDecl(this, Context, Name, NameSuffix, Tabs);
+			if (!(this.Outer is UhtProperty))
+			{
+				Builder.AppendTabs(Tabs).Append("static void ").AppendNameDecl(Context, Name, NameSuffix).Append("_SetBit(void* Obj);\r\n");
+			}
+			Builder.AppendTabs(Tabs).Append("static const UECodeGen_Private::").Append("FBoolPropertyParams").Append(' ').AppendNameDecl(Context, Name, NameSuffix).Append(";\r\n");
+			return Builder;
+		}
+
+		/// <inheritdoc/>
+		public override StringBuilder AppendMemberDef(StringBuilder Builder, IUhtPropertyMemberContext Context, string Name, string NameSuffix, string? Offset, int Tabs)
+		{
+			Builder.AppendMetaDataDef(this, Context, Name, NameSuffix, Tabs);
+			if (this.Outer == Context.OuterStruct)
+			{
+				Builder.AppendTabs(Tabs).Append("void ").AppendNameDef(Context, Name, NameSuffix).Append("_SetBit(void* Obj)\r\n");
+				Builder.AppendTabs(Tabs).Append("{\r\n");
+				Builder.AppendTabs(Tabs + 1).Append("((").Append(Context.OuterStructSourceName).Append("*)Obj)->").Append(this.SourceName).Append(" = 1;\r\n");
+				Builder.AppendTabs(Tabs).Append("}\r\n");
+			}
+
+			AppendMemberDefStart(Builder, Context, Name, NameSuffix, Offset, Tabs, "FBoolPropertyParams",
+				bIsNativeBool ?
+				"UECodeGen_Private::EPropertyGenFlags::Bool | UECodeGen_Private::EPropertyGenFlags::NativeBool" :
+				"UECodeGen_Private::EPropertyGenFlags::Bool ", 
+				false, false);
+
+			Builder.Append("sizeof(").Append(this.CppTypeText).Append("), ");
+
+			if (this.Outer == Context.OuterStruct)
+			{
+				Builder
+					.Append("sizeof(").Append(Context.OuterStructSourceName).Append("), ")
+					.Append("&").AppendNameDef(Context, Name, NameSuffix).Append("_SetBit, ");
+			}
+			else
+			{
+				Builder.Append("0, nullptr, ");
+			}
+
+			AppendMemberDefEnd(Builder, Context, Name, NameSuffix);
+			return Builder;
+		}
+
+		/// <inheritdoc/>
+		public override StringBuilder AppendFullDecl(StringBuilder Builder, UhtPropertyTextType TextType, bool bSkipParameterName = false)
+		{
+			AppendText(Builder, TextType);
+
+			//@todo we currently can't have out bools.. so this isn't really necessary, but eventually out bools may be supported, so leave here for now
+			if (TextType.IsParameter() && this.PropertyFlags.HasAnyFlags(EPropertyFlags.OutParm))
+			{
+				Builder.Append('&');
+			}
+
+			Builder.Append(' ');
+
+			if (!bSkipParameterName)
+			{
+				Builder.Append(this.SourceName);
+			}
+
+			if (this.ArrayDimensions != null)
+			{
+				Builder.Append('[').Append(this.ArrayDimensions).Append(']');
+			}
+			else if (TextType == UhtPropertyTextType.ExportMember && !this.bIsNativeBool)
+			{
+				Builder.Append(":1");
+			}
+			return Builder;
+		}
+
+		/// <inheritdoc/>
+		public override StringBuilder AppendNullConstructorArg(StringBuilder Builder, bool bIsInitializer)
+		{
+			Builder.Append("false");
+			return Builder;
+		}
+
+		/// <inheritdoc/>
+		public override bool SanitizeDefaultValue(IUhtTokenReader DefaultValueReader, StringBuilder InnerDefaultValue)
+		{
+			UhtToken Identifier = DefaultValueReader.GetIdentifier();
+			if (Identifier.IsValue("true") || Identifier.IsValue("false"))
+			{
+				InnerDefaultValue.Append(Identifier.Value.ToString());
+				return true;
+			}
+			return false;
+		}
+
+		/// <inheritdoc/>
+		public override bool IsSameType(UhtProperty Other)
+		{
+			// We don't test BoolType.
+			return Other is UhtBoolProperty;
+		}
+
+		/// <inheritdoc/>
+		public override string? GetRigVMType(ref UhtRigVMParameterFlags ParameterFlags)
+		{
+			return CppTypeText;
+		}
+
+		#region Keyword
+		[UhtPropertyType(Keyword = "bool", Options = UhtPropertyTypeOptions.Simple | UhtPropertyTypeOptions.Immediate)]
+		private static UhtProperty? BoolProperty(UhtPropertyResolvePhase ResolvePhase, UhtPropertySettings PropertySettings, IUhtTokenReader TokenReader, UhtToken MatchedToken)
+		{
+			if (PropertySettings.bIsBitfield)
+			{
+				TokenReader.LogError("bool bitfields are not supported.");
+				return null;
+			}
+			return new UhtBoolProperty(PropertySettings, UhtBoolType.Native);
+		}
+		#endregion
+	}
+}
diff --git a/Engine/Source/Programs/Shared/EpicGames.UHT/Types/Properties/UhtBooleanProperty.cs b/Engine/Source/Programs/Shared/EpicGames.UHT/Types/Properties/UhtBooleanProperty.cs
deleted file mode 100644
diff --git a/Engine/Source/Programs/Shared/EpicGames.UHT/Types/Properties/UhtByteProperty.cs b/Engine/Source/Programs/Shared/EpicGames.UHT/Types/Properties/UhtByteProperty.cs
--- a/Engine/Source/Programs/Shared/EpicGames.UHT/Types/Properties/UhtByteProperty.cs
+++ b/Engine/Source/Programs/Shared/EpicGames.UHT/Types/Properties/UhtByteProperty.cs
@@ -169,7 +169,7 @@ namespace EpicGames.UHT.Types
 		{
 			if (PropertySettings.bIsBitfield)
 			{
-				return new UhtBool8Property(PropertySettings);
+				return new UhtBoolProperty(PropertySettings, UhtBoolType.UInt8);
 			}
 			else
 			{
diff --git a/Engine/Source/Programs/Shared/EpicGames.UHT/Types/Properties/UhtUInt16Property.cs b/Engine/Source/Programs/Shared/EpicGames.UHT/Types/Properties/UhtUInt16Property.cs
--- a/Engine/Source/Programs/Shared/EpicGames.UHT/Types/Properties/UhtUInt16Property.cs
+++ b/Engine/Source/Programs/Shared/EpicGames.UHT/Types/Properties/UhtUInt16Property.cs
@@ -62,7 +62,7 @@ namespace EpicGames.UHT.Types
 		{
 			if (PropertySettings.bIsBitfield)
 			{
-				return new UhtBool8Property(PropertySettings);
+				return new UhtBoolProperty(PropertySettings, UhtBoolType.UInt8);
 			}
 			else
 			{
diff --git a/Engine/Source/Programs/Shared/EpicGames.UHT/Types/Properties/UhtUInt32Property.cs b/Engine/Source/Programs/Shared/EpicGames.UHT/Types/Properties/UhtUInt32Property.cs
--- a/Engine/Source/Programs/Shared/EpicGames.UHT/Types/Properties/UhtUInt32Property.cs
+++ b/Engine/Source/Programs/Shared/EpicGames.UHT/Types/Properties/UhtUInt32Property.cs
@@ -65,7 +65,7 @@ namespace EpicGames.UHT.Types
 		{
 			if (PropertySettings.bIsBitfield)
 			{
-				return new UhtBool8Property(PropertySettings);
+				return new UhtBoolProperty(PropertySettings, UhtBoolType.UInt8);
 			}
 			else
 			{
diff --git a/Engine/Source/Programs/Shared/EpicGames.UHT/Types/Properties/UhtUInt64Property.cs b/Engine/Source/Programs/Shared/EpicGames.UHT/Types/Properties/UhtUInt64Property.cs
--- a/Engine/Source/Programs/Shared/EpicGames.UHT/Types/Properties/UhtUInt64Property.cs
+++ b/Engine/Source/Programs/Shared/EpicGames.UHT/Types/Properties/UhtUInt64Property.cs
@@ -62,7 +62,7 @@ namespace EpicGames.UHT.Types
 		{
 			if (PropertySettings.bIsBitfield)
 			{
-				return new UhtBool8Property(PropertySettings);
+				return new UhtBoolProperty(PropertySettings, UhtBoolType.UInt8);
 			}
 			else
 			{
diff --git a/Engine/Source/Programs/Shared/EpicGames.UHT/Types/UhtProperty.cs b/Engine/Source/Programs/Shared/EpicGames.UHT/Types/UhtProperty.cs
--- a/Engine/Source/Programs/Shared/EpicGames.UHT/Types/UhtProperty.cs
+++ b/Engine/Source/Programs/Shared/EpicGames.UHT/Types/UhtProperty.cs
@@ -1017,7 +1017,7 @@ namespace EpicGames.UHT.Types
 		/// <inheritdoc/>
 		protected override string GetDisplayStringFromEngineName()
 		{
-			return UhtFCString.NameToDisplayString(this.EngineName, this is UhtBooleanProperty);
+			return UhtFCString.NameToDisplayString(this.EngineName, this is UhtBoolProperty);
 		}
 
 		#region Text and code generation support
@@ -1596,7 +1596,7 @@ namespace EpicGames.UHT.Types
 					this.LogError("Static arrays of containers are not allowed");
 				}
 
-				if (this is UhtBooleanProperty)
+				if (this is UhtBoolProperty)
 				{
 					this.LogError("Bool arrays are not allowed");
 				}
