Horde: Add support for pushing data to Horde storage without knowing the hash in advance.

diff --git a/Engine/Source/Programs/Horde/Horde.Build/Storage/Services/BasicStorageClient.cs b/Engine/Source/Programs/Horde/Horde.Build/Storage/Services/BasicStorageClient.cs
--- a/Engine/Source/Programs/Horde/Horde.Build/Storage/Services/BasicStorageClient.cs
+++ b/Engine/Source/Programs/Horde/Horde.Build/Storage/Services/BasicStorageClient.cs
@@ -91,6 +91,24 @@ namespace Horde.Build.Storage.Services
 			await _storageBackend.WriteBytesAsync(path, data);
 		}
 
+		/// <inheritdoc/>
+		public async Task<IoHash> WriteBlobAsync(NamespaceId namespaceId, Stream stream, CancellationToken cancellationToken)
+		{
+			byte[] data;
+			using (MemoryStream buffer = new MemoryStream())
+			{
+				await stream.CopyToAsync(buffer, cancellationToken);
+				data = buffer.ToArray();
+			}
+
+			IoHash dataHash = IoHash.Compute(data);
+
+			string path = GetFullBlobPath(namespaceId, dataHash);
+			await _storageBackend.WriteBytesAsync(path, data);
+
+			return dataHash;
+		}
+
 		/// <inheritdoc/>
 		public async Task<bool> HasBlobAsync(NamespaceId namespaceId, IoHash hash, CancellationToken cancellationToken)
 		{
diff --git a/Engine/Source/Programs/Shared/EpicGames.Horde/Storage/IStorageClient.cs b/Engine/Source/Programs/Shared/EpicGames.Horde/Storage/IStorageClient.cs
--- a/Engine/Source/Programs/Shared/EpicGames.Horde/Storage/IStorageClient.cs
+++ b/Engine/Source/Programs/Shared/EpicGames.Horde/Storage/IStorageClient.cs
@@ -175,6 +175,14 @@ namespace EpicGames.Horde.Storage
 		/// <param name="cancellationToken">Cancellation token for the operation</param>
 		Task WriteBlobAsync(NamespaceId namespaceId, IoHash hash, Stream stream, CancellationToken cancellationToken = default);
 
+		/// <summary>
+		/// Writes a blob to storage
+		/// </summary>
+		/// <param name="namespaceId">Namespace to operate on</param>
+		/// <param name="stream">The stream to write</param>
+		/// <param name="cancellationToken">Cancellation token for the operation</param>
+		Task<IoHash> WriteBlobAsync(NamespaceId namespaceId, Stream stream, CancellationToken cancellationToken = default);
+
 		/// <summary>
 		/// Checks if the given blob exists
 		/// </summary>
@@ -365,6 +373,43 @@ namespace EpicGames.Horde.Storage
 			return hash;
 		}
 
+		/// <summary>
+		/// Gets the given reference
+		/// </summary>
+		/// <param name="storageClient">The storage interface</param>
+		/// <param name="qualifiedRefId">Name of the reference</param>
+		/// <param name="cancellationToken">Cancellation token for the operation</param>
+		/// <returns>The reference data if the ref exists</returns>
+		public static Task<IRef> GetRefAsync(this IStorageClient storageClient, QualifiedRefId qualifiedRefId, CancellationToken cancellationToken = default)
+		{
+			return storageClient.GetRefAsync(qualifiedRefId.NamespaceId, qualifiedRefId.BucketId, qualifiedRefId.RefId, cancellationToken);
+		}
+
+		/// <summary>
+		/// Checks if the given reference exists
+		/// </summary>
+		/// <param name="storageClient">The storage interface</param>
+		/// <param name="qualifiedRefId">Name of the reference</param>
+		/// <param name="cancellationToken">Cancellation token for the operation</param>
+		/// <returns>True if the ref exists, false if it did not exist</returns>
+		public static Task<bool> HasRefAsync(this IStorageClient storageClient, QualifiedRefId qualifiedRefId, CancellationToken cancellationToken = default)
+		{
+			return storageClient.HasRefAsync(qualifiedRefId.NamespaceId, qualifiedRefId.BucketId, qualifiedRefId.RefId, cancellationToken);
+		}
+
+		/// <summary>
+		/// Attempts to sets the given reference, returning a list of missing objects on failure.
+		/// </summary>
+		/// <param name="storageClient">The storage interface</param>
+		/// <param name="qualifiedRefId">Name of the reference</param>
+		/// <param name="value">New value for the reference</param>
+		/// <param name="cancellationToken">Cancellation token for the operation</param>
+		/// <returns>List of missing references</returns>
+		public static Task<List<IoHash>> TrySetRefAsync(this IStorageClient storageClient, QualifiedRefId qualifiedRefId, CbObject value, CancellationToken cancellationToken = default)
+		{
+			return storageClient.TrySetRefAsync(qualifiedRefId.NamespaceId, qualifiedRefId.BucketId, qualifiedRefId.RefId, value, cancellationToken);
+		}
+
 		/// <summary>
 		/// Reads a reference as a specific type
 		/// </summary>
@@ -448,5 +493,17 @@ namespace EpicGames.Horde.Storage
 				throw new RefMissingBlobException(namespaceId, bucketId, refId, missingHashes);
 			}
 		}
+
+		/// <summary>
+		/// Removes the given reference
+		/// </summary>
+		/// <param name="storageClient">The storage interface</param>
+		/// <param name="qualifiedRefId">Name of the reference</param>
+		/// <param name="cancellationToken">Cancellation token for the operation</param>
+		/// <returns>True if the ref was deleted, false if it did not exist</returns>
+		public static Task<bool> DeleteRefAsync(this IStorageClient storageClient, QualifiedRefId qualifiedRefId, CancellationToken cancellationToken = default)
+		{
+			return storageClient.DeleteRefAsync(qualifiedRefId.NamespaceId, qualifiedRefId.BucketId, qualifiedRefId.RefId, cancellationToken);
+		}
 	}
 }
diff --git a/Engine/Source/Programs/Shared/EpicGames.Horde/Storage/Impl/FileStorageClient.cs b/Engine/Source/Programs/Shared/EpicGames.Horde/Storage/Impl/FileStorageClient.cs
--- a/Engine/Source/Programs/Shared/EpicGames.Horde/Storage/Impl/FileStorageClient.cs
+++ b/Engine/Source/Programs/Shared/EpicGames.Horde/Storage/Impl/FileStorageClient.cs
@@ -79,6 +79,18 @@ namespace EpicGames.Horde.Storage.Impl
 			}
 		}
 
+		/// <inheritdoc/>
+		public async Task<IoHash> WriteBlobAsync(NamespaceId namespaceId, Stream stream, CancellationToken cancellationToken = default)
+		{
+			byte[] data;
+			using (MemoryStream memoryStream = new MemoryStream())
+			{
+				await stream.CopyToAsync(memoryStream);
+				data = memoryStream.ToArray();
+			}
+			return await this.WriteBlobFromMemoryAsync(namespaceId, data, cancellationToken);
+		}
+
 		/// <inheritdoc/>
 		public Task<bool> HasBlobAsync(NamespaceId namespaceId, IoHash hash, CancellationToken cancellationToken = default)
 		{
diff --git a/Engine/Source/Programs/Shared/EpicGames.Horde/Storage/Impl/HttpStorageClient.cs b/Engine/Source/Programs/Shared/EpicGames.Horde/Storage/Impl/HttpStorageClient.cs
--- a/Engine/Source/Programs/Shared/EpicGames.Horde/Storage/Impl/HttpStorageClient.cs
+++ b/Engine/Source/Programs/Shared/EpicGames.Horde/Storage/Impl/HttpStorageClient.cs
@@ -10,6 +10,7 @@ using System.Net.Http;
 using System.Net.Http.Headers;
 using System.Text;
 using System.Text.Json;
+using System.Text.Json.Serialization;
 using System.Threading;
 using System.Threading.Tasks;
 
@@ -126,6 +127,31 @@ namespace EpicGames.Horde.Storage.Impl
 			response.EnsureSuccessStatusCode();
 		}
 
+		class WriteBlobResponse
+		{
+			[JsonPropertyName("identifier")]
+			public IoHash Identifier { get; set; }
+		}
+
+		/// <inheritdoc/>
+		public async Task<IoHash> WriteBlobAsync(NamespaceId namespaceId, Stream stream, CancellationToken cancellationToken)
+		{
+			using HttpRequestMessage request = new HttpRequestMessage(HttpMethod.Post, $"api/v1/blobs/{namespaceId}");
+
+			StreamContent content = new StreamContent(stream);
+			content.Headers.ContentType = new MediaTypeHeaderValue("application/octet-stream");
+			request.Content = content;
+
+			HttpResponseMessage response = await _httpClient.SendAsync(request, cancellationToken);
+			response.EnsureSuccessStatusCode();
+
+			using (Stream responseStream = await response.Content.ReadAsStreamAsync())
+			{
+				WriteBlobResponse responseMessage = await JsonSerializer.DeserializeAsync<WriteBlobResponse>(responseStream, cancellationToken: cancellationToken);
+				return responseMessage.Identifier;
+			}
+		}
+
 		/// <inheritdoc/>
 		public async Task<bool> HasBlobAsync(NamespaceId namespaceId, IoHash hash, CancellationToken cancellationToken = default)
 		{
diff --git a/Engine/Source/Programs/Shared/EpicGames.Horde/Storage/Impl/MemoryStorageClient.cs b/Engine/Source/Programs/Shared/EpicGames.Horde/Storage/Impl/MemoryStorageClient.cs
--- a/Engine/Source/Programs/Shared/EpicGames.Horde/Storage/Impl/MemoryStorageClient.cs
+++ b/Engine/Source/Programs/Shared/EpicGames.Horde/Storage/Impl/MemoryStorageClient.cs
@@ -55,6 +55,21 @@ namespace EpicGames.Horde.Storage.Impl
 			}
 		}
 
+		/// <inheritdoc/>
+		public async Task<IoHash> WriteBlobAsync(NamespaceId namespaceId, Stream stream, CancellationToken cancellationToken = default)
+		{
+			using (MemoryStream memoryStream = new MemoryStream())
+			{
+				await stream.CopyToAsync(memoryStream);
+
+				byte[] data = memoryStream.ToArray();
+				IoHash hash = IoHash.Compute(data);
+
+				Blobs[(namespaceId, hash)] = data;
+				return hash;
+			}
+		}
+
 		/// <inheritdoc/>
 		public Task<bool> HasBlobAsync(NamespaceId namespaceId, IoHash hash, CancellationToken cancellationToken = default)
 		{
diff --git a/Engine/Source/Programs/Shared/EpicGames.Horde/Storage/Impl/TracingStorageClient.cs b/Engine/Source/Programs/Shared/EpicGames.Horde/Storage/Impl/TracingStorageClient.cs
--- a/Engine/Source/Programs/Shared/EpicGames.Horde/Storage/Impl/TracingStorageClient.cs
+++ b/Engine/Source/Programs/Shared/EpicGames.Horde/Storage/Impl/TracingStorageClient.cs
@@ -43,6 +43,14 @@ namespace EpicGames.Horde.Storage.Impl
 			return _inner.WriteBlobAsync(namespaceId, hash, stream, cancellationToken);
 		}
 
+		/// <inheritdoc/>
+		public async Task<IoHash> WriteBlobAsync(NamespaceId namespaceId, Stream stream, CancellationToken cancellationToken = default)
+		{
+			IoHash hash = await _inner.WriteBlobAsync(namespaceId, stream, cancellationToken);
+			_logger.LogDebug("Written blob {NamespaceId}/{Hash}", namespaceId, hash);
+			return hash;
+		}
+
 		/// <inheritdoc/>
 		public Task<bool> HasBlobAsync(NamespaceId namespaceId, IoHash hash, CancellationToken cancellationToken = default)
 		{
