Horde: Missing file.

diff --git a/Engine/Source/Programs/Shared/EpicGames.Serialization/Converters/CbStringConverter.cs b/Engine/Source/Programs/Shared/EpicGames.Serialization/Converters/CbStringConverter.cs
new file mode 100644
--- /dev/null
+++ b/Engine/Source/Programs/Shared/EpicGames.Serialization/Converters/CbStringConverter.cs
@@ -0,0 +1,48 @@
+// Copyright Epic Games, Inc. All Rights Reserved.
+
+using EpicGames.Core;
+using System;
+using System.ComponentModel;
+
+namespace EpicGames.Serialization.Converters
+{
+	class CbStringConverter<T> : CbConverterBase<T>
+	{
+		readonly TypeConverter _typeConverter;
+
+		public CbStringConverter(TypeConverter typeConverter)
+		{
+			_typeConverter = typeConverter;
+		}
+
+		public override T Read(CbField field)
+		{
+			return (T)_typeConverter.ConvertFromInvariantString(field.AsString());
+		}
+
+		public override void WriteNamed(CbWriter writer, Utf8String name, T value)
+		{
+			writer.WriteString(name, _typeConverter.ConvertToInvariantString(value));
+		}
+
+		public override void Write(CbWriter writer, T value)
+		{ 
+			writer.WriteStringValue(_typeConverter.ConvertToInvariantString(value));
+		}
+	}
+
+	class CbStringConverterFactory : CbConverterFactory
+	{
+		public override ICbConverter? CreateConverter(Type type)
+		{
+			TypeConverter? frameworkTypeConverter = TypeDescriptor.GetConverter(type);
+			if (frameworkTypeConverter == null || !frameworkTypeConverter.CanConvertFrom(typeof(string)) || !frameworkTypeConverter.CanConvertTo(typeof(string)))
+			{
+				return null;
+			}
+
+			Type converterType = typeof(CbStringConverter<>).MakeGenericType(type);
+			return (ICbConverter)Activator.CreateInstance(converterType, frameworkTypeConverter)!;
+		}
+	}
+}
