Fix issue with UhtClassPtrProperty where the meta name was being placed in the output and not the class.
Add UHT support for a module's public defines.
Add external dependencies support.

#rnx

diff --git a/Engine/Source/Programs/Shared/EpicGames.Core/UHTTypes.cs b/Engine/Source/Programs/Shared/EpicGames.Core/UHTTypes.cs
--- a/Engine/Source/Programs/Shared/EpicGames.Core/UHTTypes.cs
+++ b/Engine/Source/Programs/Shared/EpicGames.Core/UHTTypes.cs
@@ -85,6 +85,7 @@ namespace EpicGames.Core
 			public List<string> PublicHeaders { get; set; } = new List<string>();
 			public List<string> InternalHeaders { get; set; } = new List<string>();
 			public List<string> PrivateHeaders { get; set; } = new List<string>();
+			public List<string> PublicDefines { get; set; } = new List<string>();
 			public string? GeneratedCPPFilenameBase { get; set; } = null;
 			public bool SaveExportedHeaders { get; set; } = false;
 			[JsonConverter(typeof(JsonStringEnumConverter))]
@@ -95,6 +96,40 @@ namespace EpicGames.Core
 			{
 				return Name;
 			}
+
+			public bool TryGetDefine(string Name, out string? Value)
+			{
+				Value = null;
+				int Length = Name.Length;
+				foreach (string Define in PublicDefines)
+				{
+					if (!Define.StartsWith(Name))
+					{
+						continue;
+					}
+					if (Define.Length > Length)
+					{
+						if (Define[Length] != '=')
+						{
+							continue;
+						}
+						Value = Define.Substring(Length + 1, Define.Length - Length - 1);
+					}
+					return true;
+				}
+				return false;
+			}
+
+			public bool TryGetDefine(string Name, out int Value)
+			{
+				string? String;
+				if (TryGetDefine(Name, out String))
+				{
+					return int.TryParse(String, out Value);
+				}
+				Value = 0;
+				return false;
+			}
 		}
 
 		/// <summary>
diff --git a/Engine/Source/Programs/Shared/EpicGames.UHT/Types/Properties/UhtClassPtrProperties.cs b/Engine/Source/Programs/Shared/EpicGames.UHT/Types/Properties/UhtClassPtrProperties.cs
--- a/Engine/Source/Programs/Shared/EpicGames.UHT/Types/Properties/UhtClassPtrProperties.cs
+++ b/Engine/Source/Programs/Shared/EpicGames.UHT/Types/Properties/UhtClassPtrProperties.cs
@@ -45,7 +45,7 @@ namespace EpicGames.UHT.Types
 			switch (TextType)
 			{
 				default:
-					Builder.Append("TObjectPtr<").Append(this.MetaClass?.SourceName).Append(">");
+					Builder.Append("TObjectPtr<").Append(this.Class.SourceName).Append(">");
 					break;
 			}
 			return Builder;
diff --git a/Engine/Source/Programs/Shared/EpicGames.UHT/Utils/UhtExport.cs b/Engine/Source/Programs/Shared/EpicGames.UHT/Utils/UhtExport.cs
--- a/Engine/Source/Programs/Shared/EpicGames.UHT/Utils/UhtExport.cs
+++ b/Engine/Source/Programs/Shared/EpicGames.UHT/Utils/UhtExport.cs
@@ -86,5 +86,11 @@ namespace EpicGames.UHT.Utils
 		/// <param name="Extension">Extension to add to the file</param>
 		/// <returns>Output file path</returns>
 		public string MakePath(string FileName, string Extension);
+
+		/// <summary>
+		/// Add an external dependency to the given file path
+		/// </summary>
+		/// <param name="FilePath">External dependency to add</param>
+		public void AddExternalDependency(string FilePath);
 	}
 }
diff --git a/Engine/Source/Programs/Shared/EpicGames.UHT/Utils/UhtSession.cs b/Engine/Source/Programs/Shared/EpicGames.UHT/Utils/UhtSession.cs
--- a/Engine/Source/Programs/Shared/EpicGames.UHT/Utils/UhtSession.cs
+++ b/Engine/Source/Programs/Shared/EpicGames.UHT/Utils/UhtSession.cs
@@ -110,6 +110,11 @@ namespace EpicGames.UHT.Utils
 		/// </summary>
 		public string VerifyDirectory = string.Empty;
 
+		/// <summary>
+		/// Collection of external dependencies
+		/// </summary>
+		public HashSet<string> ExternalDependencies = new HashSet<string>();
+
 		/// <summary>
 		/// Create a new instance of the export factory
 		/// </summary>
@@ -231,6 +236,16 @@ namespace EpicGames.UHT.Utils
 			return MakePath(this.PluginModule, FileName, Extension);
 		}
 
+
+		/// <summary>
+		/// Add an external dependency to the given file path
+		/// </summary>
+		/// <param name="FilePath">External dependency to add</param>
+		public void AddExternalDependency(string FilePath)
+		{
+			this.ExternalDependencies.Add(FilePath);
+		}
+
 		private string MakePath(UHTManifest.Module Module, string FileName, string Suffix)
 		{
 			if (PluginModule != null)
@@ -2329,6 +2344,7 @@ namespace EpicGames.UHT.Utils
 
 		private void StepExport()
 		{
+			HashSet<string> ExternalDependencies = new HashSet<string>();
 			long TotalWrittenFiles = 0;
 			Try(null, () =>
 			{
@@ -2373,12 +2389,28 @@ namespace EpicGames.UHT.Utils
 								TotalWrittenFiles++;
 							}
 						}
+						foreach (string Dep in Factory.ExternalDependencies)
+						{
+							ExternalDependencies.Add(Dep);
+						}
 					}
 					else
 					{
 						Log.Logger.LogTrace($"       Exporter {Exporter.Name} skipped");
 					}
 				}
+
+				// Save the collected external dependencies
+				if (!string.IsNullOrEmpty(this.Manifest!.ExternalDependenciesFile))
+				{
+					using (StreamWriter Out = new StreamWriter(this.Manifest!.ExternalDependenciesFile))
+					{
+						foreach (string Dep in ExternalDependencies)
+						{
+							Out.WriteLine(Dep);
+						}
+					}
+				}
 			});
 
 			Log.Logger.LogInformation($"Total of {TotalWrittenFiles} written");
