P4: Prevent updates to existing changelists clearing out the file list.

diff --git a/Engine/Source/Programs/Shared/EpicGames.Perforce/PerforceConnection.cs b/Engine/Source/Programs/Shared/EpicGames.Perforce/PerforceConnection.cs
--- a/Engine/Source/Programs/Shared/EpicGames.Perforce/PerforceConnection.cs
+++ b/Engine/Source/Programs/Shared/EpicGames.Perforce/PerforceConnection.cs
@@ -916,6 +916,10 @@ namespace EpicGames.Perforce
 			{
 				nameToValue.Add(new KeyValuePair<string, object>("Description", input.Description));
 			}
+			foreach(string file in input.Files)
+			{
+				nameToValue.Add(new KeyValuePair<string, object>("Files", file));
+			}
 			return SerializeRecord(nameToValue);
 		}
 
