Updated implementations with recent changes to C++ UHT.
Changed RigVM type code to use caps system

#rnx

diff --git a/Engine/Source/Programs/Shared/EpicGames.UHT/Exporters/CodeGen/UhtHeaderCodeGeneratorCppFile.cs b/Engine/Source/Programs/Shared/EpicGames.UHT/Exporters/CodeGen/UhtHeaderCodeGeneratorCppFile.cs
--- a/Engine/Source/Programs/Shared/EpicGames.UHT/Exporters/CodeGen/UhtHeaderCodeGeneratorCppFile.cs
+++ b/Engine/Source/Programs/Shared/EpicGames.UHT/Exporters/CodeGen/UhtHeaderCodeGeneratorCppFile.cs
@@ -471,6 +471,29 @@ namespace EpicGames.UHT.Exporters.CodeGen
 				{
 					foreach (UhtRigVMMethodInfo MethodInfo in ScriptStruct.RigVMStructInfo.Methods)
 					{
+						Builder.Append("\t\tTArray<FRigVMFunctionArgument> ").AppendArgumentsName(ScriptStruct, MethodInfo).Append(";\r\n");
+						foreach (UhtRigVMParameter Parameter in ScriptStruct.RigVMStructInfo.Members)
+						{
+							Builder
+								.Append("\t\t")
+								.AppendArgumentsName(ScriptStruct, MethodInfo)
+								.Append(".Emplace(TEXT(\"")
+								.Append(Parameter.NameOriginal())
+								.Append("\"), TEXT(\"")
+								.Append(Parameter.TypeOriginal())
+								.Append("\"));\r\n");
+						}
+						foreach (UhtRigVMParameter Parameter in MethodInfo.Parameters)
+						{
+							Builder
+								.Append("\t\t")
+								.AppendArgumentsName(ScriptStruct, MethodInfo)
+								.Append(".Emplace(TEXT(\"")
+								.Append(Parameter.NameOriginal())
+								.Append("\"), TEXT(\"")
+								.Append(Parameter.TypeOriginal())
+								.Append("\"));\r\n");
+						}
 						Builder
 							.Append("\t\tFRigVMRegistry::Get().Register(TEXT(\"")
 							.Append(ScriptStruct.SourceName)
@@ -482,7 +505,9 @@ namespace EpicGames.UHT.Exporters.CodeGen
 							.Append(MethodInfo.Name)
 							.Append(", ")
 							.Append(RegistrationName)
-							.Append(".OuterSingleton);\r\n");
+							.Append(".OuterSingleton, ")
+							.AppendArgumentsName(ScriptStruct, MethodInfo)
+							.Append(");\r\n");
 					}
 				}
 
@@ -1000,18 +1025,42 @@ namespace EpicGames.UHT.Exporters.CodeGen
 							Builder.Append("\t#if WITH_EDITORONLY_DATA\r\n");
 						}
 						Builder.Append("\t\tconst ").Append(Class.SourceName).Append("* Obj = (const ").Append(Class.SourceName).Append("*)Object;\r\n");
-						Builder
-							.Append("\t\t")
-							.AppendPropertyText(Property, UhtPropertyTextType.GetterSetterArg)
-							.Append("& Result = *(")
-							.AppendPropertyText(Property, UhtPropertyTextType.GetterSetterArg)
-							.Append("*)OutValue;\r\n");
-						Builder
-							.Append("\t\tResult = (")
-							.AppendPropertyText(Property, UhtPropertyTextType.GetterSetterArg)
-							.Append(")Obj->")
-							.Append(Property.Getter!)
-							.Append("();\r\n");
+						if (Property.bIsStaticArray)
+						{
+							Builder
+								.Append("\t\t")
+								.AppendPropertyText(Property, UhtPropertyTextType.GetterSetterArg)
+								.Append("* Source = (")
+								.AppendPropertyText(Property, UhtPropertyTextType.GetterSetterArg)
+								.Append("*)Obj->")
+								.Append(Property.Getter!)
+								.Append("();\r\n");
+							Builder
+								.Append("\t\t")
+								.AppendPropertyText(Property, UhtPropertyTextType.GetterSetterArg)
+								.Append("* Result = (")
+								.AppendPropertyText(Property, UhtPropertyTextType.GetterSetterArg)
+								.Append("*)OutValue;\r\n");
+							Builder
+								.Append("\t\tCopyAssignItems(Result, Source, ")
+								.Append(Property.ArrayDimensions)
+								.Append(");\r\n");
+						}
+						else
+						{
+							Builder
+								.Append("\t\t")
+								.AppendPropertyText(Property, UhtPropertyTextType.GetterSetterArg)
+								.Append("& Result = *(")
+								.AppendPropertyText(Property, UhtPropertyTextType.GetterSetterArg)
+								.Append("*)OutValue;\r\n");
+							Builder
+								.Append("\t\tResult = (")
+								.AppendPropertyText(Property, UhtPropertyTextType.GetterSetterArg)
+								.Append(")Obj->")
+								.Append(Property.Getter!)
+								.Append("();\r\n");
+						}
 						if (bEditorOnlyProperty)
 						{
 							Builder.Append("\t#endif // WITH_EDITORONLY_DATA\r\n");
@@ -1027,12 +1076,24 @@ namespace EpicGames.UHT.Exporters.CodeGen
 							Builder.Append("\t#if WITH_EDITORONLY_DATA\r\n");
 						}
 						Builder.Append("\t\t").Append(Class.SourceName).Append("* Obj = (").Append(Class.SourceName).Append("*)Object;\r\n");
-						Builder
-							.Append("\t\t")
-							.AppendPropertyText(Property, UhtPropertyTextType.GetterSetterArg)
-							.Append("& Value = *(")
-							.AppendPropertyText(Property, UhtPropertyTextType.GetterSetterArg)
-							.Append("*)InValue;\r\n");
+						if (Property.bIsStaticArray)
+						{
+							Builder
+								.Append("\t\t")
+								.AppendPropertyText(Property, UhtPropertyTextType.GetterSetterArg)
+								.Append("* Value = (")
+								.AppendPropertyText(Property, UhtPropertyTextType.GetterSetterArg)
+								.Append("*)InValue;\r\n");
+						}
+						else
+						{
+							Builder
+								.Append("\t\t")
+								.AppendPropertyText(Property, UhtPropertyTextType.GetterSetterArg)
+								.Append("& Value = *(")
+								.AppendPropertyText(Property, UhtPropertyTextType.GetterSetterArg)
+								.Append("*)InValue;\r\n");
+						}
 						Builder
 							.Append("\t\tObj->")
 							.Append(Property.Setter!)
diff --git a/Engine/Source/Programs/Shared/EpicGames.UHT/Exporters/CodeGen/UhtHeaderCodeGeneratorHFile.cs b/Engine/Source/Programs/Shared/EpicGames.UHT/Exporters/CodeGen/UhtHeaderCodeGeneratorHFile.cs
--- a/Engine/Source/Programs/Shared/EpicGames.UHT/Exporters/CodeGen/UhtHeaderCodeGeneratorHFile.cs
+++ b/Engine/Source/Programs/Shared/EpicGames.UHT/Exporters/CodeGen/UhtHeaderCodeGeneratorHFile.cs
@@ -15,7 +15,7 @@ namespace EpicGames.UHT.Exporters.CodeGen
 		: UhtHeaderCodeGenerator
 	{
 		public static string RigVMExecuteContextDeclaration = "FRigVMExtendedExecuteContext& RigVMExecuteContext";
-		public static string RigVMExecuteContextPublicDeclaration = "FRigVMExecuteContext& RigVMExecuteContext";
+		public static string RigVMExecuteContextPublicDeclaration = "const FRigVMExecuteContext& RigVMExecuteContext";
 
 		/// <summary>
 		/// Construct an instance of this generator object
diff --git a/Engine/Source/Programs/Shared/EpicGames.UHT/Exporters/CodeGen/UhtRigVMStringBuilderExtensions.cs b/Engine/Source/Programs/Shared/EpicGames.UHT/Exporters/CodeGen/UhtRigVMStringBuilderExtensions.cs
--- a/Engine/Source/Programs/Shared/EpicGames.UHT/Exporters/CodeGen/UhtRigVMStringBuilderExtensions.cs
+++ b/Engine/Source/Programs/Shared/EpicGames.UHT/Exporters/CodeGen/UhtRigVMStringBuilderExtensions.cs
@@ -9,6 +9,11 @@ namespace EpicGames.UHT.Exporters.CodeGen
 {
 	internal static class UhtRigVMStringBuilderExtensions
 	{
+		public static StringBuilder AppendArgumentsName(this StringBuilder Builder, UhtStruct Struct, UhtRigVMMethodInfo MethodInfo)
+		{
+			return Builder.Append("Arguments_").Append(Struct.SourceName).Append('_').Append(MethodInfo.Name);
+		}
+
 		public static StringBuilder AppendParameterNames(this StringBuilder Builder, IEnumerable<UhtRigVMParameter> Parameters, 
 			bool bLeadingSeparator = false, string Separator = ", ", bool bCastName = false, bool bIncludeEditorOnly = false)
 		{
diff --git a/Engine/Source/Programs/Shared/EpicGames.UHT/Parsers/UhtPropertyParser.cs b/Engine/Source/Programs/Shared/EpicGames.UHT/Parsers/UhtPropertyParser.cs
--- a/Engine/Source/Programs/Shared/EpicGames.UHT/Parsers/UhtPropertyParser.cs
+++ b/Engine/Source/Programs/Shared/EpicGames.UHT/Parsers/UhtPropertyParser.cs
@@ -426,12 +426,6 @@ namespace EpicGames.UHT.Parsers
 		{
 			throw new NotImplementedException();
 		}
-
-		/// <inheritdoc/>
-		public override string? GetRigVMType(ref UhtRigVMParameterFlags ParameterFlags)
-		{
-			throw new NotImplementedException();
-		}
 	}
 
 	/// <summary>
diff --git a/Engine/Source/Programs/Shared/EpicGames.UHT/Types/Properties/UhtArrayProperty.cs b/Engine/Source/Programs/Shared/EpicGames.UHT/Types/Properties/UhtArrayProperty.cs
--- a/Engine/Source/Programs/Shared/EpicGames.UHT/Types/Properties/UhtArrayProperty.cs
+++ b/Engine/Source/Programs/Shared/EpicGames.UHT/Types/Properties/UhtArrayProperty.cs
@@ -46,7 +46,7 @@ namespace EpicGames.UHT.Types
 				Value.MetaData.Remove(UhtNames.NativeConst);
 			}
 
-			this.PropertyCaps |= UhtPropertyCaps.PassCppArgsByRef;
+			this.PropertyCaps |= UhtPropertyCaps.PassCppArgsByRef | UhtPropertyCaps.SupportsRigVM | UhtPropertyCaps.IsRigVMArray;
 			this.PropertyCaps &= ~(UhtPropertyCaps.CanBeContainerValue | UhtPropertyCaps.CanBeContainerKey);
 			this.PropertyCaps = (this.PropertyCaps & ~(UhtPropertyCaps.IsParameterSupportedByBlueprint | UhtPropertyCaps.IsMemberSupportedByBlueprint | UhtPropertyCaps.CanExposeOnSpawn)) |
 				(this.ValueProperty.PropertyCaps & (UhtPropertyCaps.IsParameterSupportedByBlueprint | UhtPropertyCaps.IsMemberSupportedByBlueprint | UhtPropertyCaps.CanExposeOnSpawn));
@@ -205,15 +205,6 @@ namespace EpicGames.UHT.Types
 			return false;
 		}
 
-		/// <inheritdoc/>
-		public override string? GetRigVMType(ref UhtRigVMParameterFlags ParameterFlags)
-		{
-			ParameterFlags |= UhtRigVMParameterFlags.IsArray;
-			StringBuilder Builder = new StringBuilder();
-			Builder.AppendPropertyText(this, UhtPropertyTextType.RigVMTemplateArg);
-			return Builder.ToString();
-		}
-
 #region Keyword
 		[UhtPropertyType(Keyword = "TArray")]
 		private static UhtProperty? ArrayProperty(UhtPropertyResolvePhase ResolvePhase, UhtPropertySettings PropertySettings, IUhtTokenReader TokenReader, UhtToken MatchedToken)
diff --git a/Engine/Source/Programs/Shared/EpicGames.UHT/Types/Properties/UhtBoolProperty.cs b/Engine/Source/Programs/Shared/EpicGames.UHT/Types/Properties/UhtBoolProperty.cs
--- a/Engine/Source/Programs/Shared/EpicGames.UHT/Types/Properties/UhtBoolProperty.cs
+++ b/Engine/Source/Programs/Shared/EpicGames.UHT/Types/Properties/UhtBoolProperty.cs
@@ -114,7 +114,8 @@ namespace EpicGames.UHT.Types
 		/// <param name="BoolType">Type of the boolean</param>
 		public UhtBoolProperty(UhtPropertySettings PropertySettings, UhtBoolType BoolType) : base(PropertySettings)
 		{
-			this.PropertyCaps |= UhtPropertyCaps.RequiresNullConstructorArg | UhtPropertyCaps.IsParameterSupportedByBlueprint | UhtPropertyCaps.IsMemberSupportedByBlueprint;
+			this.PropertyCaps |= UhtPropertyCaps.RequiresNullConstructorArg | UhtPropertyCaps.IsParameterSupportedByBlueprint | 
+				UhtPropertyCaps.IsMemberSupportedByBlueprint | UhtPropertyCaps.SupportsRigVM;
 			if (BoolType == UhtBoolType.Native || BoolType == UhtBoolType.UInt8)
 			{
 				this.PropertyCaps |= UhtPropertyCaps.CanExposeOnSpawn;
@@ -131,7 +132,7 @@ namespace EpicGames.UHT.Types
 				case UhtPropertyTextType.Construction:
 				case UhtPropertyTextType.FunctionThunkParameterArrayType:
 				case UhtPropertyTextType.FunctionThunkRetVal:
-				case UhtPropertyTextType.RigVMTemplateArg:
+				case UhtPropertyTextType.RigVMType:
 				case UhtPropertyTextType.ExportMember:
 				case UhtPropertyTextType.GenericFunctionArgOrRetVal:
 					Builder.Append(this.CppTypeText);
@@ -254,12 +255,6 @@ namespace EpicGames.UHT.Types
 			return Other is UhtBoolProperty;
 		}
 
-		/// <inheritdoc/>
-		public override string? GetRigVMType(ref UhtRigVMParameterFlags ParameterFlags)
-		{
-			return CppTypeText;
-		}
-
 		#region Keyword
 		[UhtPropertyType(Keyword = "bool", Options = UhtPropertyTypeOptions.Simple | UhtPropertyTypeOptions.Immediate)]
 		private static UhtProperty? BoolProperty(UhtPropertyResolvePhase ResolvePhase, UhtPropertySettings PropertySettings, IUhtTokenReader TokenReader, UhtToken MatchedToken)
diff --git a/Engine/Source/Programs/Shared/EpicGames.UHT/Types/Properties/UhtByteProperty.cs b/Engine/Source/Programs/Shared/EpicGames.UHT/Types/Properties/UhtByteProperty.cs
--- a/Engine/Source/Programs/Shared/EpicGames.UHT/Types/Properties/UhtByteProperty.cs
+++ b/Engine/Source/Programs/Shared/EpicGames.UHT/Types/Properties/UhtByteProperty.cs
@@ -44,7 +44,12 @@ namespace EpicGames.UHT.Types
 		public UhtByteProperty(UhtPropertySettings PropertySettings, UhtPropertyIntType IntType, UhtEnum? Enum = null) : base(PropertySettings, IntType)
 		{
 			this.Enum = Enum;
-			this.PropertyCaps |= UhtPropertyCaps.CanExposeOnSpawn | UhtPropertyCaps.IsParameterSupportedByBlueprint | UhtPropertyCaps.IsMemberSupportedByBlueprint;
+			this.PropertyCaps |= UhtPropertyCaps.CanExposeOnSpawn | UhtPropertyCaps.IsParameterSupportedByBlueprint | 
+				UhtPropertyCaps.IsMemberSupportedByBlueprint | UhtPropertyCaps.SupportsRigVM;
+			if (this.Enum != null)
+			{
+				this.PropertyCaps |= UhtPropertyCaps.IsRigVMEnum;
+			}
 		}
 
 		/// <inheritdoc/>
@@ -150,19 +155,6 @@ namespace EpicGames.UHT.Types
 			}
 		}
 
-		/// <inheritdoc/>
-		public override string? GetRigVMType(ref UhtRigVMParameterFlags ParameterFlags)
-		{
-			if (this.Enum != null)
-			{
-				return UhtEnumProperty.GetEnumRigVMType(this, this.Enum, ref ParameterFlags);
-			}
-			else
-			{
-				return CppTypeText;
-			}
-		}
-
 		#region Keyword
 		[UhtPropertyType(Keyword = "uint8", Options = UhtPropertyTypeOptions.Simple | UhtPropertyTypeOptions.Immediate)]
 		private static UhtProperty? UInt8Property(UhtPropertyResolvePhase ResolvePhase, UhtPropertySettings PropertySettings, IUhtTokenReader TokenReader, UhtToken MatchedToken)
diff --git a/Engine/Source/Programs/Shared/EpicGames.UHT/Types/Properties/UhtClassProperty.cs b/Engine/Source/Programs/Shared/EpicGames.UHT/Types/Properties/UhtClassProperty.cs
--- a/Engine/Source/Programs/Shared/EpicGames.UHT/Types/Properties/UhtClassProperty.cs
+++ b/Engine/Source/Programs/Shared/EpicGames.UHT/Types/Properties/UhtClassProperty.cs
@@ -82,7 +82,7 @@ namespace EpicGames.UHT.Types
 				case UhtPropertyTextType.ExportMember:
 				case UhtPropertyTextType.Construction:
 				case UhtPropertyTextType.FunctionThunkParameterArrayType:
-				case UhtPropertyTextType.RigVMTemplateArg:
+				case UhtPropertyTextType.RigVMType:
 				case UhtPropertyTextType.GetterSetterArg:
 					AppendText(Builder);
 					break;
diff --git a/Engine/Source/Programs/Shared/EpicGames.UHT/Types/Properties/UhtDelegateProperty.cs b/Engine/Source/Programs/Shared/EpicGames.UHT/Types/Properties/UhtDelegateProperty.cs
--- a/Engine/Source/Programs/Shared/EpicGames.UHT/Types/Properties/UhtDelegateProperty.cs
+++ b/Engine/Source/Programs/Shared/EpicGames.UHT/Types/Properties/UhtDelegateProperty.cs
@@ -42,7 +42,8 @@ namespace EpicGames.UHT.Types
 		{
 			this.Function = Function;
 			this.HeaderFile.AddReferencedHeader(Function);
-			this.PropertyCaps |= UhtPropertyCaps.PassCppArgsByRef | UhtPropertyCaps.IsParameterSupportedByBlueprint | UhtPropertyCaps.IsMemberSupportedByBlueprint;
+			this.PropertyCaps |= UhtPropertyCaps.PassCppArgsByRef | UhtPropertyCaps.IsParameterSupportedByBlueprint | 
+				UhtPropertyCaps.IsMemberSupportedByBlueprint | UhtPropertyCaps.SupportsRigVM;
 		}
 
 		/// <inheritdoc/>
@@ -76,7 +77,7 @@ namespace EpicGames.UHT.Types
 			switch (TextType)
 			{
 				case UhtPropertyTextType.ExportMember:
-				case UhtPropertyTextType.RigVMTemplateArg:
+				case UhtPropertyTextType.RigVMType:
 				case UhtPropertyTextType.EventParameterFunctionMember:
 					Builder.Append(this.CppTypeText);
 					break;
@@ -158,11 +159,5 @@ namespace EpicGames.UHT.Types
 				}
 			}
 		}
-
-		/// <inheritdoc/>
-		public override string? GetRigVMType(ref UhtRigVMParameterFlags ParameterFlags)
-		{
-			return this.CppTypeText;
-		}
 	}
 }
diff --git a/Engine/Source/Programs/Shared/EpicGames.UHT/Types/Properties/UhtEnumProperty.cs b/Engine/Source/Programs/Shared/EpicGames.UHT/Types/Properties/UhtEnumProperty.cs
--- a/Engine/Source/Programs/Shared/EpicGames.UHT/Types/Properties/UhtEnumProperty.cs
+++ b/Engine/Source/Programs/Shared/EpicGames.UHT/Types/Properties/UhtEnumProperty.cs
@@ -58,7 +58,8 @@ namespace EpicGames.UHT.Types
 		/// <param name="Enum">Referenced enum</param>
 		public UhtEnumProperty(UhtPropertySettings PropertySettings, UhtEnum Enum) : base(PropertySettings)
 		{
-			this.PropertyCaps |= UhtPropertyCaps.RequiresNullConstructorArg | UhtPropertyCaps.CanExposeOnSpawn | UhtPropertyCaps.IsParameterSupportedByBlueprint | UhtPropertyCaps.IsMemberSupportedByBlueprint;
+			this.PropertyCaps |= UhtPropertyCaps.RequiresNullConstructorArg | UhtPropertyCaps.CanExposeOnSpawn | UhtPropertyCaps.IsParameterSupportedByBlueprint |
+				UhtPropertyCaps.IsMemberSupportedByBlueprint | UhtPropertyCaps.SupportsRigVM | UhtPropertyCaps.IsRigVMEnum;
 			this.Enum = Enum;
 			this.HeaderFile.AddReferencedHeader(Enum);
 			if (this.Enum.CppForm == UhtEnumCppForm.EnumClass)
@@ -125,6 +126,17 @@ namespace EpicGames.UHT.Types
 					Builder.Append(Enum.SourceName);
 					break;
 
+				case UhtPropertyTextType.RigVMType:
+					if (Enum.CppForm == UhtEnumCppForm.EnumClass || !bIsTemplateArgument)
+					{
+						Builder.Append(Enum.CppType);
+					}
+					else
+					{
+						Builder.Append("TEnumAsByte<").Append(Enum.CppType).Append(">");
+					}
+					break;
+
 				default:
 					if (Enum.CppForm == UhtEnumCppForm.EnumClass)
 					{
@@ -163,19 +175,6 @@ namespace EpicGames.UHT.Types
 			return Builder;
 		}
 
-		/// <summary>
-		/// Get the RigVM type for an enumeration
-		/// </summary>
-		/// <param name="Property">Property in question</param>
-		/// <param name="Enum">Referenced enumeration</param>
-		/// <param name="ParameterFlags">Parameter flags that can be updated</param>
-		/// <returns>Type string</returns>
-		public static string? GetEnumRigVMType(UhtProperty Property, UhtEnum Enum, ref UhtRigVMParameterFlags ParameterFlags)
-		{
-			ParameterFlags |= UhtRigVMParameterFlags.IsEnum;
-			return Enum.CppType.ToString();
-		}
-
 		/// <summary>
 		/// Sanitize the default value for an enumeration
 		/// </summary>
@@ -299,12 +298,6 @@ namespace EpicGames.UHT.Types
 			return false;
 		}
 
-		/// <inheritdoc/>
-		public override string? GetRigVMType(ref UhtRigVMParameterFlags ParameterFlags)
-		{
-			return GetEnumRigVMType(this, this.Enum, ref ParameterFlags);
-		}
-
 		private UhtProperty CreateUnderlyingProperty()
 		{
 			UhtPropertySettings PropertySettings = new UhtPropertySettings();
diff --git a/Engine/Source/Programs/Shared/EpicGames.UHT/Types/Properties/UhtFieldPathProperty.cs b/Engine/Source/Programs/Shared/EpicGames.UHT/Types/Properties/UhtFieldPathProperty.cs
--- a/Engine/Source/Programs/Shared/EpicGames.UHT/Types/Properties/UhtFieldPathProperty.cs
+++ b/Engine/Source/Programs/Shared/EpicGames.UHT/Types/Properties/UhtFieldPathProperty.cs
@@ -103,12 +103,6 @@ namespace EpicGames.UHT.Types
 			return false;
 		}
 
-		/// <inheritdoc/>
-		public override string? GetRigVMType(ref UhtRigVMParameterFlags ParameterFlags)
-		{
-			return null;
-		}
-
 		#region Keyword
 		[UhtPropertyType(Keyword = "TFieldPath", Options = UhtPropertyTypeOptions.Immediate)]
 		private static UhtProperty? FieldPathProperty(UhtPropertyResolvePhase ResolvePhase, UhtPropertySettings PropertySettings, IUhtTokenReader TokenReader, UhtToken MatchedToken)
diff --git a/Engine/Source/Programs/Shared/EpicGames.UHT/Types/Properties/UhtInterfaceProperty.cs b/Engine/Source/Programs/Shared/EpicGames.UHT/Types/Properties/UhtInterfaceProperty.cs
--- a/Engine/Source/Programs/Shared/EpicGames.UHT/Types/Properties/UhtInterfaceProperty.cs
+++ b/Engine/Source/Programs/Shared/EpicGames.UHT/Types/Properties/UhtInterfaceProperty.cs
@@ -47,6 +47,10 @@ namespace EpicGames.UHT.Types
 			this.PropertyFlags |= EPropertyFlags.UObjectWrapper;
 			this.PropertyCaps |= UhtPropertyCaps.CanExposeOnSpawn | UhtPropertyCaps.IsParameterSupportedByBlueprint | UhtPropertyCaps.IsMemberSupportedByBlueprint | UhtPropertyCaps.PassCppArgsByRef;
 			this.PropertyCaps &= ~(UhtPropertyCaps.CanHaveConfig | UhtPropertyCaps.CanBeContainerKey);
+			if (this.Session.Config!.AreRigVMUInterfaceProeprtiesEnabled)
+			{
+				this.PropertyCaps |= UhtPropertyCaps.SupportsRigVM;
+			}
 		}
 
 		/// <inheritdoc/>
@@ -171,12 +175,6 @@ namespace EpicGames.UHT.Types
 			return this.InterfaceClass.ClassFlags.HasAnyFlags(EClassFlags.Const);
 		}
 
-		/// <inheritdoc/>
-		public override string? GetRigVMType(ref UhtRigVMParameterFlags ParameterFlags)
-		{
-			return null;
-		}
-
 		#region Keywords
 		[UhtPropertyType(Keyword = "FScriptInterface", Options = UhtPropertyTypeOptions.Simple)] // This can't be immediate due to the reference to UInterface
 		private static UhtProperty? FScriptInterfaceProperty(UhtPropertyResolvePhase ResolvePhase, UhtPropertySettings PropertySettings, IUhtTokenReader TokenReader, UhtToken MatchedToken)
diff --git a/Engine/Source/Programs/Shared/EpicGames.UHT/Types/Properties/UhtMapProperty.cs b/Engine/Source/Programs/Shared/EpicGames.UHT/Types/Properties/UhtMapProperty.cs
--- a/Engine/Source/Programs/Shared/EpicGames.UHT/Types/Properties/UhtMapProperty.cs
+++ b/Engine/Source/Programs/Shared/EpicGames.UHT/Types/Properties/UhtMapProperty.cs
@@ -253,12 +253,6 @@ namespace EpicGames.UHT.Types
 			}
 		}
 
-		/// <inheritdoc/>
-		public override string? GetRigVMType(ref UhtRigVMParameterFlags ParameterFlags)
-		{
-			return null;
-		}
-
 		#region Keyword
 		[UhtPropertyType(Keyword = "TMap")]
 		private static UhtProperty? MapProperty(UhtPropertyResolvePhase ResolvePhase, UhtPropertySettings PropertySettings, IUhtTokenReader TokenReader, UhtToken MatchedToken)
diff --git a/Engine/Source/Programs/Shared/EpicGames.UHT/Types/Properties/UhtMulticastDelegateProperty.cs b/Engine/Source/Programs/Shared/EpicGames.UHT/Types/Properties/UhtMulticastDelegateProperty.cs
--- a/Engine/Source/Programs/Shared/EpicGames.UHT/Types/Properties/UhtMulticastDelegateProperty.cs
+++ b/Engine/Source/Programs/Shared/EpicGames.UHT/Types/Properties/UhtMulticastDelegateProperty.cs
@@ -43,7 +43,7 @@ namespace EpicGames.UHT.Types
 			this.Function = Function;
 			this.HeaderFile.AddReferencedHeader(Function);
 			this.PropertyCaps |= UhtPropertyCaps.PassCppArgsByRef | UhtPropertyCaps.CanBeBlueprintAssignable | UhtPropertyCaps.CanBeBlueprintCallable |
-				UhtPropertyCaps.CanBeBlueprintAuthorityOnly | UhtPropertyCaps.IsMemberSupportedByBlueprint;
+				UhtPropertyCaps.CanBeBlueprintAuthorityOnly | UhtPropertyCaps.IsMemberSupportedByBlueprint | UhtPropertyCaps.SupportsRigVM;
 			this.PropertyCaps &= ~(UhtPropertyCaps.IsParameterSupportedByBlueprint);
 		}
 
@@ -79,9 +79,9 @@ namespace EpicGames.UHT.Types
 			{
 				case UhtPropertyTextType.EventParameterFunctionMember:
 				case UhtPropertyTextType.ExportMember:
-				case UhtPropertyTextType.RigVMTemplateArg:
+				case UhtPropertyTextType.RigVMType:
 				case UhtPropertyTextType.FunctionThunkParameterArgType:
-					Builder.Append("FMulticastScriptDelegate");
+					Builder.Append(this.CppTypeText);
 					break;
 
 				default:
@@ -125,11 +125,5 @@ namespace EpicGames.UHT.Types
 			}
 			return false;
 		}
-
-		/// <inheritdoc/>
-		public override string? GetRigVMType(ref UhtRigVMParameterFlags ParameterFlags)
-		{
-			return this.CppTypeText;
-		}
 	}
 }
diff --git a/Engine/Source/Programs/Shared/EpicGames.UHT/Types/Properties/UhtNameProperty.cs b/Engine/Source/Programs/Shared/EpicGames.UHT/Types/Properties/UhtNameProperty.cs
--- a/Engine/Source/Programs/Shared/EpicGames.UHT/Types/Properties/UhtNameProperty.cs
+++ b/Engine/Source/Programs/Shared/EpicGames.UHT/Types/Properties/UhtNameProperty.cs
@@ -34,7 +34,8 @@ namespace EpicGames.UHT.Types
 		/// <param name="PropertySettings">Property settings</param>
 		public UhtNameProperty(UhtPropertySettings PropertySettings) : base(PropertySettings)
 		{
-			this.PropertyCaps |= UhtPropertyCaps.CanExposeOnSpawn | UhtPropertyCaps.IsParameterSupportedByBlueprint | UhtPropertyCaps.IsMemberSupportedByBlueprint;
+			this.PropertyCaps |= UhtPropertyCaps.CanExposeOnSpawn | UhtPropertyCaps.IsParameterSupportedByBlueprint |
+				UhtPropertyCaps.IsMemberSupportedByBlueprint | UhtPropertyCaps.SupportsRigVM;
 		}
 
 		/// <inheritdoc/>
@@ -86,12 +87,6 @@ namespace EpicGames.UHT.Types
 			return Other is UhtNameProperty;
 		}
 
-		/// <inheritdoc/>
-		public override string? GetRigVMType(ref UhtRigVMParameterFlags ParameterFlags)
-		{
-			return this.CppTypeText;
-		}
-
 		#region Keyword
 		[UhtPropertyType(Keyword = "FName", Options = UhtPropertyTypeOptions.Simple | UhtPropertyTypeOptions.Immediate)]
 		private static UhtProperty? NameProperty(UhtPropertyResolvePhase ResolvePhase, UhtPropertySettings PropertySettings, IUhtTokenReader TokenReader, UhtToken MatchedToken)
diff --git a/Engine/Source/Programs/Shared/EpicGames.UHT/Types/Properties/UhtNumericProperty.cs b/Engine/Source/Programs/Shared/EpicGames.UHT/Types/Properties/UhtNumericProperty.cs
--- a/Engine/Source/Programs/Shared/EpicGames.UHT/Types/Properties/UhtNumericProperty.cs
+++ b/Engine/Source/Programs/Shared/EpicGames.UHT/Types/Properties/UhtNumericProperty.cs
@@ -29,7 +29,7 @@ namespace EpicGames.UHT.Types
 		protected UhtNumericProperty(UhtPropertySettings PropertySettings, UhtPropertyIntType IntType) : base(PropertySettings)
 		{
 			this.IntType = IntType;
-			this.PropertyCaps |= UhtPropertyCaps.RequiresNullConstructorArg;
+			this.PropertyCaps |= UhtPropertyCaps.RequiresNullConstructorArg | UhtPropertyCaps.SupportsRigVM;
 		}
 
 		/// <inheritdoc/>
@@ -38,11 +38,5 @@ namespace EpicGames.UHT.Types
 			Builder.Append("0");
 			return Builder;
 		}
-
-		/// <inheritdoc/>
-		public override string? GetRigVMType(ref UhtRigVMParameterFlags ParameterFlags)
-		{
-			return CppTypeText;
-		}
 	}
 }
diff --git a/Engine/Source/Programs/Shared/EpicGames.UHT/Types/Properties/UhtObjectPropertyBase.cs b/Engine/Source/Programs/Shared/EpicGames.UHT/Types/Properties/UhtObjectPropertyBase.cs
--- a/Engine/Source/Programs/Shared/EpicGames.UHT/Types/Properties/UhtObjectPropertyBase.cs
+++ b/Engine/Source/Programs/Shared/EpicGames.UHT/Types/Properties/UhtObjectPropertyBase.cs
@@ -60,6 +60,10 @@ namespace EpicGames.UHT.Types
 			}
 
 			this.PropertyCaps &= ~(UhtPropertyCaps.CanHaveConfig);
+			if (this.Session.Config!.AreRigVMUObjectProeprtiesEnabled)
+			{
+				this.PropertyCaps |= UhtPropertyCaps.SupportsRigVM;
+			}
 		}
 
 		/// <inheritdoc/>
@@ -135,12 +139,6 @@ namespace EpicGames.UHT.Types
 			return this.Class.ClassFlags.HasAnyFlags(EClassFlags.Const);
 		}
 
-		/// <inheritdoc/>
-		public override string? GetRigVMType(ref UhtRigVMParameterFlags ParameterFlags)
-		{
-			return null;
-		}
-
 		#region Parsing support methods
 		/// <summary>
 		/// Parse a template type
diff --git a/Engine/Source/Programs/Shared/EpicGames.UHT/Types/Properties/UhtSetProperty.cs b/Engine/Source/Programs/Shared/EpicGames.UHT/Types/Properties/UhtSetProperty.cs
--- a/Engine/Source/Programs/Shared/EpicGames.UHT/Types/Properties/UhtSetProperty.cs
+++ b/Engine/Source/Programs/Shared/EpicGames.UHT/Types/Properties/UhtSetProperty.cs
@@ -204,12 +204,6 @@ namespace EpicGames.UHT.Types
 			}
 		}
 
-		/// <inheritdoc/>
-		public override string? GetRigVMType(ref UhtRigVMParameterFlags ParameterFlags)
-		{
-			return null;
-		}
-
 		#region Keyword
 		[UhtPropertyType(Keyword = "TSet")]
 		private static UhtProperty? SetProperty(UhtPropertyResolvePhase ResolvePhase, UhtPropertySettings PropertySettings, IUhtTokenReader TokenReader, UhtToken MatchedToken)
diff --git a/Engine/Source/Programs/Shared/EpicGames.UHT/Types/Properties/UhtStrProperty.cs b/Engine/Source/Programs/Shared/EpicGames.UHT/Types/Properties/UhtStrProperty.cs
--- a/Engine/Source/Programs/Shared/EpicGames.UHT/Types/Properties/UhtStrProperty.cs
+++ b/Engine/Source/Programs/Shared/EpicGames.UHT/Types/Properties/UhtStrProperty.cs
@@ -34,7 +34,8 @@ namespace EpicGames.UHT.Types
 		/// <param name="PropertySettings">Property settings</param>
 		public UhtStrProperty(UhtPropertySettings PropertySettings) : base(PropertySettings)
 		{
-			this.PropertyCaps |= UhtPropertyCaps.PassCppArgsByRef | UhtPropertyCaps.CanExposeOnSpawn | UhtPropertyCaps.IsParameterSupportedByBlueprint | UhtPropertyCaps.IsMemberSupportedByBlueprint;
+			this.PropertyCaps |= UhtPropertyCaps.PassCppArgsByRef | UhtPropertyCaps.CanExposeOnSpawn | UhtPropertyCaps.IsParameterSupportedByBlueprint | 
+				UhtPropertyCaps.IsMemberSupportedByBlueprint | UhtPropertyCaps.SupportsRigVM;
 		}
 
 		/// <inheritdoc/>
@@ -99,12 +100,6 @@ namespace EpicGames.UHT.Types
 			}
 		}
 
-		/// <inheritdoc/>
-		public override string? GetRigVMType(ref UhtRigVMParameterFlags ParameterFlags)
-		{
-			return this.CppTypeText;
-		}
-
 		[UhtPropertyType(Keyword = "FString")]
 		[UhtPropertyType(Keyword = "FMemoryImageString", Options = UhtPropertyTypeOptions.Immediate)]
 		private static UhtProperty? StrProperty(UhtPropertyResolvePhase ResolvePhase, UhtPropertySettings PropertySettings, IUhtTokenReader TokenReader, UhtToken MatchedToken)
diff --git a/Engine/Source/Programs/Shared/EpicGames.UHT/Types/Properties/UhtStructProperty.cs b/Engine/Source/Programs/Shared/EpicGames.UHT/Types/Properties/UhtStructProperty.cs
--- a/Engine/Source/Programs/Shared/EpicGames.UHT/Types/Properties/UhtStructProperty.cs
+++ b/Engine/Source/Programs/Shared/EpicGames.UHT/Types/Properties/UhtStructProperty.cs
@@ -44,6 +44,7 @@ namespace EpicGames.UHT.Types
 		{
 			this.ScriptStruct = ScriptStruct;
 			this.HeaderFile.AddReferencedHeader(ScriptStruct);
+			this.PropertyCaps |= UhtPropertyCaps.SupportsRigVM;
 			if (this.ScriptStruct.bHasNoOpConstructor)
 			{
 				this.PropertyCaps |= UhtPropertyCaps.RequiresNullConstructorArg;
@@ -204,12 +205,6 @@ namespace EpicGames.UHT.Types
 			return ReferencingProperty.Session.ValidateScriptStructOkForNet(ReferencingProperty, this.ScriptStruct);
 		}
 
-		/// <inheritdoc/>
-		public override string? GetRigVMType(ref UhtRigVMParameterFlags ParameterFlags)
-		{
-			return this.ScriptStruct.SourceName;
-		}
-
 		#region Structure default value sanitizers
 		[UhtStructDefaultValue(Name = "FVector")]
 		private static bool VectorStructDefaultValue(UhtStructProperty Property, IUhtTokenReader DefaultValueReader, StringBuilder InnerDefaultValue)
diff --git a/Engine/Source/Programs/Shared/EpicGames.UHT/Types/Properties/UhtTextProperty.cs b/Engine/Source/Programs/Shared/EpicGames.UHT/Types/Properties/UhtTextProperty.cs
--- a/Engine/Source/Programs/Shared/EpicGames.UHT/Types/Properties/UhtTextProperty.cs
+++ b/Engine/Source/Programs/Shared/EpicGames.UHT/Types/Properties/UhtTextProperty.cs
@@ -34,7 +34,8 @@ namespace EpicGames.UHT.Types
 		/// <param name="PropertySettings">Property settings</param>
 		public UhtTextProperty(UhtPropertySettings PropertySettings) : base(PropertySettings)
 		{
-			this.PropertyCaps |= UhtPropertyCaps.PassCppArgsByRef | UhtPropertyCaps.CanExposeOnSpawn | UhtPropertyCaps.IsParameterSupportedByBlueprint | UhtPropertyCaps.IsMemberSupportedByBlueprint;
+			this.PropertyCaps |= UhtPropertyCaps.PassCppArgsByRef | UhtPropertyCaps.CanExposeOnSpawn | UhtPropertyCaps.IsParameterSupportedByBlueprint | 
+				UhtPropertyCaps.IsMemberSupportedByBlueprint | UhtPropertyCaps.SupportsRigVM;
 			this.PropertyCaps &= ~(UhtPropertyCaps.CanBeContainerKey);
 		}
 
@@ -113,12 +114,6 @@ namespace EpicGames.UHT.Types
 			return Other is UhtTextProperty;
 		}
 
-		/// <inheritdoc/>
-		public override string? GetRigVMType(ref UhtRigVMParameterFlags ParameterFlags)
-		{
-			return this.CppTypeText;
-		}
-
 		#region Parsing keywords and default parsers		
 		[UhtPropertyType(Keyword = "FText", Options = UhtPropertyTypeOptions.Simple | UhtPropertyTypeOptions.Immediate)]
 		private static UhtProperty? TextProperty(UhtPropertyResolvePhase ResolvePhase, UhtPropertySettings PropertySettings, IUhtTokenReader TokenReader, UhtToken MatchedToken)
diff --git a/Engine/Source/Programs/Shared/EpicGames.UHT/Types/Properties/UhtVoidProperty.cs b/Engine/Source/Programs/Shared/EpicGames.UHT/Types/Properties/UhtVoidProperty.cs
--- a/Engine/Source/Programs/Shared/EpicGames.UHT/Types/Properties/UhtVoidProperty.cs
+++ b/Engine/Source/Programs/Shared/EpicGames.UHT/Types/Properties/UhtVoidProperty.cs
@@ -68,12 +68,6 @@ namespace EpicGames.UHT.Types
 			return Other is UhtVoidProperty;
 		}
 
-		/// <inheritdoc/>
-		public override string? GetRigVMType(ref UhtRigVMParameterFlags ParameterFlags)
-		{
-			throw new NotImplementedException();
-		}
-
 		#region Keyword
 		[UhtPropertyType(Keyword = "void", Options = UhtPropertyTypeOptions.Simple | UhtPropertyTypeOptions.Immediate)]
 		private static UhtProperty? VoidProperty(UhtPropertyResolvePhase ResolvePhase, UhtPropertySettings PropertySettings, IUhtTokenReader TokenReader, UhtToken MatchedToken)
diff --git a/Engine/Source/Programs/Shared/EpicGames.UHT/Types/UhtClass.cs b/Engine/Source/Programs/Shared/EpicGames.UHT/Types/UhtClass.cs
--- a/Engine/Source/Programs/Shared/EpicGames.UHT/Types/UhtClass.cs
+++ b/Engine/Source/Programs/Shared/EpicGames.UHT/Types/UhtClass.cs
@@ -818,6 +818,17 @@ namespace EpicGames.UHT.Types
 			{
 				StringBuilder Builder = Borrower.StringBuilder;
 				Builder.AppendPropertyText(Property, UhtPropertyTextType.GetterSetterArg);
+				if (Property.bIsStaticArray)
+				{
+					if (Declaration.Function == null)
+					{
+						Builder.Append("*");
+					}
+					else
+					{
+						Builder.Append("[").Append(Property.ArrayDimensions).Append("]");
+					}
+				}
 				string ExpectedType = Builder.ToString();
 				int TypeEnd = TypeIndex + TypeCount;
 
@@ -1314,6 +1325,10 @@ namespace EpicGames.UHT.Types
 						{
 							StringBuilder Builder = Borrower.StringBuilder;
 							Builder.AppendPropertyText(Property, UhtPropertyTextType.GetterSetterArg);
+							if (Property.bIsStaticArray)
+							{
+								Builder.Append($"*/[{Property.ArrayDimensions}]");
+							}
 							string ExpectedType = Builder.ToString();
 							Property.LogError($"Property '{Property.SourceName}' expected a getter function '[const] {ExpectedType} [&] {ExpectedName}() const', but it was not found");
 						}
@@ -1325,6 +1340,10 @@ namespace EpicGames.UHT.Types
 						{
 							StringBuilder Builder = Borrower.StringBuilder;
 							Builder.AppendPropertyText(Property, UhtPropertyTextType.GetterSetterArg);
+							if (Property.bIsStaticArray)
+							{
+								Builder.Append($"*/[{Property.ArrayDimensions}]");
+							}
 							string ExpectedType = Builder.ToString();
 							Property.LogError($"Property '{Property.SourceName}' expected a setter function 'void {ExpectedName}([const] {ExpectedType} [&] InArg)', but it was not found");
 						}
diff --git a/Engine/Source/Programs/Shared/EpicGames.UHT/Types/UhtProperty.cs b/Engine/Source/Programs/Shared/EpicGames.UHT/Types/UhtProperty.cs
--- a/Engine/Source/Programs/Shared/EpicGames.UHT/Types/UhtProperty.cs
+++ b/Engine/Source/Programs/Shared/EpicGames.UHT/Types/UhtProperty.cs
@@ -384,6 +384,21 @@ namespace EpicGames.UHT.Types
 		/// True to see if the member property is supported by blueprint
 		/// </summary>
 		IsMemberSupportedByBlueprint = 1 << 11,
+
+		/// <summary>
+		/// True if the property supports RigVM
+		/// </summary>
+		SupportsRigVM = 1 << 12,
+
+		/// <summary>
+		/// True if the property should codegen as an enumeration
+		/// </summary>
+		IsRigVMEnum = 1 << 13,
+
+		/// <summary>
+		/// True if the property should codegen as an array
+		/// </summary>
+		IsRigVMArray = 1 << 14,
 	}
 
 	/// <summary>
@@ -509,9 +524,9 @@ namespace EpicGames.UHT.Types
 		FunctionThunkRetVal,
 
 		/// <summary>
-		/// Used when generating types for RigVM.  Only invoked for template arguments
+		/// Basic RigVM type
 		/// </summary>
-		RigVMTemplateArg,
+		RigVMType,
 
 		/// <summary>
 		/// Type expected in a getter/setter argument list
@@ -1364,6 +1379,19 @@ namespace EpicGames.UHT.Types
 			return Builder.ToString();
 		}
 
+		/// <summary>
+		/// Return the RigVM type
+		/// </summary>
+		/// <returns></returns>
+		public string GetRigVMType()
+		{
+			using (BorrowStringBuilder Borrower = new BorrowStringBuilder(StringBuilderCache.Small))
+			{
+				AppendText(Borrower.StringBuilder, UhtPropertyTextType.RigVMType);
+				return Borrower.StringBuilder.ToString();
+			}
+		}
+
 		/// <summary>
 		/// Appends any applicable objects and child properties
 		/// </summary>
@@ -1876,15 +1904,6 @@ namespace EpicGames.UHT.Types
 		}
 		#endregion
 
-		#region RigVM support methods
-		/// <summary>
-		/// Return the RigVM type string.
-		/// </summary>
-		/// <param name="ParameterFlags">Parameter flags that can be updated for specific types</param>
-		/// <returns>RigVM type string</returns>
-		public abstract string? GetRigVMType(ref UhtRigVMParameterFlags ParameterFlags);
-		#endregion
-
 		#region Helper methods
 		/// <summary>
 		/// Generate a new name suffix based on the current suffix and the new suffix
@@ -1923,6 +1942,10 @@ namespace EpicGames.UHT.Types
 					return false;
 				}
 			}
+			if (this.bIsStaticArray != Other.bIsStaticArray)
+			{
+				return false;
+			}
 			return IsSameType(Other);
 		}
 		#endregion
diff --git a/Engine/Source/Programs/Shared/EpicGames.UHT/Types/UhtScriptStruct.cs b/Engine/Source/Programs/Shared/EpicGames.UHT/Types/UhtScriptStruct.cs
--- a/Engine/Source/Programs/Shared/EpicGames.UHT/Types/UhtScriptStruct.cs
+++ b/Engine/Source/Programs/Shared/EpicGames.UHT/Types/UhtScriptStruct.cs
@@ -205,24 +205,27 @@ namespace EpicGames.UHT.Types
 				Property.LogError($"RigVM Struct '{this.Property.Outer?.SourceName}' - Member '{this.Property.SourceName}' is editor only - WITH_EDITORONLY_DATA not allowed on structs with RIGVM_METHOD.");
 			}
 
-			string? RigVMType = Property.GetRigVMType(ref this.ParameterFlags);
-			if (RigVMType == null)
+			if (Property.PropertyCaps.HasAnyFlags(UhtPropertyCaps.SupportsRigVM))
 			{
-				Property.LogError($"RigVM Struct '{this.Property.Outer?.SourceName}' - Member '{this.Property.SourceName}' type '{this.Property.GetUserFacingDecl()}' not supported by RigVM.");
+				this.Type = Property.GetRigVMType();
+				if (Property.PropertyCaps.HasAnyFlags(UhtPropertyCaps.IsRigVMEnum))
+				{
+					this.ParameterFlags |= UhtRigVMParameterFlags.IsEnum;
+				}
+				if (Property.PropertyCaps.HasAnyFlags(UhtPropertyCaps.IsRigVMArray))
+				{
+					this.ParameterFlags |= UhtRigVMParameterFlags.IsArray;
+					if (this.IsConst())
+					{
+						string ExtendedType = this.ExtendedType(false);
+						this.CastName = $"{this.Name}_{Index}_Array";
+						this.CastType = $"TArrayView<const {ExtendedType.Substring(1, ExtendedType.Length - 2)}>";
+					}
+				}
 			}
 			else
 			{
-				this.Type = RigVMType;
-			}
-
-			if (this.bIsArray)
-			{
-				if (this.IsConst())
-				{
-					string ExtendedType = this.ExtendedType(false);
-					this.CastName = $"{this.Name}_{Index}_Array";
-					this.CastType = $"TArrayView<const {ExtendedType.Substring(1, ExtendedType.Length - 2)}>";
-				}
+				Property.LogError($"RigVM Struct '{this.Property.Outer?.SourceName}' - Member '{this.Property.SourceName}' type '{this.Property.GetUserFacingDecl()}' not supported by RigVM.");
 			}
 		}
 
diff --git a/Engine/Source/Programs/Shared/EpicGames.UHT/Utils/UhtConfig.cs b/Engine/Source/Programs/Shared/EpicGames.UHT/Utils/UhtConfig.cs
--- a/Engine/Source/Programs/Shared/EpicGames.UHT/Utils/UhtConfig.cs
+++ b/Engine/Source/Programs/Shared/EpicGames.UHT/Utils/UhtConfig.cs
@@ -59,6 +59,16 @@ namespace EpicGames.UHT.Utils
 		/// </summary>
 		public UhtPointerMemberBehavior NonEngineObjectPtrMemberBehavior { get; }
 
+		/// <summary>
+		/// If true, UObject properties are enabled in RigVM
+		/// </summary>
+		public bool AreRigVMUObjectProeprtiesEnabled { get; }
+
+		/// <summary>
+		/// If true, UInterface properties are enabled in RigVM
+		/// </summary>
+		public bool AreRigVMUInterfaceProeprtiesEnabled { get; }
+
 		/// <summary>
 		/// If the token references a remapped identifier, update the value in the token 
 		/// </summary>
diff --git a/Engine/Source/Programs/UnrealBuildTool/Modes/UnrealHeaderToolMode.cs b/Engine/Source/Programs/UnrealBuildTool/Modes/UnrealHeaderToolMode.cs
--- a/Engine/Source/Programs/UnrealBuildTool/Modes/UnrealHeaderToolMode.cs
+++ b/Engine/Source/Programs/UnrealBuildTool/Modes/UnrealHeaderToolMode.cs
@@ -79,6 +79,16 @@ namespace UnrealBuildTool.Modes
 		/// </summary>
 		private readonly UhtPointerMemberBehavior NonEngineObjectPtrMemberBehaviorInternal = UhtPointerMemberBehavior.AllowSilently;
 
+		/// <summary>
+		/// If true, UObject properties are enabled in RigVM
+		/// </summary>
+		private readonly bool AreRigVMUObjectProeprtiesEnabledInternal = false;
+
+		/// <summary>
+		/// If true, UInterface properties are enabled in RigVM
+		/// </summary>
+		private readonly bool AreRigVMUInterfaceProeprtiesEnabledInternal = false;
+
 		#region IUhtConfig Implementation
 		/// <inheritdoc/>
 		public EGeneratedCodeVersion DefaultGeneratedCodeVersion => this.DefaultGeneratedCodeVersionInternal;
@@ -95,6 +105,16 @@ namespace UnrealBuildTool.Modes
 		/// <inheritdoc/>
 		public UhtPointerMemberBehavior NonEngineObjectPtrMemberBehavior => this.NonEngineObjectPtrMemberBehaviorInternal;
 
+		/// <summary>
+		/// If true, UObject properties are enabled in RigVM
+		/// </summary>
+		public bool AreRigVMUObjectProeprtiesEnabled => this.AreRigVMUObjectProeprtiesEnabledInternal;
+
+		/// <summary>
+		/// If true, UInterface properties are enabled in RigVM
+		/// </summary>
+		public bool AreRigVMUInterfaceProeprtiesEnabled => this.AreRigVMUInterfaceProeprtiesEnabledInternal;
+
 		/// <inheritdoc/>
 		public void RedirectTypeIdentifier(ref UhtToken Token)
 		{
@@ -187,6 +207,8 @@ namespace UnrealBuildTool.Modes
 			this.EngineObjectPtrMemberBehaviorInternal = GetPointerMemberBehavior("UnrealHeaderTool", "EngineObjectPtrMemberBehavior", UhtPointerMemberBehavior.AllowSilently);
 			this.NonEngineNativePointerMemberBehaviorInternal = GetPointerMemberBehavior("UnrealHeaderTool", "NonEngineNativePointerMemberBehavior", UhtPointerMemberBehavior.AllowSilently);
 			this.NonEngineObjectPtrMemberBehaviorInternal = GetPointerMemberBehavior("UnrealHeaderTool", "NonEngineObjectPtrMemberBehavior", UhtPointerMemberBehavior.AllowSilently);
+			this.AreRigVMUObjectProeprtiesEnabledInternal = GetBoolean("UnrealHeaderTool", "AreRigVMUObjectProeprtiesEnabled", false);
+			this.AreRigVMUInterfaceProeprtiesEnabledInternal = GetBoolean("UnrealHeaderTool", "AreRigVMUInterfaceProeprtiesEnabled", false);
 		}
 
 		private bool GetBoolean(string SectionName, string KeyName, bool bDefault)
