Update Blake3 dotnet dependency to 0.5.1
#rnx

diff --git a/Engine/Source/Programs/Shared/EpicGames.IoHash/EpicGames.IoHash.csproj b/Engine/Source/Programs/Shared/EpicGames.IoHash/EpicGames.IoHash.csproj
--- a/Engine/Source/Programs/Shared/EpicGames.IoHash/EpicGames.IoHash.csproj
+++ b/Engine/Source/Programs/Shared/EpicGames.IoHash/EpicGames.IoHash.csproj
@@ -29,7 +29,7 @@
     </Compile>
   </ItemGroup>
   <ItemGroup>
-    <PackageReference Include="Blake3" Version="0.3.0" />
+    <PackageReference Include="Blake3" Version="0.5.1" />
     <PackageReference Include="System.Memory" Version="4.5.4" />
   </ItemGroup>
   <ItemGroup>
