Hack fix in EpicGames.Horde to get it building again, to unblock releasing of Horde Storage

diff --git a/Engine/Source/Programs/Shared/EpicGames.Horde/Compute/IComputeClient.cs b/Engine/Source/Programs/Shared/EpicGames.Horde/Compute/IComputeClient.cs
--- a/Engine/Source/Programs/Shared/EpicGames.Horde/Compute/IComputeClient.cs
+++ b/Engine/Source/Programs/Shared/EpicGames.Horde/Compute/IComputeClient.cs
@@ -346,6 +346,14 @@ namespace EpicGames.Horde.Compute
 		}
 	}
 
+	/// <summary>
+	/// 
+	/// </summary>
+	public interface IComputeTaskInfo
+	{
+
+	}
+
 	/// <summary>
 	/// Interface for communicating with a compute server
 	/// </summary>
@@ -378,7 +386,7 @@ namespace EpicGames.Horde.Compute
 		/// <param name="channelId">Channel to receive updates on</param>
 		/// <param name="cancellationToken">Cancellation token for the operation</param>
 		/// <returns></returns>
-		IAsyncEnumerable<ComputeTaskStatus> GetTaskUpdatesAsync(ClusterId clusterId, ChannelId channelId, CancellationToken cancellationToken = default);
+		IAsyncEnumerable<IComputeTaskInfo> GetTaskUpdatesAsync(ClusterId clusterId, ChannelId channelId, CancellationToken cancellationToken = default);
 	}
 
 	/// <summary>
