Horde: Build fix.

diff --git a/Engine/Source/Programs/Shared/EpicGames.Horde/Compute/IComputeClient.cs b/Engine/Source/Programs/Shared/EpicGames.Horde/Compute/IComputeClient.cs
--- a/Engine/Source/Programs/Shared/EpicGames.Horde/Compute/IComputeClient.cs
+++ b/Engine/Source/Programs/Shared/EpicGames.Horde/Compute/IComputeClient.cs
@@ -346,14 +346,6 @@ namespace EpicGames.Horde.Compute
 		}
 	}
 
-	/// <summary>
-	/// 
-	/// </summary>
-	public interface IComputeTaskInfo
-	{
-
-	}
-
 	/// <summary>
 	/// Interface for communicating with a compute server
 	/// </summary>
@@ -386,7 +378,7 @@ namespace EpicGames.Horde.Compute
 		/// <param name="channelId">Channel to receive updates on</param>
 		/// <param name="cancellationToken">Cancellation token for the operation</param>
 		/// <returns></returns>
-		IAsyncEnumerable<IComputeTaskInfo> GetTaskUpdatesAsync(ClusterId clusterId, ChannelId channelId, CancellationToken cancellationToken = default);
+		IAsyncEnumerable<ComputeTaskStatus> GetTaskUpdatesAsync(ClusterId clusterId, ChannelId channelId, CancellationToken cancellationToken = default);
 	}
 
 	/// <summary>
diff --git a/Engine/Source/Programs/Shared/EpicGames.Horde/Compute/Impl/HttpComputeClient.cs b/Engine/Source/Programs/Shared/EpicGames.Horde/Compute/Impl/HttpComputeClient.cs
--- a/Engine/Source/Programs/Shared/EpicGames.Horde/Compute/Impl/HttpComputeClient.cs
+++ b/Engine/Source/Programs/Shared/EpicGames.Horde/Compute/Impl/HttpComputeClient.cs
@@ -33,44 +33,6 @@ namespace EpicGames.Horde.Compute.Impl
 		public BucketId ResponseBucketId { get; set; }
 	}
 
-	/// <summary>
-	/// Supplies information about the current execution state of a task
-	/// </summary>
-	public class GetTaskUpdateResponse : IComputeTaskInfo
-	{
-		/// <inheritdoc/>
-		[CbField("h")]
-		public RefId TaskRefId { get; set; }
-
-		/// <inheritdoc/>
-		[CbField("t")]
-		public DateTime Time { get; set; }
-
-		/// <inheritdoc/>
-		[CbField("s")]
-		public ComputeTaskState State { get; set; }
-
-		/// <inheritdoc/>
-		[CbField("o")]
-		public ComputeTaskOutcome Outcome { get; set; }
-
-		/// <inheritdoc/>
-		[CbField("d")]
-		public string? Detail { get; set; }
-
-		/// <inheritdoc/>
-		[CbField("r")]
-		public RefId? ResultRefId { get; set; }
-
-		/// <inheritdoc/>
-		[CbField("a")]
-		public string? AgentId { get; set; }
-
-		/// <inheritdoc/>
-		[CbField("l")]
-		public string? LeaseId { get; set; }
-	}
-
 	/// <summary>
 	/// Request to add tasks to the compute queue
 	/// </summary>
@@ -110,7 +72,7 @@ namespace EpicGames.Horde.Compute.Impl
 		/// Task updates
 		/// </summary>
 		[CbField("u")]
-		public List<GetTaskUpdateResponse> Updates { get; set; } = new List<GetTaskUpdateResponse>();
+		public List<ComputeTaskStatus> Updates { get; set; } = new List<ComputeTaskStatus>();
 	}
 
 	#endregion
@@ -157,7 +119,7 @@ namespace EpicGames.Horde.Compute.Impl
 		}
 
 		/// <inheritdoc/>
-		public async IAsyncEnumerable<IComputeTaskInfo> GetTaskUpdatesAsync(ClusterId clusterId, ChannelId channelId, [EnumeratorCancellation] CancellationToken cancellationToken)
+		public async IAsyncEnumerable<ComputeTaskStatus> GetTaskUpdatesAsync(ClusterId clusterId, ChannelId channelId, [EnumeratorCancellation] CancellationToken cancellationToken)
 		{
 			for (; ; )
 			{
@@ -170,7 +132,7 @@ namespace EpicGames.Horde.Compute.Impl
 				byte[] data = await response.Content.ReadAsByteArrayAsync();
 				GetTaskUpdatesResponse parsedResponse = CbSerializer.Deserialize<GetTaskUpdatesResponse>(new CbField(data));
 
-				foreach (GetTaskUpdateResponse update in parsedResponse.Updates)
+				foreach (ComputeTaskStatus update in parsedResponse.Updates)
 				{
 					cancellationToken.ThrowIfCancellationRequested();
 					yield return update;
