- Allow for multiple full sdks for a platform (uses a new SDKCollection class to gather sdk versions and their validityt, etc)
- Windows will be first to use it, but in a future CL, with toolchain separate from Windows SDK
- Cleaned up some of the PlatformSDK API to have fewer public members (mainly so the functions that receive multiple versions can be automatically populated with the current single version functions in the subclasses)
- Updated the Turnkey menu to handle multiple versions

diff --git a/Engine/Source/Editor/TurnkeySupport/Private/TurnkeySupportModule.cpp b/Engine/Source/Editor/TurnkeySupport/Private/TurnkeySupportModule.cpp
--- a/Engine/Source/Editor/TurnkeySupport/Private/TurnkeySupportModule.cpp
+++ b/Engine/Source/Editor/TurnkeySupport/Private/TurnkeySupportModule.cpp
@@ -965,43 +965,49 @@ static FSlateIcon MakePlatformSdkIconAttribute(FName IniPlatformName, TSharedPtr
 //		));
 }
 
-static void FormatSdkInfo(const FString& PlatformOrDevice, const FTurnkeySdkInfo& SdkInfo, bool bIncludeAutoSdk, FText& OutInfo, FText& OutToolTip)
+static void FormatSdkInfo(const FString& PlatformOrDevice, const FTurnkeySdkInfo& SdkInfo, FText& OutInfo, FText& OutToolTip)
 {
-	FFormatOrderedArguments Args;
-	Args.Add(FText::FromString(SdkInfo.InstalledVersion));
-	Args.Add(FText::FromString(SdkInfo.AutoSDKVersion));
-	Args.Add(FText::FromString(SdkInfo.MinAllowedVersion));
-	Args.Add(FText::FromString(SdkInfo.MaxAllowedVersion));
-	Args.Add(SdkInfo.SdkErrorInformation);
-	Args.Add(FText::FromString(PlatformOrDevice));
 
 	TArray<FText> Lines;
-	Lines.Add(FText::Format(LOCTEXT("SdkInfo_Installed", "Installed SDK: {0}"), Args));
-	if (bIncludeAutoSdk)
-	{
-		Lines.Add(FText::Format(LOCTEXT("SdkInfo_AutoSDK", "AutoSDK: {1}"), Args));
-	}
-	
-	if (SdkInfo.MinAllowedVersion == SdkInfo.MaxAllowedVersion)
-	{
-		Lines.Add(FText::Format(LOCTEXT("SdkInfo_AllowedSDK_Single", "Allowed Version: {2}"), Args));
-	}
-	else if (SdkInfo.MinAllowedVersion == TEXT(""))
-	{
-		Lines.Add(FText::Format(LOCTEXT("SdkInfo_AllowedSDK_MaxOnly", "Allowed Versions: Up to {3}"), Args));
-	}
-	else if (SdkInfo.MaxAllowedVersion == TEXT(""))
-	{
-		Lines.Add(FText::Format(LOCTEXT("SdkInfo_AllowedSDK_MinOnly", "Allowed Versions: {2} and up"), Args));
-	}
-	else
+
+	for (TPair<FString, FTurnkeySdkInfo::Version> Pair : SdkInfo.SDKVersions)
 	{
-		Lines.Add(FText::Format(LOCTEXT("SdkInfo_AllowedSDK_Range", "Allowed Versions: {2} through {3}"), Args));
+		const FString& Name = Pair.Key;
+		const FString& Min = Pair.Value.Min;
+		const FString& Max = Pair.Value.Max;
+		const FString& Current = Pair.Value.Current;
+		FFormatNamedArguments VersionArgs;
+		VersionArgs.Add(TEXT("Name"), FText::FromString(Name));
+		VersionArgs.Add(TEXT("Min"), FText::FromString(Min));
+		VersionArgs.Add(TEXT("Max"), FText::FromString(Max));
+		VersionArgs.Add(TEXT("Current"), FText::FromString(Current.Len() ? Current : FString(TEXT("--"))));
+
+		if (Min == Max)
+		{
+			Lines.Add(FText::Format(LOCTEXT("SdkInfo_AllowedSDK_Single", "Allowed {Name} Version: {Min}"), VersionArgs));
+		}
+		else if (Min == TEXT(""))
+		{
+			Lines.Add(FText::Format(LOCTEXT("SdkInfo_AllowedSDK_MaxOnly", "Allowed {Name} Versions: Up to {Max}"), VersionArgs));
+		}
+		else if (Max == TEXT(""))
+		{
+			Lines.Add(FText::Format(LOCTEXT("SdkInfo_AllowedSDK_MinOnly", "Allowed {Name} Versions: {Min} and up"), VersionArgs));
+		}
+		else
+		{
+			Lines.Add(FText::Format(LOCTEXT("SdkInfo_AllowedSDK_Range", "Allowed {Name} Versions: {Min} through {Max}"), VersionArgs));
+		}
+		Lines.Add(FText::Format(LOCTEXT("SdkInfo_AllowedSDK_Current", "  Installed: {Current}"), VersionArgs));
 	}
 
+	FFormatOrderedArguments Args;
+	Args.Add(SdkInfo.SdkErrorInformation);
+	Args.Add(FText::FromString(PlatformOrDevice));
+
 	if (!SdkInfo.SdkErrorInformation.IsEmpty())
 	{
-		Lines.Add(FText::Format(LOCTEXT("SdkInfo_Error", "Error Info:\n{4}"), Args));
+		Lines.Add(FText::Format(LOCTEXT("SdkInfo_Error", "Error Info:\n{0}"), Args));
 	}
 
 	// now make a single \n delimted text
@@ -1010,11 +1016,11 @@ static void FormatSdkInfo(const FString& PlatformOrDevice, const FTurnkeySdkInfo
 	// make a tooltip
 	if (PlatformOrDevice.Contains(TEXT("@")))
 	{
-		OutToolTip = FText::Format(LOCTEXT("SdkInfo_ToolTip", "Information returned from:\nRunUAT Turnkey -command=VerifySdk -device={5}"), Args);
+		OutToolTip = FText::Format(LOCTEXT("SdkInfo_ToolTip", "Information returned from:\nRunUAT Turnkey -command=VerifySdk -device={1}"), Args);
 	}
 	else
 	{
-		OutToolTip = FText::Format(LOCTEXT("SdkInfo_ToolTip", "Information returned from:\nRunUAT Turnkey -command=VerifySdk -platform={5}"), Args);
+		OutToolTip = FText::Format(LOCTEXT("SdkInfo_ToolTip", "Information returned from:\nRunUAT Turnkey -command=VerifySdk -platform={1}"), Args);
 	}
 }
 
@@ -1318,7 +1324,7 @@ static void MakeTurnkeyPlatformMenu(UToolMenu* ToolMenu, FName IniPlatformName,
 
 					FTurnkeySdkInfo SdkInfo = ITurnkeySupportModule::Get().GetSdkInfoForDeviceId(DeviceId);
 					FText SdkText, SdkToolTip;
-					FormatSdkInfo(DeviceId, SdkInfo, false, SdkText, SdkToolTip);
+					FormatSdkInfo(DeviceId, SdkInfo, SdkText, SdkToolTip);
 
 					FToolMenuSection& Section = SubToolMenu->AddSection("DeviceSdkInfo", LOCTEXT("TurnkeySection_DeviceSdkInfo", "Sdk Info"));
 					Section.AddEntry(
@@ -1363,7 +1369,7 @@ static void MakeTurnkeyPlatformMenu(UToolMenu* ToolMenu, FName IniPlatformName,
 
 	const FTurnkeySdkInfo& SdkInfo = ITurnkeySupportModule::Get().GetSdkInfo(IniPlatformName, true);
 	FText SdkText, SdkToolTip;
-	FormatSdkInfo(IniPlatformName.ToString(), SdkInfo, true, SdkText, SdkToolTip);
+	FormatSdkInfo(IniPlatformName.ToString(), SdkInfo, SdkText, SdkToolTip);
 
 	SdkSection.AddEntry(
 		FToolMenuEntry::InitWidget(
@@ -2012,14 +2018,43 @@ bool GetSdkInfoFromTurnkey(FString Line, FName& PlatformName, FString& DeviceId,
 	FString FlagsString;
 	FParse::Value(*Info, TEXT("Status="), StatusString);
 	FParse::Value(*Info, TEXT("Flags="), FlagsString);
-	FParse::Value(*Info, TEXT("Installed="), SdkInfo.InstalledVersion);
-	FParse::Value(*Info, TEXT("AutoSDK="), SdkInfo.AutoSDKVersion);
-	FParse::Value(*Info, TEXT("MinAllowed="), SdkInfo.MinAllowedVersion);
-	FParse::Value(*Info, TEXT("MaxAllowed="), SdkInfo.MaxAllowedVersion);
 	FString ErrorString;
 	FParse::Value(*Info, TEXT("Error="), ErrorString);
 	SdkInfo.SdkErrorInformation = FText::FromString(ErrorString.Replace(TEXT("|"), TEXT("\n")));
 
+	FString SDKNamesString;
+	TArray<FString> SDKNames;
+	if (FParse::Value(*Info, TEXT("SDKs="), SDKNamesString))
+	{
+		SDKNamesString.ParseIntoArray(SDKNames, TEXT(","), true);
+	}
+	else
+	{
+		SDKNames.Add("SDK");
+	}
+	if (DeviceId.Len() == 0)
+	{
+		SDKNames.Add("AutoSDK");
+	}
+
+	for (const FString& Name : SDKNames)
+	{
+		FString Min, Max, Current;
+		// handle both AutoSDK and manual (auto only has Allowed_AutoSDK, others have Min/Max)
+		FParse::Value(*Info, *FString::Printf(TEXT("MinAllowed_%s="), *Name), Min);
+		FParse::Value(*Info, *FString::Printf(TEXT("MaxAllowed_%s="), *Name), Max);
+		FParse::Value(*Info, *FString::Printf(TEXT("Allowed_%s="), *Name), Min);
+		FParse::Value(*Info, *FString::Printf(TEXT("Allowed_%s="), *Name), Max);
+		FParse::Value(*Info, *FString::Printf(TEXT("Current_%s="), *Name), Current);
+		// also handle no name at all (for device, etc)
+		FParse::Value(*Info, *FString::Printf(TEXT("MinAllowed="), *Name), Min);
+		FParse::Value(*Info, *FString::Printf(TEXT("MaxAllowed="), *Name), Max);
+		FParse::Value(*Info, *FString::Printf(TEXT("Allowed="), *Name), Min);
+		FParse::Value(*Info, *FString::Printf(TEXT("Allowed="), *Name), Max);
+		FParse::Value(*Info, *FString::Printf(TEXT("Current=")), Current);
+		SdkInfo.SDKVersions.Add(Name, { Min, Max, Current });
+	}
+
 	SdkInfo.Status = ETurnkeyPlatformSdkStatus::Unknown;
 	if (StatusString == TEXT("Valid"))
 	{
@@ -2145,15 +2180,19 @@ void FTurnkeySupportModule::UpdateSdkInfo()
 						// future calls to Turnkey will inherit the AutoSDK env vars, and it won't be able to determine the manual SDK versions anymore. If we use the editor to
 						// install an SDK via Turnkey, it will directly update the installed version based on the result of that command, not this Update operation
 
-						FString OriginalManualInstallValue = PerPlatformSdkInfo[PlatformName].InstalledVersion;
+						TMap<FString, FTurnkeySdkInfo::Version> OriginalVersions = PerPlatformSdkInfo[PlatformName].SDKVersions;
 
 						// set it into the platform
 						PerPlatformSdkInfo[PlatformName] = SdkInfo;
 
-						// restore the original installed version if it set after the first time
-						if (OriginalManualInstallValue.Len() > 0)
+						// restore the original installed version if it set after the first time, except for AutoVersion
+						if (OriginalVersions.Num() > 0)
 						{
-							PerPlatformSdkInfo[PlatformName].InstalledVersion = OriginalManualInstallValue;
+							PerPlatformSdkInfo[PlatformName].SDKVersions = OriginalVersions;
+							if (SdkInfo.SDKVersions.Contains(TEXT("AutoSDK")))
+							{
+								PerPlatformSdkInfo[PlatformName].SDKVersions.Add(TEXT("AutoSDK"), SdkInfo.SDKVersions[TEXT("AutoSDK")]);
+							}
 						}
 
 
@@ -2460,4 +2499,4 @@ void FTurnkeySupportModule::ShutdownModule( )
 
 IMPLEMENT_MODULE(FTurnkeySupportModule, TurnkeySupport);
 
-#undef LOCTEXT_NAMESPACE
\ No newline at end of file
+#undef LOCTEXT_NAMESPACE
diff --git a/Engine/Source/Editor/TurnkeySupport/Public/Interfaces/ITurnkeySupportModule.h b/Engine/Source/Editor/TurnkeySupport/Public/Interfaces/ITurnkeySupportModule.h
--- a/Engine/Source/Editor/TurnkeySupport/Public/Interfaces/ITurnkeySupportModule.h
+++ b/Engine/Source/Editor/TurnkeySupport/Public/Interfaces/ITurnkeySupportModule.h
@@ -30,13 +30,11 @@ enum class ETurnkeyDeviceStatus : uint8
 
 struct FTurnkeySdkInfo
 {
+	struct Version { FString Min, Max, Current; };
 	ETurnkeyPlatformSdkStatus Status = ETurnkeyPlatformSdkStatus::Unknown;
 	ETurnkeyDeviceStatus DeviceStatus = ETurnkeyDeviceStatus::Unknown;
 	FText SdkErrorInformation;
-	FString InstalledVersion;
-	FString AutoSDKVersion; // only valid for platform, not device
-	FString MinAllowedVersion;
-	FString MaxAllowedVersion;
+	TMap<FString, Version> SDKVersions;
 	bool bCanInstallFullSdk;
 	bool bCanInstallAutoSdk;
 };
diff --git a/Engine/Source/Programs/AutomationTool/Turnkey/Commands/InstallSdk.cs b/Engine/Source/Programs/AutomationTool/Turnkey/Commands/InstallSdk.cs
--- a/Engine/Source/Programs/AutomationTool/Turnkey/Commands/InstallSdk.cs
+++ b/Engine/Source/Programs/AutomationTool/Turnkey/Commands/InstallSdk.cs
@@ -234,44 +234,42 @@ namespace Turnkey.Commands
 
 					int BestVersionIndex = Sdks.Any() ? Sdks.IndexOf(BestByType.First().Value) : -1;
 
-					string Current, MinVersion, MaxVersion;
+					SDKCollection SDKInfo;
 					if (DesiredType == FileSource.SourceType.Flash)
 					{
-						SDK.GetValidSoftwareVersionRange(out MinVersion, out MaxVersion);
 						DeviceInfo Device = GetDevice(AutomationPlatform, DeviceName);
-						Current = Device == null ? "N/A" : Device.SoftwareVersion;
+						SDKInfo = SDK.GetAllSoftwareInfo(Device.Type, Device.SoftwareVersion);
 					}
 					else
 					{
-						SDK.GetValidVersionRange(out MinVersion, out MaxVersion);
-						Current = SDK.GetInstalledVersion();
+						SDKInfo = SDK.GetAllSDKInfo();
 					}
 					//if (bStrippedDevices)
 					//{
 					//	Prompt += "\nNOTE: Some Flash Sdks were removed because no devices were found";
 					//}
 
-					MinVersion = MinVersion == null ? "" : MinVersion;
-					MaxVersion = MaxVersion == null ? "" : MaxVersion;
-
 					string Prompt;
 					string TypeString = DesiredType == FileSource.SourceType.Flash ? "Flash" : "Sdk";
 					
-					if (MinVersion == MaxVersion)
+//					Prompt = $"Select an {TypeString} to install [Current Version: {Current}, Allowed: {ValidVersions}]";
+					Prompt = $"Select an {TypeString} to install\n{SDKInfo.ToString("Allowed", "")}";
+					if (DesiredType != FileSource.SourceType.Flash)
 					{
-						Prompt = $"Select an {TypeString} to install [Current Version: {Current}, Required: {MinVersion}]";
-					}
-					else
-					{
-						if (DesiredType == FileSource.SourceType.Flash)
-						{
-							Prompt = $"Select an {TypeString} to install [Current: {Current}, Valid Range: {MinVersion} - {MaxVersion}]";
-						}
-						else
-						{
-							Prompt = $"Select an {TypeString} to install [Current: {Current}, Main: {SDK.GetMainVersion()}, Valid Range: {MinVersion} - {MaxVersion}]";
-						}
+						Prompt += $"\n(Preferred: {SDK.GetMainVersion()}";
 					}
+					//}
+					//else
+					//{
+					//	if (DesiredType == FileSource.SourceType.Flash)
+					//	{
+					//		Prompt = $"Select an {TypeString} to install [Current: {Current}, Valid Range: {MinVersion} - {MaxVersion}]";
+					//	}
+					//	else
+					//	{
+					//		Prompt = $"Select an {TypeString} to install [Current: {Current}, Main: {SDK.GetMainVersion()}, Valid Range: {MinVersion} - {MaxVersion}]";
+					//	}
+					//}
 
 
 					int Choice = TurnkeyUtils.ReadInputInt(Prompt, Options, true, BestVersionIndex >= 0 ?  (BestVersionIndex + 1) : -1);
diff --git a/Engine/Source/Programs/AutomationTool/Turnkey/Commands/ListPlatforms.cs b/Engine/Source/Programs/AutomationTool/Turnkey/Commands/ListPlatforms.cs
--- a/Engine/Source/Programs/AutomationTool/Turnkey/Commands/ListPlatforms.cs
+++ b/Engine/Source/Programs/AutomationTool/Turnkey/Commands/ListPlatforms.cs
@@ -26,21 +26,19 @@ namespace Turnkey.Commands
 					continue;
 				}
 
-				string ManualSDKVersion, AutoSDKVersion;
-				string MinAllowedVersion, MaxAllowedVersion;
-				string MinAllowedSoftwareVersion, MaxAllowedSoftwareVersion;
-				SDK.GetInstalledVersions(out ManualSDKVersion, out AutoSDKVersion);
-				SDK.GetValidVersionRange(out MinAllowedVersion, out MaxAllowedVersion);
-				SDK.GetValidSoftwareVersionRange(out MinAllowedSoftwareVersion, out MaxAllowedSoftwareVersion);
-				
-				bool bIsAutoSdkValid = SDK.IsVersionValid(AutoSDKVersion, bForAutoSDK:true);
-				bool bIsManualSdkValid = SDK.IsVersionValid(ManualSDKVersion, bForAutoSDK:false);
+				SDKCollection SDKInfo = SDK.GetAllSDKInfo();
+				SDKCollection SoftwareInfo = SDK.GetAllSoftwareInfo();
+
+				string AutoSDKVersion = SDKInfo.AutoSdk?.Current;
 				TurnkeyUtils.Log("  Platform: {0}", TargetPlatform.ToString());
-				TurnkeyUtils.Log("  Installed Manual Sdk: {0}", ManualSDKVersion);
-				TurnkeyUtils.Log("  Installed Auto Sdk: {0}", AutoSDKVersion);
-				TurnkeyUtils.Log("  Allowed Sdk Range: {0}-{1}", MinAllowedVersion, MaxAllowedVersion);
-				TurnkeyUtils.Log("  Valid Manual SDK Installed? {0}", bIsManualSdkValid);
-				TurnkeyUtils.Log("  Valid Auto SDK Installed? {0}", bIsAutoSdkValid);
+				TurnkeyUtils.Log("  Installed Manual Sdk: {0}", SDKInfo.ToString());
+				TurnkeyUtils.Log("  Installed Auto Sdk: {0}", AutoSDKVersion ?? "");
+				foreach (SDKDescriptor Desc in SDKInfo.FullSdks)
+				{
+					TurnkeyUtils.Log($"  Allowed {Desc.Name} Range: {Desc.Min}-{Desc.Max}");
+				}
+				TurnkeyUtils.Log("  Valid Manual SDK(s) Installed? {0}", SDKInfo.AreAllManualSDKsValid());
+				TurnkeyUtils.Log("  Valid Auto SDK Installed? {0}", AutoSDKVersion != null);
 
 				string[] AllVersions = SDK.GetAllInstalledSDKVersions();
 				TurnkeyUtils.Log("  AllVersions Installed: {0}", string.Join(",", AllVersions));
@@ -61,7 +59,8 @@ namespace Turnkey.Commands
 					}
 				}
 
-				TurnkeyUtils.Log("  Allowed Device Software Range: {0}-{1}", MinAllowedSoftwareVersion, MaxAllowedSoftwareVersion);
+	
+				TurnkeyUtils.Log("  Allowed Device Software Range: {0}", SoftwareInfo.ToString(""));
 				TurnkeyUtils.Log("  Devices: ");
 
 				DeviceInfo[] Devices = Platform.GetDevices();
@@ -73,7 +72,7 @@ namespace Turnkey.Commands
 				{
 					foreach (DeviceInfo Device in Devices)
 					{
-						bool bIsSoftwareValid = SDK.IsSoftwareVersionValid(Device.SoftwareVersion);
+						bool bIsSoftwareValid = SDK.IsSoftwareVersionValid(Device.SoftwareVersion, Device.Type);
 
 						TurnkeyUtils.Log("    Name: {0}{1}", Device.Name, Device.bIsDefault ? "*" : "");
 						TurnkeyUtils.Log("      Id: {0}", Device.Id);
diff --git a/Engine/Source/Programs/AutomationTool/Turnkey/Commands/VerifySdk.cs b/Engine/Source/Programs/AutomationTool/Turnkey/Commands/VerifySdk.cs
--- a/Engine/Source/Programs/AutomationTool/Turnkey/Commands/VerifySdk.cs
+++ b/Engine/Source/Programs/AutomationTool/Turnkey/Commands/VerifySdk.cs
@@ -64,22 +64,6 @@ namespace Turnkey.Commands
 				// checking availability may generate errors, if prerequisites are failing
 				SdkUtils.LocalAvailability LocalState = SdkUtils.GetLocalAvailability(AutomationPlatform, bUpdateIfNeeded, TurnkeyContext);
 
-				string ManualSDKVersion = "", AutoSDKVersion = "";
-				string MinAllowedVersion = "", MaxAllowedVersion = "";
-				string MinSoftwareAllowedVersion = "", MaxSoftwareAllowedVersion = "";
-
-				try
-				{
-					if (PlatformSDK != null)
-					{
-						PlatformSDK.GetInstalledVersions(out ManualSDKVersion, out AutoSDKVersion);
-						PlatformSDK.GetValidVersionRange(out MinAllowedVersion, out MaxAllowedVersion);
-						PlatformSDK.GetValidSoftwareVersionRange(out MinSoftwareAllowedVersion, out MaxSoftwareAllowedVersion);
-					}
-				}
-				catch (Exception)
-				{
-				}
 
 				SdkUtils.LocalAvailability ReportedState = LocalState;
 				string StatusString;
@@ -100,14 +84,35 @@ namespace Turnkey.Commands
 					ReportedState &= (SdkUtils.LocalAvailability.AutoSdk_ValidVersionExists | SdkUtils.LocalAvailability.InstalledSdk_ValidVersionExists | SdkUtils.LocalAvailability.Support_FullSdk | SdkUtils.LocalAvailability.Support_AutoSdk);
 				}
 
-				string ErrorString = "";
-				if (TurnkeyContext.ErrorMessages.Count() > 0)
+
+				// gather lisat of tags/values to write out
+				SDKCollection SDKInfo = null;
+				List<string> PlatformProperties = new List<string>();
 				{
-					ErrorString = string.Format(", Error=\"{0}\"", string.Join("|", TurnkeyContext.ErrorMessages));
+					// overall status (valid/invalid)
+					PlatformProperties.Add($"Status={StatusString}");
+
+					if (PlatformSDK != null)
+					{
+						SDKInfo = PlatformSDK.GetAllSDKInfo();
+						if (SDKInfo.FullSdks.Count() > 1)
+						{
+							PlatformProperties.Add($"SDKs=\"{string.Join(",", SDKInfo.FullSdks.Select(x => x.Name))}\"");
+						}
+						// get all the Min/Max versions and insert them right in, they will be in the proper format for parsing
+						PlatformProperties.Add(SDKInfo.ToString("Allowed", ""));
+					}
+
+					PlatformProperties.Add($"Flags={ReportedState}");
+
+					if (TurnkeyContext.ErrorMessages.Count() > 0)
+					{
+						PlatformProperties.Add($"Error=\"{string.Join("|", TurnkeyContext.ErrorMessages)}\"");
+					}
 				}
 
-				TurnkeyUtils.Report("{0}: (Status={1}, Installed={2}, AutoSDK={3}, MinAllowed={4}, MaxAllowed={5}, Flags=\"{6}\"{7})", Platform, StatusString, ManualSDKVersion, AutoSDKVersion,
-					MinAllowedVersion, MaxAllowedVersion, ReportedState.ToString(), ErrorString);
+				// write out all of the properties in unreal struct format: "Platform: (Props,...)"
+				TurnkeyUtils.Report("{0}: ({1})", Platform, string.Join(", ", PlatformProperties));
 
 				if (PlatformSDK == null)
 				{
@@ -122,7 +127,8 @@ namespace Turnkey.Commands
 					{
 
 						// find the highest number that is valid (because a valid version exists, we know there will be at least one valid version)
-						string BestAlternateVersion = PlatformSDK.GetAllInstalledSDKVersions().ToList().OrderByDescending(x => x).Where(x => PlatformSDK.IsVersionValid(x, false)).First();
+						// @todo turnkey another place where dealing with switching if we have multiple sdks won't work (see the "Sdk" param to IsVersionValid - we aren't allowing for althernate sdk types there)
+						string BestAlternateVersion = PlatformSDK.GetAllInstalledSDKVersions().ToList().OrderByDescending(x => x).Where(x => PlatformSDK.IsVersionValid(x, "Sdk")).First();
 
 						bool bWasSwitched = PlatformSDK.SwitchToAlternateSDK(BestAlternateVersion, false);
 
@@ -135,37 +141,96 @@ namespace Turnkey.Commands
 						}
 					}
 
-					FileSource BestSdk = null;
-					// find the best Sdk, prioritizing as request
-					if (bPreferFullSdk)
+
+					// gather the SDKs that we need to install
+					List<FileSource> SdksToInstall = new List<FileSource>();
+					if (!bPreferFullSdk)
 					{
-						BestSdk = FileSource.FindMatchingSdk(AutomationPlatform, new FileSource.SourceType[] { FileSource.SourceType.Full, FileSource.SourceType.BuildOnly, FileSource.SourceType.AutoSdk }, bSelectBest: bAutoChooseBest);
+						// attempt to install Autosdk
+						FileSource AutoSdk = FileSource.FindMatchingSdk(AutomationPlatform, new FileSource.SourceType[] { FileSource.SourceType.AutoSdk }, bSelectBest: bAutoChooseBest);
+
+						if (AutoSdk != null)
+						{
+							SdksToInstall.Add(AutoSdk);
+						}
 					}
-					else
+
+					// if no autosdk was chosen, look for manual sdk(s)
+					if (SdksToInstall.Count == 0)
 					{
-						BestSdk = FileSource.FindMatchingSdk(AutomationPlatform, new FileSource.SourceType[] { FileSource.SourceType.AutoSdk, FileSource.SourceType.BuildOnly, FileSource.SourceType.Full }, bSelectBest: bAutoChooseBest);
+						// find manual versions that need to be installed (or all if forceing)
+						IEnumerable<SDKDescriptor> NeededManualSdks = SDKInfo.FullSdks.Where(x => (bForceSdkInstall || x.Validity == SDKStatus.Invalid));
+						List<string> MissingInstallers = new List<string>();
+
+						foreach (SDKDescriptor Desc in NeededManualSdks)
+						{
+							// if we only have one manual sdk (normal), then don't specify a type
+							string SpecificType = NeededManualSdks.Count() == 1 ? null : Desc.Name;
+							FileSource SdkSource = FileSource.FindMatchingSdk(AutomationPlatform, new FileSource.SourceType[] { FileSource.SourceType.Full }, bSelectBest: bAutoChooseBest, SpecificType: SpecificType );
+
+							if (SdkSource != null)
+							{
+								SdksToInstall.Add(SdkSource);
+							}
+							else
+							{
+								MissingInstallers.Add(Desc.Name);
+							}
+						}
+
+						if (NeededManualSdks.Count() > 1 && MissingInstallers.Count > 0)
+						{
+							TurnkeyUtils.Log($"ERROR: {Platform}: Unable to find all Sdk installers. Missing '{string.Join(", ", MissingInstallers)}'. Will the ones that were found...");
+						}
 					}
 
-					if (BestSdk == null)
+					// if we went manual first, and still don't have any sdks, try for autosdk now
+					if (bPreferFullSdk && SdksToInstall.Count == 0)
 					{
-						TurnkeyUtils.Log("ERROR: {0}: Unable to find any Sdks that could be installed", Platform);
-						TurnkeyUtils.ExitCode = ExitCode.Error_SDKNotFound;
-						continue;
+						// attempt to install Autosdk
+						FileSource AutoSdk = FileSource.FindMatchingSdk(AutomationPlatform, new FileSource.SourceType[] { FileSource.SourceType.AutoSdk }, bSelectBest: bAutoChooseBest);
+
+						if (AutoSdk != null)
+						{
+							SdksToInstall.Add(AutoSdk);
+						}
 					}
 
-					TurnkeyUtils.Log("Will install {0}", BestSdk.Name);
 
-					if (BestSdk.DownloadOrInstall(Platform, TurnkeyContext, null, bUnattended) == false)
+					//FileSource BestSdk = null;
+					//// find the best Sdk, prioritizing as request
+					//if (bPreferFullSdk)
+					//{
+					//	BestSdk = FileSource.FindMatchingSdk(AutomationPlatform, new FileSource.SourceType[] { FileSource.SourceType.Full, FileSource.SourceType.BuildOnly, FileSource.SourceType.AutoSdk }, bSelectBest: bAutoChooseBest);
+					//}
+					//else
+					//{
+					//	BestSdk = FileSource.FindMatchingSdk(AutomationPlatform, new FileSource.SourceType[] { FileSource.SourceType.AutoSdk, FileSource.SourceType.BuildOnly, FileSource.SourceType.Full }, bSelectBest: bAutoChooseBest);
+					//}
+
+					if (SdksToInstall == null)
 					{
-						TurnkeyUtils.Log("Failed to install {0}", BestSdk.Name);
-						TurnkeyUtils.ExitCode = ExitCode.Error_SDKInstallFailed;
-						continue;
+						TurnkeyUtils.Log("ERROR: {0}: Unable to find any Sdks that could be installed", Platform);
+						TurnkeyUtils.ExitCode = ExitCode.Error_SDKNotFound;
 					}
+					else
+					{
+						TurnkeyUtils.Log("Will install '{0}'", string.Join(", ", SdksToInstall.Select(x => x.Name)));
 
-					// update LocalState
-//					LocalState = SdkUtils.GetLocalAvailability(AutomationPlatform, false);
+						foreach (FileSource Sdk in SdksToInstall)
+						{
+							if (Sdk.DownloadOrInstall(Platform, TurnkeyContext, null, bUnattended) == false)
+							{
+								TurnkeyUtils.Log("Failed to install {0}", Sdk.Name);
+								TurnkeyUtils.ExitCode = ExitCode.Error_SDKInstallFailed;
+							}
+						}
+
+						// update LocalState
+						//					LocalState = SdkUtils.GetLocalAvailability(AutomationPlatform, false);
 
-					// @todo turnkey: validate!
+						// @todo turnkey: validate!
+					}
 				}
 
 				// now check software verison of each device
@@ -177,7 +242,9 @@ namespace Turnkey.Commands
 						TurnkeyContext.ErrorMessages.Clear();
 
 						bool bArePrerequisitesValid = AutomationPlatform.UpdateDevicePrerequisites(Device, TurnkeyUtils.CommandUtilHelper, TurnkeyContext, !bUpdateIfNeeded);
-						bool bIsSoftwareValid = PlatformSDK.IsSoftwareVersionValid(Device.SoftwareVersion);
+
+						// get the min/max versions from PlatformSDK
+						SDKCollection SoftwareInfo = PlatformSDK.GetAllSoftwareInfo(Device.Type, Device.SoftwareVersion);
 
 						SdkUtils.LocalAvailability DeviceState = SdkUtils.LocalAvailability.None;
 						if (!bArePrerequisitesValid)
@@ -185,7 +252,7 @@ namespace Turnkey.Commands
 							StatusString = "Invalid";
 							DeviceState |= SdkUtils.LocalAvailability.Device_InvalidPrerequisites;
 						}
-						else if (bIsSoftwareValid)
+						else if (SoftwareInfo.AreAllManualSDKsValid())
 						{
 							StatusString = "Valid";
 							DeviceState |= SdkUtils.LocalAvailability.Device_InstallSoftwareValid;
@@ -201,16 +268,33 @@ namespace Turnkey.Commands
 							DeviceState |= SdkUtils.LocalAvailability.Device_CannotConnect;
 						}
 
-						string DeviceErrorString = "";
+
+						List<string> DeviceProperties = new List<string>()
+							{
+								$"Name={Device.Name}",
+								$"Type={Device.Type}",
+								$"Status={StatusString}",
+								SoftwareInfo.ToString("Allowed", ""),
+								$"Flags=\"{DeviceState}\"",
+							};
+						
+						foreach (KeyValuePair<string, string> Pair in Device.PlatformValues)
+						{
+							DeviceProperties.Add($"{Pair.Key}=\"{Pair.Value}\"");
+						}
+
 						if (TurnkeyContext.ErrorMessages.Count() > 0)
 						{
-							ErrorString = string.Format(", Error=\"{0}\"", string.Join("|", TurnkeyContext.ErrorMessages));
+							DeviceProperties.Add(string.Format("Error=\"{0}\"", string.Join("|", TurnkeyContext.ErrorMessages)));
 						}
 
-						TurnkeyUtils.Report("{0}@{1}: (Name={2}, Status={3}, Installed={4}, MinAllowed={5}, MaxAllowed={6}, Flags=\"{7}\"{8})", Platform, Device.Id, Device.Name, StatusString, Device.SoftwareVersion,
-							MinSoftwareAllowedVersion, MaxSoftwareAllowedVersion, DeviceState.ToString(), DeviceErrorString);
+						// write out all of the properties in unreal struct format: "Platform: (Props,...)"
+						TurnkeyUtils.Report("{0}@{1}: ({2})", Platform, Device.Id, string.Join(", ", DeviceProperties));
+
+						//TurnkeyUtils.Report("{0}@{1}: (Name={2}, Status={3}, Installed={4}, MinAllowed={5}, MaxAllowed={6}, Flags=\"{7}\"{8})", Platform, Device.Id, Device.Name, StatusString, Device.SoftwareVersion,
+						//	MinSoftwareAllowedVersion, MaxSoftwareAllowedVersion, DeviceState.ToString(), DeviceErrorString);
 
-						if (bForceDeviceInstall || !bIsSoftwareValid)
+						if (bForceDeviceInstall || !SoftwareInfo.AreAllManualSDKsValid())
 						{
 							if (bUpdateIfNeeded)
 							{
diff --git a/Engine/Source/Programs/AutomationTool/Turnkey/FileSource.cs b/Engine/Source/Programs/AutomationTool/Turnkey/FileSource.cs
--- a/Engine/Source/Programs/AutomationTool/Turnkey/FileSource.cs
+++ b/Engine/Source/Programs/AutomationTool/Turnkey/FileSource.cs
@@ -96,16 +96,23 @@ namespace Turnkey
 
 		[XmlElement("Platform")]
 		public string PlatformString = null;
-
 		[XmlIgnore]
 		private List<UnrealTargetPlatform> Platforms;
 
+		[XmlElement("Type")]
+		public string TypeString = null;
+		[XmlIgnore]
+		public SourceType Type = SourceType.Full;
+		[XmlIgnore]
+		public string SpecificSDKType = null;
+
 		[XmlElement("Source")]
 		public CopySource[] Sources = null;
 
+
+
 		public string Version = null;
 		public string Name = null;
-		public SourceType Type = SourceType.Full;
 		public string AllowedFlashDeviceTypes = null;
 		
 		// Project for Build types to filter on
@@ -129,10 +136,10 @@ namespace Turnkey
 		{
 			FileSource Clone = new FileSource();
 			Clone.PlatformString = PlatformString;
+			Clone.TypeString = TypeString;
 			Clone.Version = Version;
 			Clone.Name = Name;
 			Clone.AllowedFlashDeviceTypes = AllowedFlashDeviceTypes;
-			Clone.Type = Type;
 			Clone.Project = Project;
 			Clone.BuildPlatformEnumerationSuffix = BuildPlatformEnumerationSuffix;
 
@@ -223,7 +230,9 @@ namespace Turnkey
 			if (Type == SourceType.Flash)
 			{
 				UEBuildPlatformSDK SDK = UEBuildPlatformSDK.GetSDKForPlatform(Platform.ToString());
-				bool bIsValid = SDK.IsSoftwareVersionValid(Version);
+				// if we are specific to any particular device type, use one of them to pass to the validation function (different device types may have different SDK versions)
+				string DeviceTypeHint = string.IsNullOrEmpty(AllowedFlashDeviceTypes) ? null : AllowedFlashDeviceTypes.Split(",").First();
+				bool bIsValid = SDK.IsSoftwareVersionValid(Version, DeviceTypeHint);
 				// if we were passed a device, also check if this Sdk is valid for that device
 				if (Device != null)
 				{
@@ -233,7 +242,8 @@ namespace Turnkey
 			}
 			else
 			{
-				return UEBuildPlatformSDK.GetSDKForPlatform(Platform.ToString()).IsVersionValid(Version, bForAutoSDK: (Type == SourceType.AutoSdk));
+				string SDKTypeHint = SpecificSDKType ?? Type.ToString();
+				return UEBuildPlatformSDK.GetSDKForPlatform(Platform.ToString()).IsVersionValid(Version, SDKTypeHint);
 			}
 		}
 
@@ -292,7 +302,7 @@ namespace Turnkey
 			return Best;
 		}
 
-		static public FileSource FindMatchingSdk(AutomationTool.Platform Platform, SourceType[] TypePriority, bool bSelectBest, string DeviceType = null, string CurrentSdk = null)
+		static public FileSource FindMatchingSdk(AutomationTool.Platform Platform, SourceType[] TypePriority, bool bSelectBest, string DeviceType = null, string CurrentSdk = null, string SpecificType = null)
 		{
 			UEBuildPlatformSDK SDK = UEBuildPlatformSDK.GetSDKForPlatform(Platform.PlatformType.ToString());
 
@@ -304,8 +314,8 @@ namespace Turnkey
 				//				Sdks = Sdks.FindAll(x => x.Version == null || (Type == SourceType.Flash ? TurnkeyUtils.IsValueValid(x.Version, Platform.GetAllowedSoftwareVersions(), Platform) : SDK.IsVersionValid(x.Version, bForAutoSDK: x.Type == SourceType.AutoSdk)));
 				Sdks = Sdks.FindAll(x => x.Version == null ||
 					(Type == SourceType.Flash ?
-						SDK.IsSoftwareVersionValid(x.Version) :
-						SDK.IsVersionValid(x.Version, bForAutoSDK: (bSelectBest && x.Type == SourceType.AutoSdk))
+						SDK.IsSoftwareVersionValid(x.Version, DeviceType) :
+						SDK.IsVersionValid(x.Version, (bSelectBest && x.Type == SourceType.AutoSdk) ? "AutoSDK" : SpecificType ?? "Sdk")
 					));
 
 				if (DeviceType != null)
@@ -313,6 +323,11 @@ namespace Turnkey
 					Sdks = Sdks.FindAll(x => TurnkeyUtils.IsValueValid(DeviceType, x.AllowedFlashDeviceTypes, Platform));
 				}
 
+				if (SpecificType != null)
+				{
+					Sdks = Sdks.FindAll(x => TurnkeyUtils.IsValueValid(SpecificType, x.SpecificSDKType, Platform));
+				}
+
 				// if none were found try next type
 				if (Sdks.Count == 0)
 				{
@@ -903,6 +918,13 @@ namespace Turnkey
 
 		internal void PostDeserialize()
 		{
+			// if the Type is a platform specific type, then 
+			if (!Enum.TryParse<SourceType>(TypeString, out Type))
+			{
+				Type = SourceType.Full;
+				SpecificSDKType = TypeString;
+			}
+
 			// validate
 			if (Version == null && IsSdkType())
 			{
@@ -930,7 +952,7 @@ namespace Turnkey
 					else
 					{
 						// allow for shared SDK "platform" types (we use the AutoSDK platform name in the SDK object to determine the shared SDK "platform"
-						foreach (UEBuildPlatformSDK SDK in UEBuildPlatformSDK.AllSDKs)
+						foreach (UEBuildPlatformSDK SDK in UEBuildPlatformSDK.AllPlatformSDKObjects)
 						{
 							// if the platforms contains a AutoSDK platform that doesn't match a real platform name, add the platform to the set
 							if (PlatformStrings.Contains(SDK.GetAutoSDKPlatformName().ToLower()))
diff --git a/Engine/Source/Programs/AutomationTool/Turnkey/SdkInfo.cs b/Engine/Source/Programs/AutomationTool/Turnkey/SdkInfo.cs
--- a/Engine/Source/Programs/AutomationTool/Turnkey/SdkInfo.cs
+++ b/Engine/Source/Programs/AutomationTool/Turnkey/SdkInfo.cs
@@ -58,10 +58,11 @@ namespace Turnkey
 				Result |= LocalAvailability.Platform_InvalidHostPrerequisites;
 			}
 
-			string ManualSDKVersion, AutoSDKVersion;
-			SDK.GetInstalledVersions(out ManualSDKVersion, out AutoSDKVersion);
+			// get all the SDK info we can
+			SDKCollection SDKInfo = SDK.GetAllSDKInfo();
 
-			if (AutoSDKVersion == null)
+			// if we don't have an AutoSDK in good shape, look for others
+			if (SDKInfo.AutoSdk?.Current == null)
 			{
 				// look to see if other versions are around
 				string AutoSdkVar = Environment.GetEnvironmentVariable("UE_SDKS_ROOT");
@@ -86,39 +87,33 @@ namespace Turnkey
 						}
 					}
 				}
-
 			}
 			else
 			{
 				Result |= LocalAvailability.AutoSdk_ValidVersionExists | LocalAvailability.AutoSdk_VariableExists;
 			}
 
-			// if we have the variable at all, 
-
-			// if anything is installed, this will return a value
-			if (!string.IsNullOrEmpty(ManualSDKVersion))
+			// if all manual SDKs are valid, then we are good
+			if (SDKInfo.AreAllManualSDKsValid())
 			{
-				if (SDK.IsVersionValid(ManualSDKVersion, bForAutoSDK: false))
-				{
-					Result |= LocalAvailability.InstalledSdk_ValidVersionExists;
-				}
-				else
-				{
-					Result |= LocalAvailability.InstalledSdk_InvalidVersionExists;
-				}
+				Result |= LocalAvailability.InstalledSdk_ValidVersionExists;
+			}
+			// but if we have any manual SDKs that have a Current version, then we have invalid version installed
+			else if (SDKInfo.Sdks.Where(x => string.Compare(x.Name, "AutoSDK", true) != 0 && x.Current != null).Count() > 0)
+			{ 
+				Result |= LocalAvailability.InstalledSdk_InvalidVersionExists;
 			}
-
 
 			// look for other, inactive, versions
 			foreach (string AlternateVersion in SDK.GetAllInstalledSDKVersions())
 			{
-				if (SDK.IsVersionValid(AlternateVersion, bForAutoSDK: false))
+				// @todo turnkey: do we want to deal with multiple installed sdk for platforms with multiple SDK types???
+				if (SDK.IsVersionValid(AlternateVersion, "Sdk"))
 				{
 					Result |= LocalAvailability.InstalledSdk_ValidInactiveVersionExists;
 				}
 			}
 
-
 			string FullSupportedPlatforms = TurnkeyUtils.GetVariableValue("Studio_FullInstallPlatforms");
 			string AutoSdkSupportedPlatforms = TurnkeyUtils.GetVariableValue("Studio_AutoSdkPlatforms");
 			if (!string.IsNullOrEmpty(FullSupportedPlatforms))
diff --git a/Engine/Source/Programs/AutomationTool/Turnkey/TurnkeyUtils.cs b/Engine/Source/Programs/AutomationTool/Turnkey/TurnkeyUtils.cs
--- a/Engine/Source/Programs/AutomationTool/Turnkey/TurnkeyUtils.cs
+++ b/Engine/Source/Programs/AutomationTool/Turnkey/TurnkeyUtils.cs
@@ -33,7 +33,7 @@ namespace Turnkey
 			SetVariable("AllPlatforms", string.Join(",", UnrealTargetPlatform.GetValidPlatformNames()));
 
 			// walk over all the SDKs and get their AutoSDK string
-			IEnumerable<string> AutoSDKPlatforms = UEBuildPlatformSDK.AllSDKs.Select(x => x.GetAutoSDKPlatformName()).Distinct();
+			IEnumerable<string> AutoSDKPlatforms = UEBuildPlatformSDK.AllPlatformSDKObjects.Select(x => x.GetAutoSDKPlatformName()).Distinct();
 			SetVariable("AutoSDKPlatforms", string.Join(",", AutoSDKPlatforms));
 
 			// 			TurnkeyUtils.Log("AllPlatforms = {0}", GetVariableValue("AllPlatforms"));
diff --git a/Engine/Source/Programs/Shared/EpicGames.Build/UEBuildPlatformSDK.cs b/Engine/Source/Programs/Shared/EpicGames.Build/UEBuildPlatformSDK.cs
--- a/Engine/Source/Programs/Shared/EpicGames.Build/UEBuildPlatformSDK.cs
+++ b/Engine/Source/Programs/Shared/EpicGames.Build/UEBuildPlatformSDK.cs
@@ -7,6 +7,7 @@ using System.Diagnostics.CodeAnalysis;
 using System.IO;
 using System.Linq;
 using System.Text;
+using System.Text.Json;
 
 namespace EpicGames.Core
 {
@@ -26,6 +27,258 @@ namespace EpicGames.Core
 		Invalid,
 	};
 
+	/// <summary>
+	/// Range of SDK versions that are valid for a platform
+	/// </summary>
+	public class SDKDescriptor
+	{
+		public string Name;
+		public string? Min;
+		public string? Max;
+		public string? Current;
+
+		// for SDKs where you need 1 of a set of SDKs (for instance, multiple toolchains for one platform), give them all the same GroupName
+		// and Turnkey will make sure that at least one of them is installed for the SDK setup to be valid. Leave null otherwise
+		public string? GroupName = null;
+
+		public UInt64 MinInt = 0;
+		public UInt64 MaxInt = UInt64.MaxValue;
+		public UInt64 CurrentInt = 0;
+
+		public SDKStatus Validity = SDKStatus.Invalid;
+
+		public SDKDescriptor()
+		{
+			Name = "Unknown";
+		}
+		public SDKDescriptor(string InName, string? InMin, string? InMax, string? InCurrent, string? InGroupName, SDKCollection? Collection)
+		{
+			Name = InName;
+			Min = InMin;
+			Max = InMax;
+			Current = InCurrent;
+			GroupName = InGroupName;
+
+			if (Collection != null)
+			{
+				MinInt = Collection.PlatformSDK.ConvertVersionToInt(Min, Name);
+				// null will convert to 0, but we want MaxValue on null
+				MaxInt = Max == null ? UInt64.MaxValue : Collection.PlatformSDK.ConvertVersionToInt(Max, Name);
+				CurrentInt = Collection.PlatformSDK.ConvertVersionToInt(Current, Name);
+				UpdateValidity(Collection);
+			}
+		}
+
+
+		public SDKDescriptor(string InName, string? InMin, string? InMax)
+			: this(InName, InMin, InMax, null, null, null)
+		{
+		}
+
+		public SDKDescriptor(string InName, string? InCurrent)
+			: this(InName, null, null, InCurrent, null, null)
+		{
+		}
+
+		public void UpdateCurrent(string InCurrent, string? Hint,  SDKCollection? Collection)
+		{
+			Current = InCurrent;
+			if (Collection != null)
+			{
+				CurrentInt = Collection.PlatformSDK.ConvertVersionToInt(Current, Hint);
+				UpdateValidity(Collection);
+			}
+		}
+
+		private void UpdateValidity(SDKCollection Collection)
+		{
+			// instead of just checking the range numerically, we allow platforms to override the IsValid check, which makes this a little convoluted (calling back into the Collection 
+			// so it can call back into the PlatformSDK with the proper SDK vs Software function)
+			Validity = (Current != null && Collection.IsValid(Current, this)) ? SDKStatus.Valid : SDKStatus.Invalid;
+		}
+
+		public string ToString(bool bIncludeCurrent=true)
+		{
+			return $"MinVersion={Min ?? ""}, MaxVersion={Max ?? ""}" + (bIncludeCurrent ? $", Current={Current ?? ""}" : "");
+		}
+		public string ToString(string Descriptor, string? CurrentDescriptor, bool bIncludeName)
+		{
+			string DisplayName = bIncludeName ? "_" + Name : "";
+
+			// if they are different, print something like:
+			//   MinAllowed_Toolchain=1.0, MaxAllowed_Toolchain=2.0
+			if (Min != Max)
+			{
+				return $"Min{Descriptor}{DisplayName}={Min ?? ""}, Max{Descriptor}{DisplayName}={Max ?? ""}" + (CurrentDescriptor != null ? $", Current{CurrentDescriptor}{DisplayName}={Current ?? ""}" : "");
+			}
+			// if they are the same, print something like:
+			//   Allowed_Toolchain=1.0
+			return $"{Descriptor}{DisplayName}={Min}" + (CurrentDescriptor != null ? $", Current{CurrentDescriptor}{DisplayName}={Current ?? ""}" : "");
+		}
+	}
+	public class SDKCollection
+	{
+		public virtual string DefaultName { get { return "Sdk"; } }
+		public virtual bool IsValid(string Version, SDKDescriptor Info)
+		{
+			return PlatformSDK.IsVersionValid(Version, Info);
+		}
+
+
+
+		public UEBuildPlatformSDK PlatformSDK;
+		public List<SDKDescriptor> Sdks = new List<SDKDescriptor>();
+		public List<SDKDescriptor> FullSdks
+		{ get { return Sdks.Where(x => string.Compare(x.Name, "AutoSdk", true) != 0).ToList(); } }
+		public SDKDescriptor? AutoSdk
+		{ get { return Sdks.FirstOrDefault(x => string.Compare(x.Name, "AutoSdk", true) == 0); } }
+
+
+		public SDKCollection(UEBuildPlatformSDK PlatformSDK)
+		{
+			this.PlatformSDK = PlatformSDK;
+		}
+
+		public SDKCollection(SDKCollection Other, UEBuildPlatformSDK PlatformSDK)
+			: this(PlatformSDK)
+		{
+			foreach (SDKDescriptor Desc in Other.Sdks)
+			{
+				SetupSDK(Desc.Name, Desc.Min, Desc.Max, Desc.Current, Desc.GroupName);
+			}
+		}
+
+		// convenience for single unnamed Sdk range
+		public SDKCollection(string? Min, string? Max, UEBuildPlatformSDK PlatformSDK)
+			: this(PlatformSDK)
+		{
+			Sdks.Add(new SDKDescriptor(DefaultName, Min, Max));
+		}
+
+		// convenience for single unnamed Sdk current value
+		public SDKCollection(string? Current, UEBuildPlatformSDK PlatformSDK)
+			: this(PlatformSDK)
+		{
+			if (Current != null)
+			{
+				Sdks.Add(new SDKDescriptor(DefaultName, Current));
+			}
+		}
+
+		public bool AreAllManualSDKsValid()
+		{
+			IEnumerable<SDKDescriptor>? Sdks = FullSdks;
+
+			if (Sdks == null)
+			{
+				return true;
+			}
+
+			// check if all the SDKs not in groups are valid (All returns true if empty set)
+			bool bAllUngroupedAreValid = Sdks
+				.Where(Desc => Desc.GroupName == null)
+				.All(Desc => Desc.Validity == SDKStatus.Valid);
+
+			// for each groupname (ToHashSet removes duplicates), check if at least 1 SDK with that group name is valid
+			bool bAllGroupsHaveOneValid = Sdks
+				.Where(Desc => Desc.GroupName != null)
+				.Select(Desc => Desc.GroupName)
+				.ToHashSet()
+				.All(GroupName => Sdks.Any(Desc => Desc.GroupName == GroupName && Desc.Validity == SDKStatus.Valid));
+
+			return bAllUngroupedAreValid && bAllGroupsHaveOneValid;
+		}
+	
+		public bool IsAutoSDKValid()
+		{
+			return AutoSdk?.Validity == SDKStatus.Valid;
+		}
+
+
+		public string ToString(bool bIncludeCurrent=true)
+		{
+			// by default print something like "MinVersion=1.0, MaxVersion=2.0"
+			return ToString("Version", "Version");
+		}
+
+		public string ToString(string Descriptor, string? CurrentDescriptor=null)
+		{
+			if (Sdks.Count == 1)
+			{
+				return Sdks[0].ToString(Descriptor, CurrentDescriptor, false);
+			}
+			return string.Join(", ", Sdks.Select(x => x.ToString(Descriptor, CurrentDescriptor, true)));
+		}
+
+		public string ToMultilineString(bool bIncludeCurrent = true)
+		{
+			return $"";// MinVersion={Min ?? ""}, MaxVersion={Max ?? ""}" + (bIncludeCurrent ? $", Current={Current ?? ""}" : "");
+		}
+
+		public void SetupSDK(string Name, string? Min, string? Max, string? Current, string? GroupName)
+		{
+			Sdks.Add(new SDKDescriptor(Name, Min, Max, Current, GroupName, this));
+		}
+
+		public void SetupCurrent(string Name, string Current)
+		{
+			Sdks.Add(new SDKDescriptor(Name, null, null, Current, null, this));
+		}
+
+		public bool UpdateCurrentForSingle(string Name, string CurrentVersion)
+		{
+			if (PlatformSDK == null)
+			{
+				throw new Exception("Cannot call UpdateCurrentForSingle in a SDKCollection without a PlatformSDK set");
+			}
+
+			SDKDescriptor? Info;
+
+			// for ease, we assume that if there is only one "Sdk" or "Software" in this Collection, then any Name passed in will work with it
+			// so just use that one entry
+			if (Sdks.Count == 1 && (Sdks[0].Name == "Sdk" || Sdks[0].Name == "Software"))
+			{
+				Info = Sdks[0];
+			}
+			else
+			{
+				Info = Sdks.FirstOrDefault(x => string.Compare(x.Name, Name, true) == 0);
+				if (Info == null)
+				{
+					return false;
+				}
+			}
+
+			Info.UpdateCurrent(CurrentVersion, Name, this);
+			return true;
+		}
+	}
+
+	public class SoftwareCollection : SDKCollection
+	{
+		public SoftwareCollection(UEBuildPlatformSDK PlatformSDK)
+			: base(PlatformSDK)
+		{
+		}
+		public SoftwareCollection(string? Min, string? Max, UEBuildPlatformSDK PlatformSDK)
+			: base(Min, Max, PlatformSDK)
+		{
+		}
+		public SoftwareCollection(string? Current, UEBuildPlatformSDK PlatformSDK)
+			: base(Current, PlatformSDK)
+		{
+		}
+
+
+		public override string DefaultName { get { return "Software"; } }
+
+		public override bool IsValid(string Version, SDKDescriptor Info)
+		{
+			return PlatformSDK.IsSoftwareVersionValid(Version, Info);
+		}
+	}
+
+
 	/// <summary>
 	/// SDK for a platform
 	/// </summary>
@@ -63,18 +316,32 @@ namespace EpicGames.Core
 			bIsSdkAllowedOnHost = bInIsSdkAllowedOnHost;
 
 			// if the parent set up autosdk, the env vars will be wrong, but we can still get the manual SDK version from before it was setup
-			if (HasParentProcessSetupAutoSDK(out CachedManualSDKVersion))
+			string? ParentManualSDKVersions = Environment.GetEnvironmentVariable(GetPlatformManualSDKSetupEnvVar());
+			if (!string.IsNullOrEmpty(ParentManualSDKVersions))
 			{
+
 				// we pass along __None to indicate the parent didn't have a manual sdk installed
-				if (CachedManualSDKVersion == "__None")
+				if (ParentManualSDKVersions == "__None")
 				{
-					CachedManualSDKVersion = null;
+					// empty it out
+					CachedManualSDKVersions = new Dictionary<string, string>();
+				}
+				else
+				{
+					Console.WriteLine("Processing parent manual sdk: {0}", ParentManualSDKVersions);
+					CachedManualSDKVersions = JsonSerializer.Deserialize<Dictionary<string, string>>(ParentManualSDKVersions);
 				}
 			}
 			else
 			{
+				CachedManualSDKVersions = new Dictionary<string, string>();
+
 				// if there was no parent, get the SDK version before we run AutoSDK to get the manual version
-				CachedManualSDKVersion = GetInstalledSDKVersion();
+				SDKCollection InstalledVersions = GetInstalledSDKVersions();
+				foreach (SDKDescriptor Desc in InstalledVersions.FullSdks)
+				{
+					CachedManualSDKVersions[Desc.Name] = Desc.Current!;
+				}
 			}
 		}
 
@@ -98,9 +365,9 @@ namespace EpicGames.Core
 		/// <summary>
 		/// Gets the set of all known SDKs
 		/// </summary>
-		public static UEBuildPlatformSDK[] AllSDKs
+		public static UEBuildPlatformSDK[] AllPlatformSDKObjects
 		{
-			get	{ return SDKRegistry.Values.ToArray(); }
+			get { return SDKRegistry.Values.ToArray(); }
 		}
 
 		// String name of the platform (will match an UnrealTargetPlatform)
@@ -109,45 +376,129 @@ namespace EpicGames.Core
 		// True if this Sdk is allowed to be used by this host - if not, we can skip a lot 
 		public bool bIsSdkAllowedOnHost;
 
-		public string? GetInstalledVersion(out bool bIsAutoSDK)
+		public SDKCollection GetAllSDKInfo()
 		{
-			bIsAutoSDK = HasSetupAutoSDK();
-			return GetInstalledSDKVersion();
+			SDKCollection AllSdks = new SDKCollection(this);
+
+			// walk over each one version the platform supports, and get it's current installed info
+			foreach (SDKDescriptor Desc in GetValidVersions().Sdks)
+			{
+				AllSdks.SetupSDK(Desc.Name, Desc.Min, Desc.Max, CachedManualSDKVersions.GetValueOrDefault(Desc.Name), Desc.GroupName);
+			}
+
+			// now get current AutoSDK version (if autosdk is set up, then the GetInstalledSDKVersion will return AutoSDK version)
+			bool bIsAutoSDK = false;
+			string? CurrentAutoSDKVersion = (PlatformSupportsAutoSDKs() && HasRequiredAutoSDKInstalled() == SDKStatus.Valid && HasSetupAutoSDK()) ? GetInstalledVersion(out bIsAutoSDK) : null;
+			AllSdks.SetupSDK("AutoSdk", GetMainVersion(), GetMainVersion(), CurrentAutoSDKVersion, null);
+
+			// verify some assumptions
+			if (CurrentAutoSDKVersion != null && bIsAutoSDK == false)
+			{
+				throw new Exception($"AutoSDK was indicated to be setup ({CurrentAutoSDKVersion}), but GetInstalledSDKVersion returned false for bIsAutoSDK");
+			}
+			if (CurrentAutoSDKVersion != null && CurrentAutoSDKVersion != GetMainVersion())
+			{
+				throw new Exception($"AutoSDK was indicated to be setup, but the version if returned ({CurrentAutoSDKVersion}) doesn't equal the MainVersion ({GetMainVersion()}");
+			}
+
+			return AllSdks;
 		}
 
-		public string? GetInstalledVersion()
+		public SDKDescriptor? GetSDKInfo(string SDKName)
 		{
-			return GetInstalledSDKVersion();
+			return GetAllSDKInfo().Sdks.FirstOrDefault(x => string.Compare(x.Name, SDKName, true) == 0);
 		}
 
-		public void GetInstalledVersions(out string? ManualSDKVersion, out string? AutoSDKVersion)
+		public SDKDescriptor? GetSoftwareInfo(string? SDKName=null)
 		{
-			// if we support AutoSDKs, then return both versions
-			if (PlatformSupportsAutoSDKs())
+			if (SDKName == null)
 			{
-				AutoSDKVersion = (HasRequiredAutoSDKInstalled() == SDKStatus.Valid) ? GetInstalledSDKVersion() : null;
-//				AutoSDKVersion = GetInstalledSDKVersion();
+				// todo: make this queryable?
+				SDKName = "Software";
 			}
-			else
+			return GetAllSoftwareInfo().Sdks.FirstOrDefault(x => string.Compare(x.Name, SDKName, true) == 0);
+		}
+
+		public SDKCollection GetAllSoftwareInfo(string? DeviceType=null, string? Current=null)
+		{
+			SDKCollection AllSoftwares = new SDKCollection(this);
+
+			// walk over the valid software ranges, potentially setting the current verison if one was passed in
+			foreach (SDKDescriptor Desc in GetValidSoftwareVersions().Sdks)
 			{
-				AutoSDKVersion = null;
-				if (CachedManualSDKVersion != GetInstalledSDKVersion())
+				// are we restricting this to one type?
+				if (DeviceType == null || string.Compare(Desc.Name, DeviceType, true) == 0)
 				{
-					throw new Exception("Manual SDK version changed, this is not supported yet");
+					AllSoftwares.SetupSDK(Desc.Name, Desc.Min, Desc.Max, null, Desc.GroupName);
+				}
+
+				if (Current != null)
+				{
+					// set the current for this one (likely it will only be for DeviceType, but support passing in null DeviceType, and non-null Current
+					AllSoftwares.UpdateCurrentForSingle(Desc.Name, Current);
 				}
 			}
 
-			ManualSDKVersion = CachedManualSDKVersion;
+			return AllSoftwares;
 		}
 
+		//public SDKCollection GetAllSoftwareInfo(string DeviceId)
+		//{
+		//	SDKCollection AllSoftwares =
+		//}
+
+
+		public string? GetInstalledVersion(out bool bIsAutoSDK)
+		{
+			bIsAutoSDK = HasSetupAutoSDK();
+			return GetInstalledSDKVersion();
+		}
 
-		public virtual bool IsVersionValid(string? Version, bool bForAutoSDK)
+		public string? GetInstalledVersion()
+		{
+			return GetInstalledSDKVersion();
+		}
+		public string? GetInstalledVersion(string SDKName)
+		{
+			return GetSDKInfo(SDKName)?.Current;
+		}
+
+		//		public void GetInstalledVersions(out string? ManualSDKVersion, out string? AutoSDKVersion)
+		//		{
+		//			// if we support AutoSDKs, then return both versions
+		//			if (PlatformSupportsAutoSDKs())
+		//			{
+		//				AutoSDKVersion = (HasRequiredAutoSDKInstalled() == SDKStatus.Valid) ? GetInstalledSDKVersion() : null;
+		////				AutoSDKVersion = GetInstalledSDKVersion();
+		//			}
+		//			else
+		//			{
+		//				AutoSDKVersion = null;
+		//				if (CachedManualSDKVersion != GetInstalledSDKVersion())
+		//				{
+		//					throw new Exception("Manual SDK version changed, this is not supported yet");
+		//				}
+		//			}
+
+		//			ManualSDKVersion = CachedManualSDKVersion;
+		//		}
+
+
+		public virtual bool IsVersionValid(string? Version, string SDKType)
+		{
+			return IsVersionValidInternal(Version, SDKType, null);
+		}
+		public virtual bool IsVersionValid(string? Version, SDKDescriptor Info)
 		{
-			return IsVersionValidInternal(Version, bForAutoSDK);
+			return IsVersionValidInternal(Version, null, Info);
 		}
-		public virtual bool IsSoftwareVersionValid(string Version)
+		public virtual bool IsSoftwareVersionValid(string Version, string DeviceType)
 		{
-			return IsSoftwareVersionValidInternal(Version);
+			return IsSoftwareVersionValidInternal(Version, DeviceType, null);
+		}
+		public virtual bool IsSoftwareVersionValid(string Version, SDKDescriptor Info)
+		{
+			return IsSoftwareVersionValidInternal(Version, null, Info);
 		}
 
 		public void ReactivateAutoSDK()
@@ -161,16 +512,73 @@ namespace EpicGames.Core
 		#region Platform Overrides
 
 		/// <summary>
-		/// Returns the installed SDK version, used to determine if up to date or not (
+		/// Return the SDK version that the platform wants to use (AutoSDK dir must match this, full SDKs can be in a valid range)
 		/// </summary>
 		/// <returns></returns>
-		public abstract string? GetInstalledSDKVersion();
+		public abstract string GetMainVersion();
 
 		/// <summary>
-		/// Return the SDK version that the platform wants to use (AutoSDK dir must match this, full SDKs can be in a valid range)
+		/// Gets the valid string range of Sdk versions. TryConvertVersionToInt() will need to succeed to make this usable for range checks
+		/// </summary>
+		/// <param name="MinVersion">Smallest version allowed</param>
+		/// <param name="MaxVersion">Largest version allowed (inclusive)</param>
+		protected virtual void GetValidVersionRange(out string MinVersion, out string MaxVersion)
+		{
+			throw new Exception($"This platform's Sdk class ({GetType().Name}) must implement either GetValidVersionRange() or GetValidVersions(). If this triggers, and GetValidVersions() is implemented, that means GetValidVersionRange was called directly");
+		}
+
+		protected virtual SDKCollection GetValidVersions()
+		{
+			// if the platform doesn't override this, then it only has one sdk, so put it into the collection as the single sdk (this is very much the standard behavior)
+			string MinVersion, MaxVersion;
+			GetValidVersionRange(out MinVersion, out MaxVersion);
+			return new SDKCollection(MinVersion, MaxVersion, this);
+		}
+
+		/// <summary>
+		/// Gets the valid string range of software/flash versions. TryConvertVersionToInt() will need to succeed to make this usable for range checks
+		/// </summary>
+		/// <param name="MinVersion">Smallest version allowed, or null if no minmum (in other words, 0 - MaxVersion)</param>
+		/// <param name="MaxVersion">Largest version allowed (inclusive), or null if no maximum (in other words, MinVersion - infinity)y</param>
+		protected virtual void GetValidSoftwareVersionRange(out string? MinVersion, out string? MaxVersion)
+		{
+			throw new Exception($"This platform's Sdk class ({GetType().Name}) must implement either GetValidSoftwareVersionRange() or GetValidSoftwareVersionRanges(). If this triggers, and GetValidSoftwareVersionRanges() is implemented, that means GetValidSoftwareVersionRange was called directly");
+		}
+
+		protected virtual SoftwareCollection GetValidSoftwareVersions()
+		{
+			// if the platform doesn't override this, then it only has one sdk, so put it into the collection as the single software (this is very much the standard behavior)
+			string? MinVersion, MaxVersion;
+			GetValidSoftwareVersionRange(out MinVersion, out MaxVersion);
+			SoftwareCollection SoftwareVersions = new SoftwareCollection(MinVersion, MaxVersion, this);
+			return SoftwareVersions;
+		}
+
+		/// <summary>
+		/// Returns the installed SDK version, used to determine if up to date or not (
 		/// </summary>
 		/// <returns></returns>
-		public abstract string GetMainVersion();
+		protected virtual string? GetInstalledSDKVersion()
+		{
+			if (GetInstalledSDKVersions().FullSdks.Count() == 1)
+			{
+				return GetInstalledSDKVersions().FullSdks[0].Current;
+			}
+			throw new Exception($"This platform's SDK class {GetType()} did not implement GetInstalledSDKVersion(). That means it has multiple SDKs and you will need use GetAllSDKInfo() or GetInstalledVersion(SDKName)");
+		}
+
+		protected virtual SDKCollection GetInstalledSDKVersions()
+		{
+			// if the platform doesn't override this, then it only has one sdk, so put it into the collection as the single sdk (this is very much the standard behavior)
+			string? Version = GetInstalledSDKVersion();
+			if (Version != null)
+			{
+				return new SDKCollection(Version, this);
+			}
+
+			// if no manual version, then return an empty list of current sdks
+			return new SDKCollection(this);
+		}
 
 		/// <summary>
 		/// Returns a platform-specific version string that can be retrieved from anywhere without needing to typecast the SDK object
@@ -202,20 +610,6 @@ namespace EpicGames.Core
 			return false;
 		}
 
-		/// <summary>
-		/// Gets the valid string range of Sdk versions. TryConvertVersionToInt() will need to succeed to make this usable for range checks
-		/// </summary>
-		/// <param name="MinVersion">Smallest version allowed</param>
-		/// <param name="MaxVersion">Largest version allowed (inclusive)</param>
-		public abstract void GetValidVersionRange(out string MinVersion, out string MaxVersion);
-
-
-		/// <summary>
-		/// Gets the valid string range of software/flash versions. TryConvertVersionToInt() will need to succeed to make this usable for range checks
-		/// </summary>
-		/// <param name="MinVersion">Smallest version allowed, or null if no minmum (in other words, 0 - MaxVersion)</param>
-		/// <param name="MaxVersion">Largest version allowed (inclusive), or null if no maximum (in other words, MinVersion - infinity)y</param>
-		public abstract void GetValidSoftwareVersionRange(out string? MinVersion, out string? MaxVersion);
 
 		/// <summary>
 		/// For a platform that doesn't use properly named AutoSDK directories, the directory name may not be convertible to an integer,
@@ -229,42 +623,42 @@ namespace EpicGames.Core
 		}
 
 
-		/// <summary>
-		/// Gets the valid (integer) range of Sdk versions. Must be an integer to easily check a range vs a particular version
-		/// </summary>
-		/// <param name="MinVersion">Smallest version allowed</param>
-		/// <param name="MaxVersion">Largest version allowed (inclusive)</param>
-		/// <returns>True if the versions are valid, false if the platform is unable to convert its versions into an integer</returns>
-		public virtual void GetValidVersionRange(out UInt64 MinVersion, out UInt64 MaxVersion)
-		{
-			string MinVersionString, MaxVersionString;
-			GetValidVersionRange(out MinVersionString, out MaxVersionString);
-
-			// failures to convert here are bad
-			if (!TryConvertVersionToInt(MinVersionString, out MinVersion) || !TryConvertVersionToInt(MaxVersionString, out MaxVersion))
-			{
-				throw new Exception(string.Format("Unable to convert Min and Max valid versions to integers in {0} (Versions are {1} - {2})", GetType().Name, MinVersionString, MaxVersionString));
-			}
-		}
-		public virtual void GetValidSoftwareVersionRange(out UInt64 MinVersion, out UInt64 MaxVersion)
-		{
-			string? MinVersionString, MaxVersionString;
-			GetValidSoftwareVersionRange(out MinVersionString, out MaxVersionString);
-
-			MinVersion = UInt64.MinValue;
-			MaxVersion = UInt64.MaxValue - 1; // MaxValue is always bad
-
-			// failures to convert here are bad
-			if ((MinVersionString != null && !TryConvertVersionToInt(MinVersionString, out MinVersion)) || 
-				(MaxVersionString != null && !TryConvertVersionToInt(MaxVersionString, out MaxVersion)))
-			{
-				throw new Exception(string.Format("Unable to convert Min and Max valid Software versions to integers in {0} (Versions are {1} - {2})", GetType().Name, MinVersionString, MaxVersionString));
-			}
-		}
+		///// <summary>
+		///// Gets the valid (integer) range of Sdk versions. Must be an integer to easily check a range vs a particular version
+		///// </summary>
+		///// <param name="MinVersion">Smallest version allowed</param>
+		///// <param name="MaxVersion">Largest version allowed (inclusive)</param>
+		///// <returns>True if the versions are valid, false if the platform is unable to convert its versions into an integer</returns>
+		//public virtual void GetValidVersionRange(out UInt64 MinVersion, out UInt64 MaxVersion)
+		//{
+		//	string MinVersionString, MaxVersionString;
+		//	GetValidVersionRange(out MinVersionString, out MaxVersionString);
+
+		//	// failures to convert here are bad
+		//	if (!TryConvertVersionToInt(MinVersionString, out MinVersion) || !TryConvertVersionToInt(MaxVersionString, out MaxVersion))
+		//	{
+		//		throw new Exception(string.Format("Unable to convert Min and Max valid versions to integers in {0} (Versions are {1} - {2})", GetType().Name, MinVersionString, MaxVersionString));
+		//	}
+		//}
+		//public virtual void GetValidSoftwareVersionRange(out UInt64 MinVersion, out UInt64 MaxVersion)
+		//{
+		//	string? MinVersionString, MaxVersionString;
+		//	GetValidSoftwareVersionRange(out MinVersionString, out MaxVersionString);
+
+		//	MinVersion = UInt64.MinValue;
+		//	MaxVersion = UInt64.MaxValue - 1; // MaxValue is always bad
+
+		//	// failures to convert here are bad
+		//	if ((MinVersionString != null && !TryConvertVersionToInt(MinVersionString, out MinVersion)) ||
+		//		(MaxVersionString != null && !TryConvertVersionToInt(MaxVersionString, out MaxVersion)))
+		//	{
+		//		throw new Exception(string.Format("Unable to convert Min and Max valid Software versions to integers in {0} (Versions are {1} - {2})", GetType().Name, MinVersionString, MaxVersionString));
+		//	}
+		//}
 
 
 		// Let platform override behavior to determine if a version is a valid (useful for non-numeric versions)
-		protected virtual bool IsVersionValidInternal(string? Version, bool bForAutoSDK)
+		protected virtual bool IsVersionValidInternal(string? Version, string? SDKType, SDKDescriptor? VersionInfo)
 		{
 			// we could have null if no SDK is installed at all, etc, which is always a failure
 			if (Version == null)
@@ -273,38 +667,39 @@ namespace EpicGames.Core
 			}
 
 			// AutoSDK must match the desired version exactly, since that is the only one we will use
-			if (bForAutoSDK)
+			if (string.Compare(SDKType, "AutoSdk", true) == 0)
 			{
 				// if integer version checking failed, then we can detect valid autosdk if the version matches the autosdk directory by name
 				return string.Compare(Version, GetAutoSDKDirectoryForMainVersion(), true) == 0;
 			}
 
-			// convert it to an integer
-			UInt64 IntVersion;
-			if (!TryConvertVersionToInt(Version, out IntVersion))
+			// look for a range for this hinted type of SDK
+			VersionInfo = VersionInfo ?? GetSDKVersionForHint(GetValidVersions(), SDKType);
+
+			// if we couldn't find the Info, we can't do anything, so fail
+			if (VersionInfo == null)
 			{
 				return false;
 			}
 
-			UInt64 DesiredVersion;
-			if (!TryConvertVersionToInt(GetMainVersion(), out DesiredVersion))
+			// convert it to an integer
+			UInt64 IntVersion;
+			if (!TryConvertVersionToInt(Version, out IntVersion))
 			{
 				return false;
 			}
 
 			// short circuit range check if the Version is the desired version already
-			if (IntVersion == DesiredVersion)
+			if (IntVersion == ConvertVersionToInt(GetMainVersion()))
 			{
 				return true;
 			}
 
-			// get numeric range
-			UInt64 MinVersion, MaxVersion;
-			GetValidVersionRange(out MinVersion, out MaxVersion);
-			return IntVersion >= MinVersion && IntVersion <= MaxVersion;
+			// finally do the numeric comparison
+			return IntVersion >= VersionInfo.MinInt && IntVersion <= VersionInfo.MaxInt;
 		}
 
-		protected virtual bool IsSoftwareVersionValidInternal(string Version)
+		protected virtual bool IsSoftwareVersionValidInternal(string? Version, string? DeviceType, SDKDescriptor? VersionInfo)
 		{
 			// we could have null if no SDK is installed at all, etc, which is always a failure
 			if (Version == null)
@@ -319,10 +714,17 @@ namespace EpicGames.Core
 				return false;
 			}
 
-			// get numeric range
-			UInt64 MinVersion, MaxVersion;
-			GetValidSoftwareVersionRange(out MinVersion, out MaxVersion);
-			return IntVersion >= MinVersion && IntVersion <= MaxVersion;
+			// look for a range for this hinted type of SDK
+			VersionInfo = VersionInfo ?? GetSDKVersionForHint(GetValidSoftwareVersions(), DeviceType);
+
+			// if we couldn't find the Info, we can't do anything, so fail
+			if (VersionInfo == null)
+			{
+				return false;
+			}
+
+			// finally do the numeric comparison
+			return IntVersion >= VersionInfo.MinInt && IntVersion <= VersionInfo.MaxInt;
 		}
 
 
@@ -330,13 +732,31 @@ namespace EpicGames.Core
 		/// Only the platform can convert a version string into an integer that is usable for comparison
 		/// </summary>
 		/// <param name="StringValue">Version that comes from the installed SDK or a Turnkey manifest or the like</param>
-		/// <param name="OutValue">The integer version of StringValue, can be used to compare against a valid range</param>
+		/// <param name="OutValue"></param>
+		/// <param name="Hint">A platform specific hint that can help guide conversion (usually SDKName or device type)</param>
 		/// <returns>If the StringValue was able to be be converted to an integer</returns>
-		public virtual bool TryConvertVersionToInt(string? StringValue, out UInt64 OutValue)
+		public abstract bool TryConvertVersionToInt(string? StringValue, out UInt64 OutValue, string? Hint = null);
+
+		/// <summary>
+		/// Like TryConvertVersionToInt, but will throw an exception on failure
+		/// </summary>
+		/// <param name="StringValue">Version that comes from the PlatformSDK class (should not be used with Manifest or other user supplied versions)</param>
+		/// <param name="Hint">A platform specific hint that can help guide conversion (usually SDKName or device type)</param>
+		/// <returns>The integer version of StringValue, can be used to compare against a valid range</returns>
+		public UInt64 ConvertVersionToInt(string? StringValue, string? Hint = null)
 		{
-			// @todo turnkey make this abstract?
-			OutValue = 0;
-			return false;
+			// quickly handle null input, which is not necessarily an error case
+			if (StringValue == null)
+			{
+				return 0;
+			}
+
+			UInt64 Result;
+			if (!TryConvertVersionToInt(StringValue, out Result, Hint))
+			{
+				throw new Exception($"Unable to convert {GetType()} version {StringValue} to an integer. Likely this version was supplied by code, and is expected to be valid.");
+			}
+			return Result;
 		}
 
 		/// <summary>
@@ -362,58 +782,48 @@ namespace EpicGames.Core
 				return SDKInfoValidity.Value;
 			}
 
-			string? ManualSDKVersion, AutoSDKVersion;
-			GetInstalledVersions(out ManualSDKVersion, out AutoSDKVersion);
+			SDKCollection SDKInfo = GetAllSDKInfo();
 
-			SDKStatus Validity = SDKStatus.Valid;
+			// will mark invalid below if needed
+			SDKInfoValidity = SDKStatus.Valid;
 
-			bool HasValidAutoSDK = HasSetupAutoSDK();
-		
-			if (HasValidAutoSDK)
+			SDKDescriptor? AutoSDKInfo = SDKInfo.AutoSdk;
+			if (AutoSDKInfo != null && AutoSDKInfo.Validity == SDKStatus.Valid)
 			{
 				string PlatformSDKRoot = GetPathToPlatformAutoSDKs();
-
-				UInt64 Ver;
-
-				if (AutoSDKVersion != null && TryConvertVersionToInt(AutoSDKVersion, out Ver))
-				{
-					Log.WriteLine(Verbosity, Options, "{0} using Auto SDK {1} from: {2} 0x{3:X}", PlatformName, AutoSDKVersion, Path.Combine(PlatformSDKRoot, GetAutoSDKDirectoryForMainVersion()), Ver);
-				}
-				else
-				{
-					HasValidAutoSDK = false;
-				}
+				Log.WriteLine(Verbosity, Options, "{0} using Auto SDK {1} from: {2} 0x{3:X}", PlatformName, AutoSDKInfo.Current, Path.Combine(PlatformSDKRoot, GetAutoSDKDirectoryForMainVersion()), AutoSDKInfo.CurrentInt);
 			}
-
-			if (HasValidAutoSDK == false)
+			else
 			{
-				if (HasRequiredManualSDK() == SDKStatus.Valid)
+				if (SDKInfo.AreAllManualSDKsValid())
 				{
-					Log.WriteLine(Verbosity, Options, "{0} using Manual SDK {1}", PlatformName, ManualSDKVersion);
+					Log.WriteLine(Verbosity, Options, "{0} Installed SDK(s): {1}", PlatformName, SDKInfo.ToString(true));
 				}
 				else
 				{
-					Validity = SDKStatus.Invalid;
-
-					string MinVersionString, MaxVersionString;
-					GetValidVersionRange(out MinVersionString, out MaxVersionString);
+					SDKInfoValidity = SDKStatus.Invalid;
 
 					StringBuilder Msg = new StringBuilder();
-					Msg.AppendFormat("Unable to find a valid SDK for {0}.", PlatformName);
-					if (ManualSDKVersion != null)
-					{
-						Msg.AppendFormat(" Found Version: {0}.", ManualSDKVersion);
-					}
+					Msg.AppendFormat("Unable to find valid SDK(s) for {0}:", PlatformName);
 
-					if (MinVersionString != MaxVersionString)
+					foreach (SDKDescriptor Desc in SDKInfo.Sdks)
 					{
-						Msg.AppendLine(" Must be between {0} and {1}", MinVersionString, MaxVersionString);
-					}
-					else
-					{
-						Msg.AppendLine(" Must be {0}", MinVersionString);
+						if (Desc.Validity == SDKStatus.Valid)
+						{
+							Msg.Append($"  {Desc.Name} is valid ({Desc}");
+						}
+						else
+						{
+							if (Desc.Current != null)
+							{
+								Msg.AppendFormat($" Found {Desc.Name} Version: {Desc.Current}.");
+							}
+
+							Msg.AppendLine($"   {Desc.ToString("Required", null, false)}");
+						}
 					}
 
+
 					if (!bHasShownTurnkey)
 					{
 						Msg.AppendLine("  If your Studio has it set up, you can run this command to find the SDK to install:");
@@ -430,8 +840,7 @@ namespace EpicGames.Core
 				}
 			}
 
-			SDKInfoValidity = Validity;
-			return Validity;
+			return SDKInfoValidity.Value;
 		}
 
 		#endregion
@@ -443,9 +852,31 @@ namespace EpicGames.Core
 		#region Private/Protected general functionality
 
 		// this is the SDK version that was set before activating AutoSDK, since AutoSDK may remove ability to retrieve the Manual SDK version
-		protected string? CachedManualSDKVersion;
+		protected Dictionary<string, string> CachedManualSDKVersions = new Dictionary<string, string>();
 		private static Dictionary<string, UEBuildPlatformSDK> SDKRegistry = new Dictionary<string, UEBuildPlatformSDK>();
 
+		private SDKDescriptor? GetSDKVersionForHint(SDKCollection Collection, string? Hint)
+		{
+			// if the hint is found, use it always
+			SDKDescriptor? SDKDesc = Collection.Sdks.FirstOrDefault(x => string.Compare(x.Name, Hint, true) == 0);
+			if (SDKDesc != null)
+			{
+				return SDKDesc;
+			}
+
+			// only use AUtoSDK if explicitly asked for above, otherwise, remove it and look at what's left
+			IEnumerable<SDKDescriptor> FullSdks = Collection.FullSdks;
+
+			// if there's one with the special name, then use it as it's generic (common case here)
+			if (FullSdks.Count() == 1 && (FullSdks.First().Name == "Sdk" || FullSdks.First().Name == "Software"))
+			{
+				return FullSdks.First();
+			}
+
+			// finally we have to give up
+			return null;
+		}
+
 		#endregion
 
 		// AutoSDKs handling portion
@@ -475,6 +906,7 @@ namespace EpicGames.Core
 		protected const string SDKRootEnvVar = "UE_SDKS_ROOT";
 
 		protected const string AutoSetupEnvVar = "AutoSDKSetup";
+		protected const string ManualSetupEnvVar = "ManualSDKSetup";
 
 
 		private static string GetAutoSDKHostPlatform()
@@ -592,6 +1024,10 @@ namespace EpicGames.Core
 		{
 			return GetAutoSDKPlatformName() + AutoSetupEnvVar;
 		}
+		protected string GetPlatformManualSDKSetupEnvVar()
+		{
+			return GetAutoSDKPlatformName() + ManualSetupEnvVar;
+		}
 
 		/// <summary>
 		/// Gets currently installed version
@@ -1075,9 +1511,9 @@ namespace EpicGames.Core
 					// make sure we know that we've modified the local environment, invalidating manual installs for this run.
 					bLocalProcessSetupAutoSDK = true;
 
-					// tell any child processes what our manual version was before setting up autosdk
-					string ValueToWrite = CachedManualSDKVersion != null ? CachedManualSDKVersion : "__None";
-					Environment.SetEnvironmentVariable(GetPlatformAutoSDKSetupEnvVar(), ValueToWrite);
+					// tell any child processes what our manual versions were before setting up autosdk
+					string ValueToWrite = CachedManualSDKVersions.Count > 0 ? JsonSerializer.Serialize(CachedManualSDKVersions) : "__None";
+					Environment.SetEnvironmentVariable(GetPlatformManualSDKSetupEnvVar(), ValueToWrite);
 
 					return true;
 				}
@@ -1164,16 +1600,16 @@ namespace EpicGames.Core
 			return bLocalProcessSetupAutoSDK || HasParentProcessSetupAutoSDK(out _);
 		}
 
-		protected bool HasParentProcessSetupAutoSDK(out string? OutAutoSDKSetupValue)
+		protected bool HasParentProcessSetupAutoSDK([NotNullWhen(true)] out string? OutAutoSDKSetupValue)
 		{
-			bool bParentProcessSetupAutoSDK = false;
 			String AutoSDKSetupVarName = GetPlatformAutoSDKSetupEnvVar();
 			OutAutoSDKSetupValue = Environment.GetEnvironmentVariable(AutoSDKSetupVarName);
+			
 			if (!String.IsNullOrEmpty(OutAutoSDKSetupValue))
 			{
-				bParentProcessSetupAutoSDK = true;
+				return true;
 			}
-			return bParentProcessSetupAutoSDK;
+			return false;
 		}
 
 		public SDKStatus HasRequiredManualSDK()
@@ -1196,10 +1632,11 @@ namespace EpicGames.Core
 		// tells us if the user has a valid manual install.
 		protected virtual SDKStatus HasRequiredManualSDKInternal()
 		{
-			string? ManualSDKVersion;
-			GetInstalledVersions(out ManualSDKVersion, out _);
+			return GetAllSDKInfo().AreAllManualSDKsValid() ? SDKStatus.Valid : SDKStatus.Invalid;
+			//string? ManualSDKVersion;
+			//GetInstalledVersions(out ManualSDKVersion, out _);
 
-			return IsVersionValid(ManualSDKVersion, bForAutoSDK:false) ? SDKStatus.Valid : SDKStatus.Invalid;
+			//return IsVersionValid(ManualSDKVersion, bForAutoSDK:false) ? SDKStatus.Valid : SDKStatus.Invalid;
 		}
 
 		// some platforms will fail if there is a manual install that is the WRONG manual install.
diff --git a/Engine/Source/Programs/UnrealBuildTool/Platform/Android/AndroidPlatformSDK.cs b/Engine/Source/Programs/UnrealBuildTool/Platform/Android/AndroidPlatformSDK.cs
--- a/Engine/Source/Programs/UnrealBuildTool/Platform/Android/AndroidPlatformSDK.cs
+++ b/Engine/Source/Programs/UnrealBuildTool/Platform/Android/AndroidPlatformSDK.cs
@@ -20,13 +20,13 @@ namespace UnrealBuildTool
 			return "-24";
 		}
 
-		public override void GetValidVersionRange(out string MinVersion, out string MaxVersion)
+		protected override void GetValidVersionRange(out string MinVersion, out string MaxVersion)
 		{
 			MinVersion = "r21a";
 			MaxVersion = "r23a";
 		}
 
-		public override void GetValidSoftwareVersionRange(out string? MinVersion, out string? MaxVersion)
+		protected override void GetValidSoftwareVersionRange(out string? MinVersion, out string? MaxVersion)
 		{
 			MinVersion = MaxVersion = null;
 		}
@@ -45,7 +45,7 @@ namespace UnrealBuildTool
 		}
 
 
-		public override string? GetInstalledSDKVersion()
+		protected override string? GetInstalledSDKVersion()
 		{
 			string? NDKPath = Environment.GetEnvironmentVariable("NDKROOT");
 
@@ -97,7 +97,7 @@ namespace UnrealBuildTool
 		}
 
 
-		public override bool TryConvertVersionToInt(string? StringValue, out UInt64 OutValue)
+		public override bool TryConvertVersionToInt(string? StringValue, out UInt64 OutValue, string? Hint)
 		{
 			// convert r<num>[letter] to hex
 			if (!string.IsNullOrEmpty(StringValue))
diff --git a/Engine/Source/Programs/UnrealBuildTool/Platform/IOS/IOSPlatformSDK.cs b/Engine/Source/Programs/UnrealBuildTool/Platform/IOS/IOSPlatformSDK.cs
--- a/Engine/Source/Programs/UnrealBuildTool/Platform/IOS/IOSPlatformSDK.cs
+++ b/Engine/Source/Programs/UnrealBuildTool/Platform/IOS/IOSPlatformSDK.cs
@@ -12,7 +12,7 @@ namespace UnrealBuildTool
 {
 	internal class IOSPlatformSDK : ApplePlatformSDK
 	{
-		public override void GetValidSoftwareVersionRange(out string MinVersion, out string? MaxVersion)
+		protected override void GetValidSoftwareVersionRange(out string MinVersion, out string? MaxVersion)
 		{
 			// what is our min IOS version?
 			MinVersion = "12.0";
diff --git a/Engine/Source/Programs/UnrealBuildTool/Platform/Linux/LinuxPlatformSDK.cs b/Engine/Source/Programs/UnrealBuildTool/Platform/Linux/LinuxPlatformSDK.cs
--- a/Engine/Source/Programs/UnrealBuildTool/Platform/Linux/LinuxPlatformSDK.cs
+++ b/Engine/Source/Programs/UnrealBuildTool/Platform/Linux/LinuxPlatformSDK.cs
@@ -16,14 +16,14 @@ namespace UnrealBuildTool
 			return "v19_clang-11.0.1-centos7";
 		}
 
-		public override void GetValidVersionRange(out string MinVersion, out string MaxVersion)
+		protected override void GetValidVersionRange(out string MinVersion, out string MaxVersion)
 		{
 			// all that matters is the number after the v, according to TryConvertVersionToInt()
 			MinVersion = "v10_clang-5.0.0-centos7";
 			MaxVersion = "v19_clang-11.0.1-centos7";
 		}
 
-		public override void GetValidSoftwareVersionRange(out string? MinVersion, out string? MaxVersion)
+		protected override void GetValidSoftwareVersionRange(out string? MinVersion, out string? MaxVersion)
 		{
 			MinVersion = MaxVersion = null;
 		}
@@ -34,7 +34,7 @@ namespace UnrealBuildTool
 		}
 
 
-		public override string? GetInstalledSDKVersion()
+		protected override string? GetInstalledSDKVersion()
 		{
 			// @todo turnkey: ForceUseSystemCompiler() returns true, we should probably run system clang -V or similar to get version
 
@@ -60,7 +60,7 @@ namespace UnrealBuildTool
 		}
 
 
-		public override bool TryConvertVersionToInt(string? StringValue, out UInt64 OutValue)
+		public override bool TryConvertVersionToInt(string? StringValue, out UInt64 OutValue, string? Hint)
 		{
 			if (StringValue != null)
 			{
diff --git a/Engine/Source/Programs/UnrealBuildTool/Platform/Mac/ApplePlatformSDK.cs b/Engine/Source/Programs/UnrealBuildTool/Platform/Mac/ApplePlatformSDK.cs
--- a/Engine/Source/Programs/UnrealBuildTool/Platform/Mac/ApplePlatformSDK.cs
+++ b/Engine/Source/Programs/UnrealBuildTool/Platform/Mac/ApplePlatformSDK.cs
@@ -17,7 +17,7 @@ namespace UnrealBuildTool
 			return "11.5";
 		}
 
-		public override void GetValidVersionRange(out string MinVersion, out string MaxVersion)
+		protected override void GetValidVersionRange(out string MinVersion, out string MaxVersion)
 		{
 			if (RuntimePlatform.IsMac)
 			{
@@ -32,12 +32,12 @@ namespace UnrealBuildTool
 			}
 		}
 
-		public override void GetValidSoftwareVersionRange(out string? MinVersion, out string? MaxVersion)
+		protected override void GetValidSoftwareVersionRange(out string? MinVersion, out string? MaxVersion)
 		{
 			MinVersion = MaxVersion = null;
 		}
 
-		public override bool TryConvertVersionToInt(string? StringValue, out UInt64 OutValue)
+		public override bool TryConvertVersionToInt(string? StringValue, out UInt64 OutValue, string? Hint)
 		{
 			OutValue = 0;
 
@@ -65,7 +65,7 @@ namespace UnrealBuildTool
 			return false;
 		}
 
-		public override string? GetInstalledSDKVersion()
+		protected override string? GetInstalledSDKVersion()
 		{
 			return UnrealBuildBase.ApplePlatformSDK.InstalledSDKVersion;
 		}
diff --git a/Engine/Source/Programs/UnrealBuildTool/Platform/Windows/MicrosoftPlatformSDK.cs b/Engine/Source/Programs/UnrealBuildTool/Platform/Windows/MicrosoftPlatformSDK.cs
--- a/Engine/Source/Programs/UnrealBuildTool/Platform/Windows/MicrosoftPlatformSDK.cs
+++ b/Engine/Source/Programs/UnrealBuildTool/Platform/Windows/MicrosoftPlatformSDK.cs
@@ -75,13 +75,13 @@ namespace UnrealBuildTool
 			return PreferredWindowsSdkVersions.First().ToString();
 		}
 
-		public override void GetValidVersionRange(out string MinVersion, out string MaxVersion)
+		protected override void GetValidVersionRange(out string MinVersion, out string MaxVersion)
 		{
 			MinVersion = "10.0.00000.0";
 			MaxVersion = "10.9.99999.0";
 		}
 
-		public override void GetValidSoftwareVersionRange(out string? MinVersion, out string? MaxVersion)
+		protected override void GetValidSoftwareVersionRange(out string? MinVersion, out string? MaxVersion)
 		{
 			// minimum version is the oldest version in the Preferred list -
 			MinVersion = PreferredWindowsSdkVersions.Min()?.ToString();
@@ -89,7 +89,7 @@ namespace UnrealBuildTool
 		}
 
 		[SupportedOSPlatform("windows")]
-		public override string? GetInstalledSDKVersion()
+		protected override string? GetInstalledSDKVersion()
 		{
 			if (!RuntimePlatform.IsWindows)
 			{
@@ -107,7 +107,7 @@ namespace UnrealBuildTool
 			return null;
 		}
 
-		public override bool TryConvertVersionToInt(string? StringValue, out UInt64 OutValue)
+		public override bool TryConvertVersionToInt(string? StringValue, out UInt64 OutValue, string? Hint)
 		{
 			OutValue = 0;
 
@@ -422,12 +422,14 @@ namespace UnrealBuildTool
 		#region Private implementation
 
 		[SupportedOSPlatform("windows")]
-		static readonly KeyValuePair<RegistryKey, string>[] InstallDirRoots = {
-			new KeyValuePair<RegistryKey, string>(Registry.CurrentUser, "SOFTWARE\\"),
-			new KeyValuePair<RegistryKey, string>(Registry.LocalMachine, "SOFTWARE\\"),
-			new KeyValuePair<RegistryKey, string>(Registry.CurrentUser, "SOFTWARE\\Wow6432Node\\"),
-			new KeyValuePair<RegistryKey, string>(Registry.LocalMachine, "SOFTWARE\\Wow6432Node\\")
-		};
+		static readonly Lazy<KeyValuePair<RegistryKey, string>[]> InstallDirRoots = new Lazy<KeyValuePair<RegistryKey, string>[]>(() =>
+			new KeyValuePair<RegistryKey, string>[]
+			{
+				new KeyValuePair<RegistryKey, string>(Registry.CurrentUser, "SOFTWARE\\"),
+				new KeyValuePair<RegistryKey, string>(Registry.LocalMachine, "SOFTWARE\\"),
+				new KeyValuePair<RegistryKey, string>(Registry.CurrentUser, "SOFTWARE\\Wow6432Node\\"),
+				new KeyValuePair<RegistryKey, string>(Registry.LocalMachine, "SOFTWARE\\Wow6432Node\\")
+			});
 
 		/// <summary>
 		/// Reads an install directory for a 32-bit program from a registry key. This checks for per-user and machine wide settings, and under the Wow64 virtual keys (HKCU\SOFTWARE, HKLM\SOFTWARE, HKCU\SOFTWARE\Wow6432Node, HKLM\SOFTWARE\Wow6432Node).
@@ -439,7 +441,7 @@ namespace UnrealBuildTool
 		[SupportedOSPlatform("windows")]
 		public static bool TryReadInstallDirRegistryKey32(string KeySuffix, string ValueName, [NotNullWhen(true)] out DirectoryReference? InstallDir)
 		{
-			foreach (KeyValuePair<RegistryKey, string> InstallRoot in InstallDirRoots)
+			foreach (KeyValuePair<RegistryKey, string> InstallRoot in InstallDirRoots.Value)
 			{
 				using (RegistryKey? Key = InstallRoot.Key.OpenSubKey(InstallRoot.Value + KeySuffix))
 				{
@@ -463,7 +465,7 @@ namespace UnrealBuildTool
 		static string[] ReadInstallDirSubKeys32(string KeyName)
 		{
 			HashSet<string> AllSubKeys = new HashSet<string>(StringComparer.Ordinal);
-			foreach (KeyValuePair<RegistryKey, string> Root in InstallDirRoots)
+			foreach (KeyValuePair<RegistryKey, string> Root in InstallDirRoots.Value)
 			{
 				using (RegistryKey? Key = Root.Key.OpenSubKey(Root.Value + KeyName))
 				{
diff --git a/Engine/Source/Programs/UnrealBuildTool/Platform/Windows/UEBuildWindows.cs b/Engine/Source/Programs/UnrealBuildTool/Platform/Windows/UEBuildWindows.cs
--- a/Engine/Source/Programs/UnrealBuildTool/Platform/Windows/UEBuildWindows.cs
+++ b/Engine/Source/Programs/UnrealBuildTool/Platform/Windows/UEBuildWindows.cs
@@ -431,7 +431,7 @@ namespace UnrealBuildTool
 				}
 				catch(Exception) // Find function will throw if there is no visual studio installed! This can happen w/ clang builds
 				{
-					return string.Empty;
+					return null;
 				}
 			}
 		}
diff --git a/Engine/Source/Programs/UnrealBuildTool/ProjectFiles/Rider/RiderProjectFile.cs b/Engine/Source/Programs/UnrealBuildTool/ProjectFiles/Rider/RiderProjectFile.cs
--- a/Engine/Source/Programs/UnrealBuildTool/ProjectFiles/Rider/RiderProjectFile.cs
+++ b/Engine/Source/Programs/UnrealBuildTool/ProjectFiles/Rider/RiderProjectFile.cs
@@ -526,7 +526,7 @@ namespace UnrealBuildTool
 					throw new ArgumentException("Wrong Target.Architecture: {0}", Target.Architecture);
 				}
 
-				string PlatformSdkVersionString = UEBuildPlatformSDK.GetSDKForPlatform(BuildPlatform.GetPlatformName())!.GetInstalledSDKVersion()!;
+				string PlatformSdkVersionString = UEBuildPlatformSDK.GetSDKForPlatform(BuildPlatform.GetPlatformName())!.GetInstalledVersion()!;
 				var Version = GetLinuxToolchainVersionFromFullString(PlatformSdkVersionString);
 
 				string? InternalSdkPath = UEBuildPlatform.GetSDK(UnrealTargetPlatform.Linux)!.GetInternalSDKPath();
diff --git a/Engine/Source/Programs/UnrealBuildTool/System/VersionNumberRange.cs b/Engine/Source/Programs/UnrealBuildTool/System/VersionNumberRange.cs
--- a/Engine/Source/Programs/UnrealBuildTool/System/VersionNumberRange.cs
+++ b/Engine/Source/Programs/UnrealBuildTool/System/VersionNumberRange.cs
@@ -14,12 +14,12 @@ namespace UnrealBuildTool
 		/// <summary>
 		/// Minimum version number
 		/// </summary>
-		VersionNumber Min { get; }
+		public VersionNumber Min { get; }
 
 		/// <summary>
 		/// Maximum version number
 		/// </summary>
-		VersionNumber Max { get; }
+		public VersionNumber Max { get; }
 
 		/// <summary>
 		/// Constructor
