Updated EpicGames.UHT to be a net6 project

#rnx

diff --git a/Engine/Source/Programs/Shared/EpicGames.UHT/EpicGames.UHT.csproj b/Engine/Source/Programs/Shared/EpicGames.UHT/EpicGames.UHT.csproj
--- a/Engine/Source/Programs/Shared/EpicGames.UHT/EpicGames.UHT.csproj
+++ b/Engine/Source/Programs/Shared/EpicGames.UHT/EpicGames.UHT.csproj
@@ -2,7 +2,7 @@
   <Import Project="..\UnrealEngine.csproj.props" />
   
   <PropertyGroup>
-    <TargetFramework>netcoreapp3.1</TargetFramework>
+    <TargetFramework>net6.0</TargetFramework>
     <Configuration Condition=" '$(Configuration)' == '' ">Development</Configuration>
     <OutputType>Library</OutputType>
     <RootNamespace>EpicGames.UHT</RootNamespace>
diff --git a/Engine/Source/Programs/Shared/EpicGames.UHT/Exporters/Json/UhtJsonExporter.cs b/Engine/Source/Programs/Shared/EpicGames.UHT/Exporters/Json/UhtJsonExporter.cs
--- a/Engine/Source/Programs/Shared/EpicGames.UHT/Exporters/Json/UhtJsonExporter.cs
+++ b/Engine/Source/Programs/Shared/EpicGames.UHT/Exporters/Json/UhtJsonExporter.cs
@@ -6,6 +6,7 @@ using EpicGames.UHT.Types;
 using EpicGames.UHT.Utils;
 using System.Collections.Generic;
 using System.Text.Json;
+using System.Text.Json.Serialization;
 using System.Threading.Tasks;
 
 namespace EpicGames.UHT.Exporters.Json
@@ -39,7 +40,7 @@ namespace EpicGames.UHT.Exporters.Json
 					(IUhtExportFactory Factory) =>
 					{
 						string JsonPath = Factory.MakePath(Package, ".json");
-						JsonSerializerOptions Options = new JsonSerializerOptions { WriteIndented = true, IgnoreNullValues = true };
+						JsonSerializerOptions Options = new JsonSerializerOptions { WriteIndented = true, DefaultIgnoreCondition = JsonIgnoreCondition.WhenWritingNull };
 						Factory.CommitOutput(JsonPath, JsonSerializer.Serialize(Package, Options));
 					}));
 			}
