Updated EpicGames.Build to be a net6 project.

diff --git a/Engine/Source/Programs/Shared/EpicGames.Build/EpicGames.Build.csproj b/Engine/Source/Programs/Shared/EpicGames.Build/EpicGames.Build.csproj
--- a/Engine/Source/Programs/Shared/EpicGames.Build/EpicGames.Build.csproj
+++ b/Engine/Source/Programs/Shared/EpicGames.Build/EpicGames.Build.csproj
@@ -1,6 +1,6 @@
 ﻿<Project Sdk="Microsoft.NET.Sdk">
   <PropertyGroup>
-    <TargetFramework>netcoreapp3.1</TargetFramework>
+    <TargetFramework>net6.0</TargetFramework>
     <Configuration Condition=" '$(Configuration)' == '' ">Development</Configuration>
     <OutputType>Library</OutputType>
     <RootNamespace>UnrealBuildBase</RootNamespace>
diff --git a/Engine/Source/Programs/Shared/EpicGames.Build/Platform/Mac/ApplePlatformSDK.cs b/Engine/Source/Programs/Shared/EpicGames.Build/Platform/Mac/ApplePlatformSDK.cs
--- a/Engine/Source/Programs/Shared/EpicGames.Build/Platform/Mac/ApplePlatformSDK.cs
+++ b/Engine/Source/Programs/Shared/EpicGames.Build/Platform/Mac/ApplePlatformSDK.cs
@@ -3,10 +3,10 @@
 using System;
 using System.Diagnostics;
 using System.IO;
+using System.Runtime.Versioning;
 using System.Text;
 using System.Text.RegularExpressions;
 using EpicGames.Core;
-using Microsoft.VisualBasic.CompilerServices;
 using Microsoft.Win32;
 
 namespace UnrealBuildBase
@@ -18,14 +18,14 @@ namespace UnrealBuildBase
 		private static string? GetInstalledSDKVersion()
 		{
 			// get xcode version on Mac
-			if (RuntimePlatform.IsMac)
+			if (OperatingSystem.IsMacOS())
 			{
 				string Output = RunLocalProcessAndReturnStdOut("sh", "-c 'xcodebuild -version'");
 				Match Result = Regex.Match(Output, @"Xcode (\S*)");
 				return Result.Success ? Result.Groups[1].Value : null;
 			}
 
-			if (RuntimePlatform.IsWindows)
+			if (OperatingSystem.IsWindows())
 			{
 				// otherwise, get iTunes "Version"
 				string? DllPath =
@@ -63,6 +63,7 @@ namespace UnrealBuildBase
 			return null;
 		}
 
+		[SupportedOSPlatform("windows")]
 		private static string? FindWindowsStoreITunesDLL()
 		{
 			string? InstallPath = null;
@@ -126,7 +127,7 @@ namespace UnrealBuildBase
 			// for rules see https://docs.microsoft.com/en-us/cpp/cpp/main-function-command-line-args
 			Args = Args?.Replace('\'', '\"');
 
-			ProcessStartInfo StartInfo = new ProcessStartInfo(Command, Args);
+			ProcessStartInfo StartInfo = Args != null ? new ProcessStartInfo(Command, Args) : new ProcessStartInfo(Command);
 			StartInfo.UseShellExecute = false;
 			StartInfo.RedirectStandardInput = true;
 			StartInfo.RedirectStandardOutput = true;
@@ -136,30 +137,37 @@ namespace UnrealBuildBase
 
 			string FullOutput = "";
 			string ErrorOutput = "";
-			using (Process LocalProcess = Process.Start(StartInfo))
+			using (Process? LocalProcess = Process.Start(StartInfo))
 			{
-				StreamReader OutputReader = LocalProcess.StandardOutput;
-				// trim off any extraneous new lines, helpful for those one-line outputs
-				FullOutput = OutputReader.ReadToEnd().Trim();
-
-				StreamReader ErrorReader = LocalProcess.StandardError;
-				// trim off any extraneous new lines, helpful for those one-line outputs
-				ErrorOutput = ErrorReader.ReadToEnd().Trim();
-				if (LogOutput)
+				if (LocalProcess != null)
 				{
-					if(FullOutput.Length > 0)
+					StreamReader OutputReader = LocalProcess.StandardOutput;
+					// trim off any extraneous new lines, helpful for those one-line outputs
+					FullOutput = OutputReader.ReadToEnd().Trim();
+
+					StreamReader ErrorReader = LocalProcess.StandardError;
+					// trim off any extraneous new lines, helpful for those one-line outputs
+					ErrorOutput = ErrorReader.ReadToEnd().Trim();
+					if (LogOutput)
 					{
-						Log.TraceInformation(FullOutput);
-					}
+						if (FullOutput.Length > 0)
+						{
+							Log.TraceInformation(FullOutput);
+						}
 
-					if (ErrorOutput.Length > 0)
-					{
-						Log.TraceError(ErrorOutput);
+						if (ErrorOutput.Length > 0)
+						{
+							Log.TraceError(ErrorOutput);
+						}
 					}
-				}
 
-				LocalProcess.WaitForExit();
-				ExitCode = LocalProcess.ExitCode;
+					LocalProcess.WaitForExit();
+					ExitCode = LocalProcess.ExitCode;
+				}
+				else
+				{
+					ExitCode = -1; 
+				}
 			}
 
 			// trim off any extraneous new lines, helpful for those one-line outputs
diff --git a/Engine/Source/Programs/Shared/EpicGames.Build/System/AssemblyExtensions.cs b/Engine/Source/Programs/Shared/EpicGames.Build/System/AssemblyExtensions.cs
--- a/Engine/Source/Programs/Shared/EpicGames.Build/System/AssemblyExtensions.cs
+++ b/Engine/Source/Programs/Shared/EpicGames.Build/System/AssemblyExtensions.cs
@@ -17,7 +17,7 @@ namespace UnrealBuildBase
 			}
 			catch (ReflectionTypeLoadException e)
 			{
-				AllTypes = e.Types.Where(x => x != null).ToArray();
+				AllTypes = e.Types.Where(x => x != null).Cast<Type>().ToArray();
 			}
 			return AllTypes;
 		}
diff --git a/Engine/Source/Programs/Shared/EpicGames.Build/UEBuildPlatformSDK.cs b/Engine/Source/Programs/Shared/EpicGames.Build/UEBuildPlatformSDK.cs
--- a/Engine/Source/Programs/Shared/EpicGames.Build/UEBuildPlatformSDK.cs
+++ b/Engine/Source/Programs/Shared/EpicGames.Build/UEBuildPlatformSDK.cs
@@ -329,7 +329,7 @@ namespace EpicGames.Core
 				else
 				{
 					Console.WriteLine("Processing parent manual sdk: {0}", ParentManualSDKVersions);
-					CachedManualSDKVersions = JsonSerializer.Deserialize<Dictionary<string, string>>(ParentManualSDKVersions);
+					CachedManualSDKVersions = JsonSerializer.Deserialize<Dictionary<string, string>>(ParentManualSDKVersions) ?? new Dictionary<string, string>();
 				}
 			}
 			else
