Undo changes to EpicGames.UHT to make net6 assembly.

#rnx

diff --git a/Engine/Source/Programs/Shared/EpicGames.UHT/EpicGames.UHT.csproj b/Engine/Source/Programs/Shared/EpicGames.UHT/EpicGames.UHT.csproj
--- a/Engine/Source/Programs/Shared/EpicGames.UHT/EpicGames.UHT.csproj
+++ b/Engine/Source/Programs/Shared/EpicGames.UHT/EpicGames.UHT.csproj
@@ -2,7 +2,7 @@
   <Import Project="..\UnrealEngine.csproj.props" />
   
   <PropertyGroup>
-    <TargetFramework>net6.0</TargetFramework>
+    <TargetFramework>netcoreapp3.1</TargetFramework>
     <Configuration Condition=" '$(Configuration)' == '' ">Development</Configuration>
     <OutputType>Library</OutputType>
     <RootNamespace>EpicGames.UHT</RootNamespace>
diff --git a/Engine/Source/Programs/Shared/EpicGames.UHT/Exporters/Json/UhtJsonExporter.cs b/Engine/Source/Programs/Shared/EpicGames.UHT/Exporters/Json/UhtJsonExporter.cs
--- a/Engine/Source/Programs/Shared/EpicGames.UHT/Exporters/Json/UhtJsonExporter.cs
+++ b/Engine/Source/Programs/Shared/EpicGames.UHT/Exporters/Json/UhtJsonExporter.cs
@@ -6,7 +6,6 @@ using EpicGames.UHT.Types;
 using EpicGames.UHT.Utils;
 using System.Collections.Generic;
 using System.Text.Json;
-using System.Text.Json.Serialization;
 using System.Threading.Tasks;
 
 namespace EpicGames.UHT.Exporters.Json
@@ -40,7 +39,7 @@ namespace EpicGames.UHT.Exporters.Json
 					(IUhtExportFactory Factory) =>
 					{
 						string JsonPath = Factory.MakePath(Package, ".json");
-						JsonSerializerOptions Options = new JsonSerializerOptions { WriteIndented = true, DefaultIgnoreCondition = JsonIgnoreCondition.WhenWritingNull };
+						JsonSerializerOptions Options = new JsonSerializerOptions { WriteIndented = true, IgnoreNullValues = true };
 						Factory.CommitOutput(JsonPath, JsonSerializer.Serialize(Package, Options));
 					}));
 			}
