Enabled the reference assembly to be generated in the old location for all projects that include the UE C# props file.
Modified assembly cache to not include "ref" or "refint" directories which contain the reference dlls.

diff --git a/Engine/Source/Programs/Shared/EpicGames.Core/AssemblyUtils.cs b/Engine/Source/Programs/Shared/EpicGames.Core/AssemblyUtils.cs
--- a/Engine/Source/Programs/Shared/EpicGames.Core/AssemblyUtils.cs
+++ b/Engine/Source/Programs/Shared/EpicGames.Core/AssemblyUtils.cs
@@ -98,6 +98,13 @@ namespace EpicGames.Core
 
 		public static void AddFileToAssemblyCache(string assemblyPath)
 		{
+			// Ignore any reference assemblies
+			string? directory = Path.GetFileName(Path.GetDirectoryName(assemblyPath));
+			if (!string.IsNullOrEmpty(directory) && (directory == "ref" | directory == "refint"))
+			{
+				return;
+			}
+
 			string assemblyName = Path.GetFileNameWithoutExtension(assemblyPath);
 			DateTime assemblyLastWriteTime = File.GetLastWriteTimeUtc(assemblyPath);
 			if (s_assemblyLocationCache.ContainsKey(assemblyName))
