BuildGraph - Fix Switch statements always executing Default

#rnx

diff --git a/Engine/Source/Programs/Shared/EpicGames.BuildGraph/BgScriptReader.cs b/Engine/Source/Programs/Shared/EpicGames.BuildGraph/BgScriptReader.cs
--- a/Engine/Source/Programs/Shared/EpicGames.BuildGraph/BgScriptReader.cs
+++ b/Engine/Source/Programs/Shared/EpicGames.BuildGraph/BgScriptReader.cs
@@ -1871,6 +1871,7 @@ namespace EpicGames.BuildGraph
 				if (childElement.Name == "Default" || await EvaluateConditionAsync(childElement))
 				{
 					await readContentsAsync(childElement);
+					break;
 				}
 			}
 		}
