Horde: Thread-safe store for small objects backed by a memory-mapped file, optimized for read/write performance and data integrity.

Reads are typically lock free. Writes require a lock but otherwise complete quickly and in constant time. The store is transactional, and is designed not to lose data if the process is terminated between flushes. To accomplish this, blocks are not reused until the store is flushed via a call to Save().

Blobs allocated in the store may be split into multiple cells, so do not support access to a contiguous underlying memory region. The cells comprising a blob may be enumerated and accessed individually via a handle to the blob.
The store also supports storage of tree objects which can be addressed by hash. Internally, these objects are reference counted and maintain a list of references to other blobs separately to their payload.

If the store runs out of space, supplemental stores may be allocated in granularity of <see cref="LocalBlobStoreOptions.ExpandStepSize"/>. These stores are merged into the main one during the next Save() call.

diff --git a/Engine/Source/Programs/Shared/EpicGames.Horde.Tests/BlobStoreTests.cs b/Engine/Source/Programs/Shared/EpicGames.Horde.Tests/BlobStoreTests.cs
new file mode 100644
--- /dev/null
+++ b/Engine/Source/Programs/Shared/EpicGames.Horde.Tests/BlobStoreTests.cs
@@ -0,0 +1,421 @@
+// Copyright Epic Games, Inc. All Rights Reserved.
+
+using EpicGames.Core;
+using EpicGames.Horde.Storage;
+using Microsoft.VisualStudio.TestTools.UnitTesting;
+using System;
+using System.Buffers.Binary;
+using System.Collections.Generic;
+using System.IO;
+
+namespace EpicGames.Horde.Tests
+{
+	[TestClass]
+	public class BlobStoreTests
+	{
+		readonly FileReference _file;
+
+		public BlobStoreTests()
+		{
+			_file = new FileReference("test.dat");
+			FileReference.Delete(_file);
+		}
+
+		[TestMethod]
+		public void EmptyObject()
+		{
+			using LocalBlobStore store = LocalBlobStore.CreateNew(_file, new LocalBlobStoreOptions());
+			BlobHandle handle = store.AddRaw(Array.Empty<byte>());
+
+			Memory<byte> data = store.GetData(handle.FirstCell);
+			Assert.AreEqual(0, data.Length);
+		}
+
+		static byte[] MakeBlockWithIndex(int size, int index)
+		{
+			byte[] data = new byte[size];
+			BinaryPrimitives.WriteInt32LittleEndian(data, index);
+			return data;
+		}
+
+		[TestMethod]
+		public void LargeObjectPacking()
+		{
+			using LocalBlobStore store = LocalBlobStore.CreateNew(_file, new LocalBlobStoreOptions());
+
+			long baseNumCellBytes = store.NumAllocatedCellBytes;
+			long baseNumPageBytes = store.NumAllocatedPageBytes;
+
+			store.AddRaw(MakeBlockWithIndex(1000, 1));
+			Assert.AreEqual(1, store.NumItems);
+			Assert.AreEqual(baseNumCellBytes + 1024, store.NumAllocatedCellBytes);
+			Assert.AreEqual(baseNumPageBytes + 4096, store.NumAllocatedPageBytes);
+
+			store.AddRaw(MakeBlockWithIndex(1000, 2));
+			Assert.AreEqual(2, store.NumItems);
+			Assert.AreEqual(baseNumCellBytes + 2048, store.NumAllocatedCellBytes);
+			Assert.AreEqual(baseNumPageBytes + 4096, store.NumAllocatedPageBytes);
+
+			store.AddRaw(MakeBlockWithIndex(1000, 3));
+			Assert.AreEqual(3, store.NumItems);
+			Assert.AreEqual(baseNumCellBytes + 3072, store.NumAllocatedCellBytes);
+			Assert.AreEqual(baseNumPageBytes + 4096, store.NumAllocatedPageBytes);
+
+			store.AddRaw(MakeBlockWithIndex(1000, 4));
+			Assert.AreEqual(4, store.NumItems);
+			Assert.AreEqual(baseNumCellBytes + 4096, store.NumAllocatedCellBytes);
+			Assert.AreEqual(baseNumPageBytes + 4096, store.NumAllocatedPageBytes);
+
+			store.AddRaw(MakeBlockWithIndex(1000, 5));
+			Assert.AreEqual(5, store.NumItems);
+			Assert.AreEqual(baseNumCellBytes + 5120, store.NumAllocatedCellBytes);
+			Assert.AreEqual(baseNumPageBytes + 8192, store.NumAllocatedPageBytes);
+		}
+
+		[TestMethod]
+		public void SmallObjectPacking()
+		{
+			using LocalBlobStore store = LocalBlobStore.CreateNew(_file, new LocalBlobStoreOptions());
+
+			long baseNumCellBytes = store.NumAllocatedCellBytes;
+			long baseNumPageBytes = store.NumAllocatedPageBytes;
+
+			store.AddRaw(MakeBlockWithIndex(64, 1));
+			Assert.AreEqual(1, store.NumItems);
+			Assert.AreEqual(baseNumCellBytes + 64, store.NumAllocatedCellBytes);
+			Assert.AreEqual(baseNumPageBytes + 4096, store.NumAllocatedPageBytes);
+
+			store.AddRaw(MakeBlockWithIndex(64, 2));
+			Assert.AreEqual(2, store.NumItems);
+			Assert.AreEqual(baseNumCellBytes + 128, store.NumAllocatedCellBytes);
+			Assert.AreEqual(baseNumPageBytes + 4096, store.NumAllocatedPageBytes);
+
+			store.AddRaw(MakeBlockWithIndex(64, 3));
+			Assert.AreEqual(3, store.NumItems);
+			Assert.AreEqual(baseNumCellBytes + 192, store.NumAllocatedCellBytes);
+			Assert.AreEqual(baseNumPageBytes + 4096, store.NumAllocatedPageBytes);
+		}
+
+		[TestMethod]
+		public void FillItems()
+		{
+			using LocalBlobStore store = LocalBlobStore.CreateNew(_file, new LocalBlobStoreOptions());
+			for (int idx = 0; idx < 2048; idx++)
+			{
+				byte[] data = MakeBlockWithIndex(64, idx);
+				store.AddRaw(data);
+			}
+			Assert.AreEqual(2048, store.NumItems);
+		}
+
+		[TestMethod]
+		public void HugeObjects()
+		{
+			using LocalBlobStore store = LocalBlobStore.CreateNew(_file, new LocalBlobStoreOptions());
+
+			long baseNumCellBytes = store.NumAllocatedCellBytes;
+			long baseNumPageBytes = store.NumAllocatedPageBytes;
+
+			byte[] data = new byte[4096 + 4096 + 2000];
+
+			Random random = new Random(0);
+			random.NextBytes(data);
+
+			BlobHandle handle = store.AddRaw(data);
+			Assert.AreEqual(data.Length, store.GetLength(handle));
+			Assert.AreEqual(1, store.NumItems);
+			Assert.AreEqual(baseNumCellBytes + 4096 + 4096 + 2048, store.NumAllocatedCellBytes);
+			Assert.AreEqual(baseNumPageBytes + 4096 + 4096 + 4096, store.NumAllocatedPageBytes);
+
+			byte[] storedData = new byte[store.GetLength(handle)];
+			int copiedDataLength = store.ReadData(handle, storedData);
+
+			Assert.AreEqual(storedData.Length, copiedDataLength);
+			Assert.IsTrue(storedData.AsSpan().SequenceEqual(data));
+		}
+
+		[TestMethod]
+		public void Resize()
+		{
+			FileInfo file1 = new FileInfo(_file.FullName); // header
+			file1.Delete();
+
+			FileInfo file2 = new FileInfo($"{_file}.2.tmp"); // obj1
+			file2.Delete();
+
+			FileInfo file3 = new FileInfo($"{_file}.3.tmp"); // index
+			file3.Delete();
+
+			FileInfo file4 = new FileInfo($"{_file}.4.tmp");
+			file4.Delete();
+
+			using LocalBlobStore store = LocalBlobStore.CreateNew(_file, new LocalBlobStoreOptions { ReservedSize = 4096, ExpandStepSize = 4096 });
+
+			ObjectHandle obj1 = store.AddObject(IoHash.Zero, new List<ObjectHandle>(), new byte[4096]);
+
+			file1.Refresh(); // header
+			Assert.IsTrue(file1.Exists);
+			Assert.AreEqual(4096, file1.Length);
+
+			file2.Refresh(); // obj1
+			Assert.IsTrue(file2.Exists);
+			Assert.AreEqual(4096, file2.Length);
+
+			file3.Refresh(); // obj1 refs
+			Assert.IsTrue(file3.Exists);
+			Assert.AreEqual(4096, file3.Length);
+
+			file4.Refresh();
+			Assert.IsFalse(file4.Exists);
+
+			store.Save();
+
+			file1.Refresh();
+			Assert.AreEqual(4096 * 4, file1.Length);
+
+			file2.Refresh();
+			Assert.IsFalse(file2.Exists);
+
+			file3.Refresh();
+			Assert.IsFalse(file3.Exists);
+		}
+
+		[TestMethod]
+		public void RefCounts()
+		{
+			using LocalBlobStore store = LocalBlobStore.CreateNew(_file, new LocalBlobStoreOptions());
+
+			ObjectHandle obj1 = store.AddObject(IoHash.Zero, new List<ObjectHandle>(), new byte[] { 1, 2, 3 });
+			ObjectHandle obj2 = store.AddObject(IoHash.Zero, new List<ObjectHandle> { obj1 }, new byte[] { 4, 5, 6, 7 });
+
+			Assert.AreEqual(2, store.NumItems);
+			Assert.AreEqual(7, store.NumAllocatedItemBytes);
+			Assert.AreEqual(64 * 4, store.NumAllocatedCellBytes);
+			Assert.AreEqual(4096, store.NumAllocatedPageBytes);
+
+			store.AddRef(obj1);
+			store.ReleaseRef(obj2);
+			store.Save();
+
+			Assert.AreEqual(1, store.NumItems);
+			Assert.AreEqual(3, store.NumAllocatedItemBytes);
+			Assert.AreEqual(64 * 2, store.NumAllocatedCellBytes);
+			Assert.AreEqual(8192, store.NumAllocatedPageBytes); // index + data
+
+			store.ReleaseRef(obj1);
+			store.Save();
+
+			Assert.AreEqual(0, store.NumItems);
+			Assert.AreEqual(0, store.NumAllocatedItemBytes);
+			Assert.AreEqual(0, store.NumAllocatedCellBytes);
+			Assert.AreEqual(4096, store.NumAllocatedPageBytes); // index only
+		}
+
+		[TestMethod]
+		public void BasicObjectLookup()
+		{
+			using (LocalBlobStore store = LocalBlobStore.CreateNew(_file, new LocalBlobStoreOptions { ObjectHashSize = 128 }))
+			{
+				Random random = new Random(0);
+
+				byte[] data = new byte[random.Next(8192)];
+				random.NextBytes(data);
+
+				IoHash hash = IoHash.Compute(data);
+				store.AddObject(hash, new List<ObjectHandle>(), data);
+
+				ObjectHandle obj = store.FindObject(hash);
+				BlobHandle blob = store.GetObjectData(obj);
+
+				long length = store.GetLength(blob);
+				Assert.AreEqual(data.Length, length);
+			}
+		}
+
+		[TestMethod]
+		public void ObjectLookup()
+		{
+			Dictionary<IoHash, byte[]> objects = new Dictionary<IoHash, byte[]>();
+			using (LocalBlobStore store = LocalBlobStore.CreateNew(_file, new LocalBlobStoreOptions { ObjectHashSize = 128 }))
+			{
+				Random random = new Random(0);
+
+				for (int idx = 0; idx < 1024; idx++)
+				{
+					byte[] data = new byte[random.Next(8192)];
+					random.NextBytes(data);
+
+					IoHash hash = IoHash.Compute(data);
+					objects[hash] = data;
+
+					store.AddObject(hash, new List<ObjectHandle>(), data);
+				}
+
+				store.Save();
+			}
+
+			using (LocalBlobStore store = LocalBlobStore.OpenExisting(_file, new LocalBlobStoreOptions { ObjectHashSize = 128 }))
+			{
+				foreach ((IoHash hash, byte[] data) in objects)
+				{
+					ObjectHandle obj = store.FindObject(hash);
+					Assert.IsTrue(obj.IsValid());
+
+					BlobHandle blob = store.GetObjectData(obj);
+					Assert.AreEqual(data.Length, store.GetLength(blob));
+
+					byte[] copy = new byte[data.Length];
+					store.ReadData(blob, copy);
+
+					Assert.IsTrue(data.AsSpan().SequenceEqual(copy));
+				}
+			}
+		}
+
+			/*
+			[TestMethod]
+			public void HugeObjects2()
+			{
+				const int NumItems = 5;
+				for (int idx = 0; idx < NumItems; idx++)
+				{
+					byte[] data = MakeBlockWithIndex(2000 + (idx * 2048), idx);
+					_store.Add(data);
+				}
+
+				Assert.AreEqual(NumItems, _store.NumItems);
+
+				using (LruCache.View view = _store.LockView())
+				{
+					for (int idx = 0; idx < NumItems; idx++)
+					{
+						byte[] data = MakeBlockWithIndex(2000 + (idx * 2048), idx);
+						IoHash hash = IoHash.Compute(data);
+						ReadOnlyMemory<byte> otherData = view.Get(hash);
+						Assert.IsTrue(data.AsSpan().SequenceEqual(otherData.Span));
+					}
+				}
+			}
+
+			[TestMethod]
+			public async Task SmallObjectSerialization()
+			{
+				for (int idx = 0; idx < 1000; idx++)
+				{
+					byte[] data = MakeBlockWithIndex(4 + idx, idx);
+					_store.Add(data);
+				}
+
+				CheckSmallObjects();
+
+				Assert.AreEqual(1000, _store.NumItems);
+				long numBytes = _store.NumBytes;
+				long numBytesWithBlockSlack = _store.NumBytesWithBlockSlack;
+				long numBytesWithPageSlack = _store.NumBytesWithPageSlack;
+
+				await _store.SaveAsync();
+				_store.Dispose();
+
+				_store = await LruCache.OpenAsync(_indexFile, _dataFile);
+
+				Assert.AreEqual(1000, _store.NumItems);
+				Assert.AreEqual(numBytes, _store.NumBytes);
+				Assert.AreEqual(numBytesWithBlockSlack, _store.NumBytesWithBlockSlack);
+				Assert.AreEqual(numBytesWithPageSlack, _store.NumBytesWithPageSlack);
+
+				CheckSmallObjects();
+			}
+
+			void CheckSmallObjects()
+			{
+				using (LruCache.View view = _store.LockView())
+				{
+					for (int idx = 0; idx < 1000; idx++)
+					{
+						byte[] data = MakeBlockWithIndex(4 + idx, idx);
+						IoHash hash = IoHash.Compute(data);
+						ReadOnlyMemory<byte> otherData = view.Get(hash);
+						Assert.IsTrue(data.AsSpan().SequenceEqual(otherData.Span));
+					}
+				}
+			}
+
+			[TestMethod]
+			public async Task LargeObjectSerialization()
+			{
+				const int NumItems = 10;
+				for (int idx = 0; idx < NumItems; idx++)
+				{
+					byte[] data = MakeBlockWithIndex(2000 + (idx * 2048), idx);
+					_store.Add(data);
+				}
+
+				CheckLargeObjects(NumItems);
+
+				Assert.AreEqual(NumItems, _store.NumItems);
+				long numBytes = _store.NumBytes;
+				long numBytesWithBlockSlack = _store.NumBytesWithBlockSlack;
+				long numBytesWithPageSlack = _store.NumBytesWithPageSlack;
+
+				await _store.SaveAsync();
+				_store.Dispose();
+
+				_store = await LruCache.OpenAsync(_indexFile, _dataFile);
+
+				Assert.AreEqual(NumItems, _store.NumItems);
+				Assert.AreEqual(numBytes, _store.NumBytes);
+				Assert.AreEqual(numBytesWithBlockSlack, _store.NumBytesWithBlockSlack);
+				Assert.AreEqual(numBytesWithPageSlack, _store.NumBytesWithPageSlack);
+
+				CheckLargeObjects(NumItems);
+			}
+
+			void CheckLargeObjects(int numItems)
+			{
+				using (LruCache.View view = _store.LockView())
+				{
+					for (int idx = 0; idx < numItems; idx++)
+					{
+						byte[] data = MakeBlockWithIndex(2000 + (idx * 2048), idx);
+						IoHash hash = IoHash.Compute(data);
+						ReadOnlyMemory<byte> otherData = view.Get(hash);
+						Assert.IsTrue(data.AsSpan().SequenceEqual(otherData.Span));
+					}
+				}
+			}
+
+			[TestMethod]
+			public async Task ExpireItems()
+			{
+				for (int idx = 0; idx < 1000; idx++)
+				{
+					byte[] data = MakeBlockWithIndex(4 + idx, idx);
+					_store.Add(data);
+				}
+
+				long trimSize = _store.NumBytesWithBlockSlack;
+				_store.NextGeneration();
+
+				List<IoHash> testHashes = new List<IoHash>();
+				for (int idx = 1000; idx < 2000; idx++)
+				{
+					byte[] data = MakeBlockWithIndex(4 + idx, idx);
+					testHashes.Add(_store.Add(data));
+				}
+
+				long desiredSize = _store.NumBytesWithBlockSlack - trimSize;
+				await _store.TrimAsync(desiredSize);
+
+				Assert.AreEqual(desiredSize, _store.NumBytesWithBlockSlack);
+
+				using (LruCache.View view = _store.LockView())
+				{
+					for (int idx = 0; idx < 100; idx++)
+					{
+						ReadOnlyMemory<byte> memory = view.Get(testHashes[idx]);
+						Assert.IsFalse(memory.IsEmpty);
+					}
+				}
+			}*/
+		}
+	}
diff --git a/Engine/Source/Programs/Shared/EpicGames.Horde/EpicGames.Horde.csproj b/Engine/Source/Programs/Shared/EpicGames.Horde/EpicGames.Horde.csproj
--- a/Engine/Source/Programs/Shared/EpicGames.Horde/EpicGames.Horde.csproj
+++ b/Engine/Source/Programs/Shared/EpicGames.Horde/EpicGames.Horde.csproj
@@ -11,6 +11,7 @@
     <Nullable>enable</Nullable>
     <DebugType>pdbonly</DebugType>
     <GenerateDocumentationFile>True</GenerateDocumentationFile>
+    <AllowUnsafeBlocks>True</AllowUnsafeBlocks>
   </PropertyGroup>
   <PropertyGroup Condition=" '$(Configuration)|$(Platform)' == 'Debug|AnyCPU' ">
     <DocumentationFile></DocumentationFile>
diff --git a/Engine/Source/Programs/Shared/EpicGames.Horde/Storage/LocalBlobStore.cs b/Engine/Source/Programs/Shared/EpicGames.Horde/Storage/LocalBlobStore.cs
new file mode 100644
--- /dev/null
+++ b/Engine/Source/Programs/Shared/EpicGames.Horde/Storage/LocalBlobStore.cs
@@ -0,0 +1,1534 @@
+// Copyright Epic Games, Inc. All Rights Reserved.
+
+using System;
+using System.Buffers;
+using System.Buffers.Binary;
+using System.Collections;
+using System.Collections.Generic;
+using System.IO;
+using System.IO.MemoryMappedFiles;
+using System.Linq;
+using System.Numerics;
+using EpicGames.Core;
+
+namespace EpicGames.Horde.Storage
+{
+	/// <summary>
+	/// Handle to a chunk of data allocated within a compound file
+	/// </summary>
+	public readonly struct BlobHandle
+	{
+		/// <summary>
+		/// Number of bytes in a serialized handle
+		/// </summary>
+		public const int NumBytes = CellHandle.NumBytes;
+
+		/// <summary>
+		/// Invalid block handle
+		/// </summary>
+		public static BlobHandle Invalid => new BlobHandle(CellHandle.Invalid);
+
+		/// <summary>
+		/// Handle to the first piece in this block
+		/// </summary>
+		public readonly CellHandle FirstCell;
+
+		/// <summary>
+		/// Accessor for this handle as a long value
+		/// </summary>
+		public long Data => FirstCell.Data;
+
+		/// <summary>
+		/// Construct a block handle from a handle for the first piece
+		/// </summary>
+		internal BlobHandle(long data)
+		{
+			FirstCell = new CellHandle(data);
+		}
+
+		/// <summary>
+		/// Construct a block handle from a handle for the first piece
+		/// </summary>
+		internal BlobHandle(CellHandle firstCell)
+		{
+			FirstCell = firstCell;
+		}
+
+		/// <summary>
+		/// Construct a block handle from individual fields
+		/// </summary>
+		internal BlobHandle(int firstPageIdx, int tailBlockIdx, int tailSize, bool isBody)
+		{
+			FirstCell = new CellHandle(firstPageIdx, tailBlockIdx, tailSize, isBody);
+		}
+
+		/// <summary>
+		/// Test whether this handle is valid
+		/// </summary>
+		/// <returns>True if the current handle is valid</returns>
+		public bool IsValid() => FirstCell.IsValid();
+
+		/// <summary>
+		/// Deserialize a handle from a block of memory
+		/// </summary>
+		/// <param name="input">Memory to read from</param>
+		/// <returns>Handle to another piece</returns>
+		public static BlobHandle ReadFrom(ReadOnlySpan<byte> input) => new BlobHandle(CellHandle.ReadFrom(input));
+
+		/// <summary>
+		/// Serialize a handle to a block of memory
+		/// </summary>
+		/// <param name="output">Memory to write to</param>
+		public void CopyTo(Span<byte> output) => FirstCell.CopyTo(output);
+	}
+
+	/// <summary>
+	/// Handle to a cell of data allocated within a blob
+	/// </summary>
+	public readonly struct CellHandle
+	{
+		/// <summary>
+		/// Number of bytes in a serialized handle
+		/// </summary>
+		public const int NumBytes = sizeof(long);
+
+		/// <summary>
+		/// Invalid handle value
+		/// </summary>
+		public static CellHandle Invalid => new CellHandle(0);
+
+		const long BodyFlag = 1L << 62;
+		const long ValidFlag = 1L << 61;
+
+		const int PageIdxOffset = 0;
+		const int PageIdxLength = 32;
+
+		const int TailCellIdxOffset = PageIdxOffset + PageIdxLength;
+		const int TailCellIdxLength = 6;
+
+		const int TailSizeOffset = TailCellIdxOffset + TailCellIdxLength;
+		const int TailSizeLength = 13;
+
+		internal readonly long Data;
+
+		internal int PageIdx => (int)(Data >> PageIdxOffset);
+		internal int TailBlockIdx => (int)(Data >> TailCellIdxOffset) & ((1 << TailCellIdxLength) - 1);
+		internal int TailSize => (int)(Data >> TailSizeOffset) & ((1 << TailSizeLength) - 1);
+
+		/// <summary>
+		/// Tests whether the handle is valid
+		/// </summary>
+		public bool IsValid() => (Data & ValidFlag) != 0;
+
+		/// <summary>
+		/// True if this handle is to the last cell
+		/// </summary>
+		public bool IsTail() => (Data & BodyFlag) == 0;
+
+		/// <summary>
+		/// Constructor
+		/// </summary>
+		internal CellHandle(long data)
+		{
+			Data = data;
+		}
+
+		/// <summary>
+		/// Constructor
+		/// </summary>
+		internal CellHandle(int firstPageIdx, int tailBlockIdx, int tailSize, bool isBody)
+		{
+			Data = ((long)firstPageIdx << PageIdxOffset) | ((long)tailBlockIdx << TailCellIdxOffset) | ((long)tailSize << TailSizeOffset) | ValidFlag;
+			if (isBody)
+			{
+				Data |= BodyFlag;
+			}
+		}
+
+		/// <summary>
+		/// Deserialize a handle from a block of memory
+		/// </summary>
+		/// <param name="input">Memory to read from</param>
+		/// <returns>Handle to another piece</returns>
+		public static CellHandle ReadFrom(ReadOnlySpan<byte> input) => new CellHandle(BinaryPrimitives.ReadInt64LittleEndian(input));
+
+		/// <summary>
+		/// Serialize a handle to a block of memory
+		/// </summary>
+		/// <param name="output">Memory to write to</param>
+		public void CopyTo(Span<byte> output) => BinaryPrimitives.WriteInt64LittleEndian(output, Data);
+	}
+
+	/// <summary>
+	/// Identifies a position within a blob
+	/// </summary>
+	readonly struct BlobCursor
+	{
+		/// <summary>
+		/// The current cell within the blob
+		/// </summary>
+		public readonly CellHandle Cell;
+
+		/// <summary>
+		/// Offset within the cell
+		/// </summary>
+		public readonly int Offset;
+
+		/// <summary>
+		/// Constructor
+		/// </summary>
+		public BlobCursor(BlobHandle blob)
+		{
+			Cell = blob.FirstCell;
+			Offset = 0;
+		}
+
+		/// <summary>
+		/// Constructor
+		/// </summary>
+		public BlobCursor(CellHandle cell, int offset)
+		{
+			Cell = cell;
+			Offset = offset;
+		}
+
+		/// <summary>
+		/// Implicit conversion operator from a blob handle to a cursor
+		/// </summary>
+		public static implicit operator BlobCursor(BlobHandle blob) => new BlobCursor(blob);
+	}
+
+	/// <summary>
+	/// Handle to a reference-counted object
+	/// </summary>
+	readonly struct ObjectHandle
+	{
+		/// <summary>
+		/// Number of bytes in a serialized object handle
+		/// </summary>
+		public const int NumBytes = BlobHandle.NumBytes;
+
+		/// <summary>
+		/// The blob data
+		/// </summary>
+		public readonly BlobHandle Blob;
+
+		/// <summary>
+		/// Accessor for this handle as a long value
+		/// </summary>
+		public long Data => Blob.Data;
+
+		/// <summary>
+		/// An invalid handle value
+		/// </summary>
+		public static ObjectHandle Invalid => new ObjectHandle(BlobHandle.Invalid);
+
+		internal ObjectHandle(long data)
+		{
+			Blob = new BlobHandle(data);
+		}
+
+		internal ObjectHandle(BlobHandle block)
+		{
+			Blob = block;
+		}
+
+		/// <inheritdoc cref="BlobHandle.IsValid"/>
+		public bool IsValid() => Blob.IsValid();
+
+		/// <inheritdoc cref="BlobHandle.CopyTo(Span{Byte})"/>
+		public void CopyTo(Span<byte> output) => Blob.CopyTo(output);
+
+		/// <inheritdoc cref="BlobHandle.ReadFrom(ReadOnlySpan{Byte})"/>
+		public static ObjectHandle ReadFrom(ReadOnlySpan<byte> data) => new ObjectHandle(BlobHandle.ReadFrom(data));
+	}
+
+	/// <summary>
+	/// Options for the local blob store
+	/// </summary>
+	class LocalBlobStoreOptions
+	{
+		/// <summary>
+		/// Amount of space reserved for the backing file
+		/// </summary>
+		public long ReservedSize { get; set; } = 1024 * 1024;
+
+		/// <summary>
+		/// Amount to allocate for each expansion to the store. This will incur a cost on the next save, as data is merged back into the root file.
+		/// </summary>
+		public long ExpandStepSize { get; set; } = 1024 * 1024;
+
+		/// <summary>
+		/// Initial size of the object lookup hash table
+		/// </summary>
+		public int ObjectHashSize { get; set; } = 1024;
+	}
+
+	/// <summary>
+	/// A thread-safe store for small objects backed by a memory-mapped file, optimized for read/write performance and data integrity.
+	/// 
+	/// Reads are typically lock free. Writes require a lock but otherwise complete quickly and in constant time.
+	/// 
+	/// The store is transactional, and is designed not to lose data if the process is terminated between flushes. To accomplish this, blocks are not reused until the 
+	/// store is explicitly saved.
+	///
+	/// Blobs allocated in the store may be split into multiple cells, so do not support access to a contiguous underlying memory region. The cells comprising a 
+	/// blob may be enumerated and accessed individually via a handle to the blob.
+	/// 
+	/// The store also supports storage of tree objects which can be addressed by hash. Internally, these objects are reference counted and maintain a list of
+	/// references to other blobs separately to their payload.
+	/// </summary>
+	class LocalBlobStore : IDisposable
+	{
+		const int MinCellSizeLog2 = 6; // 64 bytes
+		const int MaxCellSizeLog2 = 12; // 4096 bytes
+
+		const int PageSizeLog2 = MaxCellSizeLog2;
+		const int PageSize = 1 << PageSizeLog2; // 4096 bytes
+
+		/// <summary>
+		/// Manager for memory mapped file pointers
+		/// </summary>
+		sealed unsafe class MappedFileMemoryManager : MemoryManager<byte>
+		{
+			private readonly byte* _pointer;
+			private readonly int _length;
+
+			public MappedFileMemoryManager(byte* pointer, int length)
+			{
+				_pointer = pointer;
+				_length = length;
+			}
+
+			/// <inheritdoc/>
+			public override Span<byte> GetSpan() => new Span<byte>(_pointer, _length);
+
+			/// <inheritdoc/>
+			public override MemoryHandle Pin(int elementIndex) => new MemoryHandle(_pointer + elementIndex);
+
+			/// <inheritdoc/>
+			public override void Unpin() { }
+
+			/// <inheritdoc/>
+			protected override void Dispose(bool disposing) { }
+		}
+
+		/// <summary>
+		/// A memory mapped file used to store allocated data
+		/// </summary>
+		sealed unsafe class MappedFile : IDisposable
+		{
+			const int PageCountPerRegionLog2 = 20;
+			const int PageCountPerRegion = 1 << PageCountPerRegionLog2;
+
+			public FileReference File { get; }
+			public int TotalPageCount { get; private set; }
+			public int UsedPageCount { get; set; }
+
+			MemoryMappedFile _dataFile = null!;
+			MemoryMappedViewAccessor _dataFileViewAccessor = null!;
+			MappedFileMemoryManager[] _regions = null!;
+
+			public MappedFile(FileReference file, FileMode mode, long maxSize)
+			{
+				File = file;
+				TotalPageCount = (int)(maxSize >> PageSizeLog2);
+				MapView(mode);
+			}
+
+			public void SetSize(int newPageCount)
+			{
+				UnmapView();
+				TotalPageCount = newPageCount;
+				MapView(FileMode.Open);
+			}
+
+			public void Dispose()
+			{
+				UnmapView();
+			}
+
+			private void MapView(FileMode mode)
+			{
+				long maxSize = (long)TotalPageCount << PageSizeLog2;
+
+				_dataFile = MemoryMappedFile.CreateFromFile(File.FullName, mode, null, maxSize, MemoryMappedFileAccess.ReadWrite);
+				_dataFileViewAccessor = _dataFile.CreateViewAccessor(0, (long)TotalPageCount << PageSizeLog2);
+
+				byte* data = null;
+				_dataFileViewAccessor.SafeMemoryMappedViewHandle.AcquirePointer(ref data);
+
+				_regions = new MappedFileMemoryManager[(TotalPageCount + PageCountPerRegion - 1) >> PageCountPerRegionLog2];
+				for (int idx = 0; idx < _regions.Length; idx++)
+				{
+					int pageIndex = idx << PageCountPerRegionLog2;
+					int pageCount = Math.Min(TotalPageCount - pageIndex, PageCountPerRegion);
+					_regions[idx] = new MappedFileMemoryManager(data + ((long)pageIndex << PageSizeLog2), pageCount << PageSizeLog2);
+				}
+			}
+
+			private void UnmapView()
+			{
+				_dataFileViewAccessor.SafeMemoryMappedViewHandle.ReleasePointer();
+				_dataFileViewAccessor.Dispose();
+				_dataFile.Dispose();
+				_regions = Array.Empty<MappedFileMemoryManager>();
+			}
+
+			public Memory<byte> GetPage(int pageIdx) => GetContiguousPages(pageIdx, 1);
+			public Memory<byte> GetContiguousPages(int pageIdx, int pageCount) => _regions[pageIdx >> PageCountPerRegionLog2].Memory.Slice((pageIdx & (PageCountPerRegion - 1)) << PageSizeLog2, pageCount << PageSizeLog2);
+
+			public void Flush() => _dataFileViewAccessor.Flush();
+		}
+
+		/// <summary>
+		/// A page of data allocated to store objects
+		/// </summary>
+		sealed class Page
+		{
+			public const int EncodedSize = 16;
+
+			public Memory<byte> Data;
+			public Span<byte> Span => Data.Span;
+
+			public int CellCount => 1 << (PageSizeLog2 - CellSizeLog2);
+			public int CellSize => 1 << CellSizeLog2;
+			public int CellSizeLog2;
+			public ulong FreeBitMap;
+			public ulong PendingFreeBitMap; // cells that will be released on next index save
+			public ulong EmptyBitMap => GetEmptyBitMap(CellSizeLog2);
+
+			public int NextPageIdx;
+
+			public Page(Memory<byte> data)
+			{
+				Data = data;
+				SetFree();
+			}
+
+			private void Reset(int cellSizeLog2, ulong freeBitMap, int nextPageIdx = -1)
+			{
+				CellSizeLog2 = cellSizeLog2;
+				FreeBitMap = freeBitMap;
+				PendingFreeBitMap = 0;
+				NextPageIdx = nextPageIdx;
+			}
+
+			public void SetFree() => Reset(PageSizeLog2, GetEmptyBitMap(PageSizeLog2));
+			public void SetAllocated(int cellSizeLog2) => Reset(cellSizeLog2, GetEmptyBitMap(cellSizeLog2));
+
+			public void Encode(Span<byte> span)
+			{
+				ulong data = ((ulong)CellSizeLog2 << 32) | (uint)NextPageIdx;
+				BinaryPrimitives.WriteUInt64LittleEndian(span, data);
+				BinaryPrimitives.WriteUInt64LittleEndian(span.Slice(8), FreeBitMap | PendingFreeBitMap);
+			}
+
+			public void Decode(ReadOnlySpan<byte> span)
+			{
+				ulong data = BinaryPrimitives.ReadUInt64LittleEndian(span);
+				ulong freeBitMap = BinaryPrimitives.ReadUInt64LittleEndian(span.Slice(8));
+				Reset((int)(uint)(data >> 32), freeBitMap, (int)(uint)data);
+			}
+
+			public static ulong GetEmptyBitMap(int itemSizeLog2)
+			{
+				// Note: left-shift only uses bottom 5 bits, cannot use (1UL << (PageSize >> BlockSizeLog2)) - 1
+				return ~0UL >> (64 - (1 << (PageSizeLog2 - itemSizeLog2)));
+			}
+
+			public override string ToString()
+			{
+				if (CellSizeLog2 == PageSizeLog2)
+				{
+					if (FreeBitMap == 1)
+					{
+						return "free";
+					}
+					else if (NextPageIdx == -1)
+					{
+						return "full";
+					}
+					else
+					{
+						return $"full -> {NextPageIdx}";
+					}
+				}
+				else
+				{
+					char[] chars = new char[CellCount];
+					for (int idx = 0; idx < CellCount; idx++)
+					{
+						ulong flag = 1UL << idx;
+						chars[idx] = (((FreeBitMap | PendingFreeBitMap) & flag) == 0) ? 'x' : '-';
+					}
+					return $"{CellSize}b: [{new string(chars)}]";
+				}
+			}
+		}
+
+		/// <summary>
+		/// Open-addressed hashtable used for looking up objects by hash. Uses Robin-Hood hashing.
+		/// </summary>
+		sealed class ObjectLookup : IEnumerable<KeyValuePair<IoHash, ObjectHandle>>
+		{
+			const ulong Multiplier = 11400714819323198485ul; // fibonnaci hashing
+
+			const float ShrinkLoad = 0.3f;
+			const float ExpandLoad = 0.7f;
+
+			const int MinCapacity = 128;
+
+			private IoHash[] _hashes = Array.Empty<IoHash>();
+			private ObjectHandle[] _items = Array.Empty<ObjectHandle>();
+			public int Count => _count;
+			public int Capacity
+			{
+				get => 1 << _capacityLog2;
+				set => Resize(value);
+			}
+
+			private int _count;
+			private int _capacityLog2;
+
+			public ObjectLookup(int capacity)
+			{
+				Capacity = capacity;
+			}
+
+			public void Clear(int capacity)
+			{
+				_capacityLog2 = BitOperations.Log2(Math.Max(MinCapacity, (uint)(capacity * 2 - 1)));
+				_count = 0;
+				_hashes = new IoHash[1 << _capacityLog2];
+				_items = new ObjectHandle[1 << _capacityLog2];
+			}
+
+			private void Resize(int capacity)
+			{
+				IoHash[] hashes = _hashes;
+				ObjectHandle[] items = _items;
+
+				Clear(capacity);
+
+				for (int slot = 0; slot < hashes.Length; slot++)
+				{
+					if (hashes[slot] != IoHash.Zero)
+					{
+						Add(hashes[slot], items[slot]);
+					}
+				}
+			}
+
+			int GetTargetSlot(IoHash hash) => (int)(((ulong)hash.GetHashCode() * Multiplier) >> (64 - _capacityLog2));
+
+			public void Add(IoHash hash, ObjectHandle item)
+			{
+				if (Count >= (int)(Capacity * ExpandLoad))
+				{
+					Resize(Capacity * 2);
+				}
+
+				_count++;
+				int targetSlot = GetTargetSlot(hash);
+				Add(hash, item, targetSlot, targetSlot);
+			}
+
+			private void Add(IoHash hash, ObjectHandle item, int slot, int targetSlot)
+			{
+				for (; ; )
+				{
+					// Check if this slot is empty
+					IoHash compareHash = _hashes[slot];
+					if (compareHash == IoHash.Zero)
+					{
+						break;
+					}
+
+					// If not, get the probe sequence length of the item currently in this slot
+					int compareTargetSlot = GetTargetSlot(compareHash);
+
+					// If this is a better fit, replace it. Update subsequent slots first to ensure consistent read ordering.
+					int nextSlot = (slot + 1) & (Capacity - 1);
+					if (Displace(slot, targetSlot, compareTargetSlot))
+					{
+						Add(compareHash, _items[slot], nextSlot, compareTargetSlot);
+						break;
+					}
+					slot = nextSlot;
+				}
+
+				_hashes[slot] = hash;
+				_items[slot] = item;
+			}
+
+			public void Remove(IoHash hash)
+			{
+				// Find the current location of this item
+				int slot = FindSlot(hash);
+				if (slot != -1)
+				{
+					// Move down any subsequent items over the top of it
+					for (; ; )
+					{
+						int nextSlot = GetNextSlot(slot);
+
+						IoHash compareHash = _hashes[nextSlot];
+						if (compareHash == IoHash.Zero)
+						{
+							break;
+						}
+
+						int targetSlot = GetTargetSlot(compareHash);
+						if (targetSlot == nextSlot)
+						{
+							break;
+						}
+
+						_hashes[slot] = _hashes[targetSlot];
+						_items[slot] = _items[targetSlot];
+
+						slot = nextSlot;
+					}
+
+					// Set the last slot to zero
+					_hashes[slot] = IoHash.Zero;
+					_items[slot] = ObjectHandle.Invalid;
+					_count--;
+
+					// Shrink the hash if necessary
+					if (Capacity > MinCapacity && Count < (int)(Capacity * ShrinkLoad))
+					{
+						Resize(Capacity / 2);
+					}
+				}
+			}
+
+			int GetNextSlot(int slot)
+			{
+				return (slot + 1) & (Capacity - 1);
+			}
+
+			int GetProbeLen(int slot, int targetSlot)
+			{
+				return (targetSlot + Capacity - slot) & (Capacity - 1);
+			}
+
+			public bool Displace(int slot, int targetSlot, int compareTargetSlot)
+			{
+				int probeLen = GetProbeLen(slot, targetSlot);
+				int compareProbeLen = GetProbeLen(slot, compareTargetSlot);
+				return probeLen > compareProbeLen;
+			}
+
+			public ObjectHandle Find(IoHash hash)
+			{
+				int slot = FindSlot(hash);
+				if (slot == -1)
+				{
+					return ObjectHandle.Invalid;
+				}
+				else
+				{
+					return _items[slot];
+				}
+			}
+
+			private int FindSlot(IoHash hash)
+			{
+				int targetSlot = GetTargetSlot(hash);
+				for (int slot = targetSlot; ; slot = (slot + 1) & (Capacity - 1))
+				{
+					IoHash compareHash = _hashes[slot];
+					if (compareHash == IoHash.Zero)
+					{
+						break;
+					}
+					if (compareHash == hash)
+					{
+						return slot;
+					}
+
+					int compareTargetSlot = GetTargetSlot(compareHash);
+					if (Displace(slot, targetSlot, compareTargetSlot))
+					{
+						break;
+					}
+				}
+				return -1;
+			}
+
+			/// <inheritdoc/>
+			public IEnumerator<KeyValuePair<IoHash, ObjectHandle>> GetEnumerator()
+			{
+				for (int slot = 0; slot < _hashes.Length; slot++)
+				{
+					if (_hashes[slot] != IoHash.Zero)
+					{
+						yield return KeyValuePair.Create(_hashes[slot], _items[slot]);
+					}
+				}
+			}
+
+			/// <inheritdoc/>
+			IEnumerator IEnumerable.GetEnumerator() => GetEnumerator();
+		}
+
+		const int HeaderPageCount = 1;
+		const int NumIndexEntriesPerPage = (PageSize - CellHandle.NumBytes) / Page.EncodedSize;
+
+		static readonly byte[] s_magic = new byte[] { (byte)'U', (byte)'E', (byte)'C', (byte)'F', 0, 0, 0, 1 };
+
+		readonly object _lockObject = new object();
+		readonly long _stepSize;
+		readonly List<MappedFile> _mappedFiles = new List<MappedFile>();
+		readonly int[] _firstPageWithFreeCell = new int[MaxCellSizeLog2 + 1];
+		readonly List<Page> _pages = new List<Page>();
+		readonly ObjectLookup _lookup;
+
+		/// <summary>
+		/// Number of user allocations
+		/// </summary>
+		public int NumItems { get; private set; }
+
+		/// <summary>
+		/// Sum of the size of all user allocations
+		/// </summary>
+		public long NumAllocatedItemBytes { get; private set; }
+
+		/// <summary>
+		/// Number of bytes allocated for storing items
+		/// </summary>
+		public long NumAllocatedCellBytes { get; private set; }
+
+		/// <summary>
+		/// Number of bytes allocated for pages
+		/// </summary>
+		public long NumAllocatedPageBytes { get; private set; }
+
+		/// <summary>
+		/// Constructor
+		/// </summary>
+		private LocalBlobStore(MappedFile mappedFile, LocalBlobStoreOptions options)
+		{
+			_stepSize = options.ExpandStepSize;
+
+			_mappedFiles.Add(mappedFile);
+			_mappedFiles[0].UsedPageCount = HeaderPageCount;
+
+			_lookup = new ObjectLookup(options.ObjectHashSize);
+
+			_firstPageWithFreeCell.AsSpan().Fill(-1);
+		}
+
+		/// <summary>
+		/// Opens an existing blob store, or create a new one if it does not already exist
+		/// </summary>
+		/// <param name="file">Backing file for the store</param>
+		/// <param name="options">Options for the store</param>
+		/// <returns>New store instance</returns>
+		public static LocalBlobStore CreateOrOpen(FileReference file, LocalBlobStoreOptions options)
+		{
+			if (FileReference.Exists(file))
+			{
+				return CreateNew(file, options);
+			}
+			else
+			{
+				return OpenExisting(file, options);
+			}
+		}
+
+		/// <summary>
+		/// Create a new blob store in the given location
+		/// </summary>
+		/// <param name="file">Backing file for the store</param>
+		/// <param name="options">Options for the store</param>
+		/// <returns>New store instance</returns>
+		public static LocalBlobStore CreateNew(FileReference file, LocalBlobStoreOptions options)
+		{
+			MappedFile mappedFile = new MappedFile(file, FileMode.CreateNew, options.ReservedSize);
+			LocalBlobStore store = new LocalBlobStore(mappedFile, options);
+			store.Save();
+			return store;
+		}
+
+		/// <summary>
+		/// Opens an existing blob store in the given location
+		/// </summary>
+		/// <param name="file">Backing file for the store</param>
+		/// <param name="options">Options for the store</param>
+		/// <returns>New store instance</returns>
+		public static LocalBlobStore OpenExisting(FileReference file, LocalBlobStoreOptions options)
+		{
+			long reservedSize = Math.Max(options.ReservedSize, file.ToFileInfo().Length);
+			MappedFile mappedFile = new MappedFile(file, FileMode.Open, reservedSize);
+
+			LocalBlobStore store = new LocalBlobStore(mappedFile, options);
+			try
+			{
+				store.Load();
+				return store;
+			}
+			catch
+			{
+				store.Dispose();
+				throw;
+			}
+		}
+
+		/// <inheritdoc/>
+		public void Dispose()
+		{
+			foreach (MappedFile mappedFile in _mappedFiles)
+			{
+				mappedFile.Dispose();
+			}
+		}
+
+		/// <summary>
+		/// Loads the state of the store from disk
+		/// </summary>
+		private void Load()
+		{
+			// Read the table of contents
+			ReadOnlySpan<byte> header = _mappedFiles[0].GetContiguousPages(0, HeaderPageCount).Span;
+
+			ReadOnlySpan<byte> headerEnd = header;
+			if (!headerEnd.StartsWith(s_magic))
+			{
+				throw new InvalidDataException("Missing magic bytes at start of store");
+			}
+			headerEnd = headerEnd.Slice(s_magic.Length);
+
+			int numInitializedPages = BinaryPrimitives.ReadInt32LittleEndian(headerEnd);
+			headerEnd = headerEnd.Slice(sizeof(int));
+
+			BlobHandle indexBlob = BlobHandle.ReadFrom(headerEnd);
+			headerEnd = headerEnd.Slice(BlobHandle.NumBytes);
+
+			NumItems = BinaryPrimitives.ReadInt32LittleEndian(headerEnd);
+			headerEnd = headerEnd.Slice(sizeof(int));
+
+			NumAllocatedItemBytes = BinaryPrimitives.ReadInt64LittleEndian(headerEnd);
+			headerEnd = headerEnd.Slice(sizeof(long));
+
+			int numObjects = BinaryPrimitives.ReadInt32LittleEndian(headerEnd);
+			headerEnd = headerEnd.Slice(sizeof(int));
+
+			BlobHandle lookupBlob = BlobHandle.ReadFrom(headerEnd);
+			headerEnd = headerEnd.Slice(BlobHandle.NumBytes);
+
+			// Allocate the pages list
+			_pages.Clear();
+			_pages.Capacity = numInitializedPages;
+			for (int idx = 0; idx < numInitializedPages; idx++)
+			{
+				Memory<byte> pageData = _mappedFiles[0].GetPage(HeaderPageCount + idx);
+				_pages.Add(new Page(pageData));
+			}
+
+			// Read the index pages
+			CellHandle indexCell = indexBlob.FirstCell;
+			for (int idx = 0; idx < numInitializedPages; )
+			{
+				ReadOnlySpan<byte> index = GetData(indexCell).Span;
+				ReadOnlySpan<byte> indexEnd = index;
+
+				indexCell = CellHandle.ReadFrom(indexEnd);
+				indexEnd = indexEnd.Slice(CellHandle.NumBytes);
+
+				int maxIdx = Math.Min(idx + NumIndexEntriesPerPage, numInitializedPages);
+				for (; idx < maxIdx; idx++)
+				{
+					_pages[idx].Decode(indexEnd);
+					indexEnd = indexEnd.Slice(Page.EncodedSize);
+				}
+			}
+
+			// Read the lookup pages
+			_lookup.Clear(numObjects);
+			if (numObjects > 0)
+			{
+				Span<byte> buffer = stackalloc byte[IoHash.NumBytes + ObjectHandle.NumBytes];
+
+				BlobCursor cursor = new BlobCursor(lookupBlob);
+				for (int idx = 0; idx < numObjects; idx++)
+				{
+					ReadData(ref cursor, buffer);
+					_lookup.Add(new IoHash(buffer), ObjectHandle.ReadFrom(buffer[IoHash.NumBytes..]));
+				}
+			}
+
+			// Release the index. We'll write a new one next time we save.
+			Free(indexBlob);
+
+			// Create the free page lists
+			BuildFreeLists();
+		}
+
+		/// <summary>
+		/// Saves the cache index
+		/// </summary>
+		public void Save()
+		{
+			lock (_lockObject)
+			{
+				// Serialize the object lookup
+				BlobHandle lookupHandle = BlobHandle.Invalid;
+				if (_lookup.Count > 0)
+				{
+					lookupHandle = AllocInternal(_lookup.Count * (IoHash.NumBytes + ObjectHandle.NumBytes));
+
+					BlobCursor cursor = lookupHandle;
+
+					Span<byte> entry = stackalloc byte[IoHash.NumBytes + ObjectHandle.NumBytes];
+					foreach ((IoHash hash, ObjectHandle handle) in _lookup)
+					{
+						hash.CopyTo(entry);
+						handle.CopyTo(entry[IoHash.NumBytes..]);
+						WriteData(ref cursor, entry);
+					}
+				}
+
+				// Allocate pages for the index. We need to be careful with the order of operations here, because allocating the index
+				// may cause the size of the page table - and thus the size of the index - to grow.
+				BlobHandle indexHandle = BlobHandle.Invalid;
+				for(int indexPageCount = 0; ;)
+				{
+					int newIndexPageCount = (_pages.Count + (NumIndexEntriesPerPage - 1)) / NumIndexEntriesPerPage;
+					if (newIndexPageCount <= indexPageCount)
+					{
+						break;
+					}
+					if (indexHandle.IsValid())
+					{
+						FreeInternal(indexHandle);
+					}
+					indexHandle = AllocInternal(newIndexPageCount << PageSizeLog2);
+					indexPageCount = newIndexPageCount;
+				}
+
+				// Compact all the mapped files into one.
+				Consolidate();
+
+				// Create the index
+				CellHandle chunkHandle = indexHandle.FirstCell;
+				for(int idx = 0; idx < _pages.Count; )
+				{
+					Span<byte> span = GetSpan(chunkHandle);
+					chunkHandle = GetNextCell(chunkHandle);
+
+					chunkHandle.CopyTo(span);
+					span = span.Slice(CellHandle.NumBytes);
+
+					int maxIdx = Math.Min(_pages.Count, idx + NumIndexEntriesPerPage);
+					for (; idx < maxIdx; idx++)
+					{
+						_pages[idx].Encode(span);
+						span = span.Slice(Page.EncodedSize);
+					}
+				}
+
+				// Write the header block
+				Span<byte> header = _mappedFiles[0].GetContiguousPages(0, HeaderPageCount).Span;
+
+				s_magic.CopyTo(header);
+				header = header.Slice(s_magic.Length);
+
+				BinaryPrimitives.WriteInt32LittleEndian(header, _pages.Count);
+				header = header.Slice(sizeof(int));
+
+				indexHandle.CopyTo(header);
+				header = header.Slice(BlobHandle.NumBytes);
+
+				BinaryPrimitives.WriteInt32LittleEndian(header, NumItems);
+				header = header.Slice(sizeof(int));
+
+				BinaryPrimitives.WriteInt64LittleEndian(header, NumAllocatedItemBytes);
+				header = header.Slice(sizeof(long));
+
+				BinaryPrimitives.WriteInt32LittleEndian(header, _lookup.Count);
+				header = header.Slice(sizeof(int));
+
+				lookupHandle.CopyTo(header);
+				header = header.Slice(BlobHandle.NumBytes);
+
+				// Write the data
+				_mappedFiles[0].Flush();
+
+				// Now free up the actual data
+				for (int idx = 0; idx < _pages.Count; idx++)
+				{
+					Page page = _pages[idx];
+					if (page.PendingFreeBitMap != 0)
+					{
+						NumAllocatedCellBytes -= BitOperations.PopCount(page.PendingFreeBitMap) << page.CellSizeLog2;
+
+						page.FreeBitMap |= page.PendingFreeBitMap;
+						page.PendingFreeBitMap = 0;
+
+						if (page.FreeBitMap == page.EmptyBitMap)
+						{
+							page.SetFree();
+						}
+					}
+				}
+
+				// Rebuild the free lists
+				BuildFreeLists();
+
+				// Free the index immediately
+				FreeInternal(indexHandle);
+			}
+		}
+
+		void Consolidate()
+		{
+			if (_mappedFiles.Count > 1)
+			{
+				int newPageCount = _mappedFiles.Sum(x => x.UsedPageCount);
+				_mappedFiles[0].SetSize(newPageCount);
+
+				for (int idx = 1; idx < _mappedFiles.Count; idx++)
+				{
+					for(int sourcePageIdx = 0; sourcePageIdx < _mappedFiles[idx].UsedPageCount; sourcePageIdx++)
+					{
+						int targetPageIdx = _mappedFiles[0].UsedPageCount++;
+						Page page = _pages[targetPageIdx - HeaderPageCount];
+
+						Memory<byte> source = page.Data;
+						page.Data = _mappedFiles[0].GetPage(targetPageIdx);
+						source.CopyTo(page.Data);
+					}
+
+					_mappedFiles[idx].Dispose();
+					FileReference.Delete(_mappedFiles[idx].File);
+				}
+
+				_mappedFiles.RemoveRange(1, _mappedFiles.Count - 1);
+			}
+		}
+
+		public void BuildFreeLists()
+		{
+			// Clear the current free lists
+			_firstPageWithFreeCell.AsSpan().Fill(-1);
+
+			// Add everything back into free lists
+			Page?[] lastPageWithFreeCell = new Page?[_firstPageWithFreeCell.Length];
+			for (int idx = HeaderPageCount; idx < _pages.Count; idx++)
+			{
+				Page page = _pages[idx];
+				if (page.FreeBitMap != 0)
+				{
+					LinkPage(ref _firstPageWithFreeCell[page.CellSizeLog2], lastPageWithFreeCell[page.CellSizeLog2], idx);
+					lastPageWithFreeCell[page.CellSizeLog2] = page;
+				}
+			}
+
+			// Compute the new allocation stats
+			NumAllocatedCellBytes = 0;
+			NumAllocatedPageBytes = 0;
+			foreach (Page page in _pages)
+			{
+				if (page.FreeBitMap != page.EmptyBitMap)
+				{
+					NumAllocatedPageBytes += PageSize;
+					NumAllocatedCellBytes += BitOperations.PopCount(page.FreeBitMap ^ page.EmptyBitMap) << page.CellSizeLog2;
+				}
+			}
+		}
+
+		private void LinkPage(ref int firstPageIdx, Page? lastPage, int newPageIdx)
+		{
+			if (lastPage == null)
+			{
+				firstPageIdx = newPageIdx;
+			}
+			else
+			{
+				lastPage.NextPageIdx = newPageIdx;
+			}
+		}
+
+		/// <summary>
+		/// Allocates data of the given length
+		/// </summary>
+		/// <param name="length">Length of the data to add</param>
+		public BlobHandle Alloc(long length)
+		{
+			lock (_lockObject)
+			{
+				NumItems++;
+				NumAllocatedItemBytes += length;
+				return AllocInternal(length);
+			}
+		}
+
+		private BlobHandle AllocInternal(long length)
+		{
+			int firstPageIdx = -1;
+			Page? lastPage = null;
+
+			// Allocate full pages for the data
+			long tailSize = length;
+			for (; tailSize > PageSize; tailSize -= PageSize)
+			{
+				int newPageIdx = AllocFullPage(PageSizeLog2);
+				LinkPage(ref firstPageIdx, lastPage, newPageIdx);
+				NumAllocatedCellBytes += PageSize;
+				lastPage = _pages[newPageIdx];
+				lastPage.FreeBitMap = 0;
+			}
+
+			// Get the size of the 'tail' page
+			int tailBlockSizeLog2 = GetBlockSizeLog2((int)tailSize);
+
+			// Get the tail page, or allocate it if necessary.
+			int tailPageIdx = _firstPageWithFreeCell[tailBlockSizeLog2];
+			if (tailPageIdx == -1)
+			{
+				tailPageIdx = AllocFullPage(tailBlockSizeLog2);
+				_firstPageWithFreeCell[tailBlockSizeLog2] = tailPageIdx;
+			}
+
+			// Add the page into the block list
+			LinkPage(ref firstPageIdx, lastPage, tailPageIdx);
+			NumAllocatedCellBytes += 1 << tailBlockSizeLog2;
+
+			// Take an item from the first free span
+			Page tailPage = _pages[tailPageIdx];
+			int tailBlockIdx = BitOperations.TrailingZeroCount(tailPage.FreeBitMap) & 63;
+			tailPage.FreeBitMap ^= 1UL << tailBlockIdx;
+
+			// If the page is full, remove it from the free list
+			if (tailPage.FreeBitMap == 0UL)
+			{
+				_firstPageWithFreeCell[tailBlockSizeLog2] = tailPage.NextPageIdx;
+				tailPage.NextPageIdx = -1;
+			}
+
+			// Create the handle
+			return new BlobHandle(firstPageIdx, tailBlockIdx, (int)tailSize, firstPageIdx != tailPageIdx);
+		}
+
+		/// <summary>
+		/// Adds a contiguous block of data to the store
+		/// </summary>
+		/// <param name="data">The data to add</param>
+		/// <returns>Handle to the first chunk</returns>
+		public BlobHandle AddRaw(ReadOnlySpan<byte> data)
+		{
+			BlobHandle handle = Alloc(data.Length);
+			WriteData(handle, data);
+			return handle;
+		}
+
+		/// <summary>
+		/// Free an allocated block
+		/// </summary>
+		/// <param name="handle">Handle of the block to free</param>
+		public void Free(BlobHandle handle)
+		{
+			lock (_lockObject)
+			{
+				long length = FreeInternal(handle);
+				NumItems--;
+				NumAllocatedItemBytes -= length;
+			}
+		}
+
+		private long FreeInternal(BlobHandle handle)
+		{
+			long length = 0;
+			for (CellHandle piece = handle.FirstCell; piece.IsValid(); piece = GetNextCell(piece))
+			{
+				Page page = _pages[piece.PageIdx];
+				if (piece.IsTail())
+				{
+					page.PendingFreeBitMap |= 1UL << piece.TailBlockIdx;
+					length += piece.TailSize;
+				}
+				else
+				{
+					page.PendingFreeBitMap |= 1;
+					length += PageSize;
+				}
+				NumAllocatedCellBytes -= page.CellSize;
+			}
+			return length;
+		}
+
+		/// <summary>
+		/// Copies a block to an output buffer
+		/// </summary>
+		/// <param name="handle">Handle to the block</param>
+		/// <param name="output">The output buffer to receive it</param>
+		/// <returns>Number of bytes that were copied</returns>
+		public int ReadData(BlobHandle handle, Span<byte> output)
+		{
+			int offset = 0;
+			for (CellHandle current = handle.FirstCell; current.IsValid(); current = GetNextCell(current))
+			{
+				Memory<byte> data = GetData(current);
+				data.Span.CopyTo(output.Slice(offset));
+				offset += data.Length;
+			}
+			return offset;
+		}
+
+		/// <summary>
+		/// Read data from a cursor position into a buffer
+		/// </summary>
+		/// <param name="cursor">Position to read from</param>
+		/// <param name="data">Receives the read data</param>
+		public void ReadData(ref BlobCursor cursor, Span<byte> data)
+		{
+			for(; ;)
+			{
+				ReadOnlySpan<byte> span = GetSpan(cursor.Cell);
+
+				int length = Math.Min(span.Length - cursor.Offset, data.Length);
+				span.Slice(cursor.Offset, length).CopyTo(data);
+
+				data = data[length..];
+				if(data.Length == 0)
+				{
+					cursor = new BlobCursor(cursor.Cell, cursor.Offset + length);
+					break;
+				}
+
+				CellHandle cell = GetNextCell(cursor.Cell);
+				cursor = new BlobCursor(cell, 0);
+			}
+		}
+
+		/// <summary>
+		/// Copies a block to an output buffer
+		/// </summary>
+		/// <param name="handle">Handle to the block</param>
+		/// <param name="input">Input data to copy from</param>
+		/// <returns>Number of bytes that were copied</returns>
+		public void WriteData(BlobHandle handle, ReadOnlySpan<byte> input)
+		{
+			CellHandle chunk = handle.FirstCell;
+			while (input.Length > 0)
+			{
+				Span<byte> blockData = GetSpan(chunk);
+				input.Slice(0, blockData.Length).CopyTo(blockData);
+				input = input.Slice(blockData.Length);
+				chunk = GetNextCell(chunk);
+			}
+		}
+
+		/// <summary>
+		/// Writes data to a cursor position within a blob
+		/// </summary>
+		/// <param name="cursor">The current blob cursor</param>
+		/// <param name="data">The data to write</param>
+		public void WriteData(ref BlobCursor cursor, ReadOnlySpan<byte> data)
+		{
+			for (; ; )
+			{
+				Span<byte> span = GetSpan(cursor.Cell);
+
+				int length = Math.Min(span.Length - cursor.Offset, data.Length);
+				data[0..length].CopyTo(span[cursor.Offset..]);
+
+				data = data[length..];
+				if (data.Length == 0)
+				{
+					cursor = new BlobCursor(cursor.Cell, cursor.Offset + length);
+					break;
+				}
+
+				CellHandle cell = GetNextCell(cursor.Cell);
+				cursor = new BlobCursor(cell, 0);
+			}
+		}
+
+		/// <summary>
+		/// Gets the data for a chunk
+		/// </summary>
+		/// <param name="cell"></param>
+		/// <returns>Writable block of data for the first part of the block</returns>
+		public Memory<byte> GetData(CellHandle cell)
+		{
+			if (!cell.IsValid())
+			{
+				throw new ArgumentException("Cell handle is not valid");
+			}
+
+			Page page = _pages[cell.PageIdx];
+			if (cell.IsTail())
+			{
+				return page.Data.Slice(cell.TailBlockIdx << page.CellSizeLog2, cell.TailSize);
+			}
+			else
+			{
+				return page.Data;
+			}
+		}
+
+		/// <summary>
+		/// Gets the span for a particular chunk
+		/// </summary>
+		/// <param name="chunk"></param>
+		/// <returns>Writable block of data for the first part of the block</returns>
+		public Span<byte> GetSpan(CellHandle chunk) => GetData(chunk).Span;
+
+		/// <summary>
+		/// Gets a handle to the next part of a particular block
+		/// </summary>
+		/// <param name="cell">Handle to the head of a block</param>
+		/// <returns>Handle to the next chunk in the block</returns>
+		public CellHandle GetNextCell(CellHandle cell)
+		{
+			if (cell.IsTail())
+			{
+				return CellHandle.Invalid;
+			}
+
+			int nextPageIdx = _pages[cell.PageIdx].NextPageIdx;
+			Page nextPage = _pages[nextPageIdx];
+
+			return new CellHandle(nextPageIdx, cell.TailBlockIdx, cell.TailSize, nextPage.NextPageIdx != -1);
+		}
+
+		/// <summary>
+		/// Gets the length of a block
+		/// </summary>
+		/// <param name="handle"></param>
+		/// <returns></returns>
+		public long GetLength(BlobHandle handle)
+		{
+			long length = 0;
+			for (CellHandle chunk = handle.FirstCell; chunk.IsValid(); chunk = GetNextCell(chunk))
+			{
+				length += GetData(chunk).Length;
+			}
+			return length;
+		}
+
+		/// <summary>
+		/// Allocate a new page of data from the underlying file. The first cell in the page will be marked as being used.
+		/// </summary>
+		int AllocFullPage(int cellSizeLog2)
+		{
+			int pageIdx = _firstPageWithFreeCell[MaxCellSizeLog2];
+			if (pageIdx == -1)
+			{
+				MappedFile lastMappedFile = _mappedFiles[^1];
+				if (lastMappedFile.UsedPageCount == lastMappedFile.TotalPageCount)
+				{
+					FileReference file = new FileReference($"{_mappedFiles[0].File}.{_mappedFiles.Count + 1}.tmp");
+					lastMappedFile = new MappedFile(file, FileMode.Create, _stepSize);
+					_mappedFiles.Add(lastMappedFile);
+				}
+
+				pageIdx = _pages.Count;
+				Memory<byte> data = lastMappedFile.GetPage(lastMappedFile.UsedPageCount++);
+				_pages.Add(new Page(data));
+			}
+
+			Page page = _pages[pageIdx];
+			_firstPageWithFreeCell[MaxCellSizeLog2] = page.NextPageIdx;
+			page.SetAllocated(cellSizeLog2);
+			NumAllocatedPageBytes += PageSize;
+
+			return pageIdx;
+		}
+
+		static int GetBlockSizeLog2(int itemSize)
+		{
+			int itemSizeLog2 = Math.Max(BitOperations.Log2((uint)itemSize), MinCellSizeLog2);
+			if (itemSize > (1 << itemSizeLog2))
+			{
+				itemSizeLog2++;
+			}
+			return itemSizeLog2;
+		}
+
+		#region Objects
+
+		/// <summary>
+		/// Adds an object to the store, and make it addressable by hash value.
+		/// </summary>
+		/// <param name="hash">Nominal hash of the data being added. The conversion of references/payload to the hash value is user defined, but must be unique.</param>
+		/// <param name="references">Other objects referenced by this object</param>
+		/// <param name="payload"></param>
+		public ObjectHandle AddObject(IoHash hash, List<ObjectHandle> references, ReadOnlySpan<byte> payload)
+		{
+			lock (_lockObject)
+			{
+				ObjectHandle obj = _lookup.Find(hash);
+				if (obj.IsValid())
+				{
+					AddRef(obj);
+				}
+				else
+				{
+					BlobHandle payloadBlob = AddRaw(payload);
+
+					int size = (sizeof(int) + IoHash.NumBytes + BlobHandle.NumBytes) + (references.Count * BlobHandle.NumBytes);
+					BlobHandle blob = AllocInternal(size);
+					CellHandle chunk = blob.FirstCell;
+
+					Span<byte> span = GetSpan(chunk);
+
+					BinaryPrimitives.WriteInt32LittleEndian(span, 1); // ref count
+					span = span.Slice(sizeof(int));
+
+					payloadBlob.CopyTo(span);
+					span = span.Slice(BlobHandle.NumBytes);
+
+					hash.CopyTo(span);
+					span = span.Slice(IoHash.NumBytes);
+
+					foreach (ObjectHandle reference in references)
+					{
+						if (span.Length == 0)
+						{
+							chunk = GetNextCell(chunk);
+							span = GetSpan(chunk);
+						}
+						reference.CopyTo(span);
+					}
+
+					obj = new ObjectHandle(blob);
+
+					if (hash != IoHash.Zero)
+					{
+						_lookup.Add(hash, obj);
+					}
+				}
+				return obj;
+			}
+		}
+
+		/// <summary>
+		/// Finds an object by hash value
+		/// </summary>
+		/// <param name="hash">Hash of the item to retreive</param>
+		/// <returns>Handle to the object with the given hash</returns>
+		public ObjectHandle FindObject(IoHash hash)
+		{
+			lock (_lockObject)
+			{
+				return _lookup.Find(hash);
+			}
+		}
+
+		/// <summary>
+		/// Gets the data associated with an object
+		/// </summary>
+		/// <param name="handle"></param>
+		/// <returns></returns>
+		public BlobHandle GetObjectData(ObjectHandle handle)
+		{
+			Span<byte> memory = GetSpan(handle.Blob.FirstCell);
+			return BlobHandle.ReadFrom(memory.Slice(sizeof(int), BlobHandle.NumBytes));
+		}
+
+		/// <summary>
+		/// Enumerate the references from an object
+		/// </summary>
+		/// <param name="handle"></param>
+		/// <returns></returns>
+		public IEnumerable<ObjectHandle> GetObjectRefs(ObjectHandle handle)
+		{
+			CellHandle chunk = handle.Blob.FirstCell;
+
+			Memory<byte> memory = GetData(chunk);
+			memory = memory.Slice(BlobHandle.NumBytes * 2);
+
+			for (; ; )
+			{
+				while (memory.Length > 0)
+				{
+					yield return ObjectHandle.ReadFrom(memory.Span);
+					memory = memory.Slice(ObjectHandle.NumBytes);
+				}
+
+				chunk = GetNextCell(chunk);
+				if (!chunk.IsValid())
+				{
+					break;
+				}
+
+				memory = GetData(chunk);
+			}
+		}
+
+		/// <summary>
+		/// Increment the reference count of an object
+		/// </summary>
+		/// <param name="obj">Object to increase the refcount of</param>
+		public void AddRef(ObjectHandle obj)
+		{
+			lock (_lockObject)
+			{
+				Span<byte> span = GetSpan(obj.Blob.FirstCell);
+				int count = BinaryPrimitives.ReadInt32LittleEndian(span);
+				BinaryPrimitives.WriteInt32LittleEndian(span, count + 1);
+			}
+		}
+
+		/// <summary>
+		/// Decrement the reference count of an object
+		/// </summary>
+		/// <param name="obj">Object to decrease the refcount of</param>
+		public void ReleaseRef(ObjectHandle obj)
+		{
+			lock (_lockObject)
+			{
+				CellHandle chunk = obj.Blob.FirstCell;
+				Span<byte> span = GetSpan(chunk);
+
+				int count = BinaryPrimitives.ReadInt32LittleEndian(span);
+				if (count > 1)
+				{
+					BinaryPrimitives.WriteInt32LittleEndian(span, count - 1);
+					return;
+				}
+				span = span.Slice(sizeof(int));
+
+				BlobHandle payloadBlock = BlobHandle.ReadFrom(span);
+				span = span.Slice(BlobHandle.NumBytes);
+				Free(payloadBlock);
+
+				IoHash hash = new IoHash(span);
+				if (hash != IoHash.Zero)
+				{
+					_lookup.Remove(hash);
+				}
+				span = span.Slice(IoHash.NumBytes);
+
+				for (; ; )
+				{
+					while (span.Length > 0)
+					{
+						ObjectHandle child = ObjectHandle.ReadFrom(span);
+						ReleaseRef(child);
+						span = span.Slice(BlobHandle.NumBytes);
+					}
+
+					chunk = GetNextCell(chunk);
+					if (chunk.IsValid())
+					{
+						span = GetSpan(chunk);
+					}
+					else
+					{
+						break;
+					}
+				}
+
+				FreeInternal(obj.Blob);
+			}
+		}
+
+		#endregion
+	}
+}
