UGS: Fix marshalling of UTF8 strings to native Perforce library.

diff --git a/Engine/Source/Programs/Shared/EpicGames.Perforce/NativePerforceConnection.cs b/Engine/Source/Programs/Shared/EpicGames.Perforce/NativePerforceConnection.cs
--- a/Engine/Source/Programs/Shared/EpicGames.Perforce/NativePerforceConnection.cs
+++ b/Engine/Source/Programs/Shared/EpicGames.Perforce/NativePerforceConnection.cs
@@ -74,7 +74,7 @@ namespace EpicGames.Perforce
 		static extern void Client_Login(IntPtr client, [MarshalAs(UnmanagedType.LPStr)] string password);
 
 		[DllImport(NativeDll, CallingConvention = CallingConvention.Cdecl, CharSet = CharSet.Ansi, BestFitMapping = false, ThrowOnUnmappableChar = true)]
-		static extern void Client_Command(IntPtr client, [MarshalAs(UnmanagedType.LPStr)] string command, int numArgs, string[] args, byte[]? inputData, int inputLength, bool interceptIo);
+		static extern void Client_Command(IntPtr client, [MarshalAs(UnmanagedType.LPStr)] string command, int numArgs, IntPtr[] args, byte[]? inputData, int inputLength, bool interceptIo);
 
 		[DllImport(NativeDll, CallingConvention = CallingConvention.Cdecl)]
 		static extern void Client_Destroy(IntPtr client);
@@ -337,7 +337,7 @@ namespace EpicGames.Perforce
 
 			nativeWriteBuffer._data = buffer.BasePtr;
 			nativeWriteBuffer._maxLength = buffer.Data.Length;
-			nativeWriteBuffer._maxCount = int.MaxValue;
+			nativeWriteBuffer._maxCount = Int32.MaxValue;
 		}
 
 		/// <summary>
@@ -499,10 +499,35 @@ namespace EpicGames.Perforce
 			}
 
 			Response response = new Response(this);
-			_requests.Add((() => Client_Command(_client, command, allArguments.Count, allArguments.ToArray(), specData, specData?.Length ?? 0, interceptIo), response));
+			_requests.Add((() => ExecCommand(command, allArguments, specData, interceptIo), response));
 			return Task.FromResult<IPerforceOutput>(response);
 		}
 
+		private void ExecCommand(string command, List<string> args, byte[]? specData, bool interceptIo)
+		{
+			List<IntPtr> nativeArgs = new List<IntPtr>();
+			try
+			{
+				foreach (string arg in args)
+				{
+					byte[] data = Encoding.UTF8.GetBytes(arg);
+					Array.Resize(ref data, data.Length + 1);
+					IntPtr nativeArg = Marshal.AllocHGlobal(data.Length);
+					Marshal.Copy(data, 0, nativeArg, data.Length);
+					nativeArgs.Add(nativeArg);
+				}
+
+				Client_Command(_client, command, nativeArgs.Count, nativeArgs.ToArray(), specData, specData?.Length ?? 0, interceptIo);
+			}
+			finally
+			{
+				for (int idx = 0; idx < nativeArgs.Count; idx++)
+				{
+					Marshal.FreeHGlobal(nativeArgs[idx]);
+				}
+			}
+		}
+
 		/// <inheritdoc/>
 		public Task<IPerforceOutput> LoginCommandAsync(string password, CancellationToken cancellationToken = default)
 		{
