Quick fix to two code correctness issues.
Fixed Getter/Setter check for enums.

#rnx

diff --git a/Engine/Source/Programs/Shared/EpicGames.UHT/Types/Properties/UhtEnumProperty.cs b/Engine/Source/Programs/Shared/EpicGames.UHT/Types/Properties/UhtEnumProperty.cs
--- a/Engine/Source/Programs/Shared/EpicGames.UHT/Types/Properties/UhtEnumProperty.cs
+++ b/Engine/Source/Programs/Shared/EpicGames.UHT/Types/Properties/UhtEnumProperty.cs
@@ -108,10 +108,10 @@ namespace EpicGames.UHT.Types
 			switch (TextType)
 			{
 				case UhtPropertyTextType.GenericFunctionArgOrRetVal:
+				case UhtPropertyTextType.GenericFunctionArgOrRetValImpl:
 				case UhtPropertyTextType.ClassFunctionArgOrRetVal:
 				case UhtPropertyTextType.InterfaceFunctionArgOrRetVal:
 				case UhtPropertyTextType.EventFunctionArgOrRetVal:
-				case UhtPropertyTextType.GenericFunctionArgOrRetValImpl:
 					if (Enum.CppForm == UhtEnumCppForm.EnumClass || (!bIsTemplateArgument && (Property.PropertyFlags.HasAnyFlags(EPropertyFlags.ReturnParm) || !Property.PropertyFlags.HasAnyFlags(EPropertyFlags.OutParm))))
 					{
 						Builder.Append(Enum.CppType);
@@ -137,6 +137,10 @@ namespace EpicGames.UHT.Types
 					}
 					break;
 
+				case UhtPropertyTextType.GetterSetterArg:
+					Builder.Append(Enum.CppType);
+					break;
+
 				default:
 					if (Enum.CppForm == UhtEnumCppForm.EnumClass)
 					{
diff --git a/Engine/Source/Programs/Shared/EpicGames.UHT/Types/Properties/UhtObjectPropertyBase.cs b/Engine/Source/Programs/Shared/EpicGames.UHT/Types/Properties/UhtObjectPropertyBase.cs
--- a/Engine/Source/Programs/Shared/EpicGames.UHT/Types/Properties/UhtObjectPropertyBase.cs
+++ b/Engine/Source/Programs/Shared/EpicGames.UHT/Types/Properties/UhtObjectPropertyBase.cs
@@ -250,9 +250,7 @@ namespace EpicGames.UHT.Types
 			}
 
 			string Type = TokenReader.GetStringView(TypeStartPos, TokenReader.InputPos - TypeStartPos).ToString();
-			Type.Replace("\n", " ");
-			Type.Replace("\r", "");
-			Type.Replace("\t", " ");
+			Type = Type.Replace("\n", " ").Replace("\r", "").Replace("\t", " ");
 
 			switch (Behavior)
 			{
diff --git a/Engine/Source/Programs/Shared/EpicGames.UHT/Utils/UhtStringBuilderExtensions.cs b/Engine/Source/Programs/Shared/EpicGames.UHT/Utils/UhtStringBuilderExtensions.cs
--- a/Engine/Source/Programs/Shared/EpicGames.UHT/Utils/UhtStringBuilderExtensions.cs
+++ b/Engine/Source/Programs/Shared/EpicGames.UHT/Utils/UhtStringBuilderExtensions.cs
@@ -417,7 +417,7 @@ namespace EpicGames.UHT.Utils
 							}
 							break;
 					}
-				} while (Index < Length);
+				}
 			}
 
 			Builder.Append("\"");
