Fixed missed copyright statements

#rnx

diff --git a/Engine/Source/Programs/Shared/EpicGames.Core/CsProjBuildHook.cs b/Engine/Source/Programs/Shared/EpicGames.Core/CsProjBuildHook.cs
--- a/Engine/Source/Programs/Shared/EpicGames.Core/CsProjBuildHook.cs
+++ b/Engine/Source/Programs/Shared/EpicGames.Core/CsProjBuildHook.cs
@@ -1,4 +1,6 @@
-﻿using System;
+﻿// Copyright Epic Games, Inc. All Rights Reserved.
+
+using System;
 using System.Collections.Generic;
 
 namespace EpicGames.Core
diff --git a/Engine/Source/Programs/Shared/EpicGames.MsBuild/CsProjBuilder.cs b/Engine/Source/Programs/Shared/EpicGames.MsBuild/CsProjBuilder.cs
--- a/Engine/Source/Programs/Shared/EpicGames.MsBuild/CsProjBuilder.cs
+++ b/Engine/Source/Programs/Shared/EpicGames.MsBuild/CsProjBuilder.cs
@@ -1,3 +1,5 @@
+// Copyright Epic Games, Inc. All Rights Reserved.
+
 using EpicGames.Core;
 using Microsoft.Build.Evaluation;
 using Microsoft.Build.Execution;
