Adjustments to the script module compiler code in preperation for use by UBT.

#rn trivial
#rnx

diff --git a/Engine/Source/Programs/AutomationTool/Program.cs b/Engine/Source/Programs/AutomationTool/Program.cs
--- a/Engine/Source/Programs/AutomationTool/Program.cs
+++ b/Engine/Source/Programs/AutomationTool/Program.cs
@@ -419,13 +419,23 @@ namespace AutomationToolDriver
 				: AutomationToolCommandLine.CommandsToExecute;
 			bool bBuildSuccess;
 			HashSet<FileReference> ScriptModuleAssemblyPaths = CompileScriptModule.InitializeScriptModules(
-					ScriptsForProject, AdditionalScriptDirs, bForceCompile, bNoCompile, bUseBuildRecords, out bBuildSuccess);
+					Rules.RulesFileType.AutomationModule, ScriptsForProject, AdditionalScriptDirs, bForceCompile | true, bNoCompile, bUseBuildRecords, 
+					out bBuildSuccess, (int Count) =>
+                    {
+						Log.TraceInformation($"Building {Count} projects (see Log 'Engine/Programs/AutomationTool/Saved/Logs/Log.txt' for more details)");
+					});
 
 			if (!bBuildSuccess)
             {
 				return ExitCode.Error_Unknown;
             }
 
+			// when the engine is installed, or UAT is invoked with -NoCompile, we expect to find at least one script module (AutomationUtils is a necessity)
+			if (ScriptModuleAssemblyPaths.Count == 0)
+			{
+				throw new Exception("Found no script module records.");
+			}
+
 			// Load AutomationUtils.Automation.dll
 			FileReference AssemblyPath = ScriptModuleAssemblyPaths.FirstOrDefault(x => x.GetFileNameWithoutExtension().Contains("AutomationUtils.Automation"));
 			Assembly AutomationUtilsAssembly = AssemblyPath != null ? Assembly.LoadFrom(AssemblyPath.FullName) : null;
diff --git a/Engine/Source/Programs/Shared/EpicGames.Build/System/CompileScriptModules.cs b/Engine/Source/Programs/Shared/EpicGames.Build/System/CompileScriptModules.cs
--- a/Engine/Source/Programs/Shared/EpicGames.Build/System/CompileScriptModules.cs
+++ b/Engine/Source/Programs/Shared/EpicGames.Build/System/CompileScriptModules.cs
@@ -34,7 +34,7 @@ namespace UnrealBuildBase
 			}
 
 			public void ValidateRecursively(
-				Dictionary<FileReference, CsProjBuildRecord> ValidBuildRecords,
+				Dictionary<FileReference, (CsProjBuildRecord, FileReference)> ValidBuildRecords,
 				Dictionary<FileReference, (CsProjBuildRecord, FileReference)> InvalidBuildRecords,
 				Dictionary<FileReference, (CsProjBuildRecord, FileReference)> BuildRecords,
 				FileReference ProjectPath)
@@ -54,19 +54,42 @@ namespace UnrealBuildBase
 			FileReference CsProjBuildHook.DotnetPath => Unreal.DotnetPath;
 		}
 
+		/// <summary>
+		/// Return the target paths from the collection of build records
+		/// </summary>
+		/// <param name="BuildRecords">Input build records</param>
+		/// <returns>Set of target files</returns>
+		public static HashSet<FileReference> GetTargetPaths(IReadOnlyDictionary<FileReference, (CsProjBuildRecord BuildRecord, FileReference BuildRecordPath)> BuildRecords)
+		{
+			return new HashSet<FileReference>(BuildRecords.Select(x => GetTargetPath(x)));
+		}
+
+		/// <summary>
+		/// Return the target path for the given build record
+		/// </summary>
+		/// <param name="BuildRecord">Build record</param>
+		/// <returns>File reference for the target</returns>
+		public static FileReference GetTargetPath(KeyValuePair<FileReference, (CsProjBuildRecord BuildRecord, FileReference BuildRecordPath)> BuildRecord)
+		{
+			return FileReference.Combine(BuildRecord.Key.Directory, BuildRecord.Value.BuildRecord.TargetPath!);
+		}
+
 		/// <summary>
 		/// Locates script modules, builds them if necessary, returns set of .dll files
 		/// </summary>
+		/// <param name="RulesFileType"></param>
 		/// <param name="ScriptsForProjectFileName"></param>
 		/// <param name="AdditionalScriptsFolders"></param>
 		/// <param name="bForceCompile"></param>
+		/// <param name="bNoCompile"></param>
 		/// <param name="bUseBuildRecords"></param>
 		/// <param name="bBuildSuccess"></param>
-		/// <returns></returns>
-		public static HashSet<FileReference> InitializeScriptModules(string ScriptsForProjectFileName, List<string> AdditionalScriptsFolders, bool bForceCompile, bool bNoCompile, bool bUseBuildRecords, out bool bBuildSuccess)
+		/// <param name="OnBuildingProjects">Action to invoke when projects get built</param>
+		/// <returns>Collection of all the projects.  They will have been compiled.</returns>
+		public static HashSet<FileReference> InitializeScriptModules(Rules.RulesFileType RulesFileType, 
+			string? ScriptsForProjectFileName, List<string>? AdditionalScriptsFolders, bool bForceCompile, bool bNoCompile, bool bUseBuildRecords, 
+			out bool bBuildSuccess, Action<int> OnBuildingProjects)
 		{
-			CsProjBuildHook Hook = new Hook();
-
 			List<DirectoryReference> GameDirectories = GetGameDirectories(ScriptsForProjectFileName);
 			List<DirectoryReference> AdditionalDirectories = GetAdditionalDirectories(AdditionalScriptsFolders);
 			List<DirectoryReference> GameBuildDirectories = GetAdditionalBuildDirectories(GameDirectories);
@@ -77,66 +100,93 @@ namespace UnrealBuildBase
 			BaseDirectories.AddRange(GameDirectories);
 			BaseDirectories.AddRange(AdditionalDirectories);
 
-			HashSet<FileReference> FoundAutomationProjects = new HashSet<FileReference>(
-				Rules.FindAllRulesSourceFiles(Rules.RulesFileType.AutomationModule,
-				// Project automation scripts require source engine builds
+			HashSet<FileReference> FoundProjects = new HashSet<FileReference>(
+				Rules.FindAllRulesSourceFiles(RulesFileType,
+				// Project scripts require source engine builds
 				GameFolders: Unreal.IsEngineInstalled() ? GameDirectories : new List<DirectoryReference>(), 
 				ForeignPlugins: null, AdditionalSearchPaths: AdditionalDirectories.Concat(GameBuildDirectories).ToList()));
 
+			return GetTargetPaths(Build(RulesFileType, FoundProjects, BaseDirectories, bForceCompile, bNoCompile, bUseBuildRecords, out bBuildSuccess, OnBuildingProjects));
+		}
+
+		/// <summary>
+		/// Locates script modules, builds them if necessary, returns set of .dll files
+		/// </summary>
+		/// <param name="RulesFileType"></param>
+		/// <param name="FoundProjects">Projects to be compiled</param>
+		/// <param name="BaseDirectories">Base directories for all the projects</param>
+		/// <param name="bForceCompile"></param>
+		/// <param name="bNoCompile"></param>
+		/// <param name="bUseBuildRecords"></param>
+		/// <param name="bBuildSuccess"></param>
+		/// <param name="OnBuildingProjects">Action to invoke when projects get built</param>
+		/// <returns>Collection of all the projects.  They will have been compiled.</returns>
+		public static Dictionary<FileReference, (CsProjBuildRecord, FileReference)> Build(Rules.RulesFileType RulesFileType,
+			HashSet<FileReference> FoundProjects, List<DirectoryReference> BaseDirectories, bool bForceCompile, bool bNoCompile, bool bUseBuildRecords,
+			out bool bBuildSuccess, Action<int> OnBuildingProjects)
+		{
+			CsProjBuildHook Hook = new Hook();
+
 			bool bUseBuildRecordsOnlyForProjectDiscovery = bNoCompile || Unreal.IsEngineInstalled();
 
 			// Load existing build records, validating them only if (re)compiling script projects is an option
-			Dictionary<FileReference, (CsProjBuildRecord, FileReference)> ExistingBuildRecords = LoadExistingBuildRecords(BaseDirectories);
+			Dictionary<FileReference, (CsProjBuildRecord BuildRecord, FileReference BuildRecordPath)> ExistingBuildRecords = LoadExistingBuildRecords(BaseDirectories);
 
-			Dictionary<FileReference, CsProjBuildRecord> ValidBuildRecords = new Dictionary<FileReference, CsProjBuildRecord>(ExistingBuildRecords.Count);
+			Dictionary<FileReference, (CsProjBuildRecord BuildRecord, FileReference BuildRecordPath)> ValidBuildRecords = new Dictionary<FileReference, (CsProjBuildRecord, FileReference)>(ExistingBuildRecords.Count);
 			Dictionary<FileReference, (CsProjBuildRecord BuildRecord, FileReference BuildRecordPath)> InvalidBuildRecords = new Dictionary<FileReference, (CsProjBuildRecord, FileReference)>(ExistingBuildRecords.Count);
 
 			if (bUseBuildRecords)
 			{
-				foreach (FileReference AutomationProject in FoundAutomationProjects)
+				foreach (FileReference Project in FoundProjects)
 				{
-					ValidateBuildRecordRecursively(ValidBuildRecords, InvalidBuildRecords, ExistingBuildRecords, AutomationProject, Hook);
+					ValidateBuildRecordRecursively(ValidBuildRecords, InvalidBuildRecords, ExistingBuildRecords, Project, Hook);
 				}
 			}
 
 			if (bUseBuildRecordsOnlyForProjectDiscovery)
-            {
-				// when the engine is installed, or UAT is invoked with -NoCompile, we expect to find at least one script module (AutomationUtils is a necessity)
-				if (ExistingBuildRecords.Count == 0)
-                {
-					throw new Exception("Found no script module records.");
-                }
-
-				HashSet<FileReference> BuiltTargets = new HashSet<FileReference>(ExistingBuildRecords.Count);
-				foreach ((CsProjBuildRecord BuildRecord, FileReference BuildRecordPath) in ExistingBuildRecords.Values.Where(x => x.Item2.HasExtension(".Automation.json")))
-                {
-					FileReference ProjectPath = FileReference.Combine(BuildRecordPath.Directory, BuildRecord.ProjectPath!);
-					FileReference TargetPath = FileReference.Combine(ProjectPath.Directory, BuildRecord.TargetPath!);
+			{
+				string FilterExtension = String.Empty;
+				switch (RulesFileType)
+				{
+					case Rules.RulesFileType.AutomationModule:
+						FilterExtension = ".Automation.json";
+						break;
+					case Rules.RulesFileType.UbtPlugin:
+						FilterExtension = ".ubtplugin.json";
+						break;
+					default:
+						throw new Exception("Unsupported rules file type");
+				}
+				Dictionary<FileReference, (CsProjBuildRecord, FileReference)> OutRecords = new Dictionary<FileReference, (CsProjBuildRecord, FileReference)>(ExistingBuildRecords.Count);
+				foreach (KeyValuePair<FileReference, (CsProjBuildRecord BuildRecord, FileReference BuildRecordPath)> Record in
+					ExistingBuildRecords.Where(x => x.Value.BuildRecordPath.HasExtension(FilterExtension)))
+				{
+					FileReference TargetPath = FileReference.Combine(Record.Key.Directory, Record.Value.BuildRecord.TargetPath!);
 
 					if (FileReference.Exists(TargetPath))
-                    {
-						BuiltTargets.Add(TargetPath);
-                    }
-                    else
-                    {
+					{
+						OutRecords.Add(Record.Key, Record.Value);
+					}
+					else
+					{
 						if (bNoCompile)
 						{
 							// when -NoCompile is on the command line, try to run with whatever is available
-							Log.TraceWarning($"Script module \"{TargetPath}\" not found for record \"{BuildRecordPath}\"");
+							Log.TraceWarning($"Script module \"{TargetPath}\" not found for record \"{Record.Value.BuildRecordPath}\"");
 						}
 						else
 						{
 							// when the engine is installed, expect to find a built target assembly for every record that was found
-							throw new Exception($"Script module \"{TargetPath}\" not found for record \"{BuildRecordPath}\"");
+							throw new Exception($"Script module \"{TargetPath}\" not found for record \"{Record.Value.BuildRecordPath}\"");
 						}
-						
-                    }
-                }
+
+					}
+				}
 				bBuildSuccess = true;
-				return BuiltTargets;
-            }
+				return OutRecords;
+			}
 			else
-            {
+			{
 				// when the engine is not installed, delete any build record .json file that is not valid
 				foreach ((CsProjBuildRecord _, FileReference BuildRecordPath) in InvalidBuildRecords.Values)
 				{
@@ -146,20 +196,20 @@ namespace UnrealBuildBase
 						FileReference.Delete(BuildRecordPath);
 					}
 				}
-            }
+			}
 
 			if (!bForceCompile && bUseBuildRecords)
 			{
 				// If all found records are valid, we can return their targets directly
-				if (FoundAutomationProjects.All(x => ValidBuildRecords.ContainsKey(x)))
+				if (FoundProjects.All(x => ValidBuildRecords.ContainsKey(x)))
 				{
 					bBuildSuccess = true;
-					return new HashSet<FileReference>(FoundAutomationProjects.Select(x => FileReference.Combine(x.Directory, ValidBuildRecords[x].TargetPath!)));
+					return new Dictionary<FileReference, (CsProjBuildRecord, FileReference)>(ValidBuildRecords.Where(x => FoundProjects.Contains(x.Key)));
 				}
 			}
 
 			// Fall back to the slower approach: use msbuild to load csproj files & build as necessary
-			return Build(FoundAutomationProjects, bForceCompile || !bUseBuildRecords, out bBuildSuccess, Hook, BaseDirectories);
+			return Build(FoundProjects, bForceCompile || !bUseBuildRecords, out bBuildSuccess, Hook, BaseDirectories, OnBuildingProjects);
 		}
 
 		/// <summary>
@@ -167,10 +217,11 @@ namespace UnrealBuildBase
 		/// If it is placed in the caller directly, then when the caller is invoked, the assembly will be loaded resulting
 		/// in the possible Microsoft.Build.Framework load issue later on is this method isn't invoked.
 		/// </summary>
-		static HashSet<FileReference> Build(HashSet<FileReference> FoundAutomationProjects,
-			bool bForceCompile, out bool bBuildSuccess, CsProjBuildHook Hook, List<DirectoryReference> BaseDirectories)
+		static Dictionary<FileReference, (CsProjBuildRecord, FileReference)> Build(HashSet<FileReference> FoundProjects,
+			bool bForceCompile, out bool bBuildSuccess, CsProjBuildHook Hook, List<DirectoryReference> BaseDirectories,
+			Action<int> OnBuildingProjects)
 		{
-			return CsProjBuilder.Build(FoundAutomationProjects, bForceCompile, out bBuildSuccess, Hook, BaseDirectories);
+			return CsProjBuilder.Build(FoundProjects, bForceCompile, out bBuildSuccess, Hook, BaseDirectories, OnBuildingProjects);
 		}
 
 		/// <summary>
@@ -342,7 +393,7 @@ namespace UnrealBuildBase
 		}
 
 		static void ValidateBuildRecordRecursively(
-			Dictionary<FileReference, CsProjBuildRecord> ValidBuildRecords,
+			Dictionary<FileReference, (CsProjBuildRecord BuildRecord, FileReference)> ValidBuildRecords,
 			Dictionary<FileReference, (CsProjBuildRecord, FileReference)> InvalidBuildRecords,
 			Dictionary<FileReference, (CsProjBuildRecord, FileReference)> BuildRecords,
 			FileReference ProjectPath, CsProjBuildHook Hook)
@@ -363,8 +414,8 @@ namespace UnrealBuildBase
 			// Is this particular build record valid?
 			if (!ValidateBuildRecord(Entry.BuildRecord, ProjectPath.Directory, out string ValidationFailureMessage, Hook))
 			{
-				string AutomationProjectRelativePath = Path.GetRelativePath(Unreal.EngineDirectory.FullName, ProjectPath.FullName);
-				Log.TraceLog($"[{AutomationProjectRelativePath}] {ValidationFailureMessage}");
+				string ProjectRelativePath = Path.GetRelativePath(Unreal.EngineDirectory.FullName, ProjectPath.FullName);
+				Log.TraceLog($"[{ProjectRelativePath}] {ValidationFailureMessage}");
 
 				InvalidBuildRecords.Add(ProjectPath, Entry);
 				return;
@@ -378,32 +429,32 @@ namespace UnrealBuildBase
 
 				if (!ValidBuildRecords.ContainsKey(FullProjectPath))
 				{
-					string AutomationProjectRelativePath = Path.GetRelativePath(Unreal.EngineDirectory.FullName, ProjectPath.FullName);
+					string ProjectRelativePath = Path.GetRelativePath(Unreal.EngineDirectory.FullName, ProjectPath.FullName);
 					string DependencyRelativePath = Path.GetRelativePath(Unreal.EngineDirectory.FullName, FullProjectPath.FullName);
-					Log.TraceLog($"[{AutomationProjectRelativePath}] Existing output is not valid because dependency {DependencyRelativePath} is not valid");
+					Log.TraceLog($"[{ProjectRelativePath}] Existing output is not valid because dependency {DependencyRelativePath} is not valid");
 					InvalidBuildRecords.Add(ProjectPath, Entry);
 					return;
 				}
 
 				// Ensure that the dependency was not built more recently than the project
-				if (Entry.BuildRecord.TargetBuildTime < ValidBuildRecords[FullProjectPath].TargetBuildTime)
+				if (Entry.BuildRecord.TargetBuildTime < ValidBuildRecords[FullProjectPath].BuildRecord.TargetBuildTime)
 				{
-					string AutomationProjectRelativePath = Path.GetRelativePath(Unreal.EngineDirectory.FullName, ProjectPath.FullName);
+					string ProjectRelativePath = Path.GetRelativePath(Unreal.EngineDirectory.FullName, ProjectPath.FullName);
 					string DependencyRelativePath = Path.GetRelativePath(Unreal.EngineDirectory.FullName, FullProjectPath.FullName);
-					Log.TraceLog($"[{AutomationProjectRelativePath}] Existing output is not valid because dependency {DependencyRelativePath} is newer");
+					Log.TraceLog($"[{ProjectRelativePath}] Existing output is not valid because dependency {DependencyRelativePath} is newer");
 					InvalidBuildRecords.Add(ProjectPath, Entry);
 					return;
 				}
 			}
 
-			ValidBuildRecords.Add(ProjectPath, Entry.BuildRecord);
+			ValidBuildRecords.Add(ProjectPath, Entry);
 		}
 
-		static List<DirectoryReference> GetGameDirectories(string ScriptsForProjectFileName)
+		static List<DirectoryReference> GetGameDirectories(string? ScriptsForProjectFileName)
         {
 			List<DirectoryReference> GameDirectories = new List<DirectoryReference>();
 
-			if (ScriptsForProjectFileName == null)
+			if (String.IsNullOrEmpty(ScriptsForProjectFileName))
 			{
 				GameDirectories = NativeProjectsBase.EnumerateProjectFiles().Select(x => x.Directory).ToList();
 			}
@@ -416,7 +467,7 @@ namespace UnrealBuildBase
 			return GameDirectories;
         }
 
-		static List<DirectoryReference> GetAdditionalDirectories(List<string> AdditionalScriptsFolders) =>
+		static List<DirectoryReference> GetAdditionalDirectories(List<string>? AdditionalScriptsFolders) =>
 			AdditionalScriptsFolders == null ? new List<DirectoryReference>() :
 				AdditionalScriptsFolders.Select(x => DirectoryReference.FindCorrectCase(new DirectoryReference(x))).ToList();
 
diff --git a/Engine/Source/Programs/Shared/EpicGames.Build/System/Rules.cs b/Engine/Source/Programs/Shared/EpicGames.Build/System/Rules.cs
--- a/Engine/Source/Programs/Shared/EpicGames.Build/System/Rules.cs
+++ b/Engine/Source/Programs/Shared/EpicGames.Build/System/Rules.cs
@@ -28,7 +28,12 @@ namespace UnrealBuildBase
 			/// <summary>
 			/// .automation.csproj files
 			/// </summary>
-			AutomationModule
+			AutomationModule,
+
+			/// <summary>
+			/// .ubtplugin.csproj files
+			/// </summary>
+			UbtPlugin,
 		}
 
 		/// <summary>
@@ -39,6 +44,7 @@ namespace UnrealBuildBase
 			public List<FileReference> ModuleRules = new List<FileReference>();
 			public List<FileReference> TargetRules = new List<FileReference>();
 			public List<FileReference> AutomationModules = new List<FileReference>();
+			public List<FileReference> UbtPlugins = new List<FileReference>();
 		}
 
 		/// Map of root folders to a cached list of all UBT-related source files in that folder or any of its sub-folders.
@@ -198,7 +204,7 @@ namespace UnrealBuildBase
 		/// <param name="Directories">Directories to search</param>
 		/// <param name="Type">Type of rules to return</param>
 		/// <returns>List of rules files of the given type</returns>
-		private static List<FileReference> FindAllRulesFiles(IEnumerable<DirectoryReference> Directories, RulesFileType Type)
+		public static List<FileReference> FindAllRulesFiles(IEnumerable<DirectoryReference> Directories, RulesFileType Type)
 		{
 			List<(DirectoryReference, RulesFileCache)> Caches = new List<(DirectoryReference, RulesFileCache)>(Directories.Count());
 			using (ThreadPoolWorkQueue Queue = new ThreadPoolWorkQueue())
@@ -224,6 +230,7 @@ namespace UnrealBuildBase
 					Cache.ModuleRules.Sort((A, B) => A.FullName.CompareTo(B.FullName));
 					Cache.TargetRules.Sort((A, B) => A.FullName.CompareTo(B.FullName));
 					Cache.AutomationModules.Sort((A, B) => A.FullName.CompareTo(B.FullName));
+					Cache.UbtPlugins.Sort((A, B) => A.FullName.CompareTo(B.FullName));
 					RootFolderToRulesFileCache[Directory] = Cache;
 				}
 
@@ -240,6 +247,10 @@ namespace UnrealBuildBase
 				{
 					Files.AddRange(Cache.AutomationModules);
 				}
+				else if (Type == RulesFileType.UbtPlugin)
+				{
+					Files.AddRange(Cache.UbtPlugins);
+				}
 				else
 				{
 					throw new Exception($"Unhandled rules type: {Type}");
@@ -284,6 +295,14 @@ namespace UnrealBuildBase
 					}
 					bSearchSubFolders = false;
 				}
+				else if (File.HasExtension(".ubtplugin.csproj"))
+				{
+					lock(Cache.UbtPlugins)
+					{
+						Cache.UbtPlugins.Add(File.Location);
+					}
+					bSearchSubFolders = false;
+				}
 			}
 
 			// If we didn't find anything to stop the search, search all the subdirectories too
diff --git a/Engine/Source/Programs/Shared/EpicGames.Core/CsProjBuildHook.cs b/Engine/Source/Programs/Shared/EpicGames.Core/CsProjBuildHook.cs
--- a/Engine/Source/Programs/Shared/EpicGames.Core/CsProjBuildHook.cs
+++ b/Engine/Source/Programs/Shared/EpicGames.Core/CsProjBuildHook.cs
@@ -42,7 +42,7 @@ namespace EpicGames.Core
 		/// of valid and invalid build records will usually be larger than the number of input records</param>
 		/// <param name="ProjectPath">Path of the project</param>
 		void ValidateRecursively(
-			Dictionary<FileReference, CsProjBuildRecord> ValidBuildRecords,
+			Dictionary<FileReference, (CsProjBuildRecord, FileReference)> ValidBuildRecords,
 			Dictionary<FileReference, (CsProjBuildRecord, FileReference)> InvalidBuildRecords,
 			Dictionary<FileReference, (CsProjBuildRecord, FileReference)> BuildRecords,
 			FileReference ProjectPath);
diff --git a/Engine/Source/Programs/Shared/EpicGames.MsBuild/CsProjBuilder.cs b/Engine/Source/Programs/Shared/EpicGames.MsBuild/CsProjBuilder.cs
--- a/Engine/Source/Programs/Shared/EpicGames.MsBuild/CsProjBuilder.cs
+++ b/Engine/Source/Programs/Shared/EpicGames.MsBuild/CsProjBuilder.cs
@@ -1,5 +1,3 @@
-// Copyright Epic Games, Inc. All Rights Reserved.
-
 using EpicGames.Core;
 using Microsoft.Build.Evaluation;
 using Microsoft.Build.Execution;
@@ -141,19 +139,19 @@ namespace EpicGames.MsBuild
 			return FileReference.Combine(BuildRecordDirectory, ProjectPath.GetFileName()).ChangeExtension(".json");
 		}
 
-		public static HashSet<FileReference> Build(HashSet<FileReference> FoundAutomationProjects,
-			bool bForceCompile, out bool bBuildSuccess, CsProjBuildHook Hook, List<DirectoryReference> BaseDirectories)
+		public static Dictionary<FileReference, (CsProjBuildRecord, FileReference)> Build(HashSet<FileReference> FoundProjects,
+			bool bForceCompile, out bool bBuildSuccess, CsProjBuildHook Hook, List<DirectoryReference> BaseDirectories, Action<int> OnBuildingProjects)
 		{
 
 			// Register the MS build path prior to invoking the internal routine.  By not having the internal routine
 			// inline, we avoid having the issue of the Microsoft.Build libraries being resolved prior to the build path
 			// being set.
 			RegisterMsBuildPath(Hook);
-			return BuildInternal(FoundAutomationProjects, bForceCompile, out bBuildSuccess, Hook, BaseDirectories);
+			return BuildInternal(FoundProjects, bForceCompile, out bBuildSuccess, Hook, BaseDirectories, OnBuildingProjects);
 		}
 
-		private static HashSet<FileReference> BuildInternal(HashSet<FileReference> FoundAutomationProjects,
-			bool bForceCompile, out bool bBuildSuccess, CsProjBuildHook Hook, List<DirectoryReference> BaseDirectories)
+		private static Dictionary<FileReference, (CsProjBuildRecord, FileReference)> BuildInternal(HashSet<FileReference> FoundProjects,
+			bool bForceCompile, out bool bBuildSuccess, CsProjBuildHook Hook, List<DirectoryReference> BaseDirectories, Action<int> OnBuildingProjects)
 		{
 			Dictionary<string, string> GlobalProperties = new Dictionary<string, string>
 			{
@@ -167,6 +165,7 @@ namespace EpicGames.MsBuild
 
 			Dictionary<FileReference, (CsProjBuildRecord BuildRecord, FileReference BuildRecordPath)> BuildRecords = new Dictionary<FileReference, (CsProjBuildRecord, FileReference)>();
 
+			using var ProjectCollection = new ProjectCollection(GlobalProperties);
 			var Projects = new Dictionary<string, Project>();
 			var	SkippedProjects = new HashSet<string>();
 
@@ -174,8 +173,8 @@ namespace EpicGames.MsBuild
 			// not available when using Microsoft.Build.Execution.ProjectInstance (used later in this function and
 			// in BuildProjects) - particularly, to access glob information defined in the source file.
 
-			// Load all found automation projects, and any other referenced projects.
-			foreach (FileReference ProjectPath in FoundAutomationProjects)
+			// Load all found projects, and any other referenced projects.
+			foreach (FileReference ProjectPath in FoundProjects)
 			{
 				void LoadProjectAndReferences(string ProjectPath, string ReferencedBy)
 				{
@@ -188,7 +187,7 @@ namespace EpicGames.MsBuild
 						// so make sure to print our own diagnostic info if something goes wrong
 						try
 						{
-							Project = new Project(ProjectPath, GlobalProperties, toolsVersion: null);
+							Project = new Project(ProjectPath, GlobalProperties, toolsVersion: null, projectCollection: ProjectCollection);
 						}
 						catch (Microsoft.Build.Exceptions.InvalidProjectFileException IPFEx)
 						{
@@ -363,7 +362,7 @@ namespace EpicGames.MsBuild
 				{
 					if (String.Equals("None", Glob.ItemElement.ItemType, StringComparison.Ordinal))
 					{
-						// don't record the default "None" glob - it's not (?) a trigger for any Automation rebuild
+						// don't record the default "None" glob - it's not (?) a trigger for any rebuild
 						continue;
 					}
 
@@ -386,10 +385,10 @@ namespace EpicGames.MsBuild
 			// Potential optimization: Constructing the ProjectGraph here gives the full graph of dependencies - which is nice,
 			// but not strictly necessary, and slower than doing it some other way.
 			ProjectGraph InputProjectGraph;
-			InputProjectGraph = new ProjectGraph(FoundAutomationProjects
+			InputProjectGraph = new ProjectGraph(FoundProjects
 				// Build the graph without anything that can't be built on this platform
 				.Where(x => !SkippedProjects.Contains(x.FullName))
-				.Select(P => P.FullName), GlobalProperties);
+				.Select(P => P.FullName), GlobalProperties, ProjectCollection);
 
 			// A ProjectGraph that will represent the set of projects that we actually want to build
 			ProjectGraph BuildProjectGraph = null;
@@ -401,7 +400,7 @@ namespace EpicGames.MsBuild
 			}
 			else
 			{
-				Dictionary<FileReference, CsProjBuildRecord> ValidBuildRecords = new Dictionary<FileReference, CsProjBuildRecord>(BuildRecords.Count);
+				Dictionary<FileReference, (CsProjBuildRecord, FileReference)> ValidBuildRecords = new Dictionary<FileReference, (CsProjBuildRecord, FileReference)>(BuildRecords.Count);
 				Dictionary<FileReference, (CsProjBuildRecord, FileReference)> InvalidBuildRecords = new Dictionary<FileReference, (CsProjBuildRecord, FileReference)>(BuildRecords.Count);
 
 				foreach (ProjectGraphNode Project in InputProjectGraph.ProjectNodesTopologicallySorted)
@@ -414,12 +413,13 @@ namespace EpicGames.MsBuild
 
 				if (OutOfDateProjects.Count > 0)
 				{
-					BuildProjectGraph = new ProjectGraph(OutOfDateProjects.Select(P => P.ProjectInstance.FullPath), GlobalProperties);
+					BuildProjectGraph = new ProjectGraph(OutOfDateProjects.Select(P => P.ProjectInstance.FullPath), GlobalProperties, ProjectCollection);
 				}
 			}
 
 			if (BuildProjectGraph != null)
 			{
+				OnBuildingProjects(BuildProjectGraph.EntryPointNodes.Count);
 				bBuildSuccess = BuildProjects(BuildProjectGraph, GlobalProperties);
 			}
 			else
@@ -448,11 +448,16 @@ namespace EpicGames.MsBuild
 			// todo: re-verify build records after a build to verify that everything is actually up to date
 
 			// even if only a subset was built, this function returns the full list of target assembly paths
-			return new HashSet<FileReference>(InputProjectGraph.EntryPointNodes.Select(
-				Project => FileReference.FromString(Project.ProjectInstance.GetPropertyValue("TargetPath"))));
+			var OutDict = new Dictionary<FileReference, (CsProjBuildRecord BuildRecord, FileReference BuildRecordPath)>();
+			foreach (ProjectGraphNode EntryPointNode in InputProjectGraph.EntryPointNodes)
+			{
+				FileReference ProjectPath = FileReference.FromString(EntryPointNode.ProjectInstance.FullPath);
+				OutDict.Add(ProjectPath, BuildRecords[ProjectPath]);
+			}
+			return OutDict;
 		}
 
-		private static bool BuildProjects(ProjectGraph AutomationProjectGraph, Dictionary<string, string> GlobalProperties)
+		private static bool BuildProjects(ProjectGraph ProjectGraph, Dictionary<string, string> GlobalProperties)
 		{
 			var Logger = new MLogger();
 
@@ -460,10 +465,9 @@ namespace EpicGames.MsBuild
 
 			bool Result = true;
 
-			Log.TraceInformation($"Building {AutomationProjectGraph.EntryPointNodes.Count} projects (see Log 'Engine/Programs/AutomationTool/Saved/Logs/Log.txt' for more details)");
 			foreach (string TargetToBuild in TargetsToBuild)
 			{
-				var GraphRequest = new GraphBuildRequestData(AutomationProjectGraph, new string[] { TargetToBuild });
+				var GraphRequest = new GraphBuildRequestData(ProjectGraph, new string[] { TargetToBuild });
 
 				var BuildMan = BuildManager.DefaultBuildManager;
 
