Horde: Add a generic way of attaching key/value settings to different nodes via BuildGraph.

The <Annotate> attribute allows specifying settings in the format "Key1=Value1;Key2=Value2" to be applied to individual nodes or entire targets, and this data is propagated through to Horde and stored on the job template.

diff --git a/Engine/Source/Programs/Horde/Horde.Agent/Execution/BuildGraphExecutor.cs b/Engine/Source/Programs/Horde/Horde.Agent/Execution/BuildGraphExecutor.cs
--- a/Engine/Source/Programs/Horde/Horde.Agent/Execution/BuildGraphExecutor.cs
+++ b/Engine/Source/Programs/Horde/Horde.Agent/Execution/BuildGraphExecutor.cs
@@ -34,6 +34,7 @@ namespace Horde.Agent.Execution
 			public bool? Warnings { get; set; }
 			public List<string> InputDependencies { get; set; } = new List<string>();
 			public List<string> OrderDependencies { get; set; } = new List<string>();
+			public Dictionary<string, string> Annotations { get; set; } = new Dictionary<string, string>(StringComparer.OrdinalIgnoreCase);
 		}
 
 		protected class ExportedGroup
@@ -437,6 +438,7 @@ namespace Horde.Agent.Execution
 					createNode.RunEarly = exportedNode.RunEarly;
 					createNode.Warnings = exportedNode.Warnings;
 					createNode.Priority = Priority.Normal;
+					createNode.Annotations.Add(exportedNode.Annotations);
 					createGroup.Nodes.Add(createNode);
 				}
 				updateGraph.Groups.Add(createGroup);
diff --git a/Engine/Source/Programs/Horde/Horde.Build/Api/Nodes.cs b/Engine/Source/Programs/Horde/Horde.Build/Api/Nodes.cs
--- a/Engine/Source/Programs/Horde/Horde.Build/Api/Nodes.cs
+++ b/Engine/Source/Programs/Horde/Horde.Build/Api/Nodes.cs
@@ -63,6 +63,11 @@ namespace Horde.Build.Api
 		/// </summary>
 		public IReadOnlyDictionary<string, string>? Properties { get; set; }
 
+		/// <summary>
+		/// Annotations for this node
+		/// </summary>
+		public IReadOnlyDictionary<string, string>? Annotations { get; set; }
+
 		/// <summary>
 		/// Constructor
 		/// </summary>
@@ -79,6 +84,7 @@ namespace Horde.Build.Api
 			Warnings = node.Warnings;
 			Credentials = node.Credentials;
 			Properties = node.Properties;
+			Annotations = node.Annotations;
 		}
 	}
 
diff --git a/Engine/Source/Programs/Horde/Horde.Build/Jobs/IGraph.cs b/Engine/Source/Programs/Horde/Horde.Build/Jobs/IGraph.cs
--- a/Engine/Source/Programs/Horde/Horde.Build/Jobs/IGraph.cs
+++ b/Engine/Source/Programs/Horde/Horde.Build/Jobs/IGraph.cs
@@ -6,6 +6,7 @@ using System.Diagnostics;
 using System.Diagnostics.CodeAnalysis;
 using System.Linq;
 using EpicGames.Core;
+using Horde.Build.Utilities;
 using HordeCommon;
 using MongoDB.Bson.Serialization.Attributes;
 
@@ -60,6 +61,11 @@ namespace Horde.Build.Models
 		/// Properties for this node
 		/// </summary>
 		public IReadOnlyDictionary<string, string>? Properties { get; }
+
+		/// <summary>
+		/// Annotations for this node
+		/// </summary>
+		public IReadOnlyDictionary<string, string>? Annotations { get; }
 	}
 
 	/// <summary>
@@ -413,6 +419,11 @@ namespace Horde.Build.Models
 		/// </summary>
 		public Dictionary<string, string>? Properties { get; set; }
 
+		/// <summary>
+		/// Additional user annotations for this node
+		/// </summary>
+		public CaseInsensitiveDictionary<string>? Annotations { get; set; }
+
 		/// <summary>
 		/// Constructor
 		/// </summary>
@@ -425,7 +436,8 @@ namespace Horde.Build.Models
 		/// <param name="warnings">Whether to include warnings in the diagnostic output (defaults to true)</param>
 		/// <param name="credentials">Credentials required for this node to run</param>
 		/// <param name="properties">Properties for the node</param>
-		public NewNode(string name, List<string>? inputDependencies = null, List<string>? orderDependencies = null, Priority? priority = null, bool? allowRetry = null, bool? runEarly = null, bool? warnings = null, Dictionary<string, string>? credentials = null, Dictionary<string, string>? properties = null)
+		/// <param name="annotations">User annotations for this node</param>
+		public NewNode(string name, List<string>? inputDependencies = null, List<string>? orderDependencies = null, Priority? priority = null, bool? allowRetry = null, bool? runEarly = null, bool? warnings = null, Dictionary<string, string>? credentials = null, Dictionary<string, string>? properties = null, CaseInsensitiveDictionary<string>? annotations = null)
 		{
 			Name = name;
 			InputDependencies = inputDependencies;
@@ -436,6 +448,7 @@ namespace Horde.Build.Models
 			Warnings = warnings;
 			Credentials = credentials;
 			Properties = properties;
+			Annotations = annotations;
 		}
 	}
 
diff --git a/Engine/Source/Programs/Horde/Horde.Build/Jobs/Impl/GraphCollection.cs b/Engine/Source/Programs/Horde/Horde.Build/Jobs/Impl/GraphCollection.cs
--- a/Engine/Source/Programs/Horde/Horde.Build/Jobs/Impl/GraphCollection.cs
+++ b/Engine/Source/Programs/Horde/Horde.Build/Jobs/Impl/GraphCollection.cs
@@ -10,6 +10,7 @@ using System.Threading.Tasks;
 using EpicGames.Core;
 using Horde.Build.Models;
 using Horde.Build.Server;
+using Horde.Build.Utilities;
 using HordeCommon;
 using MongoDB.Bson;
 using MongoDB.Bson.Serialization.Attributes;
@@ -48,8 +49,12 @@ namespace Horde.Build.Collections.Impl
 			[BsonIgnoreIfNull]
 			public Dictionary<string, string>? Properties { get; set; }
 
+			[BsonIgnoreIfNull]
+			public CaseInsensitiveDictionary<string>? Annotations { get; set; }
+
 			IReadOnlyDictionary<string, string>? INode.Credentials => Credentials;
 			IReadOnlyDictionary<string, string>? INode.Properties => Properties;
+			IReadOnlyDictionary<string, string>? INode.Annotations => Annotations;
 
 			[BsonConstructor]
 			private Node()
@@ -59,7 +64,7 @@ namespace Horde.Build.Collections.Impl
 				OrderDependencies = null!;
 			}
 
-			public Node(string name, NodeRef[] inputDependencies, NodeRef[] orderDependencies, Priority priority, bool allowRetry, bool runEarly, bool warnings, Dictionary<string, string>? credentials, Dictionary<string, string>? properties)
+			public Node(string name, NodeRef[] inputDependencies, NodeRef[] orderDependencies, Priority priority, bool allowRetry, bool runEarly, bool warnings, Dictionary<string, string>? credentials, Dictionary<string, string>? properties, CaseInsensitiveDictionary<string>? annotations)
 			{
 				Name = name;
 				InputDependencies = inputDependencies;
@@ -70,6 +75,7 @@ namespace Horde.Build.Collections.Impl
 				Warnings = warnings;
 				Credentials = credentials;
 				Properties = properties;
+				Annotations = annotations;
 			}
 		}
 
@@ -224,7 +230,7 @@ namespace Horde.Build.Collections.Impl
 							NodeRef[] inputDependencies = (newNodeRequest.InputDependencies == null) ? Array.Empty<NodeRef>() : newNodeRequest.InputDependencies.Select(x => nodeNameToRef[x]).ToArray();
 							NodeRef[] orderDependencies = (newNodeRequest.OrderDependencies == null) ? Array.Empty<NodeRef>() : newNodeRequest.OrderDependencies.Select(x => nodeNameToRef[x]).ToArray();
 							orderDependencies = orderDependencies.Union(inputDependencies).ToArray();
-							nodes.Add(new Node(newNodeRequest.Name, inputDependencies, orderDependencies, priority, bAllowRetry, bRunEarly, bWarnings, newNodeRequest.Credentials, newNodeRequest.Properties));
+							nodes.Add(new Node(newNodeRequest.Name, inputDependencies, orderDependencies, priority, bAllowRetry, bRunEarly, bWarnings, newNodeRequest.Credentials, newNodeRequest.Properties, newNodeRequest.Annotations));
 
 							nodeNameToRef.Add(newNodeRequest.Name, new NodeRef(newGroups.Count, nodeIdx));
 						}
@@ -384,7 +390,7 @@ namespace Horde.Build.Collections.Impl
 		/// <inheritdoc/>
 		public async Task<IGraph> AddAsync(ITemplate template)
 		{
-			Node node = new Node(IJob.SetupNodeName, Array.Empty<NodeRef>(), Array.Empty<NodeRef>(), Priority.High, true, false, true, null, null);
+			Node node = new Node(IJob.SetupNodeName, Array.Empty<NodeRef>(), Array.Empty<NodeRef>(), Priority.High, true, false, true, null, null, null);
 			NodeGroup group = new NodeGroup(template.InitialAgentType ?? "Win64", new List<Node> { node });
 
 			GraphDocument graph = new GraphDocument(new List<NodeGroup> { group }, new List<Aggregate>(), new List<Label>());
diff --git a/Engine/Source/Programs/Horde/Horde.Build/Utilities/CaseInsensitiveDictionary.cs b/Engine/Source/Programs/Horde/Horde.Build/Utilities/CaseInsensitiveDictionary.cs
new file mode 100644
--- /dev/null
+++ b/Engine/Source/Programs/Horde/Horde.Build/Utilities/CaseInsensitiveDictionary.cs
@@ -0,0 +1,37 @@
+// Copyright Epic Games, Inc. All Rights Reserved.
+
+using System;
+using System.Collections.Generic;
+
+namespace Horde.Build.Utilities
+{
+	/// <summary>
+	/// Case insensitive set of strings
+	/// </summary>
+	public class CaseInsensitiveDictionary<TValue> : Dictionary<string, TValue>
+	{
+		/// <summary>
+		/// Constructor
+		/// </summary>
+		public CaseInsensitiveDictionary()
+			: base(StringComparer.OrdinalIgnoreCase)
+		{
+		}
+
+		/// <summary>
+		/// Constructor
+		/// </summary>
+		public CaseInsensitiveDictionary(IDictionary<string, TValue> items)
+			: base(items, StringComparer.OrdinalIgnoreCase)
+		{
+		}
+
+		/// <summary>
+		/// Constructor
+		/// </summary>
+		public CaseInsensitiveDictionary(IEnumerable<KeyValuePair<string, TValue>> items)
+			: base(items, StringComparer.OrdinalIgnoreCase)
+		{
+		}
+	}
+}
diff --git a/Engine/Source/Programs/Shared/EpicGames.BuildGraph/BgGraph.cs b/Engine/Source/Programs/Shared/EpicGames.BuildGraph/BgGraph.cs
--- a/Engine/Source/Programs/Shared/EpicGames.BuildGraph/BgGraph.cs
+++ b/Engine/Source/Programs/Shared/EpicGames.BuildGraph/BgGraph.cs
@@ -527,6 +527,13 @@ namespace EpicGames.BuildGraph
 						}
 						jsonWriter.WriteArrayEnd();
 
+						jsonWriter.WriteObjectStart("Annotations");
+						foreach ((string key, string value) in node.Annotations)
+						{
+							jsonWriter.WriteValue(key, value);
+						}
+						jsonWriter.WriteObjectEnd();
+
 						jsonWriter.WriteObjectEnd();
 					}
 					jsonWriter.WriteArrayEnd();
diff --git a/Engine/Source/Programs/Shared/EpicGames.BuildGraph/BgNode.cs b/Engine/Source/Programs/Shared/EpicGames.BuildGraph/BgNode.cs
--- a/Engine/Source/Programs/Shared/EpicGames.BuildGraph/BgNode.cs
+++ b/Engine/Source/Programs/Shared/EpicGames.BuildGraph/BgNode.cs
@@ -104,6 +104,11 @@ namespace EpicGames.BuildGraph
 		/// </summary>
 		public bool NotifyOnWarnings { get; set; } = true;
 
+		/// <summary>
+		/// Custom annotations for this node
+		/// </summary>
+		public Dictionary<string, string> Annotations { get; } = new Dictionary<string, string>(StringComparer.OrdinalIgnoreCase);
+
 		/// <summary>
 		/// Constructor
 		/// </summary>
diff --git a/Engine/Source/Programs/Shared/EpicGames.BuildGraph/BgSchema.cs b/Engine/Source/Programs/Shared/EpicGames.BuildGraph/BgSchema.cs
--- a/Engine/Source/Programs/Shared/EpicGames.BuildGraph/BgSchema.cs
+++ b/Engine/Source/Programs/Shared/EpicGames.BuildGraph/BgSchema.cs
@@ -149,6 +149,7 @@ namespace EpicGames.BuildGraph
 		MacroBody,
 		Extend,
 		Expand,
+		Annotate,
 		Trace,
 		Warning,
 		Error,
@@ -332,6 +333,7 @@ namespace EpicGames.BuildGraph
 			newSchema.Items.Add(CreateMacroBodyType(taskNameToType));
 			newSchema.Items.Add(CreateExtendType());
 			newSchema.Items.Add(CreateExpandType());
+			newSchema.Items.Add(CreateAnnotateType());
 			newSchema.Items.Add(CreateDiagnosticType(ScriptSchemaStandardType.Trace));
 			newSchema.Items.Add(CreateDiagnosticType(ScriptSchemaStandardType.Warning));
 			newSchema.Items.Add(CreateDiagnosticType(ScriptSchemaStandardType.Error));
@@ -477,6 +479,7 @@ namespace EpicGames.BuildGraph
 			graphChoice.Items.Add(CreateSchemaElement("Badge", ScriptSchemaStandardType.Badge));
 			graphChoice.Items.Add(CreateSchemaElement("Label", ScriptSchemaStandardType.Label));
 			graphChoice.Items.Add(CreateSchemaElement("Notify", ScriptSchemaStandardType.Notify));
+			graphChoice.Items.Add(CreateSchemaElement("Annotate", ScriptSchemaStandardType.Annotate));
 			graphChoice.Items.Add(CreateSchemaElement("Trace", ScriptSchemaStandardType.Trace));
 			graphChoice.Items.Add(CreateSchemaElement("Warning", ScriptSchemaStandardType.Warning));
 			graphChoice.Items.Add(CreateSchemaElement("Error", ScriptSchemaStandardType.Error));
@@ -527,6 +530,7 @@ namespace EpicGames.BuildGraph
 			agentChoice.Items.Add(CreateSchemaElement("Node", ScriptSchemaStandardType.Node));
 			agentChoice.Items.Add(CreateSchemaElement("Trace", ScriptSchemaStandardType.Trace));
 			agentChoice.Items.Add(CreateSchemaElement("Label", ScriptSchemaStandardType.Label));
+			agentChoice.Items.Add(CreateSchemaElement("Annotate", ScriptSchemaStandardType.Annotate));
 			agentChoice.Items.Add(CreateSchemaElement("Warning", ScriptSchemaStandardType.Warning));
 			agentChoice.Items.Add(CreateSchemaElement("Error", ScriptSchemaStandardType.Error));
 			agentChoice.Items.Add(CreateSchemaElement("Expand", ScriptSchemaStandardType.Expand));
@@ -692,6 +696,22 @@ namespace EpicGames.BuildGraph
 			return aggregateType;
 		}
 
+		/// <summary>
+		/// Creates the schema type representing an annotation
+		/// </summary>
+		/// <returns>Type definition for a notifier</returns>
+		static XmlSchemaType CreateAnnotateType()
+		{
+			XmlSchemaComplexType aggregateType = new XmlSchemaComplexType();
+			aggregateType.Name = GetTypeName(ScriptSchemaStandardType.Annotate);
+			aggregateType.Attributes.Add(CreateSchemaAttribute("Targets", ScriptSchemaStandardType.NameOrTagList, XmlSchemaUse.Optional));
+			aggregateType.Attributes.Add(CreateSchemaAttribute("Except", ScriptSchemaStandardType.NameOrTagList, XmlSchemaUse.Optional));
+			aggregateType.Attributes.Add(CreateSchemaAttribute("Nodes", ScriptSchemaStandardType.NameOrTagList, XmlSchemaUse.Optional));
+			aggregateType.Attributes.Add(CreateSchemaAttribute("Values", ScriptSchemaStandardType.BalancedString, XmlSchemaUse.Optional));
+			aggregateType.Attributes.Add(CreateSchemaAttribute("If", ScriptSchemaStandardType.BalancedString, XmlSchemaUse.Optional));
+			return aggregateType;
+		}
+
 		/// <summary>
 		/// Creates the schema type representing an include type
 		/// </summary>
diff --git a/Engine/Source/Programs/Shared/EpicGames.BuildGraph/BgScriptReader.cs b/Engine/Source/Programs/Shared/EpicGames.BuildGraph/BgScriptReader.cs
--- a/Engine/Source/Programs/Shared/EpicGames.BuildGraph/BgScriptReader.cs
+++ b/Engine/Source/Programs/Shared/EpicGames.BuildGraph/BgScriptReader.cs
@@ -473,6 +473,9 @@ namespace EpicGames.BuildGraph
 					case "Expand":
 						await ReadExpandAsync(childElement, ReadGraphBodyAsync);
 						break;
+					case "Annotate":
+						await ReadAnnotationAsync(childElement);
+						break;
 					default:
 						LogError(childElement, "Invalid element '{0}'", childElement.Name);
 						break;
@@ -787,6 +790,9 @@ namespace EpicGames.BuildGraph
 					case "Expand":
 						await ReadExpandAsync(childElement, ReadAgentBodyAsync);
 						break;
+					case "Annotate":
+						await ReadAnnotationAsync(childElement);
+						break;
 					default:
 						LogError(childElement, "Unexpected element type '{0}'", childElement.Name);
 						break;
@@ -1012,6 +1018,12 @@ namespace EpicGames.BuildGraph
 		/// <param name="element">Xml element to read the definition from</param>
 		protected abstract Task ReadNotifierAsync(BgScriptElement element);
 
+		/// <summary>
+		/// Reads a graph annotation
+		/// </summary>
+		/// <param name="element">Xml element to read the definition from</param>
+		protected abstract Task ReadAnnotationAsync(BgScriptElement element);
+
 		/// <summary>
 		/// Reads a warning from the given element, evaluates the condition on it, and writes it to the log if the condition passes.
 		/// </summary>
@@ -1996,6 +2008,82 @@ namespace EpicGames.BuildGraph
 			}
 		}
 
+		/// <inheritdoc/>
+		protected override async Task ReadAnnotationAsync(BgScriptElement element)
+		{
+			if (await EvaluateConditionAsync(element))
+			{
+				string[] targetNames = ReadListAttribute(element, "Targets");
+				string[] exceptNames = ReadListAttribute(element, "Except");
+				string[] individualNodeNames = ReadListAttribute(element, "Nodes");
+				string[] pairs = ReadListAttribute(element, "Values");
+
+				// Find the list of targets which are included, and recurse through all their dependencies
+				HashSet<BgNode> nodes = new HashSet<BgNode>();
+				if (targetNames != null)
+				{
+					HashSet<BgNode> targetNodes = ResolveReferences(element, targetNames);
+					foreach (BgNode node in targetNodes)
+					{
+						nodes.Add(node);
+						nodes.UnionWith(node.InputDependencies);
+					}
+				}
+
+				// Add all the individually referenced nodes
+				if (individualNodeNames != null)
+				{
+					HashSet<BgNode> individualNodes = ResolveReferences(element, individualNodeNames);
+					nodes.UnionWith(individualNodes);
+				}
+
+				// Exclude all the exceptions
+				if (exceptNames != null)
+				{
+					HashSet<BgNode> exceptNodes = ResolveReferences(element, exceptNames);
+					nodes.ExceptWith(exceptNodes);
+				}
+
+				// Find the annotations to apply
+				Dictionary<string, string> pairMap = new Dictionary<string, string>(StringComparer.OrdinalIgnoreCase);
+				foreach (string pair in pairs)
+				{
+					if (!String.IsNullOrWhiteSpace(pair))
+					{
+						int idx = pair.IndexOf('=');
+						if (idx < 0)
+						{
+							LogError(element, "Invalid annotation '{0}'", pair);
+							continue;
+						}
+
+						string key = pair.Substring(0, idx).Trim();
+						if (!Regex.IsMatch(key, @"[a-zA-Z0-9_\.]+"))
+						{
+							LogError(element, "Invalid annotation key '{0}'", pair);
+							continue;
+						}
+						if (pairMap.ContainsKey(key))
+						{
+							LogError(element, "Annotation key '{0}' was specified twice", key);
+							continue;
+						}
+
+						pairMap.Add(key, pair.Substring(idx + 1).Trim());
+					}
+				}
+
+				// Update all the referenced nodes with the settings
+				foreach (BgNode node in nodes)
+				{
+					foreach ((string key, string value) in pairMap)
+					{
+						node.Annotations[key] = value;
+					}
+				}
+			}
+		}
+
 		/// <summary>
 		/// Reads a warning from the given element, evaluates the condition on it, and writes it to the log if the condition passes.
 		/// </summary>
