Updated C# UHT to include recent changes to UHT.
Fixed an issue with FieldNotify code in C# UHT

#rnx

diff --git a/Engine/Source/Programs/Shared/EpicGames.UHT/Exporters/CodeGen/UhtHeaderCodeGenerator.cs b/Engine/Source/Programs/Shared/EpicGames.UHT/Exporters/CodeGen/UhtHeaderCodeGenerator.cs
--- a/Engine/Source/Programs/Shared/EpicGames.UHT/Exporters/CodeGen/UhtHeaderCodeGenerator.cs
+++ b/Engine/Source/Programs/Shared/EpicGames.UHT/Exporters/CodeGen/UhtHeaderCodeGenerator.cs
@@ -523,7 +523,7 @@ namespace EpicGames.UHT.Exporters.CodeGen
 				{
 					if (Child is UhtProperty Property)
 					{
-						if (!Property.bIsEditorOnlyProperty)
+						if (Property.PropertyExportFlags.HasAnyFlags(UhtPropertyExportFlags.FieldNotify) && !Property.bIsEditorOnlyProperty)
 						{
 							AppendAction(Builder, Class, Property.SourceName);
 						}
@@ -537,7 +537,7 @@ namespace EpicGames.UHT.Exporters.CodeGen
 				{
 					if (Child is UhtFunction Function)
 					{
-						if (!Function.FunctionFlags.HasAnyFlags(EFunctionFlags.EditorOnly))
+						if (Function.FunctionExportFlags.HasAnyFlags(UhtFunctionExportFlags.FieldNotify) && !Function.FunctionFlags.HasAnyFlags(EFunctionFlags.EditorOnly))
 						{
 							AppendAction(Builder, Class, Function.CppImplName);
 						}
@@ -558,7 +558,7 @@ namespace EpicGames.UHT.Exporters.CodeGen
 					{
 						if (Child is UhtProperty Property)
 						{
-							if (Property.bIsEditorOnlyProperty)
+							if (Property.PropertyExportFlags.HasAnyFlags(UhtPropertyExportFlags.FieldNotify) && Property.bIsEditorOnlyProperty)
 							{
 								AppendAction(Builder, Class, Property.SourceName);
 							}
@@ -572,7 +572,7 @@ namespace EpicGames.UHT.Exporters.CodeGen
 					{
 						if (Child is UhtFunction Function)
 						{
-							if (Function.FunctionFlags.HasAnyFlags(EFunctionFlags.EditorOnly))
+							if (Function.FunctionExportFlags.HasAnyFlags(UhtFunctionExportFlags.FieldNotify) && Function.FunctionFlags.HasAnyFlags(EFunctionFlags.EditorOnly))
 							{
 								AppendAction(Builder, Class, Function.CppImplName);
 							}
diff --git a/Engine/Source/Programs/Shared/EpicGames.UHT/Types/Properties/UhtArrayProperty.cs b/Engine/Source/Programs/Shared/EpicGames.UHT/Types/Properties/UhtArrayProperty.cs
--- a/Engine/Source/Programs/Shared/EpicGames.UHT/Types/Properties/UhtArrayProperty.cs
+++ b/Engine/Source/Programs/Shared/EpicGames.UHT/Types/Properties/UhtArrayProperty.cs
@@ -66,6 +66,7 @@ namespace EpicGames.UHT.Types
 				case UhtResolvePhase.Final:
 					this.PropertyFlags |= ResolveAndReturnNewFlags(this.ValueProperty, Phase);
 					this.MetaData.Add(this.ValueProperty.MetaData);
+					this.ValueProperty.PropertyFlags = this.PropertyFlags & EPropertyFlags.PropagateToArrayInner;
 					this.ValueProperty.MetaData.Clear();
 					PropagateFlagsFromInnerAndHandlePersistentInstanceMetadata(this, this.MetaData, this.ValueProperty);
 					break;
diff --git a/Engine/Source/Programs/Shared/EpicGames.UHT/Types/Properties/UhtFieldPathProperty.cs b/Engine/Source/Programs/Shared/EpicGames.UHT/Types/Properties/UhtFieldPathProperty.cs
--- a/Engine/Source/Programs/Shared/EpicGames.UHT/Types/Properties/UhtFieldPathProperty.cs
+++ b/Engine/Source/Programs/Shared/EpicGames.UHT/Types/Properties/UhtFieldPathProperty.cs
@@ -22,7 +22,7 @@ namespace EpicGames.UHT.Types
 		protected override string CppTypeText { get => "TFieldPath"; }
 
 		/// <inheritdoc/>
-		protected override string PGetMacroText { get => "STRUCT"; }
+		protected override string PGetMacroText { get => "TFIELDPATH"; }
 
 		/// <inheritdoc/>
 		protected override UhtPGetArgumentType PGetTypeArgument { get => UhtPGetArgumentType.TypeText; }
diff --git a/Engine/Source/Programs/Shared/EpicGames.UHT/Types/Properties/UhtMapProperty.cs b/Engine/Source/Programs/Shared/EpicGames.UHT/Types/Properties/UhtMapProperty.cs
--- a/Engine/Source/Programs/Shared/EpicGames.UHT/Types/Properties/UhtMapProperty.cs
+++ b/Engine/Source/Programs/Shared/EpicGames.UHT/Types/Properties/UhtMapProperty.cs
@@ -89,6 +89,7 @@ namespace EpicGames.UHT.Types
 					this.KeyProperty.PropertyFlags |= NewFlags;
 					this.MetaData.Add(this.ValueProperty.MetaData);
 					this.ValueProperty.MetaData.Clear();
+					this.ValueProperty.PropertyFlags = this.PropertyFlags & EPropertyFlags.PropagateToMapValue;
 
 					PropagateFlagsFromInnerAndHandlePersistentInstanceMetadata(this, null, this.KeyProperty);
 					PropagateFlagsFromInnerAndHandlePersistentInstanceMetadata(this, this.MetaData, this.ValueProperty);
diff --git a/Engine/Source/Programs/Shared/EpicGames.UHT/Types/Properties/UhtSetProperty.cs b/Engine/Source/Programs/Shared/EpicGames.UHT/Types/Properties/UhtSetProperty.cs
--- a/Engine/Source/Programs/Shared/EpicGames.UHT/Types/Properties/UhtSetProperty.cs
+++ b/Engine/Source/Programs/Shared/EpicGames.UHT/Types/Properties/UhtSetProperty.cs
@@ -60,6 +60,7 @@ namespace EpicGames.UHT.Types
 					this.PropertyFlags |= ResolveAndReturnNewFlags(this.ValueProperty, Phase);
 					this.MetaData.Add(this.ValueProperty.MetaData);
 					this.ValueProperty.MetaData.Clear();
+					this.ValueProperty.PropertyFlags = this.PropertyFlags & EPropertyFlags.PropagateToSetElement;
 					PropagateFlagsFromInnerAndHandlePersistentInstanceMetadata(this, this.MetaData, this.ValueProperty);
 					break;
 			}
diff --git a/Engine/Source/Programs/Shared/EpicGames.UHT/Types/UhtScriptStruct.cs b/Engine/Source/Programs/Shared/EpicGames.UHT/Types/UhtScriptStruct.cs
--- a/Engine/Source/Programs/Shared/EpicGames.UHT/Types/UhtScriptStruct.cs
+++ b/Engine/Source/Programs/Shared/EpicGames.UHT/Types/UhtScriptStruct.cs
@@ -631,6 +631,55 @@ namespace EpicGames.UHT.Types
 		/// <inheritdoc/>
 		protected override void ResolveChildren(UhtResolvePhase Phase)
 		{
+
+			// Setup additional property as well as script struct def flags
+			// for structs / properties being used for the RigVM.
+			// The Input / Output / Constant metadata tokens can be used to mark
+			// up an input / output pin of a RigVMNode. To allow authoring of those
+			// nodes we'll mark up the property as accessible in Blueprint / Python
+			// as well as make the struct a blueprint type.
+			switch (Phase)
+			{
+				case UhtResolvePhase.Properties:
+					foreach (UhtType Child in this.Children)
+					{
+						if (Child is UhtProperty Property)
+						{
+							EPropertyFlags OriginalFlags = Property.PropertyFlags;
+							if (Property.MetaData.ContainsKey(UhtNames.Constant))
+							{
+								Property.PropertyFlags |= EPropertyFlags.Edit | EPropertyFlags.EditConst | EPropertyFlags.BlueprintVisible;
+							}
+							if (Property.MetaData.ContainsKey(UhtNames.Input) || Property.MetaData.ContainsKey(UhtNames.Visible))
+							{
+								Property.PropertyFlags |= EPropertyFlags.Edit | EPropertyFlags.BlueprintVisible;
+							}
+							if (Property.MetaData.ContainsKey(UhtNames.Output))
+							{
+								if (!Property.PropertyFlags.HasAnyFlags(EPropertyFlags.BlueprintVisible))
+								{
+									Property.PropertyFlags |= EPropertyFlags.BlueprintVisible | EPropertyFlags.BlueprintReadOnly;
+								}
+							}
+
+							if (OriginalFlags != Property.PropertyFlags &&
+								Property.PropertyFlags.HasAnyFlags(EPropertyFlags.BlueprintVisible | EPropertyFlags.BlueprintReadOnly))
+							{
+								if (!this.MetaData.GetBooleanHierarchical(UhtNames.BlueprintType))
+								{
+									this.MetaData.Add(UhtNames.BlueprintType, true);
+								}
+
+								if (!Property.MetaData.ContainsKey(UhtNames.Category))
+								{
+									Property.MetaData.Add(UhtNames.Category, UhtNames.Pins);
+								}
+							}
+						}
+					}
+					break;
+			}
+
 			base.ResolveChildren(Phase);
 
 			switch (Phase)
diff --git a/Engine/Source/Programs/Shared/EpicGames.UHT/Utils/UhtNames.cs b/Engine/Source/Programs/Shared/EpicGames.UHT/Utils/UhtNames.cs
--- a/Engine/Source/Programs/Shared/EpicGames.UHT/Utils/UhtNames.cs
+++ b/Engine/Source/Programs/Shared/EpicGames.UHT/Utils/UhtNames.cs
@@ -23,6 +23,7 @@ namespace EpicGames.UHT.Utils
 		public const string Category = "Category";
 		public const string ClassGroupNames = "ClassGroupNames";
 		public const string Comment = "Comment";
+		public const string Constant = "Constant";
 		public const string CppFromBpEvent = "CppFromBpEvent";
 		public const string CustomThunk = "CustomThunk";
 		public const string DeprecatedFunction = "DeprecatedFunction";
@@ -38,6 +39,7 @@ namespace EpicGames.UHT.Utils
 		public const string HideFunctions = "HideFunctions";
 		public const string IgnoreCategoryKeywordsInSubclasses = "IgnoreCategoryKeywordsInSubclasses";
 		public const string IncludePath = "IncludePath";
+		public const string Input = "Input";
 		public const string Interface = "Interface";
 		public const string IsBlueprintBase = "IsBlueprintBase";
 		public const string IsConversionRoot = "IsConversionRoot";
@@ -47,6 +49,8 @@ namespace EpicGames.UHT.Utils
 		public const string NativeConstTemplateArg = "NativeConstTemplateArg";
 		public const string NoGetter = "NoGetter";
 		public const string ObjectInitializerConstructorDeclared = "ObjectInitializerConstructorDeclared";
+		public const string Output = "Output";
+		public const string Pins = "Pins";
 		public const string PrioritizeCategories = "PrioritizeCategories";
 		public const string ReturnValue = "ReturnValue";
 		public const string ShortToolTip = "ShortToolTip";
@@ -56,5 +60,6 @@ namespace EpicGames.UHT.Utils
 		public const string ToolTip = "ToolTip";
 		public const string UIMin = "UIMin";
 		public const string UIMax = "UIMax";
+		public const string Visible = "Visible";
 	}
 }
