Assortment of very quick UBT optimizations.  15% redunction in UBT null build for UE

#rnx

diff --git a/Engine/Source/Programs/Shared/EpicGames.Build/System/DirectoryItem.cs b/Engine/Source/Programs/Shared/EpicGames.Build/System/DirectoryItem.cs
--- a/Engine/Source/Programs/Shared/EpicGames.Build/System/DirectoryItem.cs
+++ b/Engine/Source/Programs/Shared/EpicGames.Build/System/DirectoryItem.cs
@@ -343,7 +343,8 @@ namespace UnrealBuildBase
 		/// <returns>Instance of the serialized directory item</returns>
 		public static DirectoryItem? ReadDirectoryItem(this BinaryArchiveReader Reader)
 		{
-			return Reader.ReadObjectReference<DirectoryItem>(() => DirectoryItem.GetItemByDirectoryReference(Reader.ReadDirectoryReferenceNotNull()));
+			// Use lambda that doesn't require anything to be captured thus eliminating an allocation.
+			return Reader.ReadObjectReference<DirectoryItem>((BinaryArchiveReader Reader) => DirectoryItem.GetItemByDirectoryReference(Reader.ReadDirectoryReferenceNotNull()));
 		}
 
 		/// <summary>
@@ -353,7 +354,8 @@ namespace UnrealBuildBase
 		/// <param name="DirectoryItem">Directory item to write</param>
 		public static void WriteDirectoryItem(this BinaryArchiveWriter Writer, DirectoryItem DirectoryItem)
 		{
-			Writer.WriteObjectReference<DirectoryItem>(DirectoryItem, () => Writer.WriteDirectoryReference(DirectoryItem.Location));
+			// Use lambda that doesn't require anything to be captured thus eliminating an allocation.
+			Writer.WriteObjectReference<DirectoryItem>(DirectoryItem, (BinaryArchiveWriter Writer, DirectoryItem DirectoryItem) => Writer.WriteDirectoryReference(DirectoryItem.Location));
 		}
 	}
 }
diff --git a/Engine/Source/Programs/Shared/EpicGames.Build/System/FileItem.cs b/Engine/Source/Programs/Shared/EpicGames.Build/System/FileItem.cs
--- a/Engine/Source/Programs/Shared/EpicGames.Build/System/FileItem.cs
+++ b/Engine/Source/Programs/Shared/EpicGames.Build/System/FileItem.cs
@@ -295,7 +295,7 @@ namespace UnrealBuildBase
 		/// <returns>Instance of the serialized file item</returns>
 		public static FileItem? ReadFileItem(this BinaryArchiveReader Reader)
 		{
-			return Reader.ReadObjectReference<FileItem>(() => FileItem.GetItemByFileReference(Reader.ReadFileReference()));
+			return Reader.ReadObjectReference<FileItem>((BinaryArchiveReader Reader) => FileItem.GetItemByFileReference(Reader.ReadFileReference()));
 		}
 
 		/// <summary>
@@ -331,7 +331,8 @@ namespace UnrealBuildBase
 		/// <returns>Instance of the serialized file item</returns>
 		public static FileItem ReadCompactFileItem(this BinaryArchiveReader Reader)
 		{
-			return Reader.ReadObjectReference<FileItem>(() => ReadCompactFileItemData(Reader))!;
+			// Use lambda that doesn't require anything to be captured thus eliminating an allocation.
+			return Reader.ReadObjectReference<FileItem>((BinaryArchiveReader Reader) => ReadCompactFileItemData(Reader))!;
 		}
 
 		/// <summary>
@@ -341,7 +342,8 @@ namespace UnrealBuildBase
 		/// <param name="FileItem">File item to write</param>
 		public static void WriteCompactFileItem(this BinaryArchiveWriter Writer, FileItem FileItem)
 		{
-			Writer.WriteObjectReference<FileItem>(FileItem, () => { Writer.WriteDirectoryItem(FileItem.GetDirectoryItem()); Writer.WriteString(FileItem.Name); });
+			// Use lambda that doesn't require anything to be captured thus eliminating an allocation.
+			Writer.WriteObjectReference<FileItem>(FileItem, (BinaryArchiveWriter Writer, FileItem FileItem) => { Writer.WriteDirectoryItem(FileItem.GetDirectoryItem()); Writer.WriteString(FileItem.Name); });
 		}
 	}
 }
diff --git a/Engine/Source/Programs/Shared/EpicGames.Core/BinaryArchiveReader.cs b/Engine/Source/Programs/Shared/EpicGames.Core/BinaryArchiveReader.cs
--- a/Engine/Source/Programs/Shared/EpicGames.Core/BinaryArchiveReader.cs
+++ b/Engine/Source/Programs/Shared/EpicGames.Core/BinaryArchiveReader.cs
@@ -180,14 +180,17 @@ namespace EpicGames.Core
 		/// <returns>The value that was read</returns>
 		public string? ReadString()
 		{
-			byte[]? bytes = ReadByteArray();
-			if (bytes == null)
+			// ReadPrimitiveArray has been inlined here to avoid the transient byte array allocation
+			int length = ReadInt();
+			if (length < 0)
 			{
 				return null;
 			}
 			else
 			{
-				return Encoding.UTF8.GetString(bytes);
+				int offset = _bufferPos;
+				_bufferPos += length;
+				return Encoding.UTF8.GetString(_buffer!, offset, length);
 			}
 		}
 
@@ -533,6 +536,36 @@ namespace EpicGames.Core
 			}
 		}
 
+		/// <summary>
+		/// Reads an object reference from the stream. Each object will only be serialized once using the supplied delegate; subsequent reads reference the original.
+		/// Since the reader only receives the object reference when the CreateObject delegate returns, it is not possible for the object to serialize a reference to itself.
+		/// </summary>
+		/// <typeparam name="T">Type of the object to read.</typeparam>
+		/// <param name="readObject">Delegate used to create an object instance. The object may not reference itself recursively.</param>
+		/// <returns>Object instance</returns>
+		public object? ReadUntypedObjectReference(Func<BinaryArchiveReader, object?> readObject)
+		{
+			int index = ReadInt();
+			if (index < 0)
+			{
+				return null;
+			}
+			else
+			{
+				// Temporarily add the reader to the object list, so we can detect invalid recursive references. 
+				if (index == _objects.Count)
+				{
+					_objects.Add(null);
+					_objects[index] = readObject(this);
+				}
+				if (_objects[index] == null)
+				{
+					throw new InvalidOperationException("Attempt to serialize reference to object recursively.");
+				}
+				return _objects[index];
+			}
+		}
+
 		/// <summary>
 		/// Reads an object reference from the stream. Each object will only be serialized once using the supplied delegate; subsequent reads reference the original.
 		/// Since the reader only receives the object reference when the CreateObject delegate returns, it is not possible for the object to serialize a reference to itself.
@@ -542,6 +575,15 @@ namespace EpicGames.Core
 		/// <returns>Object instance</returns>
 		public T? ReadObjectReference<T>(Func<T> readObject) where T : class => (T?)ReadUntypedObjectReference(readObject);
 
+		/// <summary>
+		/// Reads an object reference from the stream. Each object will only be serialized once using the supplied delegate; subsequent reads reference the original.
+		/// Since the reader only receives the object reference when the CreateObject delegate returns, it is not possible for the object to serialize a reference to itself.
+		/// </summary>
+		/// <typeparam name="T">Type of the object to read.</typeparam>
+		/// <param name="readObject">Delegate used to create an object instance. The object may not reference itself recursively.</param>
+		/// <returns>Object instance</returns>
+		public T? ReadObjectReference<T>(Func<BinaryArchiveReader, T> readObject) where T : class => (T?)ReadUntypedObjectReference(readObject);
+
 		/// <summary>
 		/// Helper method for validating that deserialized objects are not null
 		/// </summary>
diff --git a/Engine/Source/Programs/Shared/EpicGames.Core/BinaryArchiveWriter.cs b/Engine/Source/Programs/Shared/EpicGames.Core/BinaryArchiveWriter.cs
--- a/Engine/Source/Programs/Shared/EpicGames.Core/BinaryArchiveWriter.cs
+++ b/Engine/Source/Programs/Shared/EpicGames.Core/BinaryArchiveWriter.cs
@@ -500,34 +500,63 @@ namespace EpicGames.Core
 		/// <param name="obj">The object to serialize</param>
 		/// <param name="writeObject">Delegate used to write the object</param>
 		public void WriteObjectReference(object? obj, Action writeObject)
+		{
+			if (WriteObjectReferenceUniqueId(obj))
+			{
+				writeObject();
+			}
+		}
+
+		/// <summary>
+		/// Writes an object to the output. If the specific instance has already been written, preserves the reference to that.
+		/// </summary>
+		/// <param name="obj">The object to serialize</param>
+		/// <param name="writeObject">Delegate used to write the object</param>
+		public void WriteObjectReference<T>(T obj, Action writeObject) where T : class?
+		{
+			WriteObjectReference((object?)obj, writeObject);
+		}
+
+		/// <summary>
+		/// Writes an object to the output. If the specific instance has already been written, preserves the reference to that.
+		/// </summary>
+		/// <param name="obj">The object to serialize</param>
+		/// <param name="writeObject">Delegate used to write the object</param>
+		public void WriteObjectReference<T>(T obj, Action<BinaryArchiveWriter, T> writeObject) where T : class?
+		{
+			if (WriteObjectReferenceUniqueId(obj))
+			{
+				writeObject(this, obj);
+			}
+		}
+
+		/// <summary>
+		/// Writes a unique id for the given object.
+		/// </summary>
+		/// <param name="obj">The object to serialize</param>
+		/// <returns>False if the object has already been written.  True if this is the first time the object has been referenced.</returns>
+		private bool WriteObjectReferenceUniqueId(object? obj)
 		{
 			if (obj == null)
 			{
 				WriteInt(-1);
+				return false;
 			}
 			else
 			{
 				if (_objectToUniqueId.TryGetValue(obj, out int index))
 				{
 					WriteInt(index);
+					return false;
 				}
 				else
 				{
 					WriteInt(_objectToUniqueId.Count);
 					_objectToUniqueId.Add(obj, _objectToUniqueId.Count);
-					writeObject();
+					return true;
 				}
 			}
 		}
 
-		/// <summary>
-		/// Writes an object to the output. If the specific instance has already been written, preserves the reference to that.
-		/// </summary>
-		/// <param name="obj">The object to serialize</param>
-		/// <param name="writeObject">Delegate used to write the object</param>
-		public void WriteObjectReference<T>(T obj, Action writeObject) where T : class?
-		{
-			WriteObjectReference((object?)obj, writeObject);
-		}
 	}
 }
diff --git a/Engine/Source/Programs/Shared/EpicGames.Core/FileReference.cs b/Engine/Source/Programs/Shared/EpicGames.Core/FileReference.cs
--- a/Engine/Source/Programs/Shared/EpicGames.Core/FileReference.cs
+++ b/Engine/Source/Programs/Shared/EpicGames.Core/FileReference.cs
@@ -442,7 +442,27 @@ namespace EpicGames.Core
 		/// <returns>Contents of the file as a single string</returns>
 		public static string ReadAllText(FileReference location)
 		{
-			return File.ReadAllText(location.FullName);
+			using (FileStream fs = new FileStream(location.FullName, FileMode.Open, FileAccess.Read, FileShare.Read, 4 * 1024, FileOptions.SequentialScan))
+			{
+				using (StreamReader sr = new StreamReader(fs, Encoding.UTF8, true))
+				{
+
+					// Try to read the whole file into a buffer created by hand.  This avoids a LOT of memory allocations which in turn reduces the
+					// GC stress on the system.  Removing the StreamReader would be nice in the future.
+					long RawFileLength = fs.Length;
+					char[] InitialBuffer = new char[RawFileLength];
+					int ReadLength = sr.Read(InitialBuffer, 0, (int)RawFileLength);
+					if (sr.EndOfStream)
+					{
+						return new String(InitialBuffer, 0, ReadLength);
+					}
+					else
+					{
+						string Remaining = sr.ReadToEnd();
+						return String.Concat(new ReadOnlySpan<char>(InitialBuffer, 0, ReadLength), Remaining);
+					}
+				}
+			}
 		}
 
 		/// <summary>
