Fixed a large number of C# diagnostics.

#rnx

diff --git a/Engine/Source/Programs/Shared/EpicGames.UHT/.editorconfig b/Engine/Source/Programs/Shared/EpicGames.UHT/.editorconfig
--- a/Engine/Source/Programs/Shared/EpicGames.UHT/.editorconfig
+++ b/Engine/Source/Programs/Shared/EpicGames.UHT/.editorconfig
@@ -11,33 +11,6 @@ dotnet_diagnostic.IDE0051.severity = none
 # IDE0060: Remove unused parameter
 dotnet_diagnostic.IDE0060.severity = none
 
-# IDE0044: Add readonly modifier
-dotnet_diagnostic.IDE0044.severity = none
-
-# IDE0059: Unnecessary assignment of a value
-dotnet_diagnostic.IDE0059.severity = none
-
-# IDE2001: Embedded statements must be on their own line
-dotnet_diagnostic.IDE2001.severity = none
-
-# IDE0036: Order modifiers
-dotnet_diagnostic.IDE0036.severity = none
-
-# IDE2002: Consecutive braces must not have blank line between them
-dotnet_diagnostic.IDE2002.severity = none
-
-# IDE2000: C#
-dotnet_diagnostic.IDE2000.severity = none
-
-# IDE0005: Using directive is unnecessary.
-dotnet_diagnostic.IDE0005.severity = none
-
-# IDE0052: Remove unread private members
-dotnet_diagnostic.IDE0052.severity = none
-
-# IDE0011: Add braces
-dotnet_diagnostic.IDE0011.severity = none
-
 # CA1051: Do not declare visible instance fields
 dotnet_diagnostic.CA1051.severity = none
 
@@ -50,92 +23,8 @@ dotnet_diagnostic.CA1307.severity = none
 # CA1045: Do not pass types by reference
 dotnet_diagnostic.CA1045.severity = none
 
-# CA2251: Use 'string.Equals'
-dotnet_diagnostic.CA2251.severity = none
-
-# CA2254: Template should be a static expression
-dotnet_diagnostic.CA2254.severity = none
-
-# CA1822: Mark members as static
-dotnet_diagnostic.CA1822.severity = none
-
-# CA1834: Consider using 'StringBuilder.Append(char)' when applicable
-dotnet_diagnostic.CA1834.severity = none
-
-# CA1309: Use ordinal string comparison
-dotnet_diagnostic.CA1309.severity = none
-
 # CA1508: Avoid dead conditional code
 dotnet_diagnostic.CA1508.severity = none
 
-# CA1063: Implement IDisposable Correctly
-dotnet_diagnostic.CA1063.severity = none
-
-# CA1721: Property names should not match get methods
-dotnet_diagnostic.CA1721.severity = none
-
-# CA1813: Avoid unsealed attributes
-dotnet_diagnostic.CA1813.severity = none
-
-# CA2211: Non-constant fields should not be visible
-dotnet_diagnostic.CA2211.severity = none
-
-# CA1825: Avoid zero-length array allocations
-dotnet_diagnostic.CA1825.severity = none
-
-# CA1052: Static holder types should be Static or NotInheritable
-dotnet_diagnostic.CA1052.severity = none
-
-# CA1028: Enum Storage should be Int32
-dotnet_diagnostic.CA1028.severity = none
-
-# CA1707: Identifiers should not contain underscores
-dotnet_diagnostic.CA1707.severity = none
-
-# CA1806: Do not ignore method results
-dotnet_diagnostic.CA1806.severity = none
-
-# CA1716: Identifiers should not match keywords
-dotnet_diagnostic.CA1716.severity = none
-
-# CA2008: Do not create tasks without passing a TaskScheduler
-dotnet_diagnostic.CA2008.severity = none
-
-# CA1033: Interface methods should be callable by child types
-dotnet_diagnostic.CA1033.severity = none
-
-# CA1012: Abstract types should not have public constructors
-dotnet_diagnostic.CA1012.severity = none
-
-# CA1816: Dispose methods should call SuppressFinalize
-dotnet_diagnostic.CA1816.severity = none
-
-# CA1830: Prefer strongly-typed Append and Insert method overloads on StringBuilder
-dotnet_diagnostic.CA1830.severity = none
-
-# CA1001: Types that own disposable fields should be disposable
-dotnet_diagnostic.CA1001.severity = none
-
 # CA1501: 'UhtClassParser' has an object hierarchy '6' levels deep within the defining module. If possible, eliminate base classes within the hierarchy to decrease its hierarchy level below '6': 'UhtClassBaseParser, UhtClass, UhtStruct, UhtField, UhtObject, UhtType'.
 dotnet_diagnostic.CA1501.severity = none
-
-# CA2227: Collection properties should be read only
-dotnet_diagnostic.CA2227.severity = none
-
-# CA2208: Instantiate argument exceptions correctly
-dotnet_diagnostic.CA2208.severity = none
-
-# CA1008: Enums should have zero value
-dotnet_diagnostic.CA1008.severity = none
-
-# CA1820: Test for empty strings using string length
-dotnet_diagnostic.CA1820.severity = none
-
-# CA1823: Avoid unused private fields
-dotnet_diagnostic.CA1823.severity = none
-
-# CA2253: Named placeholders should not be numeric values
-dotnet_diagnostic.CA2253.severity = none
-
-# CA2241: Provide correct arguments to formatting methods
-dotnet_diagnostic.CA2241.severity = none
diff --git a/Engine/Source/Programs/Shared/EpicGames.UHT/Exporters/CodeGen/UhtHeaderCodeGenerator.cs b/Engine/Source/Programs/Shared/EpicGames.UHT/Exporters/CodeGen/UhtHeaderCodeGenerator.cs
--- a/Engine/Source/Programs/Shared/EpicGames.UHT/Exporters/CodeGen/UhtHeaderCodeGenerator.cs
+++ b/Engine/Source/Programs/Shared/EpicGames.UHT/Exporters/CodeGen/UhtHeaderCodeGenerator.cs
@@ -44,7 +44,7 @@ namespace EpicGames.UHT.Exporters.CodeGen
 		}
 
 		#region Event parameter
-		protected StringBuilder AppendEventParameter(StringBuilder Builder, UhtFunction Function, string StrippedFunctionName, UhtPropertyTextType TextType, bool bOutputConstructor, int Tabs, string Endl)
+		protected static StringBuilder AppendEventParameter(StringBuilder Builder, UhtFunction Function, string StrippedFunctionName, UhtPropertyTextType TextType, bool bOutputConstructor, int Tabs, string Endl)
 		{
 			if (!WillExportEventParameters(Function))
 			{
@@ -54,7 +54,7 @@ namespace EpicGames.UHT.Exporters.CodeGen
 			string EventParameterStructName = GetEventStructParametersName(Function.Outer, StrippedFunctionName);
 
 			Builder.AppendTabs(Tabs).Append("struct ").Append(EventParameterStructName).Append(Endl);
-			Builder.AppendTabs(Tabs).Append("{").Append(Endl);
+			Builder.AppendTabs(Tabs).Append('{').Append(Endl);
 
 			++Tabs;
 			foreach (UhtProperty Property in Function.Children)
@@ -85,7 +85,7 @@ namespace EpicGames.UHT.Exporters.CodeGen
 				}
 
 				Builder.AppendFullDecl(Property, TextType, false);
-				Builder.Append(";").Append(Endl);
+				Builder.Append(';').Append(Endl);
 			}
 
 			if (bOutputConstructor)
@@ -96,9 +96,9 @@ namespace EpicGames.UHT.Exporters.CodeGen
 					Builder.Append(Endl);
 					Builder.AppendTabs(Tabs).Append("/** Constructor, initializes return property only **/").Append(Endl);
 					Builder.AppendTabs(Tabs).Append(EventParameterStructName).Append("()").Append(Endl);
-					Builder.AppendTabs(Tabs + 1).Append(": ").Append(ReturnProperty.SourceName).Append('(').AppendNullConstructorArg(ReturnProperty, true).Append(")").Append(Endl);
-					Builder.AppendTabs(Tabs).Append("{").Append(Endl);
-					Builder.AppendTabs(Tabs).Append("}").Append(Endl);
+					Builder.AppendTabs(Tabs + 1).Append(": ").Append(ReturnProperty.SourceName).Append('(').AppendNullConstructorArg(ReturnProperty, true).Append(')').Append(Endl);
+					Builder.AppendTabs(Tabs).Append('{').Append(Endl);
+					Builder.AppendTabs(Tabs).Append('}').Append(Endl);
 				}
 			}
 			--Tabs;
@@ -271,9 +271,9 @@ namespace EpicGames.UHT.Exporters.CodeGen
 			return Builder;
 		}
 
-		protected StringBuilder AppendEventFunctionPrologue(StringBuilder Builder, UhtFunction Function, string FunctionName, int Tabs, string Endl)
+		protected static StringBuilder AppendEventFunctionPrologue(StringBuilder Builder, UhtFunction Function, string FunctionName, int Tabs, string Endl)
 		{
-			Builder.AppendTabs(Tabs).Append("{").Append(Endl);
+			Builder.AppendTabs(Tabs).Append('{').Append(Endl);
 			if (Function.Children.Count == 0)
 			{
 				return Builder;
@@ -297,18 +297,18 @@ namespace EpicGames.UHT.Exporters.CodeGen
 				}
 				else
 				{
-					Builder.AppendTabs(Tabs).Append("Parms.").Append(Property.SourceName).Append("=").Append(Property.SourceName);
+					Builder.AppendTabs(Tabs).Append("Parms.").Append(Property.SourceName).Append('=').Append(Property.SourceName);
 					if (Property is UhtBoolProperty)
 					{
 						Builder.Append(" ? true : false");
 					}
-					Builder.Append(";").Append(Endl);
+					Builder.Append(';').Append(Endl);
 				}
 			}
 			return Builder;
 		}
 
-		protected StringBuilder AppendEventFunctionEpilogue(StringBuilder Builder, UhtFunction Function, int Tabs, string Endl)
+		protected static StringBuilder AppendEventFunctionEpilogue(StringBuilder Builder, UhtFunction Function, int Tabs, string Endl)
 		{
 			++Tabs;
 
@@ -336,7 +336,7 @@ namespace EpicGames.UHT.Exporters.CodeGen
 							.Append(Property.SourceName)
 							.Append("=Parms.")
 							.Append(Property.SourceName)
-							.Append(";").Append(Endl);
+							.Append(';').Append(Endl);
 					}
 				}
 			}
@@ -352,7 +352,7 @@ namespace EpicGames.UHT.Exporters.CodeGen
 						.AppendTabs(Tabs)
 						.Append("return !!Parms.")
 						.Append(ReturnProperty.SourceName)
-						.Append(";").Append(Endl);
+						.Append(';').Append(Endl);
 				}
 				else
 				{
@@ -360,16 +360,16 @@ namespace EpicGames.UHT.Exporters.CodeGen
 						.AppendTabs(Tabs)
 						.Append("return Parms.")
 						.Append(ReturnProperty.SourceName)
-						.Append(";").Append(Endl);
+						.Append(';').Append(Endl);
 				}
 			}
 
 			--Tabs;
-			Builder.AppendTabs(Tabs).Append("}").Append(Endl);
+			Builder.AppendTabs(Tabs).Append('}').Append(Endl);
 			return Builder;
 		}
 
-		protected StringBuilder AppendParameters(StringBuilder Builder, UhtFunction Function, UhtPropertyTextType TextType, string? ExtraParameter, bool bSkipParameterName)
+		protected static StringBuilder AppendParameters(StringBuilder Builder, UhtFunction Function, UhtPropertyTextType TextType, string? ExtraParameter, bool bSkipParameterName)
 		{
 			bool bNeedsSeperator = false;
 
@@ -394,7 +394,7 @@ namespace EpicGames.UHT.Exporters.CodeGen
 				Builder.AppendFullDecl(Parameter, TextType, bSkipParameterName);
 			}
 
-			Builder.Append(")");
+			Builder.Append(')');
 			return Builder;
 		}
 
@@ -473,14 +473,14 @@ namespace EpicGames.UHT.Exporters.CodeGen
 		#endregion
 
 		#region Field notify support
-		protected bool NeedFieldNotifyCodeGen(UhtClass Class)
+		protected static bool NeedFieldNotifyCodeGen(UhtClass Class)
 		{
 			return 
 				!Class.ClassExportFlags.HasAnyFlags(UhtClassExportFlags.HasCustomFieldNotify) &&
 				Class.ClassExportFlags.HasAnyFlags(UhtClassExportFlags.HasFieldNotify);
 		}
 
-		protected void GetFieldNotifyStats(UhtClass Class, out bool bHasProperties, out bool bHasFunctions, out bool bHasEditorFields, out bool bAllEditorFields)
+		protected static void GetFieldNotifyStats(UhtClass Class, out bool bHasProperties, out bool bHasFunctions, out bool bHasEditorFields, out bool bAllEditorFields)
 		{
 			// Scan the children to see what we have
 			bHasProperties = false;
@@ -513,7 +513,7 @@ namespace EpicGames.UHT.Exporters.CodeGen
 			bAllEditorFields &= bHasEditorFields;
 		}
 
-		protected StringBuilder AppendFieldNotify(StringBuilder Builder, UhtClass Class,
+		protected static StringBuilder AppendFieldNotify(StringBuilder Builder, UhtClass Class,
 			bool bHasProperties, bool bHasFunctions, bool bHasEditorFields, bool bAllEditorFields,
 			bool bIncludeEditorOnlyFields, bool bAppendDefine, Action<StringBuilder, UhtClass, string> AppendAction)
 		{
diff --git a/Engine/Source/Programs/Shared/EpicGames.UHT/Exporters/CodeGen/UhtHeaderCodeGeneratorCppFile.cs b/Engine/Source/Programs/Shared/EpicGames.UHT/Exporters/CodeGen/UhtHeaderCodeGeneratorCppFile.cs
--- a/Engine/Source/Programs/Shared/EpicGames.UHT/Exporters/CodeGen/UhtHeaderCodeGeneratorCppFile.cs
+++ b/Engine/Source/Programs/Shared/EpicGames.UHT/Exporters/CodeGen/UhtHeaderCodeGeneratorCppFile.cs
@@ -337,7 +337,7 @@ namespace EpicGames.UHT.Exporters.CodeGen
 									Builder.Append("||");
 								}
 								bDoneFirst = true;
-								Builder.Append("!int64(").Append(Value.Name).Append(")");
+								Builder.Append("!int64(").Append(Value.Name).Append(')');
 							}
 						}
 						Builder.Append(", \"'").Append(Enum.SourceName).Append("' does not have a 0 entry!(This is a problem when the enum is initalized by default)\");\r\n");
@@ -641,7 +641,7 @@ namespace EpicGames.UHT.Exporters.CodeGen
 							continue;
 						}
 						bWroteLine = true;
-						Builder.Append("\t").Append(Parameter.CastType).Append(' ').Append(Parameter.CastName).Append('(').Append(Parameter.Name).Append(");\r\n");
+						Builder.Append('\t').Append(Parameter.CastType).Append(' ').Append(Parameter.CastName).Append('(').Append(Parameter.Name).Append(");\r\n");
 					}
 					if (bWroteLine)
 					{
@@ -650,7 +650,7 @@ namespace EpicGames.UHT.Exporters.CodeGen
 					}
 
 					//COMPATIBILITY-TODO - Replace spaces with \t
-					Builder.Append("\t").Append(MethodInfo.ReturnPrefix()).Append("Static").Append(MethodInfo.Name).Append("(\r\n");
+					Builder.Append('\t').Append(MethodInfo.ReturnPrefix()).Append("Static").Append(MethodInfo.Name).Append("(\r\n");
 					Builder.Append("\t\tRigVMExecuteContext");
 					Builder.AppendParameterNames(ScriptStruct.RigVMStructInfo.Members, true, ",\r\n\t\t", true);
 					Builder.AppendParameterNames(MethodInfo.Parameters, true, ",\r\n\t\t");
@@ -734,16 +734,16 @@ namespace EpicGames.UHT.Exporters.CodeGen
 
 					if (bParamsInStatic)
 					{
-						Builder.Append("sizeof(").Append(StaticsName).Append("::").Append(GetEventStructParametersName(TempFunction.Outer, TempFunction.StrippedFunctionName)).Append(")");
+						Builder.Append("sizeof(").Append(StaticsName).Append("::").Append(GetEventStructParametersName(TempFunction.Outer, TempFunction.StrippedFunctionName)).Append(')');
 					}
 					else
 					{
-						Builder.Append("sizeof(").Append(GetEventStructParametersName(TempFunction.Outer, TempFunction.StrippedFunctionName)).Append(")");
+						Builder.Append("sizeof(").Append(GetEventStructParametersName(TempFunction.Outer, TempFunction.StrippedFunctionName)).Append(')');
 					}
 				}
 				else
 				{
-					Builder.Append("0");
+					Builder.Append('0');
 				}
 				Builder.Append(", ");
 
@@ -795,7 +795,7 @@ namespace EpicGames.UHT.Exporters.CodeGen
 			}
 		}
 
-		private StringBuilder AppendMirrorsForNoexportStruct(StringBuilder Builder, UhtScriptStruct NoExportStruct, int Tabs)
+		private static StringBuilder AppendMirrorsForNoexportStruct(StringBuilder Builder, UhtScriptStruct NoExportStruct, int Tabs)
 		{
 			Builder.AppendTabs(Tabs).Append("struct ").Append(NoExportStruct.SourceName);
 			if (NoExportStruct.SuperScriptStruct != null)
@@ -813,7 +813,7 @@ namespace EpicGames.UHT.Exporters.CodeGen
 			return Builder;
 		}
 
-		private StringBuilder AppendExportProperties(StringBuilder Builder, UhtScriptStruct ScriptStruct, int Tabs)
+		private static StringBuilder AppendExportProperties(StringBuilder Builder, UhtScriptStruct ScriptStruct, int Tabs)
 		{
 			using (UhtMacroBlockEmitter Emitter = new UhtMacroBlockEmitter(Builder, "WITH_EDITORONLY_DATA"))
 			{
@@ -967,7 +967,7 @@ namespace EpicGames.UHT.Exporters.CodeGen
 			return Builder;
 		}
 
-		private StringBuilder AppendFieldNotify(StringBuilder Builder, UhtClass Class)
+		private static StringBuilder AppendFieldNotify(StringBuilder Builder, UhtClass Class)
 		{
 			if (!NeedFieldNotifyCodeGen(Class))
 			{
@@ -1009,7 +1009,7 @@ namespace EpicGames.UHT.Exporters.CodeGen
 			return Builder;
 		}
 
-		private StringBuilder AppendPropertyAccessors(StringBuilder Builder, UhtClass Class)
+		private static StringBuilder AppendPropertyAccessors(StringBuilder Builder, UhtClass Class)
 		{
 			foreach (UhtType Type in Class.Children)
 			{
@@ -1113,7 +1113,7 @@ namespace EpicGames.UHT.Exporters.CodeGen
 		{
 			foreach (UhtFunction Function in CallbackFunctions)
 			{
-				Builder.Append("\tstatic FName NAME_").Append(Function.Outer?.SourceName).Append("_").Append(Function.SourceName).Append(" = FName(TEXT(\"").Append(Function.EngineName).Append("\"));\r\n");
+				Builder.Append("\tstatic FName NAME_").Append(Function.Outer?.SourceName).Append('_').Append(Function.SourceName).Append(" = FName(TEXT(\"").Append(Function.EngineName).Append("\"));\r\n");
 				string ExtraParameter = Function.FunctionFlags.HasAnyFlags(EFunctionFlags.Const) ? "const UObject* O" : "UObject* O";
 				AppendNativeFunctionHeader(Builder, Function, UhtPropertyTextType.InterfaceFunctionArgOrRetVal, false, null, ExtraParameter, 0, "\r\n");
 				Builder.Append("\t{\r\n");
@@ -1123,12 +1123,12 @@ namespace EpicGames.UHT.Exporters.CodeGen
 				{
 					Builder.Append("\t\t").Append(GetEventStructParametersName(Class, Function.StrippedFunctionName)).Append(" Parms;\r\n");
 				}
-				Builder.Append("\t\tUFunction* const Func = O->FindFunction(NAME_").Append(Function.Outer?.SourceName).Append("_").Append(Function.SourceName).Append(");\r\n");
+				Builder.Append("\t\tUFunction* const Func = O->FindFunction(NAME_").Append(Function.Outer?.SourceName).Append('_').Append(Function.SourceName).Append(");\r\n");
 				Builder.Append("\t\tif (Func)\r\n");
 				Builder.Append("\t\t{\r\n");
 				foreach (UhtProperty Property in Function.ParameterProperties.Span)
 				{
-					Builder.Append("\t\t\tParms.").Append(Property.SourceName).Append("=").Append(Property.SourceName).Append(";\r\n");
+					Builder.Append("\t\t\tParms.").Append(Property.SourceName).Append('=').Append(Property.SourceName).Append(";\r\n");
 				}
 				Builder
 					.Append("\t\t\t")
@@ -1168,7 +1168,7 @@ namespace EpicGames.UHT.Exporters.CodeGen
 					{
 						if (!bFirst)
 						{
-							Builder.Append(",");
+							Builder.Append(',');
 						}
 						bFirst = false;
 						Builder.Append(Property.SourceName);
@@ -1186,7 +1186,7 @@ namespace EpicGames.UHT.Exporters.CodeGen
 			return Builder;
 		}
 
-		private StringBuilder AppendSerializer(StringBuilder Builder, UhtClass Class, UhtSerializerArchiveType SerializerType, string MacroText)
+		private static StringBuilder AppendSerializer(StringBuilder Builder, UhtClass Class, UhtSerializerArchiveType SerializerType, string MacroText)
 		{
 			if (!Class.SerializerArchiveType.HasAnyFlags(SerializerType))
 			{
@@ -1194,7 +1194,7 @@ namespace EpicGames.UHT.Exporters.CodeGen
 				{
 					Builder.Append("#if ").Append(Class.EnclosingDefine).Append("\r\n");
 				}
-				Builder.Append("\t").Append(MacroText).Append('(').Append(Class.SourceName).Append(")\r\n");
+				Builder.Append('\t').Append(MacroText).Append('(').Append(Class.SourceName).Append(")\r\n");
 				if (Class.EnclosingDefine.Length > 0)
 				{
 					Builder.Append("#endif\r\n");
@@ -1213,7 +1213,7 @@ namespace EpicGames.UHT.Exporters.CodeGen
 
 			PropertyMemberContextImpl Context = new PropertyMemberContextImpl(this.CodeGenerator, Class, Class.SourceName, StaticsName);
 
-			bool bHasInterfaces = Class.Bases != null && Class.Bases.Any(x => x is UhtClass BaseClass ? BaseClass.ClassFlags.HasAnyFlags(EClassFlags.Interface) : false);
+			bool bHasInterfaces = Class.Bases.Any(x => x is UhtClass BaseClass ? BaseClass.ClassFlags.HasAnyFlags(EClassFlags.Interface) : false);
 
 			// Collect the functions to be exported
 			bool bAllEditorOnlyFunctions = true;
@@ -1317,7 +1317,7 @@ namespace EpicGames.UHT.Exporters.CodeGen
 
 				AppendPropertiesDefs(Builder, Class, Class.SourceName, StaticsName, 1);
 
-				if (Class.Bases != null && bHasInterfaces)
+				if (bHasInterfaces)
 				{
 					Builder.Append("\t\tconst UECodeGen_Private::FImplementedInterfaceParams ").Append(StaticsName).Append("::InterfaceParams[] = {\r\n");
 
@@ -1525,7 +1525,7 @@ namespace EpicGames.UHT.Exporters.CodeGen
 			return Builder;
 		}
 
-		private StringBuilder AppendNatives(StringBuilder Builder, UhtClass Class)
+		private static StringBuilder AppendNatives(StringBuilder Builder, UhtClass Class)
 		{
 			Builder.Append("\tvoid ").Append(Class.SourceName).Append("::StaticRegisterNatives").Append(Class.SourceName).Append("()\r\n");
 			Builder.Append("\t{\r\n");
@@ -1646,7 +1646,7 @@ namespace EpicGames.UHT.Exporters.CodeGen
 			return Builder;
 		}
 
-		private StringBuilder AppendRpcFunctions(StringBuilder Builder, UhtClass Class, List<UhtFunction> ReversedFunctions, bool bEditorOnly)
+		private static StringBuilder AppendRpcFunctions(StringBuilder Builder, UhtClass Class, List<UhtFunction> ReversedFunctions, bool bEditorOnly)
 		{
 			foreach (UhtFunction Function in ReversedFunctions)
 			{
@@ -1661,7 +1661,7 @@ namespace EpicGames.UHT.Exporters.CodeGen
 			return Builder;
 		}
 
-		private StringBuilder AppendFunctionThunk(StringBuilder Builder, UhtClass Class, UhtFunction Function)
+		private static StringBuilder AppendFunctionThunk(StringBuilder Builder, UhtClass Class, UhtFunction Function)
 		{
 			// Export the GET macro for the parameters
 			foreach (UhtProperty Parameter in Function.ParameterProperties.Span)
@@ -1675,7 +1675,7 @@ namespace EpicGames.UHT.Exporters.CodeGen
 			// Call the validate function if there is one
 			if (!Function.FunctionExportFlags.HasAnyFlags(UhtFunctionExportFlags.CppStatic) && Function.FunctionFlags.HasAnyFlags(EFunctionFlags.NetValidate))
 			{
-				Builder.Append("\t\tif (!P_THIS->").Append(Function.CppValidationImplName).Append("(").AppendFunctionThunkParameterNames(Function).Append("))\r\n");
+				Builder.Append("\t\tif (!P_THIS->").Append(Function.CppValidationImplName).Append('(').AppendFunctionThunkParameterNames(Function).Append("))\r\n");
 				Builder.Append("\t\t{\r\n");
 				Builder.Append("\t\t\tRPC_ValidateFailed(TEXT(\"").Append(Function.CppValidationImplName).Append("\"));\r\n");
 				Builder.Append("\t\t\treturn;\r\n");   // If we got here, the validation function check failed
@@ -1693,11 +1693,11 @@ namespace EpicGames.UHT.Exporters.CodeGen
 			// Export the call to the C++ version
 			if (Function.FunctionExportFlags.HasAnyFlags(UhtFunctionExportFlags.CppStatic))
 			{
-				Builder.Append(Function.Outer?.SourceName).Append("::").Append(Function.CppImplName).Append("(").AppendFunctionThunkParameterNames(Function).Append(");\r\n");
+				Builder.Append(Function.Outer?.SourceName).Append("::").Append(Function.CppImplName).Append('(').AppendFunctionThunkParameterNames(Function).Append(");\r\n");
 			}
 			else
 			{
-				Builder.Append("P_THIS->").Append(Function.CppImplName).Append("(").AppendFunctionThunkParameterNames(Function).Append(");\r\n");
+				Builder.Append("P_THIS->").Append(Function.CppImplName).Append('(').AppendFunctionThunkParameterNames(Function).Append(");\r\n");
 			}
 			Builder.Append("\t\tP_NATIVE_END;\r\n");
 			return Builder;
@@ -1903,11 +1903,9 @@ namespace EpicGames.UHT.Exporters.CodeGen
 					.Append(StaticsName).Append("::").Append(ArrayName)
 					.Append(", UE_ARRAY_COUNT(")
 					.Append(StaticsName).Append("::").Append(ArrayName)
-					.Append(")");
+					.Append(')');
 			}
 			return Builder;
 		}
-
 	}
-
 }
diff --git a/Engine/Source/Programs/Shared/EpicGames.UHT/Exporters/CodeGen/UhtHeaderCodeGeneratorHFile.cs b/Engine/Source/Programs/Shared/EpicGames.UHT/Exporters/CodeGen/UhtHeaderCodeGeneratorHFile.cs
--- a/Engine/Source/Programs/Shared/EpicGames.UHT/Exporters/CodeGen/UhtHeaderCodeGeneratorHFile.cs
+++ b/Engine/Source/Programs/Shared/EpicGames.UHT/Exporters/CodeGen/UhtHeaderCodeGeneratorHFile.cs
@@ -127,7 +127,7 @@ namespace EpicGames.UHT.Exporters.CodeGen
 				foreach (UhtRigVMMethodInfo MethodInfo in ScriptStruct.RigVMStructInfo.Methods)
 				{
 					Builder.Append("#define ").Append(ScriptStruct.SourceName).Append('_').Append(MethodInfo.Name).Append("() \\\r\n");
-					Builder.Append("\t").Append(MethodInfo.ReturnType).Append(' ').Append(ScriptStruct.SourceName).Append("::Static").Append(MethodInfo.Name).Append("( \\\r\n");
+					Builder.Append('\t').Append(MethodInfo.ReturnType).Append(' ').Append(ScriptStruct.SourceName).Append("::Static").Append(MethodInfo.Name).Append("( \\\r\n");
 					Builder.Append("\t\t").Append(RigVMExecuteContextPublicDeclaration);
 					Builder.AppendParameterDecls(ScriptStruct.RigVMStructInfo.Members, true, ", \\\r\n\t\t", true, false);
 					Builder.AppendParameterDecls(MethodInfo.Parameters, true, ", \\\r\n\t\t", false, false);
@@ -142,7 +142,7 @@ namespace EpicGames.UHT.Exporters.CodeGen
 				using (UhtMacroCreator Macro = new UhtMacroCreator(Builder, this, ScriptStruct, GeneratedBodyMacroSuffix))
 				{
 					Builder.Append("\tfriend struct Z_Construct_UScriptStruct_").Append(ScriptStruct.SourceName).Append("_Statics; \\\r\n");
-					Builder.Append("\t");
+					Builder.Append('\t');
 					if (!ScriptStruct.ScriptStructFlags.HasAnyFlags(EStructFlags.RequiredAPI))
 					{
 						Builder.Append(this.PackageApi);
@@ -302,7 +302,7 @@ namespace EpicGames.UHT.Exporters.CodeGen
 				AppendNativeFunctionHeader(Builder, Function, UhtPropertyTextType.EventFunctionArgOrRetVal, true, ExportFunctionName, ExtraParameter, UhtFunctionExportFlags.Inline, " \\\r\n");
 				AppendEventFunctionPrologue(Builder, Function, StrippedFunctionName, Tabs, " \\\r\n");
 				Builder
-					.Append("\t")
+					.Append('\t')
 					.Append(StrippedFunctionName)
 					.Append('.')
 					.Append(Function.FunctionFlags.HasAnyFlags(EFunctionFlags.MulticastDelegate) ? "ProcessMulticastDelegate" : "ProcessDelegate")
@@ -417,7 +417,7 @@ namespace EpicGames.UHT.Exporters.CodeGen
 			{
 				if (CallbackFunctions.Count > 0)
 				{
-					Builder.Append("\t").AppendMacroName(this, Class, EventParamsMacroSuffix).Append(" \\\r\n");
+					Builder.Append('\t').AppendMacroName(this, Class, EventParamsMacroSuffix).Append(" \\\r\n");
 				}
 			}
 
@@ -481,7 +481,7 @@ namespace EpicGames.UHT.Exporters.CodeGen
 			return Builder;
 		}
 
-		private StringBuilder AppendFieldNotify(StringBuilder Builder, UhtClass Class, 
+		private static StringBuilder AppendFieldNotify(StringBuilder Builder, UhtClass Class, 
 			bool bHasProperties, bool bHasFunctions, bool bHasEditorFields, bool bAllEditorFields,
 			bool bIncludeEditorOnlyFields)
 		{
@@ -525,12 +525,12 @@ namespace EpicGames.UHT.Exporters.CodeGen
 				{
 					foreach (string SparseDataType in SparseDataTypes)
 					{
-						Builder.Append("F").Append(SparseDataType).Append("* Get").Append(SparseDataType).Append("() \\\r\n");
+						Builder.Append('F').Append(SparseDataType).Append("* Get").Append(SparseDataType).Append("() \\\r\n");
 						Builder.Append("{ \\\r\n");
 						Builder.Append("\treturn (F").Append(SparseDataType).Append("*)(GetClass()->GetOrCreateSparseClassData()); \\\r\n");
 						Builder.Append("} \\\r\n");
 
-						Builder.Append("F").Append(SparseDataType).Append("* Get").Append(SparseDataType).Append("() const \\\r\n");
+						Builder.Append('F').Append(SparseDataType).Append("* Get").Append(SparseDataType).Append("() const \\\r\n");
 						Builder.Append("{ \\\r\n");
 						Builder.Append("\treturn (F").Append(SparseDataType).Append("*)(GetClass()->GetOrCreateSparseClassData()); \\\r\n");
 						Builder.Append("} \\\r\n");
@@ -680,11 +680,11 @@ namespace EpicGames.UHT.Exporters.CodeGen
 			return Builder;
 		}
 
-		private StringBuilder AppendNetValidateDeclaration(StringBuilder Builder, UhtFunction Function)
+		private static StringBuilder AppendNetValidateDeclaration(StringBuilder Builder, UhtFunction Function)
 		{
 			if (Function.FunctionFlags.HasAnyFlags(EFunctionFlags.NetValidate))
 			{
-				Builder.Append("\t");
+				Builder.Append('\t');
 				if (!Function.FunctionFlags.HasAnyFlags(EFunctionFlags.Static) && !Function.FunctionExportFlags.HasAnyFlags(UhtFunctionExportFlags.Final))
 				{
 					Builder.Append("virtual");
@@ -696,7 +696,7 @@ namespace EpicGames.UHT.Exporters.CodeGen
 			return Builder;
 		}
 
-		private StringBuilder AppendRpcWrappers(StringBuilder Builder, UhtClass Class, List<UhtFunction> ReversedFunctions, bool bEditorOnly)
+		private static StringBuilder AppendRpcWrappers(StringBuilder Builder, UhtClass Class, List<UhtFunction> ReversedFunctions, bool bEditorOnly)
 		{
 			bool bFirst = true;
 			foreach (UhtFunction Function in ReversedFunctions)
@@ -799,7 +799,7 @@ namespace EpicGames.UHT.Exporters.CodeGen
 				if (!Class.ClassExportFlags.HasAnyFlags(UhtClassExportFlags.HasCustomConstructor))
 				{
 					Builder.Append("\t/** Standard constructor, called after all reflected properties have been initialized */ \\\r\n");
-					Builder.Append("\t").Append(Api).Append(Class.SourceName).Append("(const FObjectInitializer& ObjectInitializer");
+					Builder.Append('\t').Append(Api).Append(Class.SourceName).Append("(const FObjectInitializer& ObjectInitializer");
 					if (!Class.ClassExportFlags.HasAnyFlags(UhtClassExportFlags.HasDefaultConstructor))
 					{
 						Builder.Append(" = FObjectInitializer::Get()");
@@ -846,7 +846,7 @@ namespace EpicGames.UHT.Exporters.CodeGen
 		/// <param name="Class">Class being exported</param>
 		/// <param name="Api">API text to be used</param>
 		/// <returns>Output builder</returns>
-		private StringBuilder AppendVTableHelperCtorAndCaller(StringBuilder Builder, UhtClass Class, string Api)
+		private static StringBuilder AppendVTableHelperCtorAndCaller(StringBuilder Builder, UhtClass Class, string Api)
 		{
 			if (!Class.ClassExportFlags.HasAnyFlags(UhtClassExportFlags.HasCustomVTableHelperConstructor))
 			{
@@ -863,12 +863,12 @@ namespace EpicGames.UHT.Exporters.CodeGen
 		/// <param name="Class">Class being exported</param>
 		/// <param name="Api">API text to be used</param>
 		/// <returns>Output builder</returns>
-		private StringBuilder AppendCopyConstructorDefinition(StringBuilder Builder, UhtClass Class, string Api)
+		private static StringBuilder AppendCopyConstructorDefinition(StringBuilder Builder, UhtClass Class, string Api)
 		{
 			Builder.Append("private: \\\r\n");
 			Builder.Append("\t/** Private move- and copy-constructors, should never be used */ \\\r\n");
-			Builder.Append("\t").Append(Api).Append(Class.SourceName).Append('(').Append(Class.SourceName).Append("&&); \\\r\n");
-			Builder.Append("\t").Append(Api).Append(Class.SourceName).Append("(const ").Append(Class.SourceName).Append("&); \\\r\n");
+			Builder.Append('\t').Append(Api).Append(Class.SourceName).Append('(').Append(Class.SourceName).Append("&&); \\\r\n");
+			Builder.Append('\t').Append(Api).Append(Class.SourceName).Append("(const ").Append(Class.SourceName).Append("&); \\\r\n");
 			Builder.Append("public: \\\r\n");
 			return Builder;
 		}
@@ -880,7 +880,7 @@ namespace EpicGames.UHT.Exporters.CodeGen
 		/// <param name="Class">Class being exported</param>
 		/// <param name="Api">API text to be used</param>
 		/// <returns>Output builder</returns>
-		private StringBuilder AppendConstructorDefinition(StringBuilder Builder, UhtClass Class, string Api)
+		private static StringBuilder AppendConstructorDefinition(StringBuilder Builder, UhtClass Class, string Api)
 		{
 			if (!Class.ClassExportFlags.HasAnyFlags(UhtClassExportFlags.HasConstructor))
 			{
@@ -900,13 +900,13 @@ namespace EpicGames.UHT.Exporters.CodeGen
 
 				if (bSuperClassObjectInitializerConstructorDeclared)
 				{
-					Builder.Append("\t").Append(Api).Append(Class.SourceName).Append("(const FObjectInitializer& ObjectInitializer = FObjectInitializer::Get()) : Super(ObjectInitializer) { }; \\\r\n");
+					Builder.Append('\t').Append(Api).Append(Class.SourceName).Append("(const FObjectInitializer& ObjectInitializer = FObjectInitializer::Get()) : Super(ObjectInitializer) { }; \\\r\n");
 					//ETSTODO - Modification of class
 					Class.ClassExportFlags |= UhtClassExportFlags.HasObjectInitializerConstructor;
 				}
 				else
 				{
-					Builder.Append("\t").Append(Api).Append(Class.SourceName).Append("() { }; \\\r\n");
+					Builder.Append('\t').Append(Api).Append(Class.SourceName).Append("() { }; \\\r\n");
 					//ETSTODO - Modification of class
 					Class.ClassExportFlags |= UhtClassExportFlags.HasDefaultConstructor;
 				}
@@ -927,7 +927,7 @@ namespace EpicGames.UHT.Exporters.CodeGen
 		/// <param name="Class">Class being exported</param>
 		/// <param name="Api">API text to be used</param>
 		/// <returns>Output builder</returns>
-		private StringBuilder AppendDefaultConstructorCallDefinition(StringBuilder Builder, UhtClass Class, string Api)
+		private static StringBuilder AppendDefaultConstructorCallDefinition(StringBuilder Builder, UhtClass Class, string Api)
 		{
 			if (Class.ClassExportFlags.HasAnyFlags(UhtClassExportFlags.HasObjectInitializerConstructor))
 			{
@@ -939,7 +939,6 @@ namespace EpicGames.UHT.Exporters.CodeGen
 				{
 					Builder.Append("\tDEFINE_DEFAULT_OBJECT_INITIALIZER_CONSTRUCTOR_CALL(").Append(Class.SourceName).Append(") \\\r\n");
 				}
-
 			}
 			else if (Class.ClassExportFlags.HasAnyFlags(UhtClassExportFlags.HasDefaultConstructor))
 			{
@@ -978,17 +977,14 @@ namespace EpicGames.UHT.Exporters.CodeGen
 			}
 
 			// export implementation of _getUObject for classes that implement interfaces
-			if (Class.Bases != null)
+			foreach (UhtStruct BaseStruct in Class.Bases)
 			{
-				foreach (UhtStruct BaseStruct in Class.Bases)
+				if (BaseStruct is UhtClass BaseClass)
 				{
-					if (BaseStruct is UhtClass BaseClass)
+					if (BaseClass.ClassFlags.HasAnyFlags(EClassFlags.Interface))
 					{
-						if (BaseClass.ClassFlags.HasAnyFlags(EClassFlags.Interface))
-						{
-							Builder.Append("\tvirtual UObject* _getUObject() const override { return const_cast<").Append(Class.SourceName).Append("*>(this); } \\\r\n");
-							break;
-						}
+						Builder.Append("\tvirtual UObject* _getUObject() const override { return const_cast<").Append(Class.SourceName).Append("*>(this); } \\\r\n");
+						break;
 					}
 				}
 			}
@@ -1049,7 +1045,7 @@ namespace EpicGames.UHT.Exporters.CodeGen
 			// Add the serialization function declaration if we generated one
 			if (Class.SerializerArchiveType != UhtSerializerArchiveType.None && Class.SerializerArchiveType != UhtSerializerArchiveType.All)
 			{
-				Builder.Append("\t").AppendMacroName(this, Class, ArchiveSerializerMacroSuffix).Append(" \\\r\n");
+				Builder.Append('\t').AppendMacroName(this, Class, ArchiveSerializerMacroSuffix).Append(" \\\r\n");
 			}
 
 			if (SuperClass != null && Class.ClassWithin != SuperClass.ClassWithin)
@@ -1066,7 +1062,7 @@ namespace EpicGames.UHT.Exporters.CodeGen
 		/// <param name="Builder">Output builder</param>
 		/// <param name="Class">Class in question</param>
 		/// <returns>Output builder</returns>
-		private StringBuilder AppendClassFlags(StringBuilder Builder, UhtClass Class)
+		private static StringBuilder AppendClassFlags(StringBuilder Builder, UhtClass Class)
 		{
 			if (Class.ClassFlags.HasAnyFlags(EClassFlags.Transient))
 			{
@@ -1109,14 +1105,14 @@ namespace EpicGames.UHT.Exporters.CodeGen
 		/// <param name="Builder">Output builder</param>
 		/// <param name="MacroName">Name ofthe macro to deprecate</param>
 		/// <returns>Output builder</returns>
-		private StringBuilder AppendGeneratedMacroDeprecationWarning(StringBuilder Builder, string MacroName)
+		private static StringBuilder AppendGeneratedMacroDeprecationWarning(StringBuilder Builder, string MacroName)
 		{
 			// Deprecation warning is disabled right now. After people get familiar with the new macro it should be re-enabled.
 			//Builder.Append("EMIT_DEPRECATED_WARNING_MESSAGE(\"").Append(MacroName).Append("() macro is deprecated. Please use GENERATED_BODY() macro instead.\") \\\r\n");
 			return Builder;
 		}
 
-		private StringBuilder AppendAccessSpecifier(StringBuilder Builder, UhtClass Class)
+		private static StringBuilder AppendAccessSpecifier(StringBuilder Builder, UhtClass Class)
 		{
 			switch (Class.GeneratedBodyAccessSpecifier)
 			{
@@ -1139,7 +1135,6 @@ namespace EpicGames.UHT.Exporters.CodeGen
 		private StringBuilder AppendInClassIInterface(StringBuilder Builder, UhtClass Class, List<UhtFunction> CallbackFunctions, string Api)
 		{
 			string InterfaceSourceName = "I" + Class.EngineName;
-			string SuperInterfaceSourceName = Class.SuperClass != null ? "I" + Class.SuperClass.EngineName : "";
 
 			Builder.Append("protected: \\\r\n");
 			Builder.Append("\tvirtual ~").Append(InterfaceSourceName).Append("() {} \\\r\n");
@@ -1166,7 +1161,7 @@ namespace EpicGames.UHT.Exporters.CodeGen
 			return Builder;
 		}
 
-		private StringBuilder AppendReplicatedMacroData(StringBuilder Builder, UhtClass Class, string Api)
+		private static StringBuilder AppendReplicatedMacroData(StringBuilder Builder, UhtClass Class, string Api)
 		{
 			if (!Class.ClassExportFlags.HasAnyFlags(UhtClassExportFlags.HasGetLifetimeReplicatedProps))
 			{
@@ -1190,7 +1185,7 @@ namespace EpicGames.UHT.Exporters.CodeGen
 			return Builder;
 		}
 
-		private StringBuilder AppendNetData(StringBuilder Builder, UhtClass Class, string Api)
+		private static StringBuilder AppendNetData(StringBuilder Builder, UhtClass Class, string Api)
 		{
 			bool bHasArray = false;
 			foreach (UhtProperty Property in Class.EnumerateReplicatedProperties(false))
@@ -1266,7 +1261,7 @@ namespace EpicGames.UHT.Exporters.CodeGen
 			}
 			Builder.Append("\t}; \\\r\n");
 
-			Builder.Append("\t").Append(Api).Append("virtual void ValidateGeneratedRepEnums(const TArray<struct FRepRecord>& ClassReps) const override; \\\r\n");
+			Builder.Append('\t').Append(Api).Append("virtual void ValidateGeneratedRepEnums(const TArray<struct FRepRecord>& ClassReps) const override; \\\r\n");
 
 			return Builder;
 		}
@@ -1369,8 +1364,6 @@ namespace EpicGames.UHT.Exporters.CodeGen
 				}
 				Builder.Append(EnableDeprecationWarnings).Append(" \\\r\n");
 			}
-
-
 			return Builder;
 		}
 
diff --git a/Engine/Source/Programs/Shared/EpicGames.UHT/Exporters/CodeGen/UhtPackageCodeGenerator.cs b/Engine/Source/Programs/Shared/EpicGames.UHT/Exporters/CodeGen/UhtPackageCodeGenerator.cs
--- a/Engine/Source/Programs/Shared/EpicGames.UHT/Exporters/CodeGen/UhtPackageCodeGenerator.cs
+++ b/Engine/Source/Programs/Shared/EpicGames.UHT/Exporters/CodeGen/UhtPackageCodeGenerator.cs
@@ -128,6 +128,7 @@ namespace EpicGames.UHT.Exporters.CodeGen
 			uint B = 0x9e3779b9;
 			A += B;
 
+#pragma warning disable IDE2001 // Embedded statements must be on their own line
 			A -= B; A -= C; A ^= (C >> 13);
 			B -= C; B -= A; B ^= (A << 8);
 			C -= A; C -= B; C ^= (B >> 13);
@@ -137,6 +138,7 @@ namespace EpicGames.UHT.Exporters.CodeGen
 			A -= B; A -= C; A ^= (C >> 3);
 			B -= C; B -= A; B ^= (A << 10);
 			C -= A; C -= B; C ^= (B >> 15);
+#pragma warning restore IDE2001 // Embedded statements must be on their own line
 
 			return C;
 		}
diff --git a/Engine/Source/Programs/Shared/EpicGames.UHT/Exporters/CodeGen/UhtPackageCodeGeneratorCppFile.cs b/Engine/Source/Programs/Shared/EpicGames.UHT/Exporters/CodeGen/UhtPackageCodeGeneratorCppFile.cs
--- a/Engine/Source/Programs/Shared/EpicGames.UHT/Exporters/CodeGen/UhtPackageCodeGeneratorCppFile.cs
+++ b/Engine/Source/Programs/Shared/EpicGames.UHT/Exporters/CodeGen/UhtPackageCodeGeneratorCppFile.cs
@@ -91,7 +91,7 @@ namespace EpicGames.UHT.Exporters.CodeGen
 				if (this.Session.bIncludeDebugOutput)
 				{
 					Builder.Append("#if 0\r\n");
-					Builder.Append(Declarations.ToString());
+					Builder.Append(Declarations);
 					Builder.Append("#endif\r\n");
 				}
 
@@ -135,7 +135,7 @@ namespace EpicGames.UHT.Exporters.CodeGen
 				Builder.Append("\t}\r\n");
 
 				// Do not change the Z_CompiledInDeferPackage_UPackage_ without changing LC_SymbolPatterns
-				Builder.Append("\tstatic FRegisterCompiledInInfo Z_CompiledInDeferPackage_UPackage_").Append(StrippedName).Append("(").Append(SingletonName)
+				Builder.Append("\tstatic FRegisterCompiledInInfo Z_CompiledInDeferPackage_UPackage_").Append(StrippedName).Append('(').Append(SingletonName)
 					.Append(", TEXT(\"").Append(this.Package.SourceName).Append("\"), Z_Registration_Info_UPackage_").Append(StrippedName).Append(", CONSTRUCT_RELOAD_VERSION_INFO(FPackageReloadVersionInfo, ")
 					.Append($"0x{BodiesHash:X8}, 0x{DeclarationsHash:X8}").Append("));\r\n");
 
diff --git a/Engine/Source/Programs/Shared/EpicGames.UHT/Exporters/Stats/UhtStatsExporter.cs b/Engine/Source/Programs/Shared/EpicGames.UHT/Exporters/Stats/UhtStatsExporter.cs
--- a/Engine/Source/Programs/Shared/EpicGames.UHT/Exporters/Stats/UhtStatsExporter.cs
+++ b/Engine/Source/Programs/Shared/EpicGames.UHT/Exporters/Stats/UhtStatsExporter.cs
@@ -4,9 +4,7 @@ using EpicGames.Core;
 using EpicGames.UHT.Tables;
 using EpicGames.UHT.Types;
 using EpicGames.UHT.Utils;
-using System;
 using System.Collections.Generic;
-using System.Text;
 
 namespace EpicGames.UHT.Exporters.Stats
 {
@@ -33,8 +31,7 @@ namespace EpicGames.UHT.Exporters.Stats
 
 		private static void Collect(SortedDictionary<string, int> CountByType, UhtType Type)
 		{
-			int Count = 1;
-			if (CountByType.TryGetValue(Type.EngineClassName, out Count))
+			if (CountByType.TryGetValue(Type.EngineClassName, out int Count))
 			{
 				CountByType[Type.EngineClassName] = Count + 1;
 			}
diff --git a/Engine/Source/Programs/Shared/EpicGames.UHT/Parsers/UhtClassParser.cs b/Engine/Source/Programs/Shared/EpicGames.UHT/Parsers/UhtClassParser.cs
--- a/Engine/Source/Programs/Shared/EpicGames.UHT/Parsers/UhtClassParser.cs
+++ b/Engine/Source/Programs/Shared/EpicGames.UHT/Parsers/UhtClassParser.cs
@@ -123,18 +123,15 @@ namespace EpicGames.UHT.Parsers
 						this.ClassFlags |= SuperClass.ClassFlags & EClassFlags.ScriptInherit & ~this.RemovedClassFlags;
 					}
 
-					if (this.Bases != null)
+					foreach (UhtStruct Base in this.Bases)
 					{
-						foreach (UhtStruct Base in this.Bases)
+						if (Base is UhtClass BaseClass)
 						{
-							if (Base is UhtClass BaseClass)
+							if (!BaseClass.ClassFlags.HasAnyFlags(EClassFlags.Interface))
 							{
-								if (!BaseClass.ClassFlags.HasAnyFlags(EClassFlags.Interface))
-								{
-									this.LogError($"Class '{BaseClass.SourceName}' is not an interface; Can only inherit from non-UObjects or UInterface derived interfaces");
-								}
-								this.ClassFlags |= BaseClass.ClassFlags & EClassFlags.ScriptInherit & ~this.RemovedClassFlags;
+								this.LogError($"Class '{BaseClass.SourceName}' is not an interface; Can only inherit from non-UObjects or UInterface derived interfaces");
 							}
+							this.ClassFlags |= BaseClass.ClassFlags & EClassFlags.ScriptInherit & ~this.RemovedClassFlags;
 						}
 					}
 
@@ -233,7 +230,7 @@ namespace EpicGames.UHT.Parsers
 								OutShowCategories.AddUnique(Value);
 								break;
 							}
-							SubCategoryPath.Append("|");
+							SubCategoryPath.Append('|');
 						}
 					}
 				}
diff --git a/Engine/Source/Programs/Shared/EpicGames.UHT/Parsers/UhtFunctionParser.cs b/Engine/Source/Programs/Shared/EpicGames.UHT/Parsers/UhtFunctionParser.cs
--- a/Engine/Source/Programs/Shared/EpicGames.UHT/Parsers/UhtFunctionParser.cs
+++ b/Engine/Source/Programs/Shared/EpicGames.UHT/Parsers/UhtFunctionParser.cs
@@ -19,11 +19,11 @@ namespace EpicGames.UHT.Parsers
 	 */
 	struct UhtAdvancedDisplayParameterHandler
 	{
-		private UhtMetaData MetaData;
-		private string[]? ParameterNames;
-		private int NumberLeaveUnmarked;
+		private readonly UhtMetaData MetaData;
+		private readonly string[]? ParameterNames;
+		private readonly int NumberLeaveUnmarked;
+		private readonly bool bUseNumber;
 		private int AlreadyLeft;
-		private bool bUseNumber;
 
 		public UhtAdvancedDisplayParameterHandler(UhtMetaData MetaData)
 		{
diff --git a/Engine/Source/Programs/Shared/EpicGames.UHT/Parsers/UhtHeaderFileParser.cs b/Engine/Source/Programs/Shared/EpicGames.UHT/Parsers/UhtHeaderFileParser.cs
--- a/Engine/Source/Programs/Shared/EpicGames.UHT/Parsers/UhtHeaderFileParser.cs
+++ b/Engine/Source/Programs/Shared/EpicGames.UHT/Parsers/UhtHeaderFileParser.cs
@@ -210,23 +210,18 @@ namespace EpicGames.UHT.Parsers
 		/// <summary>
 		/// Stack of current #if states
 		/// </summary>
-		private List<CompilerDirective> CompilerDirectives = new List<CompilerDirective>();
+		private readonly List<CompilerDirective> CompilerDirectives = new List<CompilerDirective>();
 
 		/// <summary>
 		/// Stack of current #if states saved as part of the preprocessor state
 		/// </summary>
-		private List<CompilerDirective> SavedCompilerDirectives = new List<CompilerDirective>();
+		private readonly List<CompilerDirective> SavedCompilerDirectives = new List<CompilerDirective>();
 
 		/// <summary>
 		/// Current top of the parsing scopes.  Classes, structures and functions all allocate scopes.
 		/// </summary>
 		private UhtParsingScope? TopScope = null;
 
-		/// <summary>
-		/// The total number of statements parsed
-		/// </summary>
-		private int StatementsParsed = 0;
-
 		/// <summary>
 		/// Parse the given header file
 		/// </summary>
@@ -359,7 +354,6 @@ namespace EpicGames.UHT.Parsers
 		public bool ParsePreprocessorDirective(ref UhtToken Token, bool bIsBeingIncluded, out bool bClearComments, out bool bIllegalContentsCheck)
 		{
 			bClearComments = true;
-			bIllegalContentsCheck = true;
 			if (ParseDirectiveInternal(Token, bIsBeingIncluded))
 			{
 				bClearComments = ClearCommentsCompilerDirective();
@@ -430,7 +424,6 @@ namespace EpicGames.UHT.Parsers
 				{
 					ParseStatement(this.TopScope, ref Token, bLogUnhandledKeywords);
 					this.TokenReader.ClearComments();
-					++this.StatementsParsed;
 				}
 			}
 		}
@@ -483,7 +476,7 @@ namespace EpicGames.UHT.Parsers
 				{
 					using (UhtTokenRecorder Recorder = new UhtTokenRecorder(TopScope, ref Token))
 					{
-						ParseResult = TopScope.TokenReader.SkipDeclaration(ref Token) ? UhtParseResult.Handled : UhtParseResult.Invalid;
+						TopScope.TokenReader.SkipDeclaration(ref Token);
 						if (Recorder.Stop())
 						{
 							if (Class.Declarations != null)
@@ -504,7 +497,7 @@ namespace EpicGames.UHT.Parsers
 				}
 				else
 				{
-					ParseResult = TopScope.TokenReader.SkipDeclaration(ref Token) ? UhtParseResult.Handled : UhtParseResult.Invalid;
+					TopScope.TokenReader.SkipDeclaration(ref Token);
 				}
 			}
 			return true;
@@ -909,7 +902,7 @@ namespace EpicGames.UHT.Parsers
 			}
 		}
 
-		private bool SupportsElif(UhtCompilerDirective CompilerDirective)
+		private static bool SupportsElif(UhtCompilerDirective CompilerDirective)
 		{
 			return
 				CompilerDirective == UhtCompilerDirective.WithEditor ||
@@ -917,7 +910,7 @@ namespace EpicGames.UHT.Parsers
 				CompilerDirective == UhtCompilerDirective.WithHotReload;
 		}
 
-		private string GetCompilerDirectiveText(UhtCompilerDirective CompilerDirective)
+		private static string GetCompilerDirectiveText(UhtCompilerDirective CompilerDirective)
 		{
 			switch (CompilerDirective)
 			{
@@ -932,7 +925,7 @@ namespace EpicGames.UHT.Parsers
 			}
 		}
 
-		private void SkipVirtualAndAPI(IUhtTokenReader ReplayReader)
+		private static void SkipVirtualAndAPI(IUhtTokenReader ReplayReader)
 		{
 			while (true)
 			{
@@ -950,15 +943,15 @@ namespace EpicGames.UHT.Parsers
 			IUhtTokenReader ReplayReader = UhtTokenReplayReader.GetThreadInstance(this.HeaderFile, this.HeaderFile.Data.Memory, new ReadOnlyMemory<UhtToken>(Declaration.Tokens), UhtTokenType.EndOfDeclaration);
 
 			// Allow explicit constructors
-			bool bFoundExplicit = ReplayReader.TryOptional("explicit");
-
-			bool bSkippedAPIToken = false;
-			if (bSkippedAPIToken = ReplayReader.PeekToken().Value.Span.EndsWith("_API"))
 			{
-				ReplayReader.ConsumeToken();
-				if (!bFoundExplicit)
+				bool bFoundExplicit = ReplayReader.TryOptional("explicit");
+				if (ReplayReader.PeekToken().Value.Span.EndsWith("_API"))
 				{
-					bFoundExplicit = ReplayReader.TryOptional("explicit");
+					ReplayReader.ConsumeToken();
+					if (!bFoundExplicit)
+					{
+						ReplayReader.TryOptional("explicit");
+					}
 				}
 			}
 
diff --git a/Engine/Source/Programs/Shared/EpicGames.UHT/Parsers/UhtNativeInterfaceClassParser.cs b/Engine/Source/Programs/Shared/EpicGames.UHT/Parsers/UhtNativeInterfaceClassParser.cs
--- a/Engine/Source/Programs/Shared/EpicGames.UHT/Parsers/UhtNativeInterfaceClassParser.cs
+++ b/Engine/Source/Programs/Shared/EpicGames.UHT/Parsers/UhtNativeInterfaceClassParser.cs
@@ -107,8 +107,7 @@ namespace EpicGames.UHT.Parsers
 			IUhtTokenReader TokenReader = ParentScope.TokenReader;
 
 			// Get the optional API macro
-			UhtToken APIMacroToken;
-			TokenReader.TryOptionalAPIMacro(out APIMacroToken);
+			TokenReader.TryOptionalAPIMacro(out UhtToken _);
 
 			// Get the name of the interface
 			SourceName = TokenReader.GetIdentifier();
diff --git a/Engine/Source/Programs/Shared/EpicGames.UHT/Parsers/UhtParsingScope.cs b/Engine/Source/Programs/Shared/EpicGames.UHT/Parsers/UhtParsingScope.cs
--- a/Engine/Source/Programs/Shared/EpicGames.UHT/Parsers/UhtParsingScope.cs
+++ b/Engine/Source/Programs/Shared/EpicGames.UHT/Parsers/UhtParsingScope.cs
@@ -62,7 +62,7 @@ namespace EpicGames.UHT.Parsers
 				UhtParsingScope? CurrentScope = this;
 				while (CurrentScope != null)
 				{
-					if (CurrentScope.ScopeType is UhtClass Class)
+					if (CurrentScope.ScopeType is UhtClass)
 					{
 						return CurrentScope;
 					}
@@ -116,7 +116,20 @@ namespace EpicGames.UHT.Parsers
 		/// </summary>
 		public void Dispose()
 		{
-			this.HeaderParser.PopScope(this);
+			Dispose(true);
+			GC.SuppressFinalize(this);
+		}
+
+		/// <summary>
+		/// Virtual method for disposing the object
+		/// </summary>
+		/// <param name="disposing">If true, we are disposing</param>
+		protected virtual void Dispose(bool disposing)
+		{
+			if (disposing)
+			{
+				this.HeaderParser.PopScope(this);
+			}
 		}
 
 		/// <summary>
@@ -141,7 +154,7 @@ namespace EpicGames.UHT.Parsers
 		/// Format the current token reader comments and add it as meta data
 		/// </summary>
 		/// <param name="MetaNameIndex">Index for the meta data key.  This is used for enum values</param>
-		public void AddFormattedCommentsAsTooltipMetaData(int MetaNameIndex = UhtMetaData.INDEX_NONE)
+		public void AddFormattedCommentsAsTooltipMetaData(int MetaNameIndex = UhtMetaData.IndexNone)
 		{
 			AddFormattedCommentsAsTooltipMetaData(this.ScopeType, MetaNameIndex);
 		}
@@ -151,7 +164,7 @@ namespace EpicGames.UHT.Parsers
 		/// </summary>
 		/// <param name="Type">The type to add the meta data to</param>
 		/// <param name="MetaNameIndex">Index for the meta data key.  This is used for enum values</param>
-		public void AddFormattedCommentsAsTooltipMetaData(UhtType Type, int MetaNameIndex = UhtMetaData.INDEX_NONE)
+		public void AddFormattedCommentsAsTooltipMetaData(UhtType Type, int MetaNameIndex = UhtMetaData.IndexNone)
 		{
 
 			// Don't add a tooltip if one already exists.
@@ -393,7 +406,9 @@ namespace EpicGames.UHT.Parsers
 				// Skip any leading spaces
 				int Index = 0;
 				int EndPos = Line.Length;
-				for (; Index < EndPos && UhtFCString.IsWhitespace(Line[Index]); ++Index) { }
+				for (; Index < EndPos && UhtFCString.IsWhitespace(Line[Index]); ++Index) 
+				{
+				}
 				if (Index == EndPos)
 				{
 					return true;
diff --git a/Engine/Source/Programs/Shared/EpicGames.UHT/Parsers/UhtPropertyParser.cs b/Engine/Source/Programs/Shared/EpicGames.UHT/Parsers/UhtPropertyParser.cs
--- a/Engine/Source/Programs/Shared/EpicGames.UHT/Parsers/UhtPropertyParser.cs
+++ b/Engine/Source/Programs/Shared/EpicGames.UHT/Parsers/UhtPropertyParser.cs
@@ -433,7 +433,7 @@ namespace EpicGames.UHT.Parsers
 	/// </summary>
 	public class UhtPropertyParser : IUhtMessageExtraContext
 	{
-		private UhtPropertySpecifierContext SpecifierContext;
+		private readonly UhtPropertySpecifierContext SpecifierContext;
 		private UhtPropertyParseOptions Options;
 		private UhtParsePropertyDeclarationStyle DeclarationStyle;
 
@@ -442,13 +442,12 @@ namespace EpicGames.UHT.Parsers
 
 		// Scratch pad variables used by actions
 		private List<UhtToken> CurrentTypeTokens = new List<UhtToken>();
-		private UhtToken CurrentTerminatingToken = new UhtToken();
 		private int CurrentTemplateDepth = 0;
 
 		// Actions/Delegates
-		UhtTokensUntilDelegate GatherTypeTokensDelegate;
+		private readonly UhtTokensUntilDelegate GatherTypeTokensDelegate;
 
-		private static Dictionary<StringView, UhtLayoutMacroType> LayoutMacroTypes = new Dictionary<StringView, UhtLayoutMacroType>(new[]
+		private static readonly Dictionary<StringView, UhtLayoutMacroType> LayoutMacroTypes = new Dictionary<StringView, UhtLayoutMacroType>(new[]
 		{
 			UhtLayoutMacroType.Array.MacroNameAndValue(),
 			UhtLayoutMacroType.ArrayEditorOnly.MacroNameAndValue(),
@@ -495,7 +494,7 @@ namespace EpicGames.UHT.Parsers
 			// Reset the context and create the property
 			this.SpecifierContext.PropertySettings.Reset(this.SpecifierContext.Scope.ScopeType, 0, Category, DisallowPropertyFlags);
 			this.SpecifierContext.MetaData = this.SpecifierContext.PropertySettings.MetaData;
-			this.SpecifierContext.MetaNameIndex = UhtMetaData.INDEX_NONE;
+			this.SpecifierContext.MetaNameIndex = UhtMetaData.IndexNone;
 			this.SpecifierContext.bSeenEditSpecifier = false;
 			this.SpecifierContext.bSeenBlueprintWriteSpecifier = false;
 			this.SpecifierContext.bSeenBlueprintReadOnlySpecifier = false;
@@ -507,7 +506,6 @@ namespace EpicGames.UHT.Parsers
 			this.DeclarationStyle = DeclarationStyle;
 
 			this.CurrentTypeTokens = new List<UhtToken>();
-			this.CurrentTerminatingToken = new UhtToken();
 			this.CurrentTemplateDepth = 0;
 
 			using (var TokenContext = new UhtMessageContext(this.TokenReader, this))
@@ -545,7 +543,7 @@ namespace EpicGames.UHT.Parsers
 			IUhtTokenReader ReplayReader = UhtTokenReplayReader.GetThreadInstance(PropertySettings.Outer, Data, TypeTokens, UhtTokenType.EndOfType);
 
 			// Loop through the tokens until we find a known property type or the start of a template argument list
-			UhtPropertyType PropertyType = new UhtPropertyType();
+			UhtPropertyType PropertyType;
 			for (int Index = 0; Index < TypeTokens.Length; ++Index)
 			{
 				if (TypeTokens.Span[Index].IsSymbol())
@@ -576,7 +574,7 @@ namespace EpicGames.UHT.Parsers
 			return null;
 		}
 
-		private static ThreadLocal<UhtPropertySettings> TlsPropertySettings = new ThreadLocal<UhtPropertySettings>(() => { return new UhtPropertySettings(); });
+		private static readonly ThreadLocal<UhtPropertySettings> TlsPropertySettings = new ThreadLocal<UhtPropertySettings>(() => { return new UhtPropertySettings(); });
 
 		/// <summary>
 		/// Given a type with children, resolve any children that couldn't be resolved during the parsing phase.
@@ -968,7 +966,7 @@ namespace EpicGames.UHT.Parsers
 
 				while (true)
 				{
-					UhtProperty Property = Finalize(ref NameToken, TypeTokens, LayoutMacroType, Delegate);
+					UhtProperty _ = Finalize(ref NameToken, TypeTokens, LayoutMacroType, Delegate);
 
 					// If we have reached the end
 					if (!this.TokenReader.TryOptional(','))
@@ -1164,7 +1162,6 @@ namespace EpicGames.UHT.Parsers
 		{
 			if (this.CurrentTemplateDepth == 0 && Token.IsSymbol() && (Token.IsValue(',') || Token.IsValue('(') || Token.IsValue(')') || Token.IsValue(';') || Token.IsValue('[') || Token.IsValue(':') || Token.IsValue('=') || Token.IsValue('{')))
 			{
-				this.CurrentTerminatingToken = Token;
 				return false;
 			}
 
@@ -1200,12 +1197,11 @@ namespace EpicGames.UHT.Parsers
 		}
 		#endregion
 
-		private static UhtPropertyDelegate PropertyDelegate = PropertyParsed;
+		private static readonly UhtPropertyDelegate PropertyDelegate = PropertyParsed;
 
 		private static void PropertyParsed(UhtParsingScope TopScope, UhtProperty Property, ref UhtToken NameToken, UhtLayoutMacroType LayoutMacroType)
 		{
 			IUhtTokenReader TokenReader = TopScope.TokenReader;
-			UhtStruct Struct = (UhtStruct)TopScope.ScopeType;
 
 			// Skip any initialization
 			if (TokenReader.TryOptional('='))
diff --git a/Engine/Source/Programs/Shared/EpicGames.UHT/Parsers/UhtSpecifierParser.cs b/Engine/Source/Programs/Shared/EpicGames.UHT/Parsers/UhtSpecifierParser.cs
--- a/Engine/Source/Programs/Shared/EpicGames.UHT/Parsers/UhtSpecifierParser.cs
+++ b/Engine/Source/Programs/Shared/EpicGames.UHT/Parsers/UhtSpecifierParser.cs
@@ -23,7 +23,7 @@ namespace EpicGames.UHT.Parsers
 			public object? Value;
 		}
 
-		private static List<KeyValuePair<StringView, StringView>> EmptyKVPValues = new List<KeyValuePair<StringView, StringView>>();
+		private static readonly List<KeyValuePair<StringView, StringView>> EmptyKVPValues = new List<KeyValuePair<StringView, StringView>>();
 
 		private UhtSpecifierContext SpecifierContext;
 		private IUhtTokenReader TokenReader;
@@ -377,7 +377,6 @@ namespace EpicGames.UHT.Parsers
 			return SVB.ToStringView();
 		}
 
-
 		/// <summary>
 		/// Parse the sequence of meta data
 		/// </summary>
diff --git a/Engine/Source/Programs/Shared/EpicGames.UHT/Specifiers/UhtInterfaceClassSpecifiers.cs b/Engine/Source/Programs/Shared/EpicGames.UHT/Specifiers/UhtInterfaceClassSpecifiers.cs
--- a/Engine/Source/Programs/Shared/EpicGames.UHT/Specifiers/UhtInterfaceClassSpecifiers.cs
+++ b/Engine/Source/Programs/Shared/EpicGames.UHT/Specifiers/UhtInterfaceClassSpecifiers.cs
@@ -13,7 +13,6 @@ namespace EpicGames.UHT.Parsers
 		[UhtSpecifier(Extends = UhtTableNames.Interface, ValueType = UhtSpecifierValueType.Legacy)]
 		private static void DependsOnSpecifier(UhtSpecifierContext SpecifierContext)
 		{
-			UhtClass Class = (UhtClass)SpecifierContext.Scope.ScopeType;
 			throw new UhtException(SpecifierContext.MessageSite, $"The dependsOn specifier is deprecated. Please use #include \"ClassHeaderFilename.h\" instead.");
 		}
 
diff --git a/Engine/Source/Programs/Shared/EpicGames.UHT/Tables/UhtEngineClassTable.cs b/Engine/Source/Programs/Shared/EpicGames.UHT/Tables/UhtEngineClassTable.cs
--- a/Engine/Source/Programs/Shared/EpicGames.UHT/Tables/UhtEngineClassTable.cs
+++ b/Engine/Source/Programs/Shared/EpicGames.UHT/Tables/UhtEngineClassTable.cs
@@ -12,7 +12,7 @@ namespace EpicGames.UHT.Tables
 	/// This attribute is placed on classes that represent Unreal Engine classes.
 	/// </summary>
 	[AttributeUsage(AttributeTargets.Class, AllowMultiple = true)]
-	public class UhtEngineClassAttribute : Attribute
+	public sealed class UhtEngineClassAttribute : Attribute
 	{
 
 		/// <summary>
@@ -50,7 +50,7 @@ namespace EpicGames.UHT.Tables
 		/// <summary>
 		/// Internal mapping from engine class name to information
 		/// </summary>
-		private Dictionary<StringView, UhtEngineClass> EngineClasses = new Dictionary<StringView, UhtEngineClass>();
+		private readonly Dictionary<StringView, UhtEngineClass> EngineClasses = new Dictionary<StringView, UhtEngineClass>();
 
 		/// <summary>
 		/// Test to see if the given class name is a property
diff --git a/Engine/Source/Programs/Shared/EpicGames.UHT/Tables/UhtExporterTable.cs b/Engine/Source/Programs/Shared/EpicGames.UHT/Tables/UhtExporterTable.cs
--- a/Engine/Source/Programs/Shared/EpicGames.UHT/Tables/UhtExporterTable.cs
+++ b/Engine/Source/Programs/Shared/EpicGames.UHT/Tables/UhtExporterTable.cs
@@ -79,7 +79,7 @@ namespace EpicGames.UHT.Tables
 	/// Defines an exporter
 	/// </summary>
 	[AttributeUsage(AttributeTargets.Method, AllowMultiple = true)]
-	public class UhtExporterAttribute : Attribute
+	public sealed class UhtExporterAttribute : Attribute
 	{
 
 		/// <summary>
@@ -171,7 +171,7 @@ namespace EpicGames.UHT.Tables
 	public class UhtExporterTable : IEnumerable<UhtExporter>
 	{
 
-		private Dictionary<string, UhtExporter> ExporterValues = new Dictionary<string, UhtExporter>(StringComparer.OrdinalIgnoreCase);
+		private readonly Dictionary<string, UhtExporter> ExporterValues = new Dictionary<string, UhtExporter>(StringComparer.OrdinalIgnoreCase);
 
 		/// <summary>
 		/// Return the exporter associated with the given name
diff --git a/Engine/Source/Programs/Shared/EpicGames.UHT/Tables/UhtKeywordTable.cs b/Engine/Source/Programs/Shared/EpicGames.UHT/Tables/UhtKeywordTable.cs
--- a/Engine/Source/Programs/Shared/EpicGames.UHT/Tables/UhtKeywordTable.cs
+++ b/Engine/Source/Programs/Shared/EpicGames.UHT/Tables/UhtKeywordTable.cs
@@ -15,7 +15,7 @@ namespace EpicGames.UHT.Tables
 	/// Invoke the given method when the keyword is parsed.
 	/// </summary>
 	[AttributeUsage(AttributeTargets.Method, AllowMultiple = true)]
-	public class UhtKeywordAttribute : Attribute
+	public sealed class UhtKeywordAttribute : Attribute
 	{
 
 		/// <summary>
@@ -48,7 +48,7 @@ namespace EpicGames.UHT.Tables
 	/// Invoked as a last chance processor for a keyword
 	/// </summary>
 	[AttributeUsage(AttributeTargets.Method, AllowMultiple = true)]
-	public class UhtKeywordCatchAllAttribute : Attribute
+	public sealed class UhtKeywordCatchAllAttribute : Attribute
 	{
 
 		/// <summary>
diff --git a/Engine/Source/Programs/Shared/EpicGames.UHT/Tables/UhtLocTextDefaultValueTable.cs b/Engine/Source/Programs/Shared/EpicGames.UHT/Tables/UhtLocTextDefaultValueTable.cs
--- a/Engine/Source/Programs/Shared/EpicGames.UHT/Tables/UhtLocTextDefaultValueTable.cs
+++ b/Engine/Source/Programs/Shared/EpicGames.UHT/Tables/UhtLocTextDefaultValueTable.cs
@@ -26,7 +26,7 @@ namespace EpicGames.UHT.Tables
 	/// Attribute defining the loctext sanitizer
 	/// </summary>
 	[AttributeUsage(AttributeTargets.Method, AllowMultiple = true)]
-	public class UhtLocTextDefaultValueAttribute : Attribute
+	public sealed class UhtLocTextDefaultValueAttribute : Attribute
 	{
 
 		/// <summary>
@@ -52,7 +52,7 @@ namespace EpicGames.UHT.Tables
 	public class UhtLocTextDefaultValueTable
 	{
 
-		private Dictionary<StringView, UhtLocTextDefaultValue> LocTextDefaultValues = new Dictionary<StringView, UhtLocTextDefaultValue>();
+		private readonly Dictionary<StringView, UhtLocTextDefaultValue> LocTextDefaultValues = new Dictionary<StringView, UhtLocTextDefaultValue>();
 
 		/// <summary>
 		/// Return the loc text default value associated with the given name
diff --git a/Engine/Source/Programs/Shared/EpicGames.UHT/Tables/UhtPropertyTypeTable.cs b/Engine/Source/Programs/Shared/EpicGames.UHT/Tables/UhtPropertyTypeTable.cs
--- a/Engine/Source/Programs/Shared/EpicGames.UHT/Tables/UhtPropertyTypeTable.cs
+++ b/Engine/Source/Programs/Shared/EpicGames.UHT/Tables/UhtPropertyTypeTable.cs
@@ -118,7 +118,7 @@ namespace EpicGames.UHT.Tables
 	/// Property type attribute
 	/// </summary>
 	[AttributeUsage(AttributeTargets.Method, AllowMultiple = true)]
-	public class UhtPropertyTypeAttribute : Attribute
+	public sealed class UhtPropertyTypeAttribute : Attribute
 	{
 
 		/// <summary>
@@ -154,8 +154,8 @@ namespace EpicGames.UHT.Tables
 	/// </summary>
 	public class UhtPropertyTypeTable
 	{
-		private Dictionary<StringView, UhtPropertyType> CaseSensitive = new Dictionary<StringView, UhtPropertyType>();
-		private Dictionary<StringView, UhtPropertyType> CaseInsensitive = new Dictionary<StringView, UhtPropertyType>(StringViewComparer.OrdinalIgnoreCase);
+		private readonly Dictionary<StringView, UhtPropertyType> CaseSensitive = new Dictionary<StringView, UhtPropertyType>();
+		private readonly Dictionary<StringView, UhtPropertyType> CaseInsensitive = new Dictionary<StringView, UhtPropertyType>(StringViewComparer.OrdinalIgnoreCase);
 		private UhtPropertyType? DefaultInternal = null;
 
 		/// <summary>
diff --git a/Engine/Source/Programs/Shared/EpicGames.UHT/Tables/UhtSpecifierTable.cs b/Engine/Source/Programs/Shared/EpicGames.UHT/Tables/UhtSpecifierTable.cs
--- a/Engine/Source/Programs/Shared/EpicGames.UHT/Tables/UhtSpecifierTable.cs
+++ b/Engine/Source/Programs/Shared/EpicGames.UHT/Tables/UhtSpecifierTable.cs
@@ -117,7 +117,7 @@ namespace EpicGames.UHT.Tables
 		/// <param name="MessageSite"></param>
 		/// <param name="MetaData"></param>
 		/// <param name="MetaNameIndex"></param>
-		public UhtSpecifierContext(UhtParsingScope Scope, IUhtMessageSite MessageSite, UhtMetaData MetaData, int MetaNameIndex = UhtMetaData.INDEX_NONE)
+		public UhtSpecifierContext(UhtParsingScope Scope, IUhtMessageSite MessageSite, UhtMetaData MetaData, int MetaNameIndex = UhtMetaData.IndexNone)
 		{
 			this.Scope = Scope;
 			this.MessageSite = MessageSite;
@@ -184,7 +184,7 @@ namespace EpicGames.UHT.Tables
 	/// </summary>
 	public class UhtSpecifierNone : UhtSpecifier
 	{
-		private UhtSpecifierNoneDelegate Delegate;
+		private readonly UhtSpecifierNoneDelegate Delegate;
 
 		/// <summary>
 		/// Construct the specifier
@@ -220,7 +220,7 @@ namespace EpicGames.UHT.Tables
 	/// </summary>
 	public class UhtSpecifierString : UhtSpecifier
 	{
-		private UhtSpecifierStringDelegate Delegate;
+		private readonly UhtSpecifierStringDelegate Delegate;
 
 		/// <summary>
 		/// Construct the specifier
@@ -260,7 +260,7 @@ namespace EpicGames.UHT.Tables
 	/// </summary>
 	public class UhtSpecifierOptionalString : UhtSpecifier
 	{
-		private UhtSpecifierOptionalStringDelegate Delegate;
+		private readonly UhtSpecifierOptionalStringDelegate Delegate;
 
 		/// <summary>
 		/// Construct the specifier
@@ -296,7 +296,7 @@ namespace EpicGames.UHT.Tables
 	/// </summary>
 	public class UhtSpecifierSingleString : UhtSpecifier
 	{
-		private UhtSpecifierSingleStringDelegate Delegate;
+		private readonly UhtSpecifierSingleStringDelegate Delegate;
 
 		/// <summary>
 		/// Construct the specifier
@@ -336,7 +336,7 @@ namespace EpicGames.UHT.Tables
 	/// </summary>
 	public class UhtSpecifierKeyValuePairList : UhtSpecifier
 	{
-		private UhtSpecifierKeyValuePairListDelegate Delegate;
+		private readonly UhtSpecifierKeyValuePairListDelegate Delegate;
 
 		/// <summary>
 		/// Construct the specifier
@@ -377,7 +377,7 @@ namespace EpicGames.UHT.Tables
 	/// </summary>
 	public class UhtSpecifierLegacy : UhtSpecifier
 	{
-		private UhtSpecifierLegacyDelegate Delegate;
+		private readonly UhtSpecifierLegacyDelegate Delegate;
 
 		/// <summary>
 		/// Construct the specifier
@@ -416,7 +416,7 @@ namespace EpicGames.UHT.Tables
 	/// </summary>
 	public class UhtSpecifierStringList : UhtSpecifier
 	{
-		private UhtSpecifierStringListDelegate Delegate;
+		private readonly UhtSpecifierStringListDelegate Delegate;
 
 		/// <summary>
 		/// Construct the specifier
@@ -452,7 +452,7 @@ namespace EpicGames.UHT.Tables
 	/// </summary>
 	public class UhtSpecifierNonEmptyStringList : UhtSpecifier
 	{
-		private UhtSpecifierNonEmptyStringListDelegate Delegate;
+		private readonly UhtSpecifierNonEmptyStringListDelegate Delegate;
 
 		/// <summary>
 		/// Construct the specifier
@@ -484,7 +484,7 @@ namespace EpicGames.UHT.Tables
 	/// Defines a specifier method
 	/// </summary>
 	[AttributeUsage(AttributeTargets.Method, AllowMultiple = true)]
-	public class UhtSpecifierAttribute : Attribute
+	public sealed class UhtSpecifierAttribute : Attribute
 	{
 		/// <summary>
 		/// Name of the specifier.   If not supplied, the method name must end in "Specifier" and the name will be the method name with "Specifier" stripped.
diff --git a/Engine/Source/Programs/Shared/EpicGames.UHT/Tables/UhtSpecifierValidatorTable.cs b/Engine/Source/Programs/Shared/EpicGames.UHT/Tables/UhtSpecifierValidatorTable.cs
--- a/Engine/Source/Programs/Shared/EpicGames.UHT/Tables/UhtSpecifierValidatorTable.cs
+++ b/Engine/Source/Programs/Shared/EpicGames.UHT/Tables/UhtSpecifierValidatorTable.cs
@@ -50,7 +50,7 @@ namespace EpicGames.UHT.Tables
 	/// Attribute used to create a specifier validator
 	/// </summary>
 	[AttributeUsage(AttributeTargets.Method, AllowMultiple = true)]
-	public class UhtSpecifierValidatorAttribute : Attribute
+	public sealed class UhtSpecifierValidatorAttribute : Attribute
 	{
 
 		/// <summary>
diff --git a/Engine/Source/Programs/Shared/EpicGames.UHT/Tables/UhtStructDefaultValueTable.cs b/Engine/Source/Programs/Shared/EpicGames.UHT/Tables/UhtStructDefaultValueTable.cs
--- a/Engine/Source/Programs/Shared/EpicGames.UHT/Tables/UhtStructDefaultValueTable.cs
+++ b/Engine/Source/Programs/Shared/EpicGames.UHT/Tables/UhtStructDefaultValueTable.cs
@@ -85,7 +85,7 @@ namespace EpicGames.UHT.Tables
 	/// Structure default value sanitizer attribute
 	/// </summary>
 	[AttributeUsage(AttributeTargets.Method, AllowMultiple = true)]
-	public class UhtStructDefaultValueAttribute : Attribute
+	public sealed class UhtStructDefaultValueAttribute : Attribute
 	{
 
 		/// <summary>
@@ -116,7 +116,7 @@ namespace EpicGames.UHT.Tables
 	/// </summary>
 	public class UhtStructDefaultValueTable
 	{
-		private Dictionary<StringView, UhtStructDefaultValue> StructDefaultValues = new Dictionary<StringView, UhtStructDefaultValue>();
+		private readonly Dictionary<StringView, UhtStructDefaultValue> StructDefaultValues = new Dictionary<StringView, UhtStructDefaultValue>();
 		private UhtStructDefaultValue? DefaultInternal = null;
 
 		/// <summary>
diff --git a/Engine/Source/Programs/Shared/EpicGames.UHT/Tokenizer/UhtTokenBufferReader.cs b/Engine/Source/Programs/Shared/EpicGames.UHT/Tokenizer/UhtTokenBufferReader.cs
--- a/Engine/Source/Programs/Shared/EpicGames.UHT/Tokenizer/UhtTokenBufferReader.cs
+++ b/Engine/Source/Programs/Shared/EpicGames.UHT/Tokenizer/UhtTokenBufferReader.cs
@@ -12,16 +12,16 @@ namespace EpicGames.UHT.Tokenizer
 	/// <summary>
 	/// Token reader for source buffers
 	/// </summary>
-	public class UhtTokenBufferReader : IUhtTokenReader, IUhtMessageLineNumber
+	public sealed class UhtTokenBufferReader : IUhtTokenReader, IUhtMessageLineNumber
 	{
-		private static StringView[] EmptyComments = new StringView[] { };
+		private static readonly StringView[] EmptyComments = Array.Empty<StringView>();
 
 		private readonly IUhtMessageSite MessageSiteInternal;
 		private List<StringView>? CommentsInternal = null;
-		private List<UhtToken> RecordedTokensInternal = new List<UhtToken>();
+		private readonly List<UhtToken> RecordedTokensInternal = new List<UhtToken>();
 		private IUhtTokenPreprocessor? TokenPreprocessorInternal = null;
 		private UhtToken CurrentToken = new UhtToken(); // PeekToken must have been invoked first
-		private ReadOnlyMemory<char> Data;
+		private readonly ReadOnlyMemory<char> Data;
 		private int PrevPos = 0;
 		private int PrevLine = 1;
 		private bool bHasToken = false;
@@ -86,7 +86,11 @@ namespace EpicGames.UHT.Tokenizer
 		public int InputLine
 		{
 			get => this.bHasToken ? this.PreCurrentTokenInputLine : this.InputLineInternal;
-			set { ClearToken(); this.InputLineInternal = value; }
+			set 
+			{ 
+				ClearToken(); 
+				this.InputLineInternal = value; 
+			}
 		}
 
 		/// <inheritdoc/>
@@ -451,7 +455,7 @@ namespace EpicGames.UHT.Tokenizer
 			}
 		}
 
-		private bool IsFirstTokenInLine(ReadOnlySpan<char> Span, int StartPos)
+		private static bool IsFirstTokenInLine(ReadOnlySpan<char> Span, int StartPos)
 		{
 			for (int Pos = StartPos; --Pos > 0;)
 			{
@@ -477,7 +481,6 @@ namespace EpicGames.UHT.Tokenizer
 		private UhtToken GetTokenInternal(bool bEnablePreprocessor)
 		{
 			ReadOnlySpan<char> Span = this.Data.Span;
-			int SpanLength = Span.Length;
 			bool bGotInlineComment = false;
 		Restart:
 			this.PreCurrentTokenInputLine = this.InputLineInternal;
@@ -614,13 +617,12 @@ namespace EpicGames.UHT.Tokenizer
 			{
 
 				// We try to skip the character constant value. But if it is backslash, we have to skip another character
-				char NextChar = InternalGetChar(Span);
-				if (NextChar == '\\')
+				if (InternalGetChar(Span) == '\\')
 				{
-					NextChar = InternalGetChar(Span);
+					InternalGetChar(Span);
 				}
 
-				NextChar = InternalGetChar(Span);
+				char NextChar = InternalGetChar(Span);
 				if (NextChar != '\'')
 				{
 					throw new UhtException(this, "Unterminated character constant");
@@ -705,9 +707,7 @@ namespace EpicGames.UHT.Tokenizer
 					{
 						this.PreprocessorPendingCommentsCount = this.CommentsInternal != null ? this.CommentsInternal.Count - this.CommittedComments : 0;
 						UhtToken Temp = new UhtToken(TokenType, StartPos, StartLine, StartPos, this.InputLineInternal, new StringView(this.Data.Slice(StartPos, this.InputPosInternal - StartPos)));
-						bool bClearComments = false;
-						bool bIllegalContentsCheck = false;
-						bool bInclude = this.TokenPreprocessorInternal.ParsePreprocessorDirective(ref Temp, true, out bClearComments, out bIllegalContentsCheck);
+						bool bInclude = this.TokenPreprocessorInternal.ParsePreprocessorDirective(ref Temp, true, out bool bClearComments, out bool bIllegalContentsCheck);
 						if (!bInclude)
 						{
 							++this.CommentsDisableCount;
@@ -722,7 +722,7 @@ namespace EpicGames.UHT.Tokenizer
 								ReadOnlySpan<char> LocalValueSpan = LocalToken.Value.Span;
 								if (LocalToken.IsSymbol('#') && LocalValueSpan.Length == 1 && LocalValueSpan[0] == '#' && IsFirstTokenInLine(Span, LocalToken.InputStartPos))
 								{
-									if (this.TokenPreprocessorInternal.ParsePreprocessorDirective(ref Temp, false, out bool bScratch, out bool bScratchIllegalContentsCheck))
+									if (this.TokenPreprocessorInternal.ParsePreprocessorDirective(ref Temp, false, out bool _, out bool bScratchIllegalContentsCheck))
 									{
 										break;
 									}
diff --git a/Engine/Source/Programs/Shared/EpicGames.UHT/Tokenizer/UhtTokenReader.cs b/Engine/Source/Programs/Shared/EpicGames.UHT/Tokenizer/UhtTokenReader.cs
--- a/Engine/Source/Programs/Shared/EpicGames.UHT/Tokenizer/UhtTokenReader.cs
+++ b/Engine/Source/Programs/Shared/EpicGames.UHT/Tokenizer/UhtTokenReader.cs
@@ -322,9 +322,9 @@ namespace EpicGames.UHT.Tokenizer
 	/// <summary>
 	/// Token list cache.  Token lists must be returned to the cache.
 	/// </summary>
-	public class UhtTokenListCache
+	public static class UhtTokenListCache
 	{
-		static private ThreadLocal<UhtTokenList?> Tls = new ThreadLocal<UhtTokenList?>(() => null);
+		private static readonly ThreadLocal<UhtTokenList?> Tls = new ThreadLocal<UhtTokenList?>(() => null);
 
 		/// <summary>
 		/// Borrow a token list
diff --git a/Engine/Source/Programs/Shared/EpicGames.UHT/Tokenizer/UhtTokenReaderFloatExtensions.cs b/Engine/Source/Programs/Shared/EpicGames.UHT/Tokenizer/UhtTokenReaderFloatExtensions.cs
--- a/Engine/Source/Programs/Shared/EpicGames.UHT/Tokenizer/UhtTokenReaderFloatExtensions.cs
+++ b/Engine/Source/Programs/Shared/EpicGames.UHT/Tokenizer/UhtTokenReaderFloatExtensions.cs
@@ -35,8 +35,7 @@ namespace EpicGames.UHT.Tokenizer
 		/// <returns>The token reader</returns>
 		public static IUhtTokenReader OptionalConstFloat(this IUhtTokenReader TokenReader)
 		{
-			float Value;
-			TokenReader.TryOptionalConstFloat(out Value);
+			TokenReader.TryOptionalConstFloat(out float _);
 			return TokenReader;
 		}
 
@@ -48,8 +47,7 @@ namespace EpicGames.UHT.Tokenizer
 		/// <returns>The token reader</returns>
 		public static IUhtTokenReader OptionalConstFloat(this IUhtTokenReader TokenReader, UhtTokenConstFloatDelegate Delegate)
 		{
-			float Value;
-			if (TokenReader.TryOptionalConstFloat(out Value))
+			if (TokenReader.TryOptionalConstFloat(out float Value))
 			{
 				Delegate(Value);
 			}
@@ -64,8 +62,7 @@ namespace EpicGames.UHT.Tokenizer
 		/// <returns>True if the next token was an float, false if not.</returns>
 		public static IUhtTokenReader RequireConstFloat(this IUhtTokenReader TokenReader, object? ExceptionContext = null)
 		{
-			float Value;
-			if (!TokenReader.TryOptionalConstFloat(out Value))
+			if (!TokenReader.TryOptionalConstFloat(out float _))
 			{
 				throw new UhtTokenException(TokenReader, TokenReader.PeekToken(), "constant float", ExceptionContext);
 			}
@@ -129,8 +126,7 @@ namespace EpicGames.UHT.Tokenizer
 		/// <returns>The double value</returns>
 		public static float GetConstFloatExpression(this IUhtTokenReader TokenReader)
 		{
-			float LocalValue = 0;
-			if (!TokenReader.TryOptionalConstFloatExpression(out LocalValue))
+			if (!TokenReader.TryOptionalConstFloatExpression(out float LocalValue))
 			{
 				throw new UhtTokenException(TokenReader, TokenReader.PeekToken(), "constant float", null);
 			}
@@ -162,8 +158,7 @@ namespace EpicGames.UHT.Tokenizer
 		/// <returns>The supplied token reader</returns>
 		public static IUhtTokenReader RequireConstDoubleExpression(this IUhtTokenReader TokenReader, UhtTokenConstDoubleDelegate Delegate)
 		{
-			double LocalValue = 0;
-			if (!TokenReader.TryOptionalConstDoubleExpression(out LocalValue))
+			if (!TokenReader.TryOptionalConstDoubleExpression(out double LocalValue))
 			{
 				throw new UhtTokenException(TokenReader, TokenReader.PeekToken(), "constant double", null);
 			}
@@ -178,8 +173,7 @@ namespace EpicGames.UHT.Tokenizer
 		/// <returns>The double value</returns>
 		public static double GetConstDoubleExpression(this IUhtTokenReader TokenReader)
 		{
-			double LocalValue = 0;
-			if (!TokenReader.TryOptionalConstDoubleExpression(out LocalValue))
+			if (!TokenReader.TryOptionalConstDoubleExpression(out double LocalValue))
 			{
 				throw new UhtTokenException(TokenReader, TokenReader.PeekToken(), "constant double", null);
 			}
diff --git a/Engine/Source/Programs/Shared/EpicGames.UHT/Tokenizer/UhtTokenReaderIntegerExtensions.cs b/Engine/Source/Programs/Shared/EpicGames.UHT/Tokenizer/UhtTokenReaderIntegerExtensions.cs
--- a/Engine/Source/Programs/Shared/EpicGames.UHT/Tokenizer/UhtTokenReaderIntegerExtensions.cs
+++ b/Engine/Source/Programs/Shared/EpicGames.UHT/Tokenizer/UhtTokenReaderIntegerExtensions.cs
@@ -44,8 +44,7 @@ namespace EpicGames.UHT.Tokenizer
 		/// <returns>The value of the constant</returns>
 		public static IUhtTokenReader OptionalConstInt(this IUhtTokenReader TokenReader, object? ExceptionContext = null)
 		{
-			int Value;
-			TokenReader.TryOptionalConstInt(out Value);
+			TokenReader.TryOptionalConstInt(out int _);
 			return TokenReader;
 		}
 
@@ -57,8 +56,7 @@ namespace EpicGames.UHT.Tokenizer
 		/// <returns>The value of the constant</returns>
 		public static IUhtTokenReader RequireConstInt(this IUhtTokenReader TokenReader, object? ExceptionContext = null)
 		{
-			int Value;
-			if (!TokenReader.TryOptionalConstInt(out Value))
+			if (!TokenReader.TryOptionalConstInt(out int _))
 			{
 				throw new UhtTokenException(TokenReader, TokenReader.PeekToken(), "constant integer", ExceptionContext);
 			}
@@ -109,8 +107,7 @@ namespace EpicGames.UHT.Tokenizer
 		{
 			ref UhtToken Token = ref TokenReader.PeekToken();
 
-			long Value;
-			if (Token.GetConstLong(out Value))
+			if (Token.GetConstLong(out long _))
 			{
 				TokenReader.ConsumeToken();
 			}
@@ -125,8 +122,7 @@ namespace EpicGames.UHT.Tokenizer
 		/// <returns>The value of the constant</returns>
 		public static IUhtTokenReader RequireConstLong(this IUhtTokenReader TokenReader, object? ExceptionContext = null)
 		{
-			long Value;
-			if (!TokenReader.TryOptionalConstLong(out Value))
+			if (!TokenReader.TryOptionalConstLong(out long _))
 			{
 				throw new UhtTokenException(TokenReader, TokenReader.PeekToken(), "constant long integer", ExceptionContext);
 			}
diff --git a/Engine/Source/Programs/Shared/EpicGames.UHT/Tokenizer/UhtTokenReaderUtilityExtensions.cs b/Engine/Source/Programs/Shared/EpicGames.UHT/Tokenizer/UhtTokenReaderUtilityExtensions.cs
--- a/Engine/Source/Programs/Shared/EpicGames.UHT/Tokenizer/UhtTokenReaderUtilityExtensions.cs
+++ b/Engine/Source/Programs/Shared/EpicGames.UHT/Tokenizer/UhtTokenReaderUtilityExtensions.cs
@@ -13,23 +13,23 @@ namespace EpicGames.UHT.Tokenizer
 	/// </summary>
 	public static class UhtTokenReaderUtilityExtensions
 	{
-		private static HashSet<StringView> SkipDeclarationWarningStrings = new HashSet<StringView>
-			{
-				"GENERATED_BODY",
-				"GENERATED_IINTERFACE_BODY",
-				"GENERATED_UCLASS_BODY",
-				"GENERATED_UINTERFACE_BODY",
-				"GENERATED_USTRUCT_BODY",
-				// Leaving these disabled ATM since they can exist in the code without causing compile issues
-				//"RIGVM_METHOD",
-				//"UCLASS",
-				//"UDELEGATE",
-				//"UENUM",
-				//"UFUNCTION",
-				//"UINTERFACE",
-				//"UPROPERTY",
-				//"USTRUCT",
-			};
+		private static readonly HashSet<StringView> SkipDeclarationWarningStrings = new HashSet<StringView>
+		{
+			"GENERATED_BODY",
+			"GENERATED_IINTERFACE_BODY",
+			"GENERATED_UCLASS_BODY",
+			"GENERATED_UINTERFACE_BODY",
+			"GENERATED_USTRUCT_BODY",
+			// Leaving these disabled ATM since they can exist in the code without causing compile issues
+			//"RIGVM_METHOD",
+			//"UCLASS",
+			//"UDELEGATE",
+			//"UENUM",
+			//"UFUNCTION",
+			//"UINTERFACE",
+			//"UPROPERTY",
+			//"USTRUCT",
+		};
 
 		/// <summary>
 		/// When processing type, make sure that the next token is the expected token
@@ -123,8 +123,6 @@ namespace EpicGames.UHT.Tokenizer
 		/// <returns>true if there could be more header to process, false if the end was reached.</returns>
 		public static bool SkipDeclaration(this IUhtTokenReader TokenReader, ref UhtToken Token)
 		{
-			UhtToken StarToken = Token;
-
 			// Consume all tokens until the end of declaration/definition has been found.
 			int NestedScopes = 0;
 			bool bEndOfDeclarationFound = false;
diff --git a/Engine/Source/Programs/Shared/EpicGames.UHT/Tokenizer/UhtTokenReplayReader.cs b/Engine/Source/Programs/Shared/EpicGames.UHT/Tokenizer/UhtTokenReplayReader.cs
--- a/Engine/Source/Programs/Shared/EpicGames.UHT/Tokenizer/UhtTokenReplayReader.cs
+++ b/Engine/Source/Programs/Shared/EpicGames.UHT/Tokenizer/UhtTokenReplayReader.cs
@@ -12,9 +12,9 @@ namespace EpicGames.UHT.Tokenizer
 	/// <summary>
 	/// Token reader to replay previously recorded token stream
 	/// </summary>
-	public class UhtTokenReplayReader : IUhtTokenReader, IUhtMessageLineNumber
+	public sealed class UhtTokenReplayReader : IUhtTokenReader, IUhtMessageLineNumber
 	{
-		static private ThreadLocal<UhtTokenReplayReader> Tls = new ThreadLocal<UhtTokenReplayReader>(() => new UhtTokenReplayReader());
+		private static readonly ThreadLocal<UhtTokenReplayReader> Tls = new ThreadLocal<UhtTokenReplayReader>(() => new UhtTokenReplayReader());
 
 		const int MaxSavedStates = 2;
 		private struct SavedState
@@ -29,7 +29,7 @@ namespace EpicGames.UHT.Tokenizer
 		private int CurrentTokenIndex = -1;
 		private bool bHasToken = false;
 		private UhtToken CurrentToken = new UhtToken();
-		private SavedState[] SavedStates = new SavedState[MaxSavedStates];
+		private readonly SavedState[] SavedStates = new SavedState[MaxSavedStates];
 		private int SavedStateCount = 0;
 		private UhtTokenType EndTokenType = UhtTokenType.EndOfFile;
 
@@ -43,8 +43,8 @@ namespace EpicGames.UHT.Tokenizer
 		public UhtTokenReplayReader(IUhtMessageSite MessageSite, ReadOnlyMemory<char> Data, ReadOnlyMemory<UhtToken> Tokens, UhtTokenType EndTokenType)
 		{
 			this.MessageSiteInternal = MessageSite;
-			this.Tokens = new UhtToken[0].AsMemory();
-			this.Data = new char[0].AsMemory();
+			this.Tokens = Array.Empty<UhtToken>().AsMemory();
+			this.Data = Array.Empty<char>().AsMemory();
 			this.EndTokenType = EndTokenType;
 			this.CurrentToken = new UhtToken(EndTokenType);
 		}
@@ -55,8 +55,8 @@ namespace EpicGames.UHT.Tokenizer
 		public UhtTokenReplayReader()
 		{
 			this.MessageSiteInternal = new UhtEmptyMessageSite();
-			this.Tokens = new UhtToken[0].AsMemory();
-			this.Data = new char[0].AsMemory();
+			this.Tokens = Array.Empty<UhtToken>().AsMemory();
+			this.Data = Array.Empty<char>().AsMemory();
 		}
 
 		/// <summary>
diff --git a/Engine/Source/Programs/Shared/EpicGames.UHT/Types/Properties/UhtArrayProperty.cs b/Engine/Source/Programs/Shared/EpicGames.UHT/Types/Properties/UhtArrayProperty.cs
--- a/Engine/Source/Programs/Shared/EpicGames.UHT/Types/Properties/UhtArrayProperty.cs
+++ b/Engine/Source/Programs/Shared/EpicGames.UHT/Types/Properties/UhtArrayProperty.cs
@@ -113,7 +113,7 @@ namespace EpicGames.UHT.Types
 					if (Builder[Builder.Length - 1] == '>')
 					{
 						// if our internal property type is a template class, add a space between the closing brackets b/c VS.NET cannot parse this correctly
-						Builder.Append(" ");
+						Builder.Append(' ');
 					}
 					Builder.Append('>');
 					break;
diff --git a/Engine/Source/Programs/Shared/EpicGames.UHT/Types/Properties/UhtBoolProperty.cs b/Engine/Source/Programs/Shared/EpicGames.UHT/Types/Properties/UhtBoolProperty.cs
--- a/Engine/Source/Programs/Shared/EpicGames.UHT/Types/Properties/UhtBoolProperty.cs
+++ b/Engine/Source/Programs/Shared/EpicGames.UHT/Types/Properties/UhtBoolProperty.cs
@@ -189,7 +189,7 @@ namespace EpicGames.UHT.Types
 			{
 				Builder
 					.Append("sizeof(").Append(Context.OuterStructSourceName).Append("), ")
-					.Append("&").AppendNameDef(Context, Name, NameSuffix).Append("_SetBit, ");
+					.Append('&').AppendNameDef(Context, Name, NameSuffix).Append("_SetBit, ");
 			}
 			else
 			{
diff --git a/Engine/Source/Programs/Shared/EpicGames.UHT/Types/Properties/UhtByteProperty.cs b/Engine/Source/Programs/Shared/EpicGames.UHT/Types/Properties/UhtByteProperty.cs
--- a/Engine/Source/Programs/Shared/EpicGames.UHT/Types/Properties/UhtByteProperty.cs
+++ b/Engine/Source/Programs/Shared/EpicGames.UHT/Types/Properties/UhtByteProperty.cs
@@ -123,7 +123,7 @@ namespace EpicGames.UHT.Types
 			}
 
 			int Value = DefaultValueReader.GetConstIntExpression();
-			InnerDefaultValue.Append(Value.ToString());
+			InnerDefaultValue.Append(Value);
 			return Value >= byte.MinValue && Value <= byte.MaxValue;
 		}
 
diff --git a/Engine/Source/Programs/Shared/EpicGames.UHT/Types/Properties/UhtClassProperty.cs b/Engine/Source/Programs/Shared/EpicGames.UHT/Types/Properties/UhtClassProperty.cs
--- a/Engine/Source/Programs/Shared/EpicGames.UHT/Types/Properties/UhtClassProperty.cs
+++ b/Engine/Source/Programs/Shared/EpicGames.UHT/Types/Properties/UhtClassProperty.cs
@@ -61,7 +61,7 @@ namespace EpicGames.UHT.Types
 			}
 			else
 			{
-				Builder.Append(this.Class.SourceName).Append("*");
+				Builder.Append(this.Class.SourceName).Append('*');
 			}
 			return Builder;
 		}
diff --git a/Engine/Source/Programs/Shared/EpicGames.UHT/Types/Properties/UhtClassPtrProperties.cs b/Engine/Source/Programs/Shared/EpicGames.UHT/Types/Properties/UhtClassPtrProperties.cs
--- a/Engine/Source/Programs/Shared/EpicGames.UHT/Types/Properties/UhtClassPtrProperties.cs
+++ b/Engine/Source/Programs/Shared/EpicGames.UHT/Types/Properties/UhtClassPtrProperties.cs
@@ -45,7 +45,7 @@ namespace EpicGames.UHT.Types
 			switch (TextType)
 			{
 				default:
-					Builder.Append("TObjectPtr<").Append(this.Class.SourceName).Append(">");
+					Builder.Append("TObjectPtr<").Append(this.Class.SourceName).Append('>');
 					break;
 			}
 			return Builder;
diff --git a/Engine/Source/Programs/Shared/EpicGames.UHT/Types/Properties/UhtContainerProperty.cs b/Engine/Source/Programs/Shared/EpicGames.UHT/Types/Properties/UhtContainerProperty.cs
--- a/Engine/Source/Programs/Shared/EpicGames.UHT/Types/Properties/UhtContainerProperty.cs
+++ b/Engine/Source/Programs/Shared/EpicGames.UHT/Types/Properties/UhtContainerProperty.cs
@@ -19,7 +19,7 @@ namespace EpicGames.UHT.Types
 		/// </summary>
 		/// <param name="PropertySettings">Property settings</param>
 		/// <param name="Value">Value property</param>
-		public UhtContainerBaseProperty(UhtPropertySettings PropertySettings, UhtProperty Value) : base(PropertySettings)
+		protected UhtContainerBaseProperty(UhtPropertySettings PropertySettings, UhtProperty Value) : base(PropertySettings)
 		{
 			this.ValueProperty = Value;
 			this.PropertyCaps = (this.PropertyCaps & ~(UhtPropertyCaps.CanBeInstanced | UhtPropertyCaps.CanHaveConfig)) | 
diff --git a/Engine/Source/Programs/Shared/EpicGames.UHT/Types/Properties/UhtDelegateProperty.cs b/Engine/Source/Programs/Shared/EpicGames.UHT/Types/Properties/UhtDelegateProperty.cs
--- a/Engine/Source/Programs/Shared/EpicGames.UHT/Types/Properties/UhtDelegateProperty.cs
+++ b/Engine/Source/Programs/Shared/EpicGames.UHT/Types/Properties/UhtDelegateProperty.cs
@@ -107,7 +107,7 @@ namespace EpicGames.UHT.Types
 		/// <inheritdoc/>
 		public override StringBuilder AppendFunctionThunkParameterArg(StringBuilder Builder)
 		{
-			return Builder.Append(this.Function.SourceName).Append("(").AppendFunctionThunkParameterName(this).Append(")");
+			return Builder.Append(this.Function.SourceName).Append('(').AppendFunctionThunkParameterName(this).Append(')');
 		}
 
 		/// <inheritdoc/>
diff --git a/Engine/Source/Programs/Shared/EpicGames.UHT/Types/Properties/UhtEnumProperty.cs b/Engine/Source/Programs/Shared/EpicGames.UHT/Types/Properties/UhtEnumProperty.cs
--- a/Engine/Source/Programs/Shared/EpicGames.UHT/Types/Properties/UhtEnumProperty.cs
+++ b/Engine/Source/Programs/Shared/EpicGames.UHT/Types/Properties/UhtEnumProperty.cs
@@ -118,7 +118,7 @@ namespace EpicGames.UHT.Types
 					}
 					else
 					{
-						Builder.Append("TEnumAsByte<").Append(Enum.CppType).Append(">");
+						Builder.Append("TEnumAsByte<").Append(Enum.CppType).Append('>');
 					}
 					break;
 
@@ -133,7 +133,7 @@ namespace EpicGames.UHT.Types
 					}
 					else
 					{
-						Builder.Append("TEnumAsByte<").Append(Enum.CppType).Append(">");
+						Builder.Append("TEnumAsByte<").Append(Enum.CppType).Append('>');
 					}
 					break;
 
@@ -148,7 +148,7 @@ namespace EpicGames.UHT.Types
 					}
 					else
 					{
-						Builder.Append("TEnumAsByte<").Append(Enum.CppType).Append(">");
+						Builder.Append("TEnumAsByte<").Append(Enum.CppType).Append('>');
 					}
 					break;
 			}
@@ -273,7 +273,7 @@ namespace EpicGames.UHT.Types
 		{
 			if (this.Enum.CppForm != UhtEnumCppForm.EnumClass)
 			{
-				Builder.Append("0");
+				Builder.Append('0');
 			}
 			else
 			{
diff --git a/Engine/Source/Programs/Shared/EpicGames.UHT/Types/Properties/UhtInt64Property.cs b/Engine/Source/Programs/Shared/EpicGames.UHT/Types/Properties/UhtInt64Property.cs
--- a/Engine/Source/Programs/Shared/EpicGames.UHT/Types/Properties/UhtInt64Property.cs
+++ b/Engine/Source/Programs/Shared/EpicGames.UHT/Types/Properties/UhtInt64Property.cs
@@ -47,7 +47,7 @@ namespace EpicGames.UHT.Types
 		/// <inheritdoc/>
 		public override bool SanitizeDefaultValue(IUhtTokenReader DefaultValueReader, StringBuilder InnerDefaultValue)
 		{
-			InnerDefaultValue.Append(DefaultValueReader.GetConstLongExpression().ToString());
+			InnerDefaultValue.Append(DefaultValueReader.GetConstLongExpression());
 			return true;
 		}
 
diff --git a/Engine/Source/Programs/Shared/EpicGames.UHT/Types/Properties/UhtIntProperty.cs b/Engine/Source/Programs/Shared/EpicGames.UHT/Types/Properties/UhtIntProperty.cs
--- a/Engine/Source/Programs/Shared/EpicGames.UHT/Types/Properties/UhtIntProperty.cs
+++ b/Engine/Source/Programs/Shared/EpicGames.UHT/Types/Properties/UhtIntProperty.cs
@@ -50,7 +50,7 @@ namespace EpicGames.UHT.Types
 		/// <inheritdoc/>
 		public override bool SanitizeDefaultValue(IUhtTokenReader DefaultValueReader, StringBuilder InnerDefaultValue)
 		{
-			InnerDefaultValue.Append(DefaultValueReader.GetConstIntExpression().ToString());
+			InnerDefaultValue.Append(DefaultValueReader.GetConstIntExpression());
 			return true;
 		}
 
diff --git a/Engine/Source/Programs/Shared/EpicGames.UHT/Types/Properties/UhtLazyObjectPtrProperty.cs b/Engine/Source/Programs/Shared/EpicGames.UHT/Types/Properties/UhtLazyObjectPtrProperty.cs
--- a/Engine/Source/Programs/Shared/EpicGames.UHT/Types/Properties/UhtLazyObjectPtrProperty.cs
+++ b/Engine/Source/Programs/Shared/EpicGames.UHT/Types/Properties/UhtLazyObjectPtrProperty.cs
@@ -53,7 +53,7 @@ namespace EpicGames.UHT.Types
 					break;
 
 				default:
-					Builder.Append("TLazyObjectPtr<").Append(this.Class.SourceName).Append(">");
+					Builder.Append("TLazyObjectPtr<").Append(this.Class.SourceName).Append('>');
 					break;
 			}
 			return Builder;
diff --git a/Engine/Source/Programs/Shared/EpicGames.UHT/Types/Properties/UhtMapProperty.cs b/Engine/Source/Programs/Shared/EpicGames.UHT/Types/Properties/UhtMapProperty.cs
--- a/Engine/Source/Programs/Shared/EpicGames.UHT/Types/Properties/UhtMapProperty.cs
+++ b/Engine/Source/Programs/Shared/EpicGames.UHT/Types/Properties/UhtMapProperty.cs
@@ -134,7 +134,7 @@ namespace EpicGames.UHT.Types
 					break;
 
 				case UhtPropertyTextType.FunctionThunkParameterArgType:
-					Builder.AppendFunctionThunkParameterArrayType(this.KeyProperty).Append(",").AppendFunctionThunkParameterArrayType(this.ValueProperty);
+					Builder.AppendFunctionThunkParameterArrayType(this.KeyProperty).Append(',').AppendFunctionThunkParameterArrayType(this.ValueProperty);
 					break;
 
 				default:
@@ -142,13 +142,13 @@ namespace EpicGames.UHT.Types
 					if (Builder[Builder.Length - 1] == '>')
 					{
 						// if our internal property type is a template class, add a space between the closing brackets b/c VS.NET cannot parse this correctly
-						Builder.Append(" ");
+						Builder.Append(' ');
 					}
 					Builder.Append(',').AppendPropertyText(this.ValueProperty, TextType, true);
 					if (Builder[Builder.Length - 1] == '>')
 					{
 						// if our internal property type is a template class, add a space between the closing brackets b/c VS.NET cannot parse this correctly
-						Builder.Append(" ");
+						Builder.Append(' ');
 					}
 					Builder.Append('>');
 					break;
diff --git a/Engine/Source/Programs/Shared/EpicGames.UHT/Types/Properties/UhtMulticastDelegateProperty.cs b/Engine/Source/Programs/Shared/EpicGames.UHT/Types/Properties/UhtMulticastDelegateProperty.cs
--- a/Engine/Source/Programs/Shared/EpicGames.UHT/Types/Properties/UhtMulticastDelegateProperty.cs
+++ b/Engine/Source/Programs/Shared/EpicGames.UHT/Types/Properties/UhtMulticastDelegateProperty.cs
@@ -38,7 +38,7 @@ namespace EpicGames.UHT.Types
 		/// </summary>
 		/// <param name="PropertySettings">Property settings</param>
 		/// <param name="Function">Referenced function</param>
-		public UhtMulticastDelegateProperty(UhtPropertySettings PropertySettings, UhtFunction Function) : base(PropertySettings)
+		protected UhtMulticastDelegateProperty(UhtPropertySettings PropertySettings, UhtFunction Function) : base(PropertySettings)
 		{
 			this.Function = Function;
 			this.HeaderFile.AddReferencedHeader(Function);
@@ -100,7 +100,7 @@ namespace EpicGames.UHT.Types
 		/// <inheritdoc/>
 		public override StringBuilder AppendFunctionThunkParameterArg(StringBuilder Builder)
 		{
-			return Builder.Append(this.Function.SourceName).Append("(").AppendFunctionThunkParameterName(this).Append(")");
+			return Builder.Append(this.Function.SourceName).Append('(').AppendFunctionThunkParameterName(this).Append(')');
 		}
 
 		/// <inheritdoc/>
diff --git a/Engine/Source/Programs/Shared/EpicGames.UHT/Types/Properties/UhtNumericProperty.cs b/Engine/Source/Programs/Shared/EpicGames.UHT/Types/Properties/UhtNumericProperty.cs
--- a/Engine/Source/Programs/Shared/EpicGames.UHT/Types/Properties/UhtNumericProperty.cs
+++ b/Engine/Source/Programs/Shared/EpicGames.UHT/Types/Properties/UhtNumericProperty.cs
@@ -35,7 +35,7 @@ namespace EpicGames.UHT.Types
 		/// <inheritdoc/>
 		public override StringBuilder AppendNullConstructorArg(StringBuilder Builder, bool bIsInitializer)
 		{
-			Builder.Append("0");
+			Builder.Append('0');
 			return Builder;
 		}
 	}
diff --git a/Engine/Source/Programs/Shared/EpicGames.UHT/Types/Properties/UhtObjectProperty.cs b/Engine/Source/Programs/Shared/EpicGames.UHT/Types/Properties/UhtObjectProperty.cs
--- a/Engine/Source/Programs/Shared/EpicGames.UHT/Types/Properties/UhtObjectProperty.cs
+++ b/Engine/Source/Programs/Shared/EpicGames.UHT/Types/Properties/UhtObjectProperty.cs
@@ -67,7 +67,7 @@ namespace EpicGames.UHT.Types
 					{
 						Builder.Append("const ");
 					}
-					Builder.Append(this.Class.SourceName).Append("*");
+					Builder.Append(this.Class.SourceName).Append('*');
 					break;
 
 				case UhtPropertyTextType.FunctionThunkParameterArgType:
@@ -75,7 +75,7 @@ namespace EpicGames.UHT.Types
 					break;
 
 				default:
-					Builder.Append(this.Class.SourceName).Append("*");
+					Builder.Append(this.Class.SourceName).Append('*');
 					break;
 			}
 			return Builder;
diff --git a/Engine/Source/Programs/Shared/EpicGames.UHT/Types/Properties/UhtObjectPropertyBase.cs b/Engine/Source/Programs/Shared/EpicGames.UHT/Types/Properties/UhtObjectPropertyBase.cs
--- a/Engine/Source/Programs/Shared/EpicGames.UHT/Types/Properties/UhtObjectPropertyBase.cs
+++ b/Engine/Source/Programs/Shared/EpicGames.UHT/Types/Properties/UhtObjectPropertyBase.cs
@@ -44,7 +44,7 @@ namespace EpicGames.UHT.Types
 		/// <param name="PropertySettings">Property settings</param>
 		/// <param name="Class">Referenced UCLASS</param>
 		/// <param name="MetaClass">Referenced UCLASS used by class properties</param>
-		public UhtObjectPropertyBase(UhtPropertySettings PropertySettings, UhtClass Class, UhtClass? MetaClass) : base(PropertySettings)
+		protected UhtObjectPropertyBase(UhtPropertySettings PropertySettings, UhtClass Class, UhtClass? MetaClass) : base(PropertySettings)
 		{
 			this.Class = Class;
 			this.MetaClass = MetaClass;
diff --git a/Engine/Source/Programs/Shared/EpicGames.UHT/Types/Properties/UhtObjectPtrProperty.cs b/Engine/Source/Programs/Shared/EpicGames.UHT/Types/Properties/UhtObjectPtrProperty.cs
--- a/Engine/Source/Programs/Shared/EpicGames.UHT/Types/Properties/UhtObjectPtrProperty.cs
+++ b/Engine/Source/Programs/Shared/EpicGames.UHT/Types/Properties/UhtObjectPtrProperty.cs
@@ -48,11 +48,11 @@ namespace EpicGames.UHT.Types
 				case UhtPropertyTextType.GetterSetterArg:
 					if (bIsTemplateArgument)
 					{
-						Builder.Append("TObjectPtr<").Append(this.Class.SourceName).Append(">");
+						Builder.Append("TObjectPtr<").Append(this.Class.SourceName).Append('>');
 					}
 					else
 					{
-						Builder.Append(this.Class.SourceName).Append("*");
+						Builder.Append(this.Class.SourceName).Append('*');
 					}
 					break;
 
@@ -61,11 +61,11 @@ namespace EpicGames.UHT.Types
 					{
 						Builder.Append("const ");
 					}
-					Builder.Append("TObjectPtr<").Append(this.Class.SourceName).Append(">");
+					Builder.Append("TObjectPtr<").Append(this.Class.SourceName).Append('>');
 					break;
 
 				default:
-					Builder.Append("TObjectPtr<").Append(this.Class.SourceName).Append(">");
+					Builder.Append("TObjectPtr<").Append(this.Class.SourceName).Append('>');
 					break;
 			}
 			return Builder;
diff --git a/Engine/Source/Programs/Shared/EpicGames.UHT/Types/Properties/UhtSetProperty.cs b/Engine/Source/Programs/Shared/EpicGames.UHT/Types/Properties/UhtSetProperty.cs
--- a/Engine/Source/Programs/Shared/EpicGames.UHT/Types/Properties/UhtSetProperty.cs
+++ b/Engine/Source/Programs/Shared/EpicGames.UHT/Types/Properties/UhtSetProperty.cs
@@ -106,7 +106,7 @@ namespace EpicGames.UHT.Types
 					if (Builder[Builder.Length - 1] == '>')
 					{
 						// if our internal property type is a template class, add a space between the closing brackets b/c VS.NET cannot parse this correctly
-						Builder.Append(" ");
+						Builder.Append(' ');
 					}
 					Builder.Append('>');
 					break;
diff --git a/Engine/Source/Programs/Shared/EpicGames.UHT/Types/Properties/UhtSoftObjectProperty.cs b/Engine/Source/Programs/Shared/EpicGames.UHT/Types/Properties/UhtSoftObjectProperty.cs
--- a/Engine/Source/Programs/Shared/EpicGames.UHT/Types/Properties/UhtSoftObjectProperty.cs
+++ b/Engine/Source/Programs/Shared/EpicGames.UHT/Types/Properties/UhtSoftObjectProperty.cs
@@ -54,7 +54,7 @@ namespace EpicGames.UHT.Types
 			switch (TextType)
 			{
 				default:
-					Builder.Append("TSoftObjectPtr<").Append(this.Class.SourceName).Append(">");
+					Builder.Append("TSoftObjectPtr<").Append(this.Class.SourceName).Append('>');
 					break;
 			}
 			return Builder;
diff --git a/Engine/Source/Programs/Shared/EpicGames.UHT/Types/Properties/UhtTextProperty.cs b/Engine/Source/Programs/Shared/EpicGames.UHT/Types/Properties/UhtTextProperty.cs
--- a/Engine/Source/Programs/Shared/EpicGames.UHT/Types/Properties/UhtTextProperty.cs
+++ b/Engine/Source/Programs/Shared/EpicGames.UHT/Types/Properties/UhtTextProperty.cs
@@ -79,7 +79,7 @@ namespace EpicGames.UHT.Types
 						DefaultValueReader.Require('(');
 						StringView Value = DefaultValueReader.GetWrappedConstString();
 						DefaultValueReader.Require(')');
-						InnerDefaultValue.Append("\"").Append(Value).Append("\"");
+						InnerDefaultValue.Append('\"').Append(Value).Append('\"');
 						this.LogWarning("FText::FromString should be replaced with INVTEXT for default parameter values");
 					}
 					else
@@ -312,7 +312,7 @@ namespace EpicGames.UHT.Types
 			DefaultValueReader.Require(')');
 
 			// With UHT, we end up outputting just the integer part of the specified value.  It is using default locale information in UHT
-			InnerDefaultValue.Append("LOCGEN_CURRENCY(").Append((int)BaseValue).Append(", ").Append(CurrencyCode).Append(", ").Append(CultureNameString).Append(")");
+			InnerDefaultValue.Append("LOCGEN_CURRENCY(").Append((int)BaseValue).Append(", ").Append(CurrencyCode).Append(", ").Append(CultureNameString).Append(')');
 			return true;
 		}
 
@@ -588,7 +588,7 @@ namespace EpicGames.UHT.Types
 				InnerDefaultValue.Append(", ");
 			}
 
-			StringView Culture = DefaultValueReader.GetConstQuotedString();
+			StringView _ = DefaultValueReader.GetConstQuotedString();
 			InnerDefaultValue.Append("\"\""); // UHT doesn't write out the culture
 
 			DefaultValueReader.Require(')');
@@ -610,7 +610,7 @@ namespace EpicGames.UHT.Types
 			DefaultValueReader.Require('(');
 			InnerDefaultValue.Append(MacroToken.Value).Append('(');
 
-			long UnixTimestampValue = 0;
+			long UnixTimestampValue;
 			UhtToken TimestampToken = DefaultValueReader.GetToken();
 			switch (TimestampToken.TokenType)
 			{
@@ -664,11 +664,11 @@ namespace EpicGames.UHT.Types
 			}
 
 			DefaultValueReader.Require(',');
-			StringView CultureName = DefaultValueReader.GetConstQuotedString();
+			StringView _ = DefaultValueReader.GetConstQuotedString();
 			InnerDefaultValue.Append(", ").Append("\"\""); // We don't write out the culture
 
 			DefaultValueReader.Require(')');
-			InnerDefaultValue.Append(")");
+			InnerDefaultValue.Append(')');
 			return true;
 		}
 
diff --git a/Engine/Source/Programs/Shared/EpicGames.UHT/Types/Properties/UhtWeakObjectPtrProperty.cs b/Engine/Source/Programs/Shared/EpicGames.UHT/Types/Properties/UhtWeakObjectPtrProperty.cs
--- a/Engine/Source/Programs/Shared/EpicGames.UHT/Types/Properties/UhtWeakObjectPtrProperty.cs
+++ b/Engine/Source/Programs/Shared/EpicGames.UHT/Types/Properties/UhtWeakObjectPtrProperty.cs
@@ -50,11 +50,11 @@ namespace EpicGames.UHT.Types
 				default:
 					if (this.PropertyFlags.HasAnyFlags(EPropertyFlags.AutoWeak))
 					{
-						Builder.Append("TAutoWeakObjectPtr<").Append(this.Class.SourceName).Append(">");
+						Builder.Append("TAutoWeakObjectPtr<").Append(this.Class.SourceName).Append('>');
 					}
 					else
 					{
-						Builder.Append("TWeakObjectPtr<").Append(this.Class.SourceName).Append(">");
+						Builder.Append("TWeakObjectPtr<").Append(this.Class.SourceName).Append('>');
 					}
 					break;
 			}
diff --git a/Engine/Source/Programs/Shared/EpicGames.UHT/Types/UhtClass.cs b/Engine/Source/Programs/Shared/EpicGames.UHT/Types/UhtClass.cs
--- a/Engine/Source/Programs/Shared/EpicGames.UHT/Types/UhtClass.cs
+++ b/Engine/Source/Programs/Shared/EpicGames.UHT/Types/UhtClass.cs
@@ -7,7 +7,6 @@ using EpicGames.UHT.Tokenizer;
 using EpicGames.UHT.Utils;
 using System;
 using System.Collections.Generic;
-using System.Diagnostics;
 using System.Text;
 using System.Text.Json.Serialization;
 
@@ -17,7 +16,7 @@ namespace EpicGames.UHT.Types
 	/// Series of flags not part of the engine's class flags that affect code generation or verification
 	/// </summary>
 	[Flags]
-	public enum UhtClassExportFlags : UInt32
+	public enum UhtClassExportFlags : Int32
 	{
 
 		/// <summary>
@@ -629,7 +628,7 @@ namespace EpicGames.UHT.Types
 					{
 						if (this.ClassExportFlags.HasAnyFlags(UhtClassExportFlags.SelfHasReplicatedProperties))
 						{
-							if (TryGetDeclaration("GetLifetimeReplicatedProps", out UhtFoundDeclaration FoundDeclaration))
+							if (TryGetDeclaration("GetLifetimeReplicatedProps", out UhtFoundDeclaration _))
 							{
 								this.ClassExportFlags |= UhtClassExportFlags.HasGetLifetimeReplicatedProps;
 							}
@@ -728,7 +727,7 @@ namespace EpicGames.UHT.Types
 		/// <param name="Declaration">The declaration being tested</param>
 		/// <param name="NameIndex">The index of the token with the expected getter/setter name</param>
 		/// <returns>True if the declaration matches</returns>
-		private bool TryMatchGetterSetter(UhtProperty Property, bool bSetter, UhtDeclaration Declaration, int NameIndex)
+		private static bool TryMatchGetterSetter(UhtProperty Property, bool bSetter, UhtDeclaration Declaration, int NameIndex)
 		{
 			UhtToken[] Tokens = Declaration.Tokens;
 			int TokenIndex = 0;
@@ -740,8 +739,8 @@ namespace EpicGames.UHT.Types
 				// Nothing to do in the body
 			}
 
-			int TypeIndex = 0;
-			int TypeCount = 0;
+			int TypeIndex;
+			int TypeCount;
 
 			// Verify the format of the function declaration and extract the type
 			if (bSetter)
@@ -810,7 +809,9 @@ namespace EpicGames.UHT.Types
 				{
 					return false;
 				}
+#pragma warning disable IDE0059 // Unnecessary assignment of a value
 				TokenIndex++;
+#pragma warning restore IDE0059 // Unnecessary assignment of a value
 			}
 
 			// Verify the type
@@ -822,11 +823,11 @@ namespace EpicGames.UHT.Types
 				{
 					if (Declaration.Function == null)
 					{
-						Builder.Append("*");
+						Builder.Append('*');
 					}
 					else
 					{
-						Builder.Append("[").Append(Property.ArrayDimensions).Append("]");
+						Builder.Append('[').Append(Property.ArrayDimensions).Append(']');
 					}
 				}
 				string ExpectedType = Builder.ToString();
@@ -921,14 +922,13 @@ namespace EpicGames.UHT.Types
 		/// <param name="Property">Property requesting the getter/setter</param>
 		/// <param name="bSetter">True if this is a setter</param>
 		/// <returns>Resulting dictionary</returns>
-		private Dictionary<string, List<GetterSetterToResolve>> AddGetterSetter(Dictionary<string, List<GetterSetterToResolve>>? GSToResolve, string Name, UhtProperty Property, bool bSetter)
+		private static Dictionary<string, List<GetterSetterToResolve>> AddGetterSetter(Dictionary<string, List<GetterSetterToResolve>>? GSToResolve, string Name, UhtProperty Property, bool bSetter)
 		{
 			if (GSToResolve == null)
 			{
 				GSToResolve = new Dictionary<string, List<GetterSetterToResolve>>();
 			}
-			List<GetterSetterToResolve>? GSList = null;
-			if (!GSToResolve.TryGetValue(Name, out GSList))
+			if (!GSToResolve.TryGetValue(Name, out List<GetterSetterToResolve>? GSList))
 			{
 				GSList = new List<GetterSetterToResolve>();
 				GSToResolve.Add(Name, GSList);
@@ -1020,87 +1020,83 @@ namespace EpicGames.UHT.Types
 						}
 
 						// Make sure that all interface functions are implemented property
-						if (this.Bases != null)
+						// Iterate over all the bases looking for any interfaces
+						foreach (UhtStruct BaseStruct in this.Bases)
 						{
 
-							// Iterate over all the bases looking for any interfaces
-							foreach (UhtStruct BaseStruct in this.Bases)
+							// Skip children that aren't interfaces or if a common ancestor
+							UhtClass? BaseClass = BaseStruct as UhtClass;
+							if (BaseClass == null || BaseClass.ClassType == UhtClassType.Class || IsChildOf(BaseClass) )
 							{
+								continue;
+							}
 
-								// Skip children that aren't interfaces or if a common ancestor
-								UhtClass? BaseClass = BaseStruct as UhtClass;
-								if (BaseClass == null || BaseClass.ClassType == UhtClassType.Class || IsChildOf(BaseClass) )
+							// Loop through the function to make sure they are implemented
+							List<UhtType> BaseChildren = BaseClass.AlternateObject != null ? BaseClass.AlternateObject.Children : BaseClass.Children;
+							foreach (UhtType BaseChild in BaseChildren)
+							{
+								UhtFunction? BaseFunction = BaseChild as UhtFunction;
+								if (BaseFunction == null)
 								{
-									continue;
+									continue; 
 								}
 
-								// Loop through the function to make sure they are implemented
-								List<UhtType> BaseChildren = BaseClass.AlternateObject != null ? BaseClass.AlternateObject.Children : BaseClass.Children;
-								foreach (UhtType BaseChild in BaseChildren)
+								// Delegate signature functions are simple stubs and aren't required to be implemented (they are not callable)
+								bool bImplemented = BaseFunction.FunctionFlags.HasAnyFlags(EFunctionFlags.Delegate);
+
+								// Try to find an existing function
+								foreach (UhtType Child in this.Children)
 								{
-									UhtFunction? BaseFunction = BaseChild as UhtFunction;
-									if (BaseFunction == null)
+									UhtFunction? Function = Child as UhtFunction;
+									if (Function == null || Function.SourceName != BaseFunction.SourceName)
 									{
-										continue; 
+										continue;
 									}
 
-									// Delegate signature functions are simple stubs and aren't required to be implemented (they are not callable)
-									bool bImplemented = BaseFunction.FunctionFlags.HasAnyFlags(EFunctionFlags.Delegate);
-
-									// Try to find an existing function
-									foreach (UhtType Child in this.Children)
+									if (BaseFunction.FunctionFlags.HasAnyFlags(EFunctionFlags.Event) && !Function.FunctionFlags.HasAnyFlags(EFunctionFlags.Event))
 									{
-										UhtFunction? Function = Child as UhtFunction;
-										if (Function == null || Function.SourceName != BaseFunction.SourceName)
-										{
-											continue;
-										}
-
-										if (BaseFunction.FunctionFlags.HasAnyFlags(EFunctionFlags.Event) && !Function.FunctionFlags.HasAnyFlags(EFunctionFlags.Event))
-										{
-											Function.LogError($"Implementation of function '{Function.SourceName}' must be declared as 'event' to match declaration in interface '{BaseClass.SourceName}'");
-										}
+										Function.LogError($"Implementation of function '{Function.SourceName}' must be declared as 'event' to match declaration in interface '{BaseClass.SourceName}'");
+									}
 
-										if (BaseFunction.FunctionFlags.HasAnyFlags(EFunctionFlags.Delegate) && !Function.FunctionFlags.HasAnyFlags(EFunctionFlags.Delegate))
-										{
-											Function.LogError($"Implementation of function '{Function.SourceName}' must be declared as 'delegate' to match declaration in interface '{BaseClass.SourceName}'");
-										}
+									if (BaseFunction.FunctionFlags.HasAnyFlags(EFunctionFlags.Delegate) && !Function.FunctionFlags.HasAnyFlags(EFunctionFlags.Delegate))
+									{
+										Function.LogError($"Implementation of function '{Function.SourceName}' must be declared as 'delegate' to match declaration in interface '{BaseClass.SourceName}'");
+									}
 
-										bImplemented = true;
+									bImplemented = true;
 
-										if (BaseFunction.Children.Count != Function.Children.Count)
-										{
-											Function.LogError($"Implementation of function '{Function.SourceName}' conflicts with interface '{BaseClass.SourceName}' - different number of parameters ({Function.Children.Count}/{BaseFunction.Children.Count})");
-											continue;
-										}
+									if (BaseFunction.Children.Count != Function.Children.Count)
+									{
+										Function.LogError($"Implementation of function '{Function.SourceName}' conflicts with interface '{BaseClass.SourceName}' - different number of parameters ({Function.Children.Count}/{BaseFunction.Children.Count})");
+										continue;
+									}
 
-										for (int Index = 0; Index < Function.Children.Count; ++Index)
+									for (int Index = 0; Index < Function.Children.Count; ++Index)
+									{
+										UhtProperty BaseProperty = (UhtProperty)BaseFunction.Children[Index];
+										UhtProperty Property = (UhtProperty)Function.Children[Index];
+										if (!BaseProperty.MatchesType(Property))
 										{
-											UhtProperty BaseProperty = (UhtProperty)BaseFunction.Children[Index];
-											UhtProperty Property = (UhtProperty)Function.Children[Index];
-											if (!BaseProperty.MatchesType(Property))
+											if (BaseProperty.PropertyFlags.HasAnyFlags(EPropertyFlags.ReturnParm))
 											{
-												if (BaseProperty.PropertyFlags.HasAnyFlags(EPropertyFlags.ReturnParm))
-												{
-													Function.LogError($"Implementation of function '{Function.SourceName}' conflicts only by return type with interface '{BaseClass.SourceName}'");
-												}
-												else
-												{
-													Function.LogError($"Implementation of function '{Function.SourceName}' conflicts type with interface '{BaseClass.SourceName}' - parameter {Index} '{Property.SourceName}'");
-												}
+												Function.LogError($"Implementation of function '{Function.SourceName}' conflicts only by return type with interface '{BaseClass.SourceName}'");
+											}
+											else
+											{
+												Function.LogError($"Implementation of function '{Function.SourceName}' conflicts type with interface '{BaseClass.SourceName}' - parameter {Index} '{Property.SourceName}'");
 											}
 										}
 									}
+								}
 
-									// Verify that if this has blueprint-callable functions that are not implementable events, we've implemented them as a UFunction in the target class
-									if (!bImplemented
-										&& BaseFunction.FunctionFlags.HasAnyFlags(EFunctionFlags.BlueprintCallable)
-										&& !BaseFunction.FunctionFlags.HasAnyFlags(EFunctionFlags.BlueprintEvent)
-										&& !BaseClass.MetaData.ContainsKey(UhtNames.CannotImplementInterfaceInBlueprint)
-										&& (BaseClass.AlternateObject == null || !BaseClass.AlternateObject.MetaData.ContainsKey(UhtNames.CannotImplementInterfaceInBlueprint)))
-									{
-										this.LogError($"Missing UFunction implementation of function '{BaseFunction.SourceName}' from interface '{BaseClass.SourceName}'.  This function needs a UFUNCTION() declaration.");
-									}
+								// Verify that if this has blueprint-callable functions that are not implementable events, we've implemented them as a UFunction in the target class
+								if (!bImplemented
+									&& BaseFunction.FunctionFlags.HasAnyFlags(EFunctionFlags.BlueprintCallable)
+									&& !BaseFunction.FunctionFlags.HasAnyFlags(EFunctionFlags.BlueprintEvent)
+									&& !BaseClass.MetaData.ContainsKey(UhtNames.CannotImplementInterfaceInBlueprint)
+									&& (BaseClass.AlternateObject == null || !BaseClass.AlternateObject.MetaData.ContainsKey(UhtNames.CannotImplementInterfaceInBlueprint)))
+								{
+									this.LogError($"Missing UFunction implementation of function '{BaseFunction.SourceName}' from interface '{BaseClass.SourceName}'.  This function needs a UFUNCTION() declaration.");
 								}
 							}
 						}
@@ -1125,7 +1121,7 @@ namespace EpicGames.UHT.Types
 						bool bCannotImplementInBlueprints = (!bCanImplementInBlueprints && this.MetaData.ContainsKey(UhtNames.IsBlueprintBase))
 							|| this.MetaData.ContainsKey(UhtNames.CannotImplementInterfaceInBlueprint);
 
-						bool bCanImplementInterfaceInBlueprint = !this.MetaData.ContainsKey(UhtNames.CannotImplementInterfaceInBlueprint);
+						//bool bCanImplementInterfaceInBlueprint = !this.MetaData.ContainsKey(UhtNames.CannotImplementInterfaceInBlueprint);
 						foreach (UhtType Child in this.Children)
 						{
 							if (Child is UhtFunction Function)
@@ -1171,7 +1167,7 @@ namespace EpicGames.UHT.Types
 			return Options |= UhtValidationOptions.Shadowing | UhtValidationOptions.Deprecated;
 		}
 
-		void ValidateBlueprintPopertyGetter(UhtProperty Property, string FunctionName, UhtFunction? TargetFunction)
+		static void ValidateBlueprintPopertyGetter(UhtProperty Property, string FunctionName, UhtFunction? TargetFunction)
 		{
 			if (TargetFunction == null)
 			{
@@ -1201,7 +1197,7 @@ namespace EpicGames.UHT.Types
 			}
 		}
 
-		void ValidateBlueprintPopertySetter(UhtProperty Property, string FunctionName, UhtFunction? TargetFunction)
+		static void ValidateBlueprintPopertySetter(UhtProperty Property, string FunctionName, UhtFunction? TargetFunction)
 		{
 			if (TargetFunction == null)
 			{
@@ -1235,7 +1231,7 @@ namespace EpicGames.UHT.Types
 			}
 		}
 
-		void ValidateRepNotifyCallback(UhtProperty Property, string FunctionName, UhtFunction? TargetFunction)
+		static void ValidateRepNotifyCallback(UhtProperty Property, string FunctionName, UhtFunction? TargetFunction)
 		{
 			if (TargetFunction == null)
 			{
@@ -1383,14 +1379,11 @@ namespace EpicGames.UHT.Types
 			}
 
 			// Add any other base classes
-			if (this.Bases != null)
+			foreach (UhtStruct Base in this.Bases)
 			{
-				foreach (UhtStruct Base in this.Bases)
+				if (Base is UhtClass BaseClass)
 				{
-					if (Base is UhtClass BaseClass)
-					{
-						Collector.AddCrossModuleReference(BaseClass, false);
-					}
+					Collector.AddCrossModuleReference(BaseClass, false);
 				}
 			}
 
@@ -1427,14 +1420,11 @@ namespace EpicGames.UHT.Types
 		{
 			for (UhtClass? SuperClass = this; SuperClass != null; SuperClass = SuperClass.SuperClass)
 			{
-				if (SuperClass.Bases != null)
+				foreach (UhtStruct Struct in SuperClass.Bases)
 				{
-					foreach (UhtStruct Struct in SuperClass.Bases)
+					if (Struct.IsChildOf(Interface))
 					{
-						if (Struct.IsChildOf(Interface))
-						{
-							return true;
-						}
+						return true;
 					}
 				}
 			}
diff --git a/Engine/Source/Programs/Shared/EpicGames.UHT/Types/UhtEnum.cs b/Engine/Source/Programs/Shared/EpicGames.UHT/Types/UhtEnum.cs
--- a/Engine/Source/Programs/Shared/EpicGames.UHT/Types/UhtEnum.cs
+++ b/Engine/Source/Programs/Shared/EpicGames.UHT/Types/UhtEnum.cs
@@ -136,7 +136,7 @@ namespace EpicGames.UHT.Types
 		/// <summary>
 		/// Collection of enumeration values
 		/// </summary>
-		public List<UhtEnumValue> EnumValues { get; set; }
+		public List<UhtEnumValue> EnumValues { get; }
 
 		/// <inheritdoc/>
 		[JsonIgnore]
diff --git a/Engine/Source/Programs/Shared/EpicGames.UHT/Types/UhtField.cs b/Engine/Source/Programs/Shared/EpicGames.UHT/Types/UhtField.cs
--- a/Engine/Source/Programs/Shared/EpicGames.UHT/Types/UhtField.cs
+++ b/Engine/Source/Programs/Shared/EpicGames.UHT/Types/UhtField.cs
@@ -16,7 +16,7 @@ namespace EpicGames.UHT.Types
 		/// </summary>
 		/// <param name="Outer">Outer object</param>
 		/// <param name="LineNumber">Line number of declaration</param>
-		public UhtField(UhtType Outer, int LineNumber) : base(Outer, LineNumber)
+		protected UhtField(UhtType Outer, int LineNumber) : base(Outer, LineNumber)
 		{
 		}
 	}
diff --git a/Engine/Source/Programs/Shared/EpicGames.UHT/Types/UhtFunction.cs b/Engine/Source/Programs/Shared/EpicGames.UHT/Types/UhtFunction.cs
--- a/Engine/Source/Programs/Shared/EpicGames.UHT/Types/UhtFunction.cs
+++ b/Engine/Source/Programs/Shared/EpicGames.UHT/Types/UhtFunction.cs
@@ -14,7 +14,7 @@ namespace EpicGames.UHT.Types
 	/// Series of flags not part of the engine's function flags that affect code generation or verification
 	/// </summary>
 	[Flags]
-	public enum UhtFunctionExportFlags : UInt32
+	public enum UhtFunctionExportFlags : Int32
 	{
 
 		/// <summary>
@@ -201,7 +201,7 @@ namespace EpicGames.UHT.Types
 		/// <summary>
 		/// Suffix added to delegate engine names 
 		/// </summary>
-		public static string GeneratedDelegateSignatureSuffix = "__DelegateSignature";
+		public static readonly string GeneratedDelegateSignatureSuffix = "__DelegateSignature";
 
 		/// <summary>
 		/// Engine function flags
@@ -836,7 +836,7 @@ namespace EpicGames.UHT.Types
 			}
 			else
 			{
-				Builder.Append("\"");
+				Builder.Append('\"');
 			}
 			this.LogError(Builder.ToString());
 		}
@@ -978,7 +978,7 @@ namespace EpicGames.UHT.Types
 		/// <param name="Name">Name of the parameter</param>
 		/// <param name="Text">Text of the parameter tooltip</param>
 		/// <param name="ParamMap">Parameter map to add the tooltip</param>
-		private void AddParameterToolTip(StringView Name, StringView Text, Dictionary<StringView, StringView> ParamMap)
+		private static void AddParameterToolTip(StringView Name, StringView Text, Dictionary<StringView, StringView> ParamMap)
 		{
 			ParamMap.Add(Name, new StringView(UhtFCString.TabsToSpaces(Text.Memory.Trim(), 4, false)));
 		}
@@ -988,7 +988,7 @@ namespace EpicGames.UHT.Types
 		/// </summary>
 		/// <param name="Text">Text of the parameter tooltip</param>
 		/// <param name="ParamMap">Parameter map to add the tooltip</param>
-		private void AddParameterToolTip(StringView Text, Dictionary<StringView, StringView> ParamMap)
+		private static void AddParameterToolTip(StringView Text, Dictionary<StringView, StringView> ParamMap)
 		{
 			ReadOnlyMemory<char> Trimmed = Text.Memory.Trim();
 			int NameEnd = Trimmed.Span.IndexOf(' ');
@@ -1003,7 +1003,7 @@ namespace EpicGames.UHT.Types
 		/// </summary>
 		/// <param name="Input">The function input comment</param>
 		/// <returns>Dictionary of parameter names and the documentation.  The return value will have a name of "ReturnValue"</returns>
-		private Dictionary<StringView, StringView> GetParameterToolTipsFromFunctionComment(StringView Input)
+		private static Dictionary<StringView, StringView> GetParameterToolTipsFromFunctionComment(StringView Input)
 		{
 			const string ParamTag = "@param";
 			const string ReturnTag = "@return";
diff --git a/Engine/Source/Programs/Shared/EpicGames.UHT/Types/UhtHeaderFile.cs b/Engine/Source/Programs/Shared/EpicGames.UHT/Types/UhtHeaderFile.cs
--- a/Engine/Source/Programs/Shared/EpicGames.UHT/Types/UhtHeaderFile.cs
+++ b/Engine/Source/Programs/Shared/EpicGames.UHT/Types/UhtHeaderFile.cs
@@ -42,7 +42,7 @@ namespace EpicGames.UHT.Types
 	/// Series of flags not part of the engine's class flags that affect code generation or verification
 	/// </summary>
 	[Flags]
-	public enum UhtHeaderFileExportFlags : UInt32
+	public enum UhtHeaderFileExportFlags : Int32
 	{
 		/// <summary>
 		/// No export flags
@@ -103,9 +103,9 @@ namespace EpicGames.UHT.Types
 	/// </summary>
 	public class UhtHeaderFile : UhtType
 	{
-		private UhtSimpleMessageSite MessageSite;
-		private UhtSourceFile SourceFile;
-		private List<UhtHeaderFile> ReferencedHeaders = new List<UhtHeaderFile>();
+		private readonly UhtSimpleMessageSite MessageSite;
+		private readonly UhtSourceFile SourceFile;
+		private readonly List<UhtHeaderFile> ReferencedHeaders = new List<UhtHeaderFile>();
 
 		/// <summary>
 		/// Contents of the header
@@ -200,7 +200,7 @@ namespace EpicGames.UHT.Types
 		/// Collection of headers directly included by this header
 		/// </summary>
 		[JsonIgnore]
-		public List<UhtHeaderFile> IncludedHeaders { get; set; } = new List<UhtHeaderFile>();
+		public List<UhtHeaderFile> IncludedHeaders { get; } = new List<UhtHeaderFile>();
 
 		#region IUHTMessageSite implementation
 
@@ -227,7 +227,7 @@ namespace EpicGames.UHT.Types
 			this.FileNameWithoutExtension = System.IO.Path.GetFileNameWithoutExtension(this.SourceFile.FilePath);
 			this.GeneratedHeaderFileName = this.FileNameWithoutExtension + ".generated.h";
 			this.SourceName = System.IO.Path.GetFileName(this.SourceFile.FilePath);
-			this.bIsNoExportTypes = string.Compare(this.SourceFile.FileName, "NoExportTypes", true) == 0;
+			this.bIsNoExportTypes = string.Equals(this.SourceFile.FileName, "NoExportTypes", StringComparison.OrdinalIgnoreCase);
 		}
 
 		/// <summary>
@@ -324,12 +324,12 @@ namespace EpicGames.UHT.Types
 		}
 
 		/// <inheritdoc/>
-		public override void GetPathName(StringBuilder Builder, UhtType? StopOuter = null)
+		public override void AppendPathName(StringBuilder Builder, UhtType? StopOuter = null)
 		{
 			// Headers do not contribute to path names
 			if (this != StopOuter && this.Outer != null)
 			{
-				this.Outer.GetPathName(Builder, StopOuter);
+				this.Outer.AppendPathName(Builder, StopOuter);
 			}
 		}
 
diff --git a/Engine/Source/Programs/Shared/EpicGames.UHT/Types/UhtManifestFile.cs b/Engine/Source/Programs/Shared/EpicGames.UHT/Types/UhtManifestFile.cs
--- a/Engine/Source/Programs/Shared/EpicGames.UHT/Types/UhtManifestFile.cs
+++ b/Engine/Source/Programs/Shared/EpicGames.UHT/Types/UhtManifestFile.cs
@@ -18,7 +18,7 @@ namespace EpicGames.UHT.Types
 		/// </summary>
 		public UHTManifest? Manifest = null;
 
-		private UhtSourceFile SourceFile;
+		private readonly UhtSourceFile SourceFile;
 
 		#region IUHTMessageSite implementation
 
diff --git a/Engine/Source/Programs/Shared/EpicGames.UHT/Types/UhtMetaData.cs b/Engine/Source/Programs/Shared/EpicGames.UHT/Types/UhtMetaData.cs
--- a/Engine/Source/Programs/Shared/EpicGames.UHT/Types/UhtMetaData.cs
+++ b/Engine/Source/Programs/Shared/EpicGames.UHT/Types/UhtMetaData.cs
@@ -45,7 +45,7 @@ namespace EpicGames.UHT.Types
 		/// </summary>
 		/// <param name="Name">Meta data name</param>
 		/// <param name="Index">Meta data index</param>
-		public UhtMetaDataKey(string Name, int Index = UhtMetaData.INDEX_NONE)
+		public UhtMetaDataKey(string Name, int Index = UhtMetaData.IndexNone)
 		{
 			this.Name = Name;
 			this.Index = Index;
@@ -57,7 +57,7 @@ namespace EpicGames.UHT.Types
 		/// <returns>String representation</returns>
 		public override string ToString()
 		{
-			if (Index == UhtMetaData.INDEX_NONE)
+			if (Index == UhtMetaData.IndexNone)
 			{
 				return this.Name.ToString();
 			}
@@ -124,12 +124,12 @@ namespace EpicGames.UHT.Types
 		/// <summary>
 		/// Helper comparer for meta data keys
 		/// </summary>
-		private static UhtMetaDataKeyComparer Comparer = new UhtMetaDataKeyComparer();
+		private static readonly UhtMetaDataKeyComparer Comparer = new UhtMetaDataKeyComparer();
 
 		/// <summary>
 		/// Empty collection of meta data
 		/// </summary>
-		public static UhtMetaData Empty = new UhtMetaData(null, null);
+		public static readonly UhtMetaData Empty = new UhtMetaData(null, null);
 
 		/// <summary>
 		/// The meta data of the outer object for the type that owns this meta data
@@ -165,7 +165,7 @@ namespace EpicGames.UHT.Types
 		/// <summary>
 		/// Index for a meta data key associated with the owning object
 		/// </summary>
-		public const int INDEX_NONE = -1;
+		public const int IndexNone = -1;
 
 		/// <summary>
 		/// Construct new meta data
@@ -213,7 +213,7 @@ namespace EpicGames.UHT.Types
 		/// <param name="Name">Name of the meta data</param>
 		/// <param name="NameIndex">Enumeration value index or -1 for the type's meta data</param>
 		/// <returns>True if the key is found</returns>
-		public bool ContainsKey(string Name, int NameIndex = INDEX_NONE)
+		public bool ContainsKey(string Name, int NameIndex = IndexNone)
 		{
 			if (this.Dictionary == null)
 			{
@@ -228,7 +228,7 @@ namespace EpicGames.UHT.Types
 		/// <param name="Name">Name of the meta data</param>
 		/// <param name="NameIndex">Enumeration value index or -1 for the type's meta data</param>
 		/// <returns>True if the key is found</returns>
-		public bool ContainsKeyHierarchical(string Name, int NameIndex = INDEX_NONE)
+		public bool ContainsKeyHierarchical(string Name, int NameIndex = IndexNone)
 		{
 			for (UhtMetaData? Current = this; Current != null; Current = Current.Parent)
 			{
@@ -248,7 +248,7 @@ namespace EpicGames.UHT.Types
 		/// <returns>True if the key is found</returns>
 		public bool TryGetValue(string Name, [NotNullWhen(true)] out string? Value)
 		{
-			return TryGetValue(Name, INDEX_NONE, out Value);
+			return TryGetValue(Name, IndexNone, out Value);
 		}
 
 		/// <summary>
@@ -259,7 +259,7 @@ namespace EpicGames.UHT.Types
 		/// <returns>True if the key is found</returns>
 		public bool TryGetValueHierarchical(string Name, [NotNullWhen(true)] out string? Value)
 		{
-			return TryGetValueHierarchical(Name, INDEX_NONE, out Value);
+			return TryGetValueHierarchical(Name, IndexNone, out Value);
 		}
 
 		/// <summary>
@@ -305,7 +305,7 @@ namespace EpicGames.UHT.Types
 		/// <param name="Name">Name of the meta data key</param>
 		/// <param name="NameIndex">Index of the meta data key</param>
 		/// <returns>Meta data value or empty string if not found.</returns>
-		public string GetValueOrDefault(string Name, int NameIndex = INDEX_NONE)
+		public string GetValueOrDefault(string Name, int NameIndex = IndexNone)
 		{
 			string? Out;
 			if (TryGetValue(Name, NameIndex, out Out))
@@ -321,7 +321,7 @@ namespace EpicGames.UHT.Types
 		/// <param name="Name">Name of the meta data key</param>
 		/// <param name="NameIndex">Index of the meta data key</param>
 		/// <returns>Meta data value or empty string if not found.</returns>
-		public string GetValueOrDefaultHierarchical(string Name, int NameIndex = INDEX_NONE)
+		public string GetValueOrDefaultHierarchical(string Name, int NameIndex = IndexNone)
 		{
 			string? Out;
 			if (TryGetValueHierarchical(Name, NameIndex, out Out))
@@ -337,7 +337,7 @@ namespace EpicGames.UHT.Types
 		/// <param name="Name">Name of the meta data key</param>
 		/// <param name="NameIndex">Index of the meta data key</param>
 		/// <returns>Boolean value or false if not found</returns>
-		public bool GetBoolean(string Name, int NameIndex = INDEX_NONE)
+		public bool GetBoolean(string Name, int NameIndex = IndexNone)
 		{
 			return UhtFCString.ToBool(GetValueOrDefault(Name, NameIndex));
 		}
@@ -348,7 +348,7 @@ namespace EpicGames.UHT.Types
 		/// <param name="Name">Name of the meta data key</param>
 		/// <param name="NameIndex">Index of the meta data key</param>
 		/// <returns>Boolean value or false if not found</returns>
-		public bool GetBooleanHierarchical(string Name, int NameIndex = INDEX_NONE)
+		public bool GetBooleanHierarchical(string Name, int NameIndex = IndexNone)
 		{
 			return UhtFCString.ToBool(GetValueOrDefaultHierarchical(Name, NameIndex));
 		}
@@ -359,7 +359,7 @@ namespace EpicGames.UHT.Types
 		/// <param name="Name">Name of the meta data key</param>
 		/// <param name="NameIndex">Index of the meta data key</param>
 		/// <returns>Double value or zero if not found</returns>
-		public double GetDouble(string Name, int NameIndex = INDEX_NONE)
+		public double GetDouble(string Name, int NameIndex = IndexNone)
 		{
 			string Value = GetValueOrDefault(Name, NameIndex);
 			double Result;
@@ -376,7 +376,7 @@ namespace EpicGames.UHT.Types
 		/// <param name="Name">Name of the meta data key</param>
 		/// <param name="NameIndex">Index of the meta data key</param>
 		/// <returns>String array or null if not found</returns>
-		public string[]? GetStringArray(string Name, int NameIndex = INDEX_NONE)
+		public string[]? GetStringArray(string Name, int NameIndex = IndexNone)
 		{
 			string? Temp;
 			if (TryGetValue(Name, NameIndex, out Temp))
@@ -392,7 +392,7 @@ namespace EpicGames.UHT.Types
 		/// <param name="Name">Name of the meta data key</param>
 		/// <param name="NameIndex">Index of the meta data key</param>
 		/// <returns>String array or null if not found</returns>
-		public string[]? GetStringArrayHierarchical(string Name, int NameIndex = INDEX_NONE)
+		public string[]? GetStringArrayHierarchical(string Name, int NameIndex = IndexNone)
 		{
 			string? Temp;
 			if (TryGetValueHierarchical(Name, NameIndex, out Temp))
@@ -409,7 +409,7 @@ namespace EpicGames.UHT.Types
 		/// <param name="Value">Value of the meta data</param>
 		public void Add(string Name, string Value)
 		{
-			Add(Name, INDEX_NONE, Value);
+			Add(Name, IndexNone, Value);
 		}
 
 		/// <summary>
@@ -440,7 +440,7 @@ namespace EpicGames.UHT.Types
 		/// <param name="Name">Name of the meta data key</param>
 		/// <param name="Value">Value of the meta data</param>
 		/// <param name="NameIndex">Index of the meta data key</param>
-		public void Add(string Name, bool Value, int NameIndex = INDEX_NONE)
+		public void Add(string Name, bool Value, int NameIndex = IndexNone)
 		{
 			Add(Name, NameIndex, Value ? "true" : "false");
 		}
@@ -491,7 +491,7 @@ namespace EpicGames.UHT.Types
 		/// </summary>
 		/// <param name="Name">Name of the meta data key</param>
 		/// <param name="NameIndex">Index of the meta data key</param>
-		public void Remove(string Name, int NameIndex = INDEX_NONE)
+		public void Remove(string Name, int NameIndex = IndexNone)
 		{
 			if (this.Dictionary != null)
 			{
@@ -518,7 +518,7 @@ namespace EpicGames.UHT.Types
 		/// <exception cref="UhtIceException">Thrown if an index is supplied (not -1) and no key conversion interface is set</exception>
 		public string GetKeyString(string Name, int NameIndex)
 		{
-			if (NameIndex == INDEX_NONE)
+			if (NameIndex == IndexNone)
 			{
 				return Name;
 			}
diff --git a/Engine/Source/Programs/Shared/EpicGames.UHT/Types/UhtProperty.cs b/Engine/Source/Programs/Shared/EpicGames.UHT/Types/UhtProperty.cs
--- a/Engine/Source/Programs/Shared/EpicGames.UHT/Types/UhtProperty.cs
+++ b/Engine/Source/Programs/Shared/EpicGames.UHT/Types/UhtProperty.cs
@@ -608,17 +608,17 @@ namespace EpicGames.UHT.Types
 		/// <summary>
 		/// Return the hash code for a given object
 		/// </summary>
-		/// <param name="Object">Object in question</param>
+		/// <param name="Obj">Object in question</param>
 		/// <returns>Hash code</returns>
-		public uint GetTypeHash(UhtObject Object);
+		public uint GetTypeHash(UhtObject Obj);
 
 		/// <summary>
 		/// Return the singleton name for the given object
 		/// </summary>
-		/// <param name="Object">Object in question</param>
+		/// <param name="Obj">Object in question</param>
 		/// <param name="bRegistered">If true, the singleton that returns the registered object is returned.</param>
 		/// <returns>Singleton function name</returns>
-		public string GetSingletonName(UhtObject? Object, bool bRegistered);
+		public string GetSingletonName(UhtObject? Obj, bool bRegistered);
 	}
 
 	/// <summary>
@@ -826,7 +826,7 @@ namespace EpicGames.UHT.Types
 		/// <summary>
 		/// Collection of recognized casts when parsing array dimensions
 		/// </summary>
-		private static string[] Casts = new string[]
+		private static readonly string[] Casts = new string[]
 		{
 			"(uint32)",
 			"(int32)",
@@ -996,7 +996,7 @@ namespace EpicGames.UHT.Types
 		/// </summary>
 		/// <param name="Outer">Outer type of the property</param>
 		/// <param name="LineNumber">Line number where property was declared</param>
-		public UhtProperty(UhtType Outer, int LineNumber) : base(Outer, LineNumber)
+		protected UhtProperty(UhtType Outer, int LineNumber) : base(Outer, LineNumber)
 		{
 			this.PropertyFlags = EPropertyFlags.None;
 			this.PropertyCaps = UhtPropertyCaps.CanBeContainerValue | UhtPropertyCaps.CanBeContainerKey | UhtPropertyCaps.CanHaveConfig;
@@ -1006,7 +1006,7 @@ namespace EpicGames.UHT.Types
 		/// Construct a new property
 		/// </summary>
 		/// <param name="PropertySettings">Property settings from parsing</param>
-		public UhtProperty(UhtPropertySettings PropertySettings) : base(PropertySettings.Outer, PropertySettings.LineNumber, PropertySettings.MetaData)
+		protected UhtProperty(UhtPropertySettings PropertySettings) : base(PropertySettings.Outer, PropertySettings.LineNumber, PropertySettings.MetaData)
 		{
 			this.SourceName = PropertySettings.SourceName;
 			// Engine name defaults to source name.  If it doesn't match what is coming in, then set it.
@@ -1254,7 +1254,7 @@ namespace EpicGames.UHT.Types
 		/// <param name="bRegistered">True if the registered singleton name is to be used</param>
 		/// <param name="bAppendNull">True if a "nullptr" is to be appended if the object is null</param>
 		/// <returns>Output builder</returns>
-		protected StringBuilder AppendMemberDefRef(StringBuilder Builder, IUhtPropertyMemberContext Context, UhtObject? Object, bool bRegistered, bool bAppendNull = false)
+		protected static StringBuilder AppendMemberDefRef(StringBuilder Builder, IUhtPropertyMemberContext Context, UhtObject? Object, bool bRegistered, bool bAppendNull = false)
 		{
 			if (Object != null)
 			{
@@ -1309,10 +1309,10 @@ namespace EpicGames.UHT.Types
 					Builder.Append("_REF");
 				}
 			}
-			Builder.Append("(");
+			Builder.Append('(');
 			if (this.ArrayDimensions != null)
 			{
-				Builder.AppendFunctionThunkParameterArrayType(this).Append(",");
+				Builder.AppendFunctionThunkParameterArrayType(this).Append(',');
 			}
 			else
 			{
@@ -1326,11 +1326,11 @@ namespace EpicGames.UHT.Types
 						break;
 
 					case UhtPGetArgumentType.TypeText:
-						Builder.AppendPropertyText(this, UhtPropertyTextType.FunctionThunkParameterArgType).Append(",");
+						Builder.AppendPropertyText(this, UhtPropertyTextType.FunctionThunkParameterArgType).Append(',');
 						break;
 				}
 			}
-			Builder.AppendFunctionThunkParameterName(this).Append(")");
+			Builder.AppendFunctionThunkParameterName(this).Append(')');
 			return Builder;
 		}
 
@@ -1428,7 +1428,7 @@ namespace EpicGames.UHT.Types
 					{
 						ReadOnlySpan<char> Dim = this.ArrayDimensions.AsSpan();
 
-						bool bAgain = false;
+						bool bAgain;
 						do
 						{
 							bAgain = false;
@@ -1470,7 +1470,6 @@ namespace EpicGames.UHT.Types
 									}
 								}
 							}
-
 						} while (bAgain && Dim.Length > 0);
 
 						//COMPATIBILITY-TODO - This method is more robust, but causes differences.  See UhtSession for future
@@ -1509,7 +1508,7 @@ namespace EpicGames.UHT.Types
 							}
 							if (Enum != null)
 							{
-								this.MetaData.Add(UhtNames.ArraySizeEnum, Enum.GetPathName());
+								this.MetaData.Add(UhtNames.ArraySizeEnum, Enum.PathName);
 							}
 						}
 					}
@@ -1704,16 +1703,16 @@ namespace EpicGames.UHT.Types
 		/// <summary>
 		/// Verify function argument
 		/// </summary>
-		protected virtual void ValidateFunctionArgument(UhtFunction Function, UhtValidationOptions Options)
+		protected virtual void ValidateFunctionArgument(UhtFunction Func, UhtValidationOptions Options)
 		{
-			if (Function.FunctionFlags.HasAnyFlags(EFunctionFlags.Net))
+			if (Func.FunctionFlags.HasAnyFlags(EFunctionFlags.Net))
 			{
 				if (this.PropertyFlags.HasExactFlags(EPropertyFlags.ReferenceParm | EPropertyFlags.ConstParm, EPropertyFlags.ReferenceParm))
 				{
 					this.LogError($"Replicated parameters cannot be passed by non-const reference");
 				}
 
-				if (Function.FunctionFlags.HasAnyFlags(EFunctionFlags.NetRequest))
+				if (Func.FunctionFlags.HasAnyFlags(EFunctionFlags.NetRequest))
 				{
 					if (this.PropertyFlags.HasExactFlags(EPropertyFlags.OutParm | EPropertyFlags.RepSkip, EPropertyFlags.OutParm))
 					{
@@ -1737,7 +1736,7 @@ namespace EpicGames.UHT.Types
 			}
 
 			// The following checks are not performed on the value of a container
-			if (Function.FunctionFlags.HasAnyFlags(EFunctionFlags.BlueprintEvent | EFunctionFlags.BlueprintCallable))
+			if (Func.FunctionFlags.HasAnyFlags(EFunctionFlags.BlueprintEvent | EFunctionFlags.BlueprintCallable))
 			{
 				// Check that the parameter name is valid and does not conflict with pre-defined types
 				if (InvalidParamNames.Contains(this.SourceName))
@@ -1911,7 +1910,7 @@ namespace EpicGames.UHT.Types
 		/// <param name="OuterSuffix">Current suffix</param>
 		/// <param name="NewSuffix">Suffix to be added</param>
 		/// <returns>Combination of the two suffixes</returns>
-		protected string GetNameSuffix(string OuterSuffix, string NewSuffix)
+		protected static string GetNameSuffix(string OuterSuffix, string NewSuffix)
 		{
 			return string.IsNullOrEmpty(OuterSuffix) ? NewSuffix : $"{OuterSuffix}{NewSuffix}";
 		}
@@ -2154,7 +2153,7 @@ namespace EpicGames.UHT.Types
 			}
 			else
 			{
-				Builder.Append("1");
+				Builder.Append('1');
 			}
 			return Builder;
 		}
diff --git a/Engine/Source/Programs/Shared/EpicGames.UHT/Types/UhtScriptStruct.cs b/Engine/Source/Programs/Shared/EpicGames.UHT/Types/UhtScriptStruct.cs
--- a/Engine/Source/Programs/Shared/EpicGames.UHT/Types/UhtScriptStruct.cs
+++ b/Engine/Source/Programs/Shared/EpicGames.UHT/Types/UhtScriptStruct.cs
@@ -13,7 +13,7 @@ namespace EpicGames.UHT.Types
 	/// Flags to represent information about a RigVM parameter
 	/// </summary>
 	[Flags]
-	public enum UhtRigVMParameterFlags : UInt32
+	public enum UhtRigVMParameterFlags : Int32
 	{
 		/// <summary>
 		/// No RigVM flags
@@ -404,8 +404,8 @@ namespace EpicGames.UHT.Types
 	/// </summary>
 	public class UhtRigVMMethodInfo
 	{
-		private static string NoPrefixInternal = string.Empty;
-		private static string ReturnPrefixInternal = "return ";
+		private static readonly string NoPrefixInternal = string.Empty;
+		private const string ReturnPrefixInternal = "return ";
 
 		/// <summary>
 		/// Return type of the method
@@ -420,7 +420,7 @@ namespace EpicGames.UHT.Types
 		/// <summary>
 		/// Method parameters
 		/// </summary>
-		public List<UhtRigVMParameter> Parameters { get; set; } = new List<UhtRigVMParameter>();
+		public List<UhtRigVMParameter> Parameters { get; } = new List<UhtRigVMParameter>();
 
 		/// <summary>
 		/// If the method has a return value, return "return".  Otherwise return nothing.
@@ -462,14 +462,14 @@ namespace EpicGames.UHT.Types
 		/// <summary>
 		/// List of the methods
 		/// </summary>
-		public List<UhtRigVMMethodInfo> Methods { get; set; } = new List<UhtRigVMMethodInfo>();
+		public List<UhtRigVMMethodInfo> Methods { get; } = new List<UhtRigVMMethodInfo>();
 	};
 
 	/// <summary>
 	/// Series of flags not part of the engine's script struct flags that affect code generation or verification
 	/// </summary>
 	[Flags]
-	public enum UhtScriptStructExportFlags : UInt32
+	public enum UhtScriptStructExportFlags : Int32
 	{
 
 		/// <summary>
diff --git a/Engine/Source/Programs/Shared/EpicGames.UHT/Types/UhtSourceFile.cs b/Engine/Source/Programs/Shared/EpicGames.UHT/Types/UhtSourceFile.cs
--- a/Engine/Source/Programs/Shared/EpicGames.UHT/Types/UhtSourceFile.cs
+++ b/Engine/Source/Programs/Shared/EpicGames.UHT/Types/UhtSourceFile.cs
@@ -12,7 +12,7 @@ namespace EpicGames.UHT.Types
 	/// </summary>
 	public class UhtSourceFile : IUhtMessageSite, IUhtMessageSource
 	{
-		private UhtSimpleMessageSite MessageSite;
+		private readonly UhtSimpleMessageSite MessageSite;
 		private UhtSourceFragment SourceFragment;
 
 		/// <summary>
diff --git a/Engine/Source/Programs/Shared/EpicGames.UHT/Types/UhtStruct.cs b/Engine/Source/Programs/Shared/EpicGames.UHT/Types/UhtStruct.cs
--- a/Engine/Source/Programs/Shared/EpicGames.UHT/Types/UhtStruct.cs
+++ b/Engine/Source/Programs/Shared/EpicGames.UHT/Types/UhtStruct.cs
@@ -84,8 +84,12 @@ namespace EpicGames.UHT.Types
 		/// <summary>
 		/// Base types
 		/// </summary>
-		[JsonConverter(typeof(UhtNullableTypeListJsonConverter<UhtStruct>))]
-		public List<UhtStruct>? Bases { get; set; } = null;
+		//[JsonConverter(typeof(UhtNullableTypeListJsonConverter<UhtStruct>))]
+		[JsonIgnore]
+		public IReadOnlyList<UhtStruct> Bases 
+		{ 
+			get => _basesInternal != null ? _basesInternal : _emptyBases;
+		}
 
 		/// <summary>
 		/// Super struct type
@@ -98,7 +102,7 @@ namespace EpicGames.UHT.Types
 		/// </summary>
 		/// <param name="Outer">Outer type</param>
 		/// <param name="LineNumber">Line number where definition begins</param>
-		public UhtStruct(UhtType Outer, int LineNumber) : base(Outer, LineNumber)
+		protected UhtStruct(UhtType Outer, int LineNumber) : base(Outer, LineNumber)
 		{
 		}
 
@@ -122,6 +126,9 @@ namespace EpicGames.UHT.Types
 			}
 			return false;
 		}
+		
+		private static readonly List<UhtStruct> _emptyBases = new List<UhtStruct>();
+		private List<UhtStruct>? _basesInternal = null;
 
 		#region Resolution support
 		/// <inheritdoc/>
@@ -132,12 +139,9 @@ namespace EpicGames.UHT.Types
 				this.Super.Resolve(ResolvePhase);
 			}
 
-			if (this.Bases != null)
+			foreach (UhtStruct Base in this.Bases)
 			{
-				foreach (UhtStruct Base in this.Bases)
-				{
-					Base.Resolve(ResolvePhase);
-				}
+				Base.Resolve(ResolvePhase);
 			}
 
 			base.ResolveSuper(ResolvePhase);
@@ -207,11 +211,11 @@ namespace EpicGames.UHT.Types
 					UhtStruct? Base = (UhtStruct?)this.FindType(FindOptions | UhtFindOptions.Class | UhtFindOptions.ScriptStruct | UhtFindOptions.SourceName | UhtFindOptions.NoSelf, BaseIdentifier);
 					if (Base != null)
 					{
-						if (this.Bases == null)
+						if (this._basesInternal == null)
 						{
-							this.Bases = new List<UhtStruct>();
+							this._basesInternal = new List<UhtStruct>();
 						}
-						this.Bases.Add(Base);
+						this._basesInternal.Add(Base);
 						this.HeaderFile.AddReferencedHeader(Base);
 						Base.Resolve(UhtResolvePhase.Bases);
 					}
diff --git a/Engine/Source/Programs/Shared/EpicGames.UHT/Types/UhtType.cs b/Engine/Source/Programs/Shared/EpicGames.UHT/Types/UhtType.cs
--- a/Engine/Source/Programs/Shared/EpicGames.UHT/Types/UhtType.cs
+++ b/Engine/Source/Programs/Shared/EpicGames.UHT/Types/UhtType.cs
@@ -150,7 +150,6 @@ namespace EpicGames.UHT.Types
 			}
 		}
 
-
 		/// <summary>
 		/// Get the name of the type
 		/// </summary>
@@ -556,6 +555,11 @@ namespace EpicGames.UHT.Types
 	/// </summary>
 	public enum UhtResolvePhase : Int32
 	{
+		/// <summary>
+		/// Initial resolve phase
+		/// </summary>
+		None,
+
 		/// <summary>
 		/// Check for any type that is invalid.  Invalid items will be removed from the children. 
 		/// Used to detect things that looked like native interfaces but no interface was found
@@ -583,7 +587,7 @@ namespace EpicGames.UHT.Types
 	/// </summary>
 	public class UhtDocumentationPolicy
 	{
-		private static UhtDocumentationPolicy StrictDocumentationPolicy = new UhtDocumentationPolicy
+		private static readonly UhtDocumentationPolicy StrictDocumentationPolicy = new UhtDocumentationPolicy
 		{
 			bPolicySet = true,
 			bClassOrStructCommentRequired = true,
@@ -593,7 +597,7 @@ namespace EpicGames.UHT.Types
 			bFloatRangesRequired = true,
 		};
 
-		private static UhtDocumentationPolicy NoDocumentationPolicy = new UhtDocumentationPolicy
+		private static readonly UhtDocumentationPolicy NoDocumentationPolicy = new UhtDocumentationPolicy
 		{
 			bPolicySet = false,
 			bClassOrStructCommentRequired = false,
@@ -661,7 +665,7 @@ namespace EpicGames.UHT.Types
 		/// <summary>
 		/// Empty list of type used when children are requested but no children have been added.
 		/// </summary>
-		private static List<UhtType> EmptyTypeList = new List<UhtType>();
+		private static readonly List<UhtType> EmptyTypeList = new List<UhtType>();
 
 		/// <summary>
 		/// All UHT runs are associated with a given session.  The session holds all the global information for a run.
@@ -753,13 +757,29 @@ namespace EpicGames.UHT.Types
 		/// Return a combination of the engine type name followed by the path name of the type
 		/// </summary>
 		[JsonIgnore]
-		public virtual string FullName { get => GetFullName(); }
+		public virtual string FullName 
+		{ 
+			get
+			{
+				StringBuilder Builder = new StringBuilder();
+				AppendFullName(Builder);
+				return Builder.ToString();
+			}
+		}
 
 		/// <summary>
 		/// Return the path name of the type which includes all parent outer objects excluding the header file
 		/// </summary>
 		[JsonIgnore]
-		public string PathName { get => GetPathName(); }
+		public string PathName
+		{
+			get
+			{
+				StringBuilder Builder = new StringBuilder();
+				AppendPathName(Builder, null);
+				return Builder.ToString();
+			}
+		}
 
 		/// <summary>
 		/// Return the name of the documentation policy to be used
@@ -781,6 +801,7 @@ namespace EpicGames.UHT.Types
 		/// Children types of this type
 		/// </summary>
 		[JsonConverter(typeof(UhtNullableTypeListJsonConverter<UhtType>))]
+		//TODO - this should return a IReadOnlyList
 		public List<UhtType> Children { get => ChildrenInternal == null ? UhtType.EmptyTypeList : ChildrenInternal; }
 
 		/// <summary>
@@ -791,7 +812,7 @@ namespace EpicGames.UHT.Types
 
 		private string? EngineNameInternal = null;
 		private List<UhtType>? ChildrenInternal = null;
-		private int ResolveState = 0;
+		private int ResolveState = (int)UhtResolvePhase.None * 2;
 
 		#region IUHTMessageSite implementation
 		/// <inheritdoc/>
@@ -1117,40 +1138,28 @@ namespace EpicGames.UHT.Types
 		#endregion
 
 		#region Name and user facing text support
-		/// <summary>
-		/// Return the path name of the type which includes the outers
-		/// </summary>
-		/// <param name="StopOuter">Type to stop at when generating the path</param>
-		/// <return>The path string</return>
-		public string GetPathName(UhtType? StopOuter = null)
-		{
-			StringBuilder Builder = new StringBuilder();
-			GetPathName(Builder, StopOuter);
-			return Builder.ToString();
-		}
-
 		/// <summary>
 		/// Return the path name of the type which includes the outers
 		/// </summary>
 		/// <param name="Builder">Destination builder</param>
 		/// <param name="StopOuter">Type to stop at when generating the path</param>
-		public virtual void GetPathName(StringBuilder Builder, UhtType? StopOuter = null)
+		public virtual void AppendPathName(StringBuilder Builder, UhtType? StopOuter = null)
 		{
-			GetPathNameInternal(Builder, StopOuter);
+			AppendPathNameInternal(Builder, StopOuter);
 			if (Builder.Length == 0)
 			{
 				Builder.Append("None");
 			}
 		}
 
-		private void GetPathNameInternal(StringBuilder Builder, UhtType? StopOuter = null)
+		private void AppendPathNameInternal(StringBuilder Builder, UhtType? StopOuter = null)
 		{
 			if (this != StopOuter)
 			{
 				UhtType? Outer = this.Outer;
 				if (Outer != null && Outer != StopOuter)
 				{
-					Outer.GetPathName(Builder, StopOuter);
+					Outer.AppendPathName(Builder, StopOuter);
 
 					// SubObjectDelimiter is used to indicate that this object's outer is not a UPackage
 					// In the C# version of UHT, we key off the header file
@@ -1173,26 +1182,15 @@ namespace EpicGames.UHT.Types
 			}
 		}
 
-		/// <summary>
-		/// Get the full type name which is the engine class name followed by the path name
-		/// </summary>
-		/// <returns>Full name string</returns>
-		public string GetFullName()
-		{
-			StringBuilder Builder = new StringBuilder();
-			GetFullName(Builder);
-			return Builder.ToString();
-		}
-
 		/// <summary>
 		/// Get the full type name which is the engine class name followed by the path name
 		/// </summary>
 		/// <param name="Builder">Destination builder</param>
-		public virtual void GetFullName(StringBuilder Builder)
+		public virtual void AppendFullName(StringBuilder Builder)
 		{
 			Builder.Append(this.EngineClassName);
-			Builder.Append(" ");
-			GetPathName(Builder);
+			Builder.Append(' ');
+			AppendPathName(Builder);
 		}
 
 		/// <summary>
@@ -1275,7 +1273,7 @@ namespace EpicGames.UHT.Types
 				return Builder;
 			}
 
-			if (Outer is UhtClass Class)
+			if (Outer is UhtClass)
 			{
 				Builder.Append('_');
 				Builder.Append(Outer.SourceName);
diff --git a/Engine/Source/Programs/Shared/EpicGames.UHT/Utils/UhtBuffer.cs b/Engine/Source/Programs/Shared/EpicGames.UHT/Utils/UhtBuffer.cs
--- a/Engine/Source/Programs/Shared/EpicGames.UHT/Utils/UhtBuffer.cs
+++ b/Engine/Source/Programs/Shared/EpicGames.UHT/Utils/UhtBuffer.cs
@@ -18,27 +18,27 @@ namespace EpicGames.UHT.Utils
 		/// <summary>
 		/// Any requests of the given size or smaller will be placed in bucket zero with the given size.
 		/// </summary>
-		private static int MinSize = 1024 * 16;
+		private const int MinSize = 1024 * 16;
 
 		/// <summary>
 		/// Adjustment to the bucket index to account for the minimum bucket size
 		/// </summary>
-		private static int BuckedAdjustment = BitOperations.Log2((uint)UhtBuffer.MinSize);
+		private static readonly int BuckedAdjustment = BitOperations.Log2((uint)UhtBuffer.MinSize);
 		
 		/// <summary>
 		/// Total number of supported buckets
 		/// </summary>
-		private static int BucketCount = 32 - UhtBuffer.BuckedAdjustment;
+		private static readonly int BucketCount = 32 - UhtBuffer.BuckedAdjustment;
 
 		/// <summary>
 		/// Bucket lookaside list
 		/// </summary>
-		private static UhtBuffer?[] LookAsideArray = new UhtBuffer?[UhtBuffer.BucketCount];
+		private static readonly UhtBuffer?[] LookAsideArray = new UhtBuffer?[UhtBuffer.BucketCount];
 
 		/// <summary>
 		/// The bucket index associated with the buffer
 		/// </summary>
-		private int Bucket;
+		private readonly int Bucket;
 
 		/// <summary>
 		/// Single list link to the next cached buffer
diff --git a/Engine/Source/Programs/Shared/EpicGames.UHT/Utils/UhtException.cs b/Engine/Source/Programs/Shared/EpicGames.UHT/Utils/UhtException.cs
--- a/Engine/Source/Programs/Shared/EpicGames.UHT/Utils/UhtException.cs
+++ b/Engine/Source/Programs/Shared/EpicGames.UHT/Utils/UhtException.cs
@@ -95,7 +95,7 @@ namespace EpicGames.UHT.Utils
 		/// <param name="MessageSite">Message site that may contain extra context</param>
 		/// <param name="ExtraContext">Additional caller supplied context</param>
 		/// <returns></returns>
-		private string FormatExtraContext(IUhtMessageSite MessageSite, object? ExtraContext = null)
+		private static string FormatExtraContext(IUhtMessageSite MessageSite, object? ExtraContext = null)
 		{
 			StringBuilder Builder = new StringBuilder(" while parsing ");
 			int StartingLength = Builder.Length;
diff --git a/Engine/Source/Programs/Shared/EpicGames.UHT/Utils/UhtExport.cs b/Engine/Source/Programs/Shared/EpicGames.UHT/Utils/UhtExport.cs
--- a/Engine/Source/Programs/Shared/EpicGames.UHT/Utils/UhtExport.cs
+++ b/Engine/Source/Programs/Shared/EpicGames.UHT/Utils/UhtExport.cs
@@ -2,7 +2,6 @@
 
 using EpicGames.Core;
 using EpicGames.UHT.Types;
-using System;
 using System.Collections.Generic;
 using System.Text;
 using System.Threading.Tasks;
diff --git a/Engine/Source/Programs/Shared/EpicGames.UHT/Utils/UhtFCString.cs b/Engine/Source/Programs/Shared/EpicGames.UHT/Utils/UhtFCString.cs
--- a/Engine/Source/Programs/Shared/EpicGames.UHT/Utils/UhtFCString.cs
+++ b/Engine/Source/Programs/Shared/EpicGames.UHT/Utils/UhtFCString.cs
@@ -16,7 +16,7 @@ namespace EpicGames.UHT.Utils
 		/// <summary>
 		/// Delimiter between the subobjects
 		/// </summary>
-		public static char SubObjectDelimiter = ':';
+		public static readonly char SubObjectDelimiter = ':';
 
 		/// <summary>
 		/// Test to see if the string is a boolean
@@ -62,15 +62,15 @@ namespace EpicGames.UHT.Utils
 			{
 				return false;
 			}
-			else if (string.Compare(Value, "true", true) == 0 ||
-				string.Compare(Value, "yes", true) == 0 ||
-				string.Compare(Value, "on", true) == 0)
+			else if (string.Equals(Value, "true", StringComparison.OrdinalIgnoreCase) ||
+				string.Equals(Value, "yes", StringComparison.OrdinalIgnoreCase) ||
+				string.Equals(Value, "on", StringComparison.OrdinalIgnoreCase))
 			{
 				return true;
 			}
-			else if (string.Compare(Value, "false", true) == 0 ||
-				string.Compare(Value, "no", true) == 0 ||
-				string.Compare(Value, "off", true) == 0)
+			else if (string.Equals(Value, "false", StringComparison.OrdinalIgnoreCase) ||
+				string.Equals(Value, "no", StringComparison.OrdinalIgnoreCase) ||
+				string.Equals(Value, "off", StringComparison.OrdinalIgnoreCase))
 			{
 				return false;
 			}
@@ -309,7 +309,7 @@ namespace EpicGames.UHT.Utils
 		/// Words that are considered articles when parsing comment strings
 		/// Some words are always forced lowercase
 		/// </summary>
-		private static string[] Articles = new string[]
+		private static readonly string[] Articles = new string[]
 		{
 			"In",
 			"As",
diff --git a/Engine/Source/Programs/Shared/EpicGames.UHT/Utils/UhtFileManager.cs b/Engine/Source/Programs/Shared/EpicGames.UHT/Utils/UhtFileManager.cs
--- a/Engine/Source/Programs/Shared/EpicGames.UHT/Utils/UhtFileManager.cs
+++ b/Engine/Source/Programs/Shared/EpicGames.UHT/Utils/UhtFileManager.cs
@@ -171,7 +171,7 @@ namespace EpicGames.UHT.Utils
 		/// <param name="FilePath">Full file path</param>
 		/// <param name="Contents">Contents of the file</param>
 		/// <returns>True if the file was read, false if not</returns>
-		private bool ReadFile(string FilePath, out StringView Contents)
+		private static bool ReadFile(string FilePath, out StringView Contents)
 		{
 			// Exceptions are very expensive.  Don't bother trying to open the file if it doesn't exist
 			if (!File.Exists(FilePath))
diff --git a/Engine/Source/Programs/Shared/EpicGames.UHT/Utils/UhtMessage.cs b/Engine/Source/Programs/Shared/EpicGames.UHT/Utils/UhtMessage.cs
--- a/Engine/Source/Programs/Shared/EpicGames.UHT/Utils/UhtMessage.cs
+++ b/Engine/Source/Programs/Shared/EpicGames.UHT/Utils/UhtMessage.cs
@@ -330,7 +330,7 @@ namespace EpicGames.UHT.Utils
 		public int MessageFragmentLineNumber { get => -1; }
 		#endregion
 
-		private string FilePath;
+		private readonly string FilePath;
 
 		/// <summary>
 		/// Create a simple file site
@@ -484,10 +484,10 @@ namespace EpicGames.UHT.Utils
 	/// <summary>
 	/// Thread based message context.  Used to improve performance by avoiding allocations.
 	/// </summary>
-	public class UhtTlsMessageExtraContext : IUhtMessageExtraContext
+	public sealed class UhtTlsMessageExtraContext : IUhtMessageExtraContext
 	{
 		private Stack<object?>? ExtraContexts;
-		private static ThreadLocal<UhtTlsMessageExtraContext> Tls = new ThreadLocal<UhtTlsMessageExtraContext>(() => new UhtTlsMessageExtraContext());
+		private static readonly ThreadLocal<UhtTlsMessageExtraContext> Tls = new ThreadLocal<UhtTlsMessageExtraContext>(() => new UhtTlsMessageExtraContext());
 
 		#region IUHTMessageExtraContext implementation
 		IEnumerable<object?>? IUhtMessageExtraContext.MessageExtraContext => this.ExtraContexts;
diff --git a/Engine/Source/Programs/Shared/EpicGames.UHT/Utils/UhtReferenceCollector.cs b/Engine/Source/Programs/Shared/EpicGames.UHT/Utils/UhtReferenceCollector.cs
--- a/Engine/Source/Programs/Shared/EpicGames.UHT/Utils/UhtReferenceCollector.cs
+++ b/Engine/Source/Programs/Shared/EpicGames.UHT/Utils/UhtReferenceCollector.cs
@@ -18,16 +18,16 @@ namespace EpicGames.UHT.Utils
 		/// <summary>
 		/// Add a cross module reference to a given object type.
 		/// </summary>
-		/// <param name="Object">Object type being referenced</param>
+		/// <param name="Obj">Object type being referenced</param>
 		/// <param name="bRegistered">If true, the method being invoked must return the registered object.  This only applies to classes.</param>
-		void AddCrossModuleReference(UhtObject? Object, bool bRegistered);
+		void AddCrossModuleReference(UhtObject? Obj, bool bRegistered);
 
 		/// <summary>
 		/// Add an object declaration
 		/// </summary>
-		/// <param name="Object">Object in question</param>
+		/// <param name="Obj">Object in question</param>
 		/// <param name="bRegistered">If true, the method being invoked must return the registered object.  This only applies to classes.</param>
-		void AddDeclaration(UhtObject Object, bool bRegistered);
+		void AddDeclaration(UhtObject Obj, bool bRegistered);
 
 		/// <summary>
 		/// Add a field as a singleton for exporting
@@ -64,12 +64,12 @@ namespace EpicGames.UHT.Utils
 		/// <summary>
 		/// Collection use to quickly detect if a reference is already in the collection
 		/// </summary>
-		private HashSet<int> Uniques = new HashSet<int>();
+		private readonly HashSet<int> Uniques = new HashSet<int>();
 
 		/// <summary>
 		/// List of all unique reference keys.  Use UngetKey to get the object index and the flag.
 		/// </summary>
-		public List<int> References = new List<int>();
+		public readonly List<int> References = new List<int>();
 
 		/// <summary>
 		/// Return an encoded key that represents the object and registered flag. 
diff --git a/Engine/Source/Programs/Shared/EpicGames.UHT/Utils/UhtSession.cs b/Engine/Source/Programs/Shared/EpicGames.UHT/Utils/UhtSession.cs
--- a/Engine/Source/Programs/Shared/EpicGames.UHT/Utils/UhtSession.cs
+++ b/Engine/Source/Programs/Shared/EpicGames.UHT/Utils/UhtSession.cs
@@ -73,7 +73,7 @@ namespace EpicGames.UHT.Utils
 		/// Limiter for the number of files being saved to the reference directory.
 		/// The OS can get swamped on high core systems
 		/// </summary>
-		private Semaphore WriteRefSemaphore = new Semaphore(32, 32);
+		private static readonly Semaphore WriteRefSemaphore = new Semaphore(32, 32);
 
 		/// <summary>
 		/// Requesting exporter
@@ -178,7 +178,7 @@ namespace EpicGames.UHT.Utils
 				}
 				else
 				{
-					return Task.Factory.StartNew(() => { Action(this); });
+					return Task.Factory.StartNew(() => { Action(this); }, CancellationToken.None, TaskCreationOptions.None, TaskScheduler.Default);
 				}
 			}
 			else
@@ -220,7 +220,6 @@ namespace EpicGames.UHT.Utils
 			return MakePath(Package.Module, Package.ShortName, Suffix);
 		}
 
-
 		/// <summary>
 		/// Make a path for an output based on the package output directory.
 		/// </summary>
@@ -236,7 +235,6 @@ namespace EpicGames.UHT.Utils
 			return MakePath(this.PluginModule, FileName, Extension);
 		}
 
-
 		/// <summary>
 		/// Add an external dependency to the given file path
 		/// </summary>
@@ -273,7 +271,7 @@ namespace EpicGames.UHT.Utils
 				// Writing billions of files to the same directory causes issues.  Use ourselves to throttle reference writes
 				try
 				{
-					this.WriteRefSemaphore.WaitOne();
+					UhtExportFactory.WriteRefSemaphore.WaitOne();
 					{
 						string OutPath = Path.Combine(this.Session.ReferenceMode == UhtReferenceMode.Reference ? this.ReferenceDirectory : this.VerifyDirectory, FileName);
 						if (!this.Session.WriteSource(OutPath, Exported.Span))
@@ -284,7 +282,7 @@ namespace EpicGames.UHT.Utils
 				}
 				finally
 				{
-					this.WriteRefSemaphore.Release();
+					UhtExportFactory.WriteRefSemaphore.Release();
 				}
 
 				// If we are verifying, read the existing file and check the contents
@@ -307,9 +305,9 @@ namespace EpicGames.UHT.Utils
 						Message = $"********************************* {FileName} appears to be a new generated file.";
 					}
 
-					if (Message != String.Empty)
+					if (!string.IsNullOrEmpty(Message))
 					{
-						Log.Logger.LogInformation(Message);
+						Log.TraceInformation(Message);
 						lock (this.ReferenceErrorMessages)
 						{
 							this.ReferenceErrorMessages.Add(Message, true);
@@ -566,7 +564,7 @@ namespace EpicGames.UHT.Utils
 		/// <summary>
 		/// Collection of reserved names
 		/// </summary>
-		private static HashSet<string> ReservedNames = new HashSet<string> { "none" };
+		private static readonly HashSet<string> ReservedNames = new HashSet<string> { "none" };
 
 		#region Configurable settings
 
@@ -771,30 +769,29 @@ namespace EpicGames.UHT.Utils
 		/// </summary>
 		public UhtClass? INotifyFieldValueChanged = null;
 
-		private List<UhtPackage> PackagesInternal = new List<UhtPackage>();
-		private List<UhtHeaderFile> HeaderFilesInternal = new List<UhtHeaderFile>();
-		private List<UhtHeaderFile> SortedHeaderFilesInternal = new List<UhtHeaderFile>();
-		private Dictionary<string, UhtHeaderFile> HeaderFileDictionaryInternal = new Dictionary<string, UhtHeaderFile>(StringComparer.OrdinalIgnoreCase);
+		private readonly List<UhtPackage> PackagesInternal = new List<UhtPackage>();
+		private readonly List<UhtHeaderFile> HeaderFilesInternal = new List<UhtHeaderFile>();
+		private readonly List<UhtHeaderFile> SortedHeaderFilesInternal = new List<UhtHeaderFile>();
+		private readonly Dictionary<string, UhtHeaderFile> HeaderFileDictionaryInternal = new Dictionary<string, UhtHeaderFile>(StringComparer.OrdinalIgnoreCase);
 		private long ErrorCountInternal = 0;
 		private long WarningCountInternal = 0;
-		private List<UhtMessage> Messages = new List<UhtMessage>();
+		private readonly List<UhtMessage> Messages = new List<UhtMessage>();
 		private Task? MessageTask = null;
-		private Dictionary<string, UhtSourceFragment> SourceFragments = new Dictionary<string, UhtSourceFragment>();
 		private UhtClass? UObjectInternal = null;
 		private UhtClass? UClassInternal = null;
 		private UhtClass? UInterfaceInternal = null;
 		private UhtClass? IInterfaceInternal = null;
-		private TypeCounter TypeCounterInternal = new TypeCounter();
-		private TypeCounter PackageTypeCountInternal = new TypeCounter();
-		private TypeCounter HeaderFileTypeCountInternal = new TypeCounter();
-		private TypeCounter ObjectTypeCountInternal = new TypeCounter();
+		private readonly TypeCounter TypeCounterInternal = new TypeCounter();
+		private readonly TypeCounter PackageTypeCountInternal = new TypeCounter();
+		private readonly TypeCounter HeaderFileTypeCountInternal = new TypeCounter();
+		private readonly TypeCounter ObjectTypeCountInternal = new TypeCounter();
 		private UhtSymbolTable SourceNameSymbolTable = new UhtSymbolTable(0);
 		private UhtSymbolTable EngineNameSymbolTable = new UhtSymbolTable(0);
 		private bool bSymbolTablePopulated = false;
 		private Task? ReferenceDeleteTask = null;
-		private Dictionary<string, bool> ExporterStates = new Dictionary<string, bool>(StringComparer.OrdinalIgnoreCase);
-		private Dictionary<string, EnumAndValue> FullEnumValueLookup = new Dictionary<string, EnumAndValue>();
-		private Dictionary<string, UhtEnum> ShortEnumValueLookup = new Dictionary<string, UhtEnum>();
+		private readonly Dictionary<string, bool> ExporterStates = new Dictionary<string, bool>(StringComparer.OrdinalIgnoreCase);
+		private readonly Dictionary<string, EnumAndValue> FullEnumValueLookup = new Dictionary<string, EnumAndValue>();
+		private readonly Dictionary<string, UhtEnum> ShortEnumValueLookup = new Dictionary<string, UhtEnum>();
 		private JsonDocument? ProjectJson = null;
 
 		/// <summary>
@@ -828,8 +825,6 @@ namespace EpicGames.UHT.Utils
 		public IUhtMessageSource? MessageSource => null;
 		/// <inheritdoc/>
 		public IUhtMessageLineNumber? MessageLineNumber => null;
-		/// <inheritdoc/>
-		public IUhtMessageExtraContext? MessageExtraContext => null;
 		#endregion
 
 		/// <summary>
@@ -1046,7 +1041,7 @@ namespace EpicGames.UHT.Utils
 					}
 					catch (Exception)
 					{ }
-				});
+				}, CancellationToken.None, TaskCreationOptions.None, TaskScheduler.Default);
 			}
 
 			StepReadManifestFile(ManifestFilePath);
@@ -1250,7 +1245,7 @@ namespace EpicGames.UHT.Utils
 				}
 			}
 
-			UhtType? Type = null;
+			UhtType? Type;
 			if (Identifiers.Next != null)
 			{
 				Type = FindTypeTwoNamesInternal(StartingType, Options, Identifiers.Token.Value.ToString(), Identifiers.Next.Token.Value.ToString());
@@ -1294,7 +1289,7 @@ namespace EpicGames.UHT.Utils
 				}
 			}
 
-			UhtType? Type = null;
+			UhtType? Type;
 			if (Identifiers.Length == 0)
 			{
 				Type = FindTypeTwoNamesInternal(StartingType, Options, Identifiers[0].Value.ToString(), Identifiers[1].Value.ToString());
@@ -1364,34 +1359,32 @@ namespace EpicGames.UHT.Utils
 		/// <returns>The located type of null if not found</returns>
 		public UhtType? FindTypeInternal(UhtType? StartingType, UhtFindOptions Options, string Name)
 		{
-			UhtType? Type = null;
 			if (Options.HasAnyFlags(UhtFindOptions.EngineName))
 			{
 				if (Options.HasAnyFlags(UhtFindOptions.CaseCompare))
 				{
-					Type = this.EngineNameSymbolTable.FindCasedType(StartingType, Options, Name);
+					return this.EngineNameSymbolTable.FindCasedType(StartingType, Options, Name);
 				}
 				else
 				{
-					Type = this.EngineNameSymbolTable.FindCaselessType(StartingType, Options, Name);
+					return this.EngineNameSymbolTable.FindCaselessType(StartingType, Options, Name);
 				}
 			}
 			else if (Options.HasAnyFlags(UhtFindOptions.SourceName))
 			{
 				if (Options.HasAnyFlags(UhtFindOptions.CaselessCompare))
 				{
-					Type = this.SourceNameSymbolTable.FindCaselessType(StartingType, Options, Name);
+					return this.SourceNameSymbolTable.FindCaselessType(StartingType, Options, Name);
 				}
 				else
 				{
-					Type = this.SourceNameSymbolTable.FindCasedType(StartingType, Options, Name);
+					return this.SourceNameSymbolTable.FindCasedType(StartingType, Options, Name);
 				}
 			}
 			else
 			{
 				throw new UhtIceException("Either EngineName or SourceName must be specified in the options");
 			}
-			return Type;
 		}
 
 		/// <summary>
@@ -1490,7 +1483,7 @@ namespace EpicGames.UHT.Utils
 				// start a task to flush the messages.
 				if (!this.bCacheMessages && this.Messages.Count == 1)
 				{
-					this.MessageTask = Task.Factory.StartNew(() => FlushMessages());
+					this.MessageTask = Task.Factory.StartNew(() => FlushMessages(), CancellationToken.None, TaskCreationOptions.None, TaskScheduler.Default);
 				}
 			}
 
@@ -1556,7 +1549,7 @@ namespace EpicGames.UHT.Utils
 		private void LogMessage(UhtMessage Message)
 		{
 			string FormattedMessage = FormatMessage(Message);
-			LogLevel LogLevel = LogLevel.Information;
+			LogLevel LogLevel;
 			switch (Message.MessageType)
 			{
 				default:
@@ -1578,7 +1571,7 @@ namespace EpicGames.UHT.Utils
 					break;
 			}
 
-			Log.Logger.Log(LogLevel, "{0}", FormattedMessage);
+			Log.Logger.Log(LogLevel, "{FormattedMessage}", FormattedMessage);
 		}
 
 		/// <summary>
@@ -1728,12 +1721,7 @@ namespace EpicGames.UHT.Utils
 		/// <returns>Normalized path possibly converted to a full path.</returns>
 		private string GetNormalizedFullFilePath(string FilePath)
 		{
-			return NormalizePath(this.FileManager!.GetFullFilePath(FilePath));
-		}
-
-		private string NormalizePath(string FilePath)
-		{
-			return FilePath.Replace('\\', '/');
+			return this.FileManager!.GetFullFilePath(FilePath).Replace('\\', '/');
 		}
 
 		private UhtHeaderFileParser? ParseHeaderFile(UhtHeaderFile HeaderFile)
@@ -1842,7 +1830,7 @@ namespace EpicGames.UHT.Utils
 				if (HeaderFileDictionaryInternal.TryGetValue(FileName, out ExistingHeaderFile) && ExistingHeaderFile != null)
 				{
 					string NormalizedExistingFullFilePath = GetNormalizedFullFilePath(ExistingHeaderFile.FilePath);
-					if (string.Compare(NormalizedFullFilePath, NormalizedExistingFullFilePath, true) != 0)
+					if (!string.Equals(NormalizedFullFilePath, NormalizedExistingFullFilePath, StringComparison.OrdinalIgnoreCase))
 					{
 						IUhtMessageSite Site = (IUhtMessageSite?)this.ManifestFile ?? this;
 						Site.LogError($"Two headers with the same name is not allowed. '{HeaderFilePath}' conflicts with '{ExistingHeaderFile.FilePath}'");
@@ -1914,7 +1902,6 @@ namespace EpicGames.UHT.Utils
 						}
 					}
 				}
-
 			});
 		}
 
@@ -2099,9 +2086,9 @@ namespace EpicGames.UHT.Utils
 					UhtType? ExistingType = this.EngineNameSymbolTable.FindCaselessType(null, EngineExtendedType.MustBeUniqueFindOptions(), Type.EngineName);
 					if (ExistingType != null)
 					{
-						Type.HeaderFile.LogError(Type.LineNumber, string.Format("{0} '{1}' shares engine name '{2}' with {6} '{3}' in {4}(5)",
-							EngineExtendedType.CapitalizedText(), Type.SourceName, Type.EngineName, ExistingType.SourceName,
-							ExistingType.HeaderFile.FilePath, ExistingType.LineNumber, ExistingType.EngineType.LowercaseText()));
+						Type.HeaderFile.LogError(Type.LineNumber, 
+							$"{EngineExtendedType.CapitalizedText()} '{Type.SourceName}' shares engine name '{Type.EngineName}' with " +
+							$"{ExistingType.EngineType.LowercaseText()} '{ExistingType.SourceName}' in {ExistingType.HeaderFile.FilePath}({ExistingType.LineNumber})");
 					}
 				}
 
@@ -2154,7 +2141,6 @@ namespace EpicGames.UHT.Utils
 			switch (States[Visit.HeaderFileTypeIndex])
 			{
 				case TopologicalState.Unmarked:
-					UhtHeaderFile? Recursion = null;
 					States[Visit.HeaderFileTypeIndex] = TopologicalState.Temporary;
 					foreach (UhtHeaderFile Referenced in Visit.ReferencedHeadersNoLock)
 					{
@@ -2163,8 +2149,7 @@ namespace EpicGames.UHT.Utils
 							UhtHeaderFile? Out = TopologicalVisit(States, Referenced);
 							if (Out != null)
 							{
-								Recursion = Out;
-								break;
+								return Out;
 							}
 						}
 					}
@@ -2215,7 +2200,7 @@ namespace EpicGames.UHT.Utils
 		#endregion
 
 		#region Validation helpers
-		private HashSet<UhtScriptStruct> ScriptStructsValidForNet = new HashSet<UhtScriptStruct>();
+		private readonly HashSet<UhtScriptStruct> ScriptStructsValidForNet = new HashSet<UhtScriptStruct>();
 
 		/// <summary>
 		/// Validate that the given referenced script structure is valid for network operations.  If the structure
@@ -2311,7 +2296,7 @@ namespace EpicGames.UHT.Utils
 			{
 				foreach (JsonObject Plugin in Plugins)
 				{
-					if (!Plugin.TryGetStringField("Name", out string? TestPluginName) || string.Compare(PluginName, TestPluginName, StringComparison.OrdinalIgnoreCase) != 0)
+					if (!Plugin.TryGetStringField("Name", out string? TestPluginName) || !string.Equals(PluginName, TestPluginName, StringComparison.OrdinalIgnoreCase))
 					{
 						continue;
 					}
@@ -2364,7 +2349,7 @@ namespace EpicGames.UHT.Utils
 					{
 						foreach (UHTManifest.Module Module in this.Manifest!.Modules)
 						{
-							if (string.Compare(Module.Name, Exporter.ModuleName, StringComparison.OrdinalIgnoreCase) == 0)
+							if (string.Equals(Module.Name, Exporter.ModuleName, StringComparison.OrdinalIgnoreCase))
 							{
 								PluginModule = Module;
 								break;
@@ -2372,14 +2357,14 @@ namespace EpicGames.UHT.Utils
 						}
 						if (PluginModule == null)
 						{
-							Log.Logger.LogWarning($"Exporter \"{Exporter.Name}\" skipped because module \"{Exporter.ModuleName}\" was not found in manifest");
+							Log.TraceWarning($"Exporter \"{Exporter.Name}\" skipped because module \"{Exporter.ModuleName}\" was not found in manifest");
 							continue;
 						}
 					}
 
 					if (Run)
 					{
-						Log.Logger.LogTrace($"       Running exporter {Exporter.Name}");
+						Log.TraceLog($"       Running exporter {Exporter.Name}");
 						UhtExportFactory Factory = new UhtExportFactory(this, PluginModule, Exporter);
 						Factory.Run();
 						foreach (var Output in Factory.Outputs)
@@ -2396,7 +2381,7 @@ namespace EpicGames.UHT.Utils
 					}
 					else
 					{
-						Log.Logger.LogTrace($"       Exporter {Exporter.Name} skipped");
+						Log.TraceLog($"       Exporter {Exporter.Name} skipped");
 					}
 				}
 
@@ -2413,7 +2398,7 @@ namespace EpicGames.UHT.Utils
 				}
 			});
 
-			Log.Logger.LogInformation($"Total of {TotalWrittenFiles} written");
+			Log.TraceInformation($"Total of {TotalWrittenFiles} written");
 		}
 		#endregion
 	}
diff --git a/Engine/Source/Programs/Shared/EpicGames.UHT/Utils/UhtStringBuilderExtensions.cs b/Engine/Source/Programs/Shared/EpicGames.UHT/Utils/UhtStringBuilderExtensions.cs
--- a/Engine/Source/Programs/Shared/EpicGames.UHT/Utils/UhtStringBuilderExtensions.cs
+++ b/Engine/Source/Programs/Shared/EpicGames.UHT/Utils/UhtStringBuilderExtensions.cs
@@ -31,7 +31,7 @@ namespace EpicGames.UHT.Utils
 		{
 			if (Tabs < 0 || Tabs > TabsString.Length)
 			{
-				throw new ArgumentOutOfRangeException();
+				throw new ArgumentOutOfRangeException($"Number of tabs specified must be between 0 and {TabsString.Length - 1} inclusive");
 			}
 			else if (Tabs > 0)
 			{
@@ -51,7 +51,7 @@ namespace EpicGames.UHT.Utils
 		{
 			if (Spaces < 0 || Spaces > SpacesString.Length)
 			{
-				throw new ArgumentOutOfRangeException();
+				throw new ArgumentOutOfRangeException($"Number of spaces specified must be between 0 and {SpacesString.Length - 1} inclusive");
 			}
 			else if (Spaces > 0)
 			{
@@ -315,7 +315,7 @@ namespace EpicGames.UHT.Utils
 		/// <returns>Destination builder</returns>
 		public static StringBuilder AppendUTF8LiteralString(this StringBuilder Builder, StringView Text)
 		{
-			Builder.Append("\"");
+			Builder.Append('\"');
 
 			ReadOnlySpan<char> Span = Text.Span;
 			int Length = Span.Length;
@@ -420,7 +420,7 @@ namespace EpicGames.UHT.Utils
 				}
 			}
 
-			Builder.Append("\"");
+			Builder.Append('\"');
 			return Builder;
 		}
 
@@ -492,28 +492,28 @@ namespace EpicGames.UHT.Utils
 		/// <summary>
 		/// Cache of StringBuilders with large initial buffer sizes
 		/// </summary>
-		public static StringBuilderCache Big = new StringBuilderCache(256, 256 * 1024);
+		public static readonly StringBuilderCache Big = new StringBuilderCache(256, 256 * 1024);
 
 		/// <summary>
 		/// Cache of StringBuilders with small initial buffer sizes
 		/// </summary>
-		public static StringBuilderCache Small = new StringBuilderCache(256, 1 * 1024);
+		public static readonly StringBuilderCache Small = new StringBuilderCache(256, 1 * 1024);
 
 		/// <summary>
 		/// Capacity of the cache
 		/// </summary>
-		private int Capacity;
+		private readonly int Capacity;
 
 		/// <summary>
 		/// Initial buffer size for new StringBuilders.  Resetting StringBuilders might result
 		/// in the initial chunk size being smaller.
 		/// </summary>
-		private int InitialBufferSize;
+		private readonly int InitialBufferSize;
 
 		/// <summary>
 		/// Stack of cached StringBuilders
 		/// </summary>
-		private Stack<StringBuilder> Stack;
+		private readonly Stack<StringBuilder> Stack;
 
 		/// <summary>
 		/// Create a new StringBuilder cache
diff --git a/Engine/Source/Programs/Shared/EpicGames.UHT/Utils/UhtSymbolTable.cs b/Engine/Source/Programs/Shared/EpicGames.UHT/Utils/UhtSymbolTable.cs
--- a/Engine/Source/Programs/Shared/EpicGames.UHT/Utils/UhtSymbolTable.cs
+++ b/Engine/Source/Programs/Shared/EpicGames.UHT/Utils/UhtSymbolTable.cs
@@ -73,17 +73,17 @@ namespace EpicGames.UHT.Utils
 		/// <summary>
 		/// Case name lookup table that returns the symbol index and the case index
 		/// </summary>
-		private Dictionary<string, Lookup> CasedDictionary = new Dictionary<string, Lookup>(StringComparer.Ordinal);
+		private readonly Dictionary<string, Lookup> CasedDictionary = new Dictionary<string, Lookup>(StringComparer.Ordinal);
 
 		/// <summary>
 		/// Caseless name lookup table that returns the symbol index
 		/// </summary>
-		private Dictionary<string, int> CaselessDictionary = new Dictionary<string, int>(StringComparer.OrdinalIgnoreCase);
+		private readonly Dictionary<string, int> CaselessDictionary = new Dictionary<string, int>(StringComparer.OrdinalIgnoreCase);
 
 		/// <summary>
 		/// Collection of symbols in the table
 		/// </summary>
-		private Symbol[] Symbols;
+		private readonly Symbol[] Symbols;
 
 		/// <summary>
 		/// Constructs a new symbol table.
diff --git a/Engine/Source/Programs/Shared/EpicGames.UHT/Utils/UhtTables.cs b/Engine/Source/Programs/Shared/EpicGames.UHT/Utils/UhtTables.cs
--- a/Engine/Source/Programs/Shared/EpicGames.UHT/Utils/UhtTables.cs
+++ b/Engine/Source/Programs/Shared/EpicGames.UHT/Utils/UhtTables.cs
@@ -171,7 +171,7 @@ namespace EpicGames.UHT.Utils
 		/// <summary>
 		/// Lookup dictionary for the specifiers
 		/// </summary>
-		private Dictionary<StringView, TValue> Lookup;
+		private readonly Dictionary<StringView, TValue> Lookup;
 
 		/// <summary>
 		/// Construct a new table
@@ -411,7 +411,7 @@ namespace EpicGames.UHT.Utils
 	/// This attribute must be applied to any class that contains other UHT attributes.
 	/// </summary>
 	[AttributeUsage(AttributeTargets.Class)]
-	public class UnrealHeaderToolAttribute : Attribute
+	public sealed class UnrealHeaderToolAttribute : Attribute
 	{
 
 		/// <summary>
@@ -549,10 +549,7 @@ namespace EpicGames.UHT.Utils
 			Assembly? Assembly = FindAssemblyByName(Path.GetFileNameWithoutExtension(AssemblyFilePath));
 			if (Assembly == null)
 			{
-				if (Assembly == null)
-				{
-					Assembly = Assembly.LoadFile(AssemblyFilePath);
-				}
+				Assembly = Assembly.LoadFile(AssemblyFilePath);
 			}
 			return Assembly;
 		}
@@ -569,7 +566,7 @@ namespace EpicGames.UHT.Utils
 				foreach (Assembly Assembly in AppDomain.CurrentDomain.GetAssemblies())
 				{
 					AssemblyName AssemblyName = Assembly.GetName();
-					if (AssemblyName.Name != null && String.Compare(AssemblyName.Name, Name, true) == 0)
+					if (AssemblyName.Name != null && String.Equals(AssemblyName.Name, Name, StringComparison.OrdinalIgnoreCase))
 					{
 						return Assembly;
 					}
@@ -637,7 +634,7 @@ namespace EpicGames.UHT.Utils
 				{
 					throw new Exception($"Unable to find UnrealHeaderTool attribute InitMethod {ParserAttribute.InitMethod}");
 				}
-				InitInfo.Invoke(null, new Object[] { });
+				InitInfo.Invoke(null, Array.Empty<object>());
 			}
 
 			// Scan the methods for things we are interested in
diff --git a/Engine/Source/Programs/Shared/EpicGames.UHT/Utils/UhtUtilities.cs b/Engine/Source/Programs/Shared/EpicGames.UHT/Utils/UhtUtilities.cs
--- a/Engine/Source/Programs/Shared/EpicGames.UHT/Utils/UhtUtilities.cs
+++ b/Engine/Source/Programs/Shared/EpicGames.UHT/Utils/UhtUtilities.cs
@@ -111,7 +111,7 @@ namespace EpicGames.UHT.Utils
 	/// <summary>
 	/// Assorted utility functions
 	/// </summary>
-	public class UhtUtilities
+	public static class UhtUtilities
 	{
 
 		/// <summary>
@@ -146,9 +146,9 @@ namespace EpicGames.UHT.Utils
 				}
 				if (bQuote)
 				{
-					Builder.Append("'");
+					Builder.Append('\'');
 					Builder.Append(Local[Index]);
-					Builder.Append("'");
+					Builder.Append('\'');
 				}
 				else
 				{
