More code warning fixes

#rnx

diff --git a/Engine/Source/Programs/Shared/EpicGames.UHT/.editorconfig b/Engine/Source/Programs/Shared/EpicGames.UHT/.editorconfig
--- a/Engine/Source/Programs/Shared/EpicGames.UHT/.editorconfig
+++ b/Engine/Source/Programs/Shared/EpicGames.UHT/.editorconfig
@@ -5,21 +5,12 @@
 # Default severity for all analyzer diagnostics
 dotnet_analyzer_diagnostic.severity = none
 
-# IDE0051: Remove unused private members
-dotnet_diagnostic.IDE0051.severity = none
-
 # IDE0060: Remove unused parameter
 dotnet_diagnostic.IDE0060.severity = none
 
 # CA1051: Do not declare visible instance fields
 dotnet_diagnostic.CA1051.severity = none
 
-# CA1310: Specify StringComparison for correctness
-dotnet_diagnostic.CA1310.severity = none
-
-# CA1307: Specify StringComparison for clarity
-dotnet_diagnostic.CA1307.severity = none
-
 # CA1045: Do not pass types by reference
 dotnet_diagnostic.CA1045.severity = none
 
diff --git a/Engine/Source/Programs/Shared/EpicGames.UHT/Exporters/CodeGen/UhtCodeGenerator.cs b/Engine/Source/Programs/Shared/EpicGames.UHT/Exporters/CodeGen/UhtCodeGenerator.cs
--- a/Engine/Source/Programs/Shared/EpicGames.UHT/Exporters/CodeGen/UhtCodeGenerator.cs
+++ b/Engine/Source/Programs/Shared/EpicGames.UHT/Exporters/CodeGen/UhtCodeGenerator.cs
@@ -4,6 +4,7 @@ using EpicGames.Core;
 using EpicGames.UHT.Tables;
 using EpicGames.UHT.Types;
 using EpicGames.UHT.Utils;
+using System;
 using System.Collections.Generic;
 using System.IO;
 using System.Text;
@@ -277,7 +278,7 @@ namespace EpicGames.UHT.Exporters.CodeGen
 			FilePath = FilePath.Replace('\\', '/');
 			if (bIsRelative)
 			{
-				while (FilePath.StartsWith("../"))
+				while (FilePath.StartsWith("../", StringComparison.Ordinal))
 				{
 					FilePath = FilePath.Substring(3);
 				}
diff --git a/Engine/Source/Programs/Shared/EpicGames.UHT/Exporters/CodeGen/UhtHeaderCodeGenerator.cs b/Engine/Source/Programs/Shared/EpicGames.UHT/Exporters/CodeGen/UhtHeaderCodeGenerator.cs
--- a/Engine/Source/Programs/Shared/EpicGames.UHT/Exporters/CodeGen/UhtHeaderCodeGenerator.cs
+++ b/Engine/Source/Programs/Shared/EpicGames.UHT/Exporters/CodeGen/UhtHeaderCodeGenerator.cs
@@ -458,7 +458,7 @@ namespace EpicGames.UHT.Exporters.CodeGen
 			{
 				throw new UhtIceException("Attempt to export a function that isn't a delegate as a delegate");
 			}
-			if (!Function.MarshalAndCallName.StartsWith(DelegatePrefix))
+			if (!Function.MarshalAndCallName.StartsWith(DelegatePrefix, StringComparison.Ordinal))
 			{
 				throw new UhtIceException("Marshal and call name must begin with 'delegate'");
 			}
diff --git a/Engine/Source/Programs/Shared/EpicGames.UHT/Exporters/CodeGen/UhtHeaderCodeGeneratorHFile.cs b/Engine/Source/Programs/Shared/EpicGames.UHT/Exporters/CodeGen/UhtHeaderCodeGeneratorHFile.cs
--- a/Engine/Source/Programs/Shared/EpicGames.UHT/Exporters/CodeGen/UhtHeaderCodeGeneratorHFile.cs
+++ b/Engine/Source/Programs/Shared/EpicGames.UHT/Exporters/CodeGen/UhtHeaderCodeGeneratorHFile.cs
@@ -552,7 +552,7 @@ namespace EpicGames.UHT.Exporters.CodeGen
 								{
 									string PropertyName = SparseProperty.SourceName;
 									string CleanPropertyName = PropertyName;
-									if (SparseProperty is UhtBoolProperty && PropertyName.StartsWith("b"))
+									if (SparseProperty is UhtBoolProperty && PropertyName.StartsWith("b", StringComparison.Ordinal))
 									{
 										CleanPropertyName = PropertyName.Substring(1);
 									}
@@ -1370,7 +1370,7 @@ namespace EpicGames.UHT.Exporters.CodeGen
 		#region Enum helper methods
 		private static bool IsFullEnumName(string InEnumName)
 		{
-			return InEnumName.Contains("::");
+			return InEnumName.Contains("::", StringComparison.Ordinal);
 		}
 
 		private static StringView GenerateEnumPrefix(UhtEnum Enum)
diff --git a/Engine/Source/Programs/Shared/EpicGames.UHT/Exporters/Stats/UhtStatsExporter.cs b/Engine/Source/Programs/Shared/EpicGames.UHT/Exporters/Stats/UhtStatsExporter.cs
--- a/Engine/Source/Programs/Shared/EpicGames.UHT/Exporters/Stats/UhtStatsExporter.cs
+++ b/Engine/Source/Programs/Shared/EpicGames.UHT/Exporters/Stats/UhtStatsExporter.cs
@@ -5,6 +5,7 @@ using EpicGames.UHT.Tables;
 using EpicGames.UHT.Types;
 using EpicGames.UHT.Utils;
 using System.Collections.Generic;
+using System.Diagnostics.CodeAnalysis;
 
 namespace EpicGames.UHT.Exporters.Stats
 {
@@ -12,6 +13,7 @@ namespace EpicGames.UHT.Exporters.Stats
 	internal class UhtStatsExporter
 	{
 		[UhtExporter(Name = "Stats", Description = "Type Stats", Options = UhtExporterOptions.None)]
+		[SuppressMessage("CodeQuality", "IDE0051:Remove unused private members", Justification = "Attribute accessed method")]
 		private static void StatsExporter(IUhtExportFactory Factory)
 		{
 			SortedDictionary<string, int> CountByType = new SortedDictionary<string, int>();
diff --git a/Engine/Source/Programs/Shared/EpicGames.UHT/Parsers/UhtClassParser.cs b/Engine/Source/Programs/Shared/EpicGames.UHT/Parsers/UhtClassParser.cs
--- a/Engine/Source/Programs/Shared/EpicGames.UHT/Parsers/UhtClassParser.cs
+++ b/Engine/Source/Programs/Shared/EpicGames.UHT/Parsers/UhtClassParser.cs
@@ -7,6 +7,7 @@ using EpicGames.UHT.Types;
 using EpicGames.UHT.Utils;
 using System;
 using System.Collections.Generic;
+using System.Diagnostics.CodeAnalysis;
 using System.Text;
 
 namespace EpicGames.UHT.Parsers
@@ -365,6 +366,7 @@ namespace EpicGames.UHT.Parsers
 
 		#region Keywords
 		[UhtKeyword(Extends = UhtTableNames.Global)]
+		[SuppressMessage("CodeQuality", "IDE0051:Remove unused private members", Justification = "Attribute accessed method")]
 		private static UhtParseResult UCLASSKeyword(UhtParsingScope TopScope, UhtParsingScope ActionScope, ref UhtToken Token)
 		{
 			return ParseUClass(TopScope, ref Token);
@@ -372,6 +374,7 @@ namespace EpicGames.UHT.Parsers
 
 		[UhtKeyword(Extends = UhtTableNames.Class)]
 		[UhtKeyword(Extends = UhtTableNames.Class, Keyword = "GENERATED_BODY")]
+		[SuppressMessage("CodeQuality", "IDE0051:Remove unused private members", Justification = "Attribute accessed method")]
 		private static UhtParseResult GENERATED_UCLASS_BODYKeyword(UhtParsingScope TopScope, UhtParsingScope ActionScope, ref UhtToken Token)
 		{
 			UhtClass Class = (UhtClass)TopScope.ScopeType;
diff --git a/Engine/Source/Programs/Shared/EpicGames.UHT/Parsers/UhtEnumParser.cs b/Engine/Source/Programs/Shared/EpicGames.UHT/Parsers/UhtEnumParser.cs
--- a/Engine/Source/Programs/Shared/EpicGames.UHT/Parsers/UhtEnumParser.cs
+++ b/Engine/Source/Programs/Shared/EpicGames.UHT/Parsers/UhtEnumParser.cs
@@ -1,5 +1,6 @@
 // Copyright Epic Games, Inc. All Rights Reserved.
 
+using System.Diagnostics.CodeAnalysis;
 using EpicGames.Core;
 using EpicGames.UHT.Tables;
 using EpicGames.UHT.Tokenizer;
@@ -25,6 +26,7 @@ namespace EpicGames.UHT.Parsers
 
 		#region Keywords
 		[UhtKeyword(Extends = UhtTableNames.Global)]
+		[SuppressMessage("CodeQuality", "IDE0051:Remove unused private members", Justification = "Attribute accessed method")]
 		private static UhtParseResult UENUMKeyword(UhtParsingScope TopScope, UhtParsingScope ActionScope, ref UhtToken Token)
 		{
 			return ParseUEnum(TopScope, Token);
diff --git a/Engine/Source/Programs/Shared/EpicGames.UHT/Parsers/UhtFunctionParser.cs b/Engine/Source/Programs/Shared/EpicGames.UHT/Parsers/UhtFunctionParser.cs
--- a/Engine/Source/Programs/Shared/EpicGames.UHT/Parsers/UhtFunctionParser.cs
+++ b/Engine/Source/Programs/Shared/EpicGames.UHT/Parsers/UhtFunctionParser.cs
@@ -7,6 +7,7 @@ using EpicGames.UHT.Types;
 using EpicGames.UHT.Utils;
 using System;
 using System.Collections.Generic;
+using System.Diagnostics.CodeAnalysis;
 using System.Linq;
 using System.Text;
 
@@ -181,6 +182,7 @@ namespace EpicGames.UHT.Parsers
 		#region Keywords
 		[UhtKeyword(Extends = UhtTableNames.Global)]
 		[UhtKeyword(Extends = UhtTableNames.Class)]
+		[SuppressMessage("CodeQuality", "IDE0051:Remove unused private members", Justification = "Attribute accessed method")]
 		private static UhtParseResult UDELEGATEKeyword(UhtParsingScope TopScope, UhtParsingScope ActionScope, ref UhtToken Token)
 		{
 			return ParseUDelegate(TopScope, Token, true);
@@ -189,12 +191,14 @@ namespace EpicGames.UHT.Parsers
 		[UhtKeyword(Extends = UhtTableNames.Class)]
 		[UhtKeyword(Extends = UhtTableNames.Interface)]
 		[UhtKeyword(Extends = UhtTableNames.NativeInterface)]
+		[SuppressMessage("CodeQuality", "IDE0051:Remove unused private members", Justification = "Attribute accessed method")]
 		private static UhtParseResult UFUNCTIONKeyword(UhtParsingScope TopScope, UhtParsingScope ActionScope, ref UhtToken Token)
 		{
 			return ParseUFunction(TopScope, Token);
 		}
 
 		[UhtKeywordCatchAll(Extends = UhtTableNames.Global)]
+		[SuppressMessage("CodeQuality", "IDE0051:Remove unused private members", Justification = "Attribute accessed method")]
 		private static UhtParseResult ParseCatchAllKeyword(UhtParsingScope TopScope, ref UhtToken Token)
 		{
 			if (UhtFunctionParser.IsValidateDelegateDeclaration(Token))
@@ -712,7 +716,7 @@ namespace EpicGames.UHT.Parsers
 			// The source name won't have the suffix applied to delegate names, however, the engine name will
 			// We use the engine name because we need to detect the suffix for delegates
 			string FunctionName = Function.EngineName;
-			if (FunctionName.EndsWith(UhtFunction.GeneratedDelegateSignatureSuffix))
+			if (FunctionName.EndsWith(UhtFunction.GeneratedDelegateSignatureSuffix, StringComparison.Ordinal))
 			{
 				FunctionName = FunctionName.Substring(0, FunctionName.Length - UhtFunction.GeneratedDelegateSignatureSuffix.Length);
 			}
diff --git a/Engine/Source/Programs/Shared/EpicGames.UHT/Parsers/UhtHeaderFileParser.cs b/Engine/Source/Programs/Shared/EpicGames.UHT/Parsers/UhtHeaderFileParser.cs
--- a/Engine/Source/Programs/Shared/EpicGames.UHT/Parsers/UhtHeaderFileParser.cs
+++ b/Engine/Source/Programs/Shared/EpicGames.UHT/Parsers/UhtHeaderFileParser.cs
@@ -7,6 +7,7 @@ using EpicGames.UHT.Types;
 using EpicGames.UHT.Utils;
 using System;
 using System.Collections.Generic;
+using System.Diagnostics.CodeAnalysis;
 using System.IO;
 
 namespace EpicGames.UHT.Parsers
@@ -137,6 +138,7 @@ namespace EpicGames.UHT.Parsers
 		#region Keywords
 		[UhtKeyword(Extends = UhtTableNames.ClassBase, Keyword = "public")]
 		[UhtKeyword(Extends = UhtTableNames.ScriptStruct, Keyword = "public")]
+		[SuppressMessage("CodeQuality", "IDE0051:Remove unused private members", Justification = "Attribute accessed method")]
 		private static UhtParseResult PublicKeyword(UhtParsingScope TopScope, UhtParsingScope ActionScope, ref UhtToken Token)
 		{
 			return SetAccessSpecifier(TopScope, UhtAccessSpecifier.Public);
@@ -144,6 +146,7 @@ namespace EpicGames.UHT.Parsers
 
 		[UhtKeyword(Extends = UhtTableNames.ClassBase, Keyword = "protected")]
 		[UhtKeyword(Extends = UhtTableNames.ScriptStruct, Keyword = "protected")]
+		[SuppressMessage("CodeQuality", "IDE0051:Remove unused private members", Justification = "Attribute accessed method")]
 		private static UhtParseResult ProtectedKeyword(UhtParsingScope TopScope, UhtParsingScope ActionScope, ref UhtToken Token)
 		{
 			return SetAccessSpecifier(TopScope, UhtAccessSpecifier.Protected);
@@ -151,6 +154,7 @@ namespace EpicGames.UHT.Parsers
 
 		[UhtKeyword(Extends = UhtTableNames.ClassBase, Keyword = "private")]
 		[UhtKeyword(Extends = UhtTableNames.ScriptStruct, Keyword = "private")]
+		[SuppressMessage("CodeQuality", "IDE0051:Remove unused private members", Justification = "Attribute accessed method")]
 		private static UhtParseResult PrivateKeyword(UhtParsingScope TopScope, UhtParsingScope ActionScope, ref UhtToken Token)
 		{
 			return SetAccessSpecifier(TopScope, UhtAccessSpecifier.Private);
@@ -692,7 +696,7 @@ namespace EpicGames.UHT.Parsers
 						}
 						if (IncludeNameString.Span.Contains(".generated.h", StringComparison.Ordinal))
 						{
-							string RefName = Path.GetFileName(IncludeNameString.ToString()).Replace(".generated.h", ".h");
+							string RefName = Path.GetFileName(IncludeNameString.ToString()).Replace(".generated.h", ".h", StringComparison.Ordinal);
 							this.HeaderFile.AddReferencedHeader(RefName, true);
 						}
 						else
diff --git a/Engine/Source/Programs/Shared/EpicGames.UHT/Parsers/UhtInterfaceClassParser.cs b/Engine/Source/Programs/Shared/EpicGames.UHT/Parsers/UhtInterfaceClassParser.cs
--- a/Engine/Source/Programs/Shared/EpicGames.UHT/Parsers/UhtInterfaceClassParser.cs
+++ b/Engine/Source/Programs/Shared/EpicGames.UHT/Parsers/UhtInterfaceClassParser.cs
@@ -1,5 +1,6 @@
 // Copyright Epic Games, Inc. All Rights Reserved.
 
+using System.Diagnostics.CodeAnalysis;
 using EpicGames.Core;
 using EpicGames.UHT.Tables;
 using EpicGames.UHT.Tokenizer;
@@ -77,6 +78,7 @@ namespace EpicGames.UHT.Parsers
 		}
 		#region Keywords
 		[UhtKeyword(Extends = UhtTableNames.Global)]
+		[SuppressMessage("CodeQuality", "IDE0051:Remove unused private members", Justification = "Attribute accessed method")]
 		private static UhtParseResult UINTERFACEKeyword(UhtParsingScope TopScope, UhtParsingScope ActionScope, ref UhtToken Token)
 		{
 			return ParseUInterface(TopScope, ref Token);
@@ -84,6 +86,7 @@ namespace EpicGames.UHT.Parsers
 
 		[UhtKeyword(Extends = UhtTableNames.Interface)]
 		[UhtKeyword(Extends = UhtTableNames.Interface, Keyword = "GENERATED_BODY")]
+		[SuppressMessage("CodeQuality", "IDE0051:Remove unused private members", Justification = "Attribute accessed method")]
 		private static UhtParseResult GENERATED_UINTERFACE_BODYKeyword(UhtParsingScope TopScope, UhtParsingScope ActionScope, ref UhtToken Token)
 		{
 			UhtClass Class = (UhtClass)TopScope.ScopeType;
diff --git a/Engine/Source/Programs/Shared/EpicGames.UHT/Parsers/UhtNativeInterfaceClassParser.cs b/Engine/Source/Programs/Shared/EpicGames.UHT/Parsers/UhtNativeInterfaceClassParser.cs
--- a/Engine/Source/Programs/Shared/EpicGames.UHT/Parsers/UhtNativeInterfaceClassParser.cs
+++ b/Engine/Source/Programs/Shared/EpicGames.UHT/Parsers/UhtNativeInterfaceClassParser.cs
@@ -1,5 +1,6 @@
 // Copyright Epic Games, Inc. All Rights Reserved.
 
+using System.Diagnostics.CodeAnalysis;
 using EpicGames.Core;
 using EpicGames.UHT.Tables;
 using EpicGames.UHT.Tokenizer;
@@ -66,6 +67,7 @@ namespace EpicGames.UHT.Parsers
 
 		#region Keywords
 		[UhtKeyword(Extends = UhtTableNames.Global, Keyword = "class", DisableUsageError = true)]
+		[SuppressMessage("CodeQuality", "IDE0051:Remove unused private members", Justification = "Attribute accessed method")]
 		private static UhtParseResult ClassKeyword(UhtParsingScope TopScope, UhtParsingScope ActionScope, ref UhtToken Token)
 		{
 			using (var SaveState = new UhtTokenSaveState(TopScope.TokenReader))
@@ -84,6 +86,7 @@ namespace EpicGames.UHT.Parsers
 
 		[UhtKeyword(Extends = UhtTableNames.NativeInterface)]
 		[UhtKeyword(Extends = UhtTableNames.NativeInterface, Keyword = "GENERATED_BODY")]
+		[SuppressMessage("CodeQuality", "IDE0051:Remove unused private members", Justification = "Attribute accessed method")]
 		private static UhtParseResult GENERATED_IINTERFACE_BODYKeyword(UhtParsingScope TopScope, UhtParsingScope ActionScope, ref UhtToken Token)
 		{
 			UhtClass Class = (UhtClass)TopScope.ScopeType;
diff --git a/Engine/Source/Programs/Shared/EpicGames.UHT/Parsers/UhtParsingScope.cs b/Engine/Source/Programs/Shared/EpicGames.UHT/Parsers/UhtParsingScope.cs
--- a/Engine/Source/Programs/Shared/EpicGames.UHT/Parsers/UhtParsingScope.cs
+++ b/Engine/Source/Programs/Shared/EpicGames.UHT/Parsers/UhtParsingScope.cs
@@ -218,24 +218,6 @@ namespace EpicGames.UHT.Parsers
 			//this.TokenReader.ClearComments();
 		}
 
-		// We consider any alpha/digit or code point > 0xFF as a valid comment char
-		/// <summary>
-		/// Given a list of comments, check to see if any have alpha, numeric, or unicode code points with a value larger than 0xFF.
-		/// </summary>
-		/// <param name="Comments">Comments to search</param>
-		/// <returns>True is a character in question was found</returns>
-		private static bool HasValidCommentChar(ReadOnlySpan<StringView> Comments)
-		{
-			foreach (StringView Comment in Comments)
-			{
-				if (HasValidCommentChar(Comment.Span))
-				{
-					return true;
-				}
-			}
-			return false;
-		}
-
 		// We consider any alpha/digit or code point > 0xFF as a valid comment char
 		/// <summary>
 		/// Given a list of comments, check to see if any have alpha, numeric, or unicode code points with a value larger than 0xFF.
@@ -286,9 +268,9 @@ namespace EpicGames.UHT.Parsers
 				int CommentStart, CommentEnd;
 
 				// Block comments go first
-				while ((CommentStart = Result.IndexOf("/*~")) != -1)
+				while ((CommentStart = Result.IndexOf("/*~", StringComparison.Ordinal)) != -1)
 				{
-					CommentEnd = Result.IndexOf("*/", CommentStart);
+					CommentEnd = Result.IndexOf("*/", CommentStart, StringComparison.Ordinal);
 					if (CommentEnd != -1)
 					{
 						Result = Result.Remove(CommentStart, CommentEnd - CommentStart + 2);
@@ -301,9 +283,9 @@ namespace EpicGames.UHT.Parsers
 				}
 
 				// Leftover line comments go next
-				while ((CommentStart = Result.IndexOf("//~")) != -1)
+				while ((CommentStart = Result.IndexOf("//~", StringComparison.Ordinal)) != -1)
 				{
-					CommentEnd = Result.IndexOf("\n", CommentStart);
+					CommentEnd = Result.IndexOf("\n", CommentStart, StringComparison.Ordinal);
 					if (CommentEnd != -1)
 					{
 						Result = Result.Remove(CommentStart, CommentEnd - CommentStart + 1);
@@ -317,38 +299,38 @@ namespace EpicGames.UHT.Parsers
 			}
 
 			// Check for known commenting styles.
-			bool bJavaDocStyle = Result.Contains("/**");
-			bool bCStyle = Result.Contains("/*");
-			bool bCPPStyle = Result.StartsWith("//");
+			bool bJavaDocStyle = Result.Contains("/**", StringComparison.Ordinal);
+			bool bCStyle = Result.Contains("/*", StringComparison.Ordinal);
+			bool bCPPStyle = Result.StartsWith("//", StringComparison.Ordinal);
 
 			if (bJavaDocStyle || bCStyle)
 			{
 				// Remove beginning and end markers.
 				if (bJavaDocStyle)
 				{
-					Result = Result.Replace("/**", "");
+					Result = Result.Replace("/**", "", StringComparison.Ordinal);
 				}
 				if (bCStyle)
 				{
-					Result = Result.Replace("/*", "");
+					Result = Result.Replace("/*", "", StringComparison.Ordinal);
 				}
-				Result = Result.Replace("*/", "");
+				Result = Result.Replace("*/", "", StringComparison.Ordinal);
 			}
 
 			if (bCPPStyle)
 			{
 				// Remove c++-style comment markers.  Also handle javadoc-style comments 
-				Result = Result.Replace("///", "");
-				Result = Result.Replace("//", "");
+				Result = Result.Replace("///", "", StringComparison.Ordinal);
+				Result = Result.Replace("//", "", StringComparison.Ordinal);
 
 				// Parser strips cpptext and replaces it with "// (cpptext)" -- prevent
 				// this from being treated as a comment on variables declared below the
 				// cpptext section
-				Result = Result.Replace("(cpptext)", "");
+				Result = Result.Replace("(cpptext)", "", StringComparison.Ordinal);
 			}
 
 			// Get rid of carriage return or tab characters, which mess up tooltips.
-			Result = Result.Replace("\r", "");
+			Result = Result.Replace("\r", "", StringComparison.Ordinal);
 
 			//wx widgets has a hard coded tab size of 8
 			{
diff --git a/Engine/Source/Programs/Shared/EpicGames.UHT/Parsers/UhtPropertyParser.cs b/Engine/Source/Programs/Shared/EpicGames.UHT/Parsers/UhtPropertyParser.cs
--- a/Engine/Source/Programs/Shared/EpicGames.UHT/Parsers/UhtPropertyParser.cs
+++ b/Engine/Source/Programs/Shared/EpicGames.UHT/Parsers/UhtPropertyParser.cs
@@ -8,6 +8,7 @@ using EpicGames.UHT.Types;
 using EpicGames.UHT.Utils;
 using System;
 using System.Collections.Generic;
+using System.Diagnostics.CodeAnalysis;
 using System.Text;
 using System.Threading;
 
@@ -1063,8 +1064,8 @@ namespace EpicGames.UHT.Parsers
 
 			// Check to see if the variable is deprecated, and if so set the flag
 			{
-				int DeprecatedIndex = NewProperty.SourceName.IndexOf("_DEPRECATED");
-				int NativizedPropertyPostfixIndex = NewProperty.SourceName.IndexOf("__pf"); //@TODO: check OverrideNativeName in Meta Data, to be sure it's not a random occurrence of the "__pf" string.
+				int DeprecatedIndex = NewProperty.SourceName.IndexOf("_DEPRECATED", StringComparison.Ordinal);
+				int NativizedPropertyPostfixIndex = NewProperty.SourceName.IndexOf("__pf", StringComparison.Ordinal); //@TODO: check OverrideNativeName in Meta Data, to be sure it's not a random occurrence of the "__pf" string.
 				bool bIgnoreDeprecatedWord = (NativizedPropertyPostfixIndex != -1) && (NativizedPropertyPostfixIndex > DeprecatedIndex);
 				if ((DeprecatedIndex != -1) && !bIgnoreDeprecatedWord)
 				{
@@ -1188,6 +1189,7 @@ namespace EpicGames.UHT.Parsers
 		#region Keywords
 		[UhtKeyword(Extends = UhtTableNames.Class)]
 		[UhtKeyword(Extends = UhtTableNames.ScriptStruct)]
+		[SuppressMessage("CodeQuality", "IDE0051:Remove unused private members", Justification = "Attribute accessed method")]
 		private static UhtParseResult UPROPERTYKeyword(UhtParsingScope TopScope, UhtParsingScope ActionScope, ref UhtToken Token)
 		{
 			UhtPropertyParseOptions Options = UhtPropertyParseOptions.ParseLayoutMacro | UhtPropertyParseOptions.List | UhtPropertyParseOptions.AddModuleRelativePath;
@@ -1219,6 +1221,7 @@ namespace EpicGames.UHT.Parsers
 	static class UhtDefaultPropertyParser
 	{
 		[UhtPropertyType(Options = UhtPropertyTypeOptions.Default)]
+		[SuppressMessage("CodeQuality", "IDE0051:Remove unused private members", Justification = "Attribute accessed method")]
 		private static UhtProperty? DefaultProperty(UhtPropertyResolvePhase ResolvePhase, UhtPropertySettings PropertySettings, IUhtTokenReader TokenReader, UhtToken MatchedToken)
 		{
 			UhtSession Session = PropertySettings.Outer.Session;
diff --git a/Engine/Source/Programs/Shared/EpicGames.UHT/Parsers/UhtScriptStructParser.cs b/Engine/Source/Programs/Shared/EpicGames.UHT/Parsers/UhtScriptStructParser.cs
--- a/Engine/Source/Programs/Shared/EpicGames.UHT/Parsers/UhtScriptStructParser.cs
+++ b/Engine/Source/Programs/Shared/EpicGames.UHT/Parsers/UhtScriptStructParser.cs
@@ -6,6 +6,7 @@ using EpicGames.UHT.Tokenizer;
 using EpicGames.UHT.Types;
 using EpicGames.UHT.Utils;
 using System.Collections.Generic;
+using System.Diagnostics.CodeAnalysis;
 
 namespace EpicGames.UHT.Parsers
 {
@@ -71,6 +72,7 @@ namespace EpicGames.UHT.Parsers
 
 		#region Keywords
 		[UhtKeyword(Extends = UhtTableNames.Global)]
+		[SuppressMessage("CodeQuality", "IDE0051:Remove unused private members", Justification = "Attribute accessed method")]
 		private static UhtParseResult USTRUCTKeyword(UhtParsingScope TopScope, UhtParsingScope ActionScope, ref UhtToken Token)
 		{
 			return ParseUScriptStruct(TopScope, Token);
@@ -78,6 +80,7 @@ namespace EpicGames.UHT.Parsers
 
 		[UhtKeyword(Extends = UhtTableNames.ScriptStruct)]
 		[UhtKeyword(Extends = UhtTableNames.ScriptStruct, Keyword = "GENERATED_BODY")]
+		[SuppressMessage("CodeQuality", "IDE0051:Remove unused private members", Justification = "Attribute accessed method")]
 		private static UhtParseResult GENERATED_USTRUCT_BODYKeyword(UhtParsingScope TopScope, UhtParsingScope ActionScope, ref UhtToken Token)
 		{
 			UhtScriptStruct ScriptStruct = (UhtScriptStruct)TopScope.ScopeType;
@@ -101,6 +104,7 @@ namespace EpicGames.UHT.Parsers
 		}
 
 		[UhtKeyword(Extends = UhtTableNames.ScriptStruct)]
+		[SuppressMessage("CodeQuality", "IDE0051:Remove unused private members", Justification = "Attribute accessed method")]
 		private static UhtParseResult RIGVM_METHODKeyword(UhtParsingScope TopScope, UhtParsingScope ActionScope, ref UhtToken Token)
 		{
 			ParseRigVM(TopScope);
diff --git a/Engine/Source/Programs/Shared/EpicGames.UHT/Specifiers/UhtClassSpecifiers.cs b/Engine/Source/Programs/Shared/EpicGames.UHT/Specifiers/UhtClassSpecifiers.cs
--- a/Engine/Source/Programs/Shared/EpicGames.UHT/Specifiers/UhtClassSpecifiers.cs
+++ b/Engine/Source/Programs/Shared/EpicGames.UHT/Specifiers/UhtClassSpecifiers.cs
@@ -5,6 +5,7 @@ using EpicGames.UHT.Tables;
 using EpicGames.UHT.Types;
 using EpicGames.UHT.Utils;
 using System.Collections.Generic;
+using System.Diagnostics.CodeAnalysis;
 
 namespace EpicGames.UHT.Parsers
 {
@@ -15,6 +16,7 @@ namespace EpicGames.UHT.Parsers
 	public static class UhtClassSpecifiers
 	{
 		[UhtSpecifier(Extends = UhtTableNames.Class, ValueType = UhtSpecifierValueType.Legacy)]
+		[SuppressMessage("CodeQuality", "IDE0051:Remove unused private members", Justification = "Attribute accessed method")]
 		private static void NoExportSpecifier(UhtSpecifierContext SpecifierContext)
 		{
 			UhtClassParser Class = (UhtClassParser)SpecifierContext.Scope.ScopeType;
@@ -22,6 +24,7 @@ namespace EpicGames.UHT.Parsers
 		}
 
 		[UhtSpecifier(Extends = UhtTableNames.Class, ValueType = UhtSpecifierValueType.Legacy)]
+		[SuppressMessage("CodeQuality", "IDE0051:Remove unused private members", Justification = "Attribute accessed method")]
 		private static void IntrinsicSpecifier(UhtSpecifierContext SpecifierContext)
 		{
 			UhtClassParser Class = (UhtClassParser)SpecifierContext.Scope.ScopeType;
@@ -29,12 +32,14 @@ namespace EpicGames.UHT.Parsers
 		}
 
 		[UhtSpecifier(Extends = UhtTableNames.Class, ValueType = UhtSpecifierValueType.Legacy)]
+		[SuppressMessage("CodeQuality", "IDE0051:Remove unused private members", Justification = "Attribute accessed method")]
 		private static void ComponentWrapperClassSpecifier(UhtSpecifierContext SpecifierContext)
 		{
 			SpecifierContext.MetaData.Add(UhtNames.IgnoreCategoryKeywordsInSubclasses, true);
 		}
 
 		[UhtSpecifier(Extends = UhtTableNames.Class, ValueType = UhtSpecifierValueType.SingleString)]
+		[SuppressMessage("CodeQuality", "IDE0051:Remove unused private members", Justification = "Attribute accessed method")]
 		private static void WithinSpecifier(UhtSpecifierContext SpecifierContext, StringView Value)
 		{
 			UhtClassParser Class = (UhtClassParser)SpecifierContext.Scope.ScopeType;
@@ -42,6 +47,7 @@ namespace EpicGames.UHT.Parsers
 		}
 
 		[UhtSpecifier(Extends = UhtTableNames.Class, ValueType = UhtSpecifierValueType.Legacy)]
+		[SuppressMessage("CodeQuality", "IDE0051:Remove unused private members", Justification = "Attribute accessed method")]
 		private static void EditInlineNewSpecifier(UhtSpecifierContext SpecifierContext)
 		{
 			UhtClassParser Class = (UhtClassParser)SpecifierContext.Scope.ScopeType;
@@ -49,6 +55,7 @@ namespace EpicGames.UHT.Parsers
 		}
 
 		[UhtSpecifier(Extends = UhtTableNames.Class, ValueType = UhtSpecifierValueType.Legacy)]
+		[SuppressMessage("CodeQuality", "IDE0051:Remove unused private members", Justification = "Attribute accessed method")]
 		private static void NotEditInlineNewSpecifier(UhtSpecifierContext SpecifierContext)
 		{
 			UhtClassParser Class = (UhtClassParser)SpecifierContext.Scope.ScopeType;
@@ -56,6 +63,7 @@ namespace EpicGames.UHT.Parsers
 		}
 
 		[UhtSpecifier(Extends = UhtTableNames.Class, ValueType = UhtSpecifierValueType.Legacy)]
+		[SuppressMessage("CodeQuality", "IDE0051:Remove unused private members", Justification = "Attribute accessed method")]
 		private static void PlaceableSpecifier(UhtSpecifierContext SpecifierContext)
 		{
 			UhtClassParser Class = (UhtClassParser)SpecifierContext.Scope.ScopeType;
@@ -63,6 +71,7 @@ namespace EpicGames.UHT.Parsers
 		}
 
 		[UhtSpecifier(Extends = UhtTableNames.Class, ValueType = UhtSpecifierValueType.Legacy)]
+		[SuppressMessage("CodeQuality", "IDE0051:Remove unused private members", Justification = "Attribute accessed method")]
 		private static void NotPlaceableSpecifier(UhtSpecifierContext SpecifierContext)
 		{
 			UhtClassParser Class = (UhtClassParser)SpecifierContext.Scope.ScopeType;
@@ -70,6 +79,7 @@ namespace EpicGames.UHT.Parsers
 		}
 
 		[UhtSpecifier(Extends = UhtTableNames.Class, ValueType = UhtSpecifierValueType.Legacy)]
+		[SuppressMessage("CodeQuality", "IDE0051:Remove unused private members", Justification = "Attribute accessed method")]
 		private static void DefaultToInstancedSpecifier(UhtSpecifierContext SpecifierContext)
 		{
 			UhtClassParser Class = (UhtClassParser)SpecifierContext.Scope.ScopeType;
@@ -77,6 +87,7 @@ namespace EpicGames.UHT.Parsers
 		}
 
 		[UhtSpecifier(Extends = UhtTableNames.Class, ValueType = UhtSpecifierValueType.Legacy)]
+		[SuppressMessage("CodeQuality", "IDE0051:Remove unused private members", Justification = "Attribute accessed method")]
 		private static void HideDropdownSpecifier(UhtSpecifierContext SpecifierContext)
 		{
 			UhtClassParser Class = (UhtClassParser)SpecifierContext.Scope.ScopeType;
@@ -84,6 +95,7 @@ namespace EpicGames.UHT.Parsers
 		}
 
 		[UhtSpecifier(Extends = UhtTableNames.Class, ValueType = UhtSpecifierValueType.Legacy)]
+		[SuppressMessage("CodeQuality", "IDE0051:Remove unused private members", Justification = "Attribute accessed method")]
 		private static void HiddenSpecifier(UhtSpecifierContext SpecifierContext)
 		{
 			// Prevents class from appearing in the editor class browser and edit inline menus.
@@ -92,12 +104,14 @@ namespace EpicGames.UHT.Parsers
 		}
 
 		[UhtSpecifier(Extends = UhtTableNames.Class, ValueType = UhtSpecifierValueType.Legacy)]
+		[SuppressMessage("CodeQuality", "IDE0051:Remove unused private members", Justification = "Attribute accessed method")]
 		private static void DependsOnSpecifier(UhtSpecifierContext SpecifierContext)
 		{
 			SpecifierContext.MessageSite.LogError("The dependsOn specifier is deprecated. Please use #include \"ClassHeaderFilename.h\" instead.");
 		}
 
 		[UhtSpecifier(Extends = UhtTableNames.Class, ValueType = UhtSpecifierValueType.Legacy)]
+		[SuppressMessage("CodeQuality", "IDE0051:Remove unused private members", Justification = "Attribute accessed method")]
 		private static void MinimalAPISpecifier(UhtSpecifierContext SpecifierContext)
 		{
 			UhtClassParser Class = (UhtClassParser)SpecifierContext.Scope.ScopeType;
@@ -105,6 +119,7 @@ namespace EpicGames.UHT.Parsers
 		}
 
 		[UhtSpecifier(Extends = UhtTableNames.Class, ValueType = UhtSpecifierValueType.Legacy)]
+		[SuppressMessage("CodeQuality", "IDE0051:Remove unused private members", Justification = "Attribute accessed method")]
 		private static void ConstSpecifier(UhtSpecifierContext SpecifierContext)
 		{
 			UhtClassParser Class = (UhtClassParser)SpecifierContext.Scope.ScopeType;
@@ -112,6 +127,7 @@ namespace EpicGames.UHT.Parsers
 		}
 
 		[UhtSpecifier(Extends = UhtTableNames.Class, ValueType = UhtSpecifierValueType.Legacy)]
+		[SuppressMessage("CodeQuality", "IDE0051:Remove unused private members", Justification = "Attribute accessed method")]
 		private static void PerObjectConfigSpecifier(UhtSpecifierContext SpecifierContext)
 		{
 			UhtClassParser Class = (UhtClassParser)SpecifierContext.Scope.ScopeType;
@@ -119,6 +135,7 @@ namespace EpicGames.UHT.Parsers
 		}
 
 		[UhtSpecifier(Extends = UhtTableNames.Class, ValueType = UhtSpecifierValueType.Legacy)]
+		[SuppressMessage("CodeQuality", "IDE0051:Remove unused private members", Justification = "Attribute accessed method")]
 		private static void ConfigDoNotCheckDefaultsSpecifier(UhtSpecifierContext SpecifierContext)
 		{
 			UhtClassParser Class = (UhtClassParser)SpecifierContext.Scope.ScopeType;
@@ -126,6 +143,7 @@ namespace EpicGames.UHT.Parsers
 		}
 
 		[UhtSpecifier(Extends = UhtTableNames.Class, ValueType = UhtSpecifierValueType.Legacy)]
+		[SuppressMessage("CodeQuality", "IDE0051:Remove unused private members", Justification = "Attribute accessed method")]
 		private static void AbstractSpecifier(UhtSpecifierContext SpecifierContext)
 		{
 			UhtClassParser Class = (UhtClassParser)SpecifierContext.Scope.ScopeType;
@@ -133,6 +151,7 @@ namespace EpicGames.UHT.Parsers
 		}
 
 		[UhtSpecifier(Extends = UhtTableNames.Class, ValueType = UhtSpecifierValueType.Legacy)]
+		[SuppressMessage("CodeQuality", "IDE0051:Remove unused private members", Justification = "Attribute accessed method")]
 		private static void DeprecatedSpecifier(UhtSpecifierContext SpecifierContext)
 		{
 			UhtClassParser Class = (UhtClassParser)SpecifierContext.Scope.ScopeType;
@@ -140,6 +159,7 @@ namespace EpicGames.UHT.Parsers
 		}
 
 		[UhtSpecifier(Extends = UhtTableNames.Class, ValueType = UhtSpecifierValueType.Legacy)]
+		[SuppressMessage("CodeQuality", "IDE0051:Remove unused private members", Justification = "Attribute accessed method")]
 		private static void TransientSpecifier(UhtSpecifierContext SpecifierContext)
 		{
 			UhtClassParser Class = (UhtClassParser)SpecifierContext.Scope.ScopeType;
@@ -147,6 +167,7 @@ namespace EpicGames.UHT.Parsers
 		}
 
 		[UhtSpecifier(Extends = UhtTableNames.Class, ValueType = UhtSpecifierValueType.Legacy)]
+		[SuppressMessage("CodeQuality", "IDE0051:Remove unused private members", Justification = "Attribute accessed method")]
 		private static void NonTransientSpecifier(UhtSpecifierContext SpecifierContext)
 		{
 			UhtClassParser Class = (UhtClassParser)SpecifierContext.Scope.ScopeType;
@@ -154,6 +175,7 @@ namespace EpicGames.UHT.Parsers
 		}
 
 		[UhtSpecifier(Extends = UhtTableNames.Class, ValueType = UhtSpecifierValueType.Legacy)]
+		[SuppressMessage("CodeQuality", "IDE0051:Remove unused private members", Justification = "Attribute accessed method")]
 		private static void OptionalSpecifier(UhtSpecifierContext SpecifierContext)
 		{
 			// Optional class
@@ -162,6 +184,7 @@ namespace EpicGames.UHT.Parsers
 		}
 
 		[UhtSpecifier(Extends = UhtTableNames.Class, ValueType = UhtSpecifierValueType.Legacy)]
+		[SuppressMessage("CodeQuality", "IDE0051:Remove unused private members", Justification = "Attribute accessed method")]
 		private static void CustomConstructorSpecifier(UhtSpecifierContext SpecifierContext)
 		{
 			// we will not export a constructor for this class, assuming it is in the CPP block
@@ -170,6 +193,7 @@ namespace EpicGames.UHT.Parsers
 		}
 
 		[UhtSpecifier(Extends = UhtTableNames.Class, ValueType = UhtSpecifierValueType.SingleString)]
+		[SuppressMessage("CodeQuality", "IDE0051:Remove unused private members", Justification = "Attribute accessed method")]
 		private static void ConfigSpecifier(UhtSpecifierContext SpecifierContext, StringView Value)
 		{
 			UhtClass Class = (UhtClass)SpecifierContext.Scope.ScopeType;
@@ -177,6 +201,7 @@ namespace EpicGames.UHT.Parsers
 		}
 
 		[UhtSpecifier(Extends = UhtTableNames.Class, ValueType = UhtSpecifierValueType.Legacy)]
+		[SuppressMessage("CodeQuality", "IDE0051:Remove unused private members", Justification = "Attribute accessed method")]
 		private static void DefaultConfigSpecifier(UhtSpecifierContext SpecifierContext)
 		{
 			// Save object config only to Default INIs, never to local INIs.
@@ -185,6 +210,7 @@ namespace EpicGames.UHT.Parsers
 		}
 
 		[UhtSpecifier(Extends = UhtTableNames.Class, ValueType = UhtSpecifierValueType.Legacy)]
+		[SuppressMessage("CodeQuality", "IDE0051:Remove unused private members", Justification = "Attribute accessed method")]
 		private static void GlobalUserConfigSpecifier(UhtSpecifierContext SpecifierContext)
 		{
 			// Save object config only to global user overrides, never to local INIs
@@ -193,6 +219,7 @@ namespace EpicGames.UHT.Parsers
 		}
 
 		[UhtSpecifier(Extends = UhtTableNames.Class, ValueType = UhtSpecifierValueType.Legacy)]
+		[SuppressMessage("CodeQuality", "IDE0051:Remove unused private members", Justification = "Attribute accessed method")]
 		private static void ProjectUserConfigSpecifier(UhtSpecifierContext SpecifierContext)
 		{
 			// Save object config only to project user overrides, never to INIs that are checked in
@@ -201,12 +228,14 @@ namespace EpicGames.UHT.Parsers
 		}
 
 		[UhtSpecifier(Extends = UhtTableNames.Class, ValueType = UhtSpecifierValueType.SingleString)]
+		[SuppressMessage("CodeQuality", "IDE0051:Remove unused private members", Justification = "Attribute accessed method")]
 		private static void EditorConfigSpecifier(UhtSpecifierContext SpecifierContext, StringView Value)
 		{
 			SpecifierContext.MetaData.Add(UhtNames.EditorConfig, Value.ToString());
 		}
 
 		[UhtSpecifier(Extends = UhtTableNames.Class, ValueType = UhtSpecifierValueType.NonEmptyStringList)]
+		[SuppressMessage("CodeQuality", "IDE0051:Remove unused private members", Justification = "Attribute accessed method")]
 		private static void ShowCategoriesSpecifier(UhtSpecifierContext SpecifierContext, List<StringView> Value)
 		{
 			UhtClassParser Class = (UhtClassParser)SpecifierContext.Scope.ScopeType;
@@ -214,6 +243,7 @@ namespace EpicGames.UHT.Parsers
 		}
 
 		[UhtSpecifier(Extends = UhtTableNames.Class, ValueType = UhtSpecifierValueType.NonEmptyStringList)]
+		[SuppressMessage("CodeQuality", "IDE0051:Remove unused private members", Justification = "Attribute accessed method")]
 		private static void HideCategoriesSpecifier(UhtSpecifierContext SpecifierContext, List<StringView> Value)
 		{
 			UhtClassParser Class = (UhtClassParser)SpecifierContext.Scope.ScopeType;
@@ -221,6 +251,7 @@ namespace EpicGames.UHT.Parsers
 		}
 
 		[UhtSpecifier(Extends = UhtTableNames.Class, ValueType = UhtSpecifierValueType.NonEmptyStringList)]
+		[SuppressMessage("CodeQuality", "IDE0051:Remove unused private members", Justification = "Attribute accessed method")]
 		private static void ShowFunctionsSpecifier(UhtSpecifierContext SpecifierContext, List<StringView> Value)
 		{
 			UhtClassParser Class = (UhtClassParser)SpecifierContext.Scope.ScopeType;
@@ -228,6 +259,7 @@ namespace EpicGames.UHT.Parsers
 		}
 
 		[UhtSpecifier(Extends = UhtTableNames.Class, ValueType = UhtSpecifierValueType.NonEmptyStringList)]
+		[SuppressMessage("CodeQuality", "IDE0051:Remove unused private members", Justification = "Attribute accessed method")]
 		private static void HideFunctionsSpecifier(UhtSpecifierContext SpecifierContext, List<StringView> Value)
 		{
 			UhtClassParser Class = (UhtClassParser)SpecifierContext.Scope.ScopeType;
@@ -235,6 +267,7 @@ namespace EpicGames.UHT.Parsers
 		}
 
 		[UhtSpecifier(Extends = UhtTableNames.Class, ValueType = UhtSpecifierValueType.SingleString)]
+		[SuppressMessage("CodeQuality", "IDE0051:Remove unused private members", Justification = "Attribute accessed method")]
 		private static void SparseClassDataTypesSpecifier(UhtSpecifierContext SpecifierContext, StringView Value)
 		{
 			UhtClassParser Class = (UhtClassParser)SpecifierContext.Scope.ScopeType;
@@ -242,6 +275,7 @@ namespace EpicGames.UHT.Parsers
 		}
 
 		[UhtSpecifier(Extends = UhtTableNames.Class, ValueType = UhtSpecifierValueType.NonEmptyStringList)]
+		[SuppressMessage("CodeQuality", "IDE0051:Remove unused private members", Justification = "Attribute accessed method")]
 		private static void ClassGroupSpecifier(UhtSpecifierContext SpecifierContext, List<StringView> Value)
 		{
 			UhtClassParser Class = (UhtClassParser)SpecifierContext.Scope.ScopeType;
@@ -252,6 +286,7 @@ namespace EpicGames.UHT.Parsers
 		}
 
 		[UhtSpecifier(Extends = UhtTableNames.Class, ValueType = UhtSpecifierValueType.NonEmptyStringList)]
+		[SuppressMessage("CodeQuality", "IDE0051:Remove unused private members", Justification = "Attribute accessed method")]
 		private static void AutoExpandCategoriesSpecifier(UhtSpecifierContext SpecifierContext, List<StringView> Value)
 		{
 			UhtClassParser Class = (UhtClassParser)SpecifierContext.Scope.ScopeType;
@@ -260,6 +295,7 @@ namespace EpicGames.UHT.Parsers
 		}
 
 		[UhtSpecifier(Extends = UhtTableNames.Class, ValueType = UhtSpecifierValueType.NonEmptyStringList)]
+		[SuppressMessage("CodeQuality", "IDE0051:Remove unused private members", Justification = "Attribute accessed method")]
 		private static void AutoCollapseCategoriesSpecifier(UhtSpecifierContext SpecifierContext, List<StringView> Value)
 		{
 			UhtClassParser Class = (UhtClassParser)SpecifierContext.Scope.ScopeType;
@@ -268,6 +304,7 @@ namespace EpicGames.UHT.Parsers
 		}
 
 		[UhtSpecifier(Extends = UhtTableNames.Class, ValueType = UhtSpecifierValueType.NonEmptyStringList)]
+		[SuppressMessage("CodeQuality", "IDE0051:Remove unused private members", Justification = "Attribute accessed method")]
 		private static void PrioritizeCategoriesSpecifier(UhtSpecifierContext SpecifierContext, List<StringView> Value)
 		{
 			UhtClassParser Class = (UhtClassParser)SpecifierContext.Scope.ScopeType;
@@ -275,6 +312,7 @@ namespace EpicGames.UHT.Parsers
 		}
 
 		[UhtSpecifier(Extends = UhtTableNames.Class, ValueType = UhtSpecifierValueType.NonEmptyStringList)]
+		[SuppressMessage("CodeQuality", "IDE0051:Remove unused private members", Justification = "Attribute accessed method")]
 		private static void DontAutoCollapseCategoriesSpecifier(UhtSpecifierContext SpecifierContext, List<StringView> Value)
 		{
 			UhtClassParser Class = (UhtClassParser)SpecifierContext.Scope.ScopeType;
@@ -282,6 +320,7 @@ namespace EpicGames.UHT.Parsers
 		}
 
 		[UhtSpecifier(Extends = UhtTableNames.Class, ValueType = UhtSpecifierValueType.Legacy)]
+		[SuppressMessage("CodeQuality", "IDE0051:Remove unused private members", Justification = "Attribute accessed method")]
 		private static void CollapseCategoriesSpecifier(UhtSpecifierContext SpecifierContext)
 		{
 			// Class' properties should not be shown categorized in the editor.
@@ -290,6 +329,7 @@ namespace EpicGames.UHT.Parsers
 		}
 
 		[UhtSpecifier(Extends = UhtTableNames.Class, ValueType = UhtSpecifierValueType.Legacy)]
+		[SuppressMessage("CodeQuality", "IDE0051:Remove unused private members", Justification = "Attribute accessed method")]
 		private static void DontCollapseCategoriesSpecifier(UhtSpecifierContext SpecifierContext)
 		{
 			// Class' properties should be shown categorized in the editor.
@@ -298,6 +338,7 @@ namespace EpicGames.UHT.Parsers
 		}
 
 		[UhtSpecifier(Extends = UhtTableNames.Class, ValueType = UhtSpecifierValueType.Legacy)]
+		[SuppressMessage("CodeQuality", "IDE0051:Remove unused private members", Justification = "Attribute accessed method")]
 		private static void AdvancedClassDisplaySpecifier(UhtSpecifierContext SpecifierContext)
 		{
 			// By default the class properties are shown in advanced sections in UI
@@ -306,6 +347,7 @@ namespace EpicGames.UHT.Parsers
 		}
 
 		[UhtSpecifier(Extends = UhtTableNames.Class, ValueType = UhtSpecifierValueType.Legacy)]
+		[SuppressMessage("CodeQuality", "IDE0051:Remove unused private members", Justification = "Attribute accessed method")]
 		private static void ConversionRootSpecifier(UhtSpecifierContext SpecifierContext)
 		{
 			UhtClass Class = (UhtClass)SpecifierContext.Scope.ScopeType;
@@ -313,6 +355,7 @@ namespace EpicGames.UHT.Parsers
 		}
 
 		[UhtSpecifier(Extends = UhtTableNames.Class, ValueType = UhtSpecifierValueType.Legacy)]
+		[SuppressMessage("CodeQuality", "IDE0051:Remove unused private members", Justification = "Attribute accessed method")]
 		private static void NeedsDeferredDependencyLoadingSpecifier(UhtSpecifierContext SpecifierContext)
 		{
 			UhtClassParser Class = (UhtClassParser)SpecifierContext.Scope.ScopeType;
@@ -320,6 +363,7 @@ namespace EpicGames.UHT.Parsers
 		}
 
 		[UhtSpecifier(Extends = UhtTableNames.Class, ValueType = UhtSpecifierValueType.Legacy)]
+		[SuppressMessage("CodeQuality", "IDE0051:Remove unused private members", Justification = "Attribute accessed method")]
 		private static void MatchedSerializersSpecifier(UhtSpecifierContext SpecifierContext)
 		{
 			if (!SpecifierContext.Scope.HeaderParser.HeaderFile.bIsNoExportTypes)
@@ -331,6 +375,7 @@ namespace EpicGames.UHT.Parsers
 		}
 
 		[UhtSpecifier(Extends = UhtTableNames.Class, ValueType = UhtSpecifierValueType.Legacy)]
+		[SuppressMessage("CodeQuality", "IDE0051:Remove unused private members", Justification = "Attribute accessed method")]
 		private static void InterfaceSpecifier(UhtSpecifierContext SpecifierContext)
 		{
 			if (!SpecifierContext.Scope.HeaderParser.HeaderFile.bIsNoExportTypes)
@@ -342,6 +387,7 @@ namespace EpicGames.UHT.Parsers
 		}
 
 		[UhtSpecifier(Extends = UhtTableNames.Class, ValueType = UhtSpecifierValueType.None)]
+		[SuppressMessage("CodeQuality", "IDE0051:Remove unused private members", Justification = "Attribute accessed method")]
 		private static void CustomFieldNotifySpecifier(UhtSpecifierContext SpecifierContext)
 		{
 			UhtClassParser Class = (UhtClassParser)SpecifierContext.Scope.ScopeType;
diff --git a/Engine/Source/Programs/Shared/EpicGames.UHT/Specifiers/UhtDefaultSpecifiers.cs b/Engine/Source/Programs/Shared/EpicGames.UHT/Specifiers/UhtDefaultSpecifiers.cs
--- a/Engine/Source/Programs/Shared/EpicGames.UHT/Specifiers/UhtDefaultSpecifiers.cs
+++ b/Engine/Source/Programs/Shared/EpicGames.UHT/Specifiers/UhtDefaultSpecifiers.cs
@@ -6,6 +6,7 @@ using EpicGames.UHT.Types;
 using EpicGames.UHT.Utils;
 using System;
 using System.Collections.Generic;
+using System.Diagnostics.CodeAnalysis;
 
 namespace EpicGames.UHT.Parsers
 {
@@ -80,6 +81,7 @@ namespace EpicGames.UHT.Parsers
 
 		#region Specifiers
 		[UhtSpecifier(Extends = UhtTableNames.Default, ValueType = UhtSpecifierValueType.KeyValuePairList, When = UhtSpecifierWhen.Immediate)]
+		[SuppressMessage("CodeQuality", "IDE0051:Remove unused private members", Justification = "Attribute accessed method")]
 		private static void MetaSpecifier(UhtSpecifierContext SpecifierContext, List<KeyValuePair<StringView, StringView>> Value)
 		{
 			foreach (KeyValuePair<StringView, StringView> KVP in (List<KeyValuePair<StringView, StringView>>)Value)
@@ -89,18 +91,21 @@ namespace EpicGames.UHT.Parsers
 		}
 
 		[UhtSpecifier(Extends = UhtTableNames.Default, ValueType = UhtSpecifierValueType.String, When = UhtSpecifierWhen.Immediate)]
+		[SuppressMessage("CodeQuality", "IDE0051:Remove unused private members", Justification = "Attribute accessed method")]
 		private static void DisplayNameSpecifier(UhtSpecifierContext SpecifierContext, StringView Value)
 		{
 			SetMetaData(SpecifierContext, "DisplayName", Value);
 		}
 
 		[UhtSpecifier(Extends = UhtTableNames.Default, ValueType = UhtSpecifierValueType.String, When = UhtSpecifierWhen.Immediate)]
+		[SuppressMessage("CodeQuality", "IDE0051:Remove unused private members", Justification = "Attribute accessed method")]
 		private static void FriendlyNameSpecifier(UhtSpecifierContext SpecifierContext, StringView Value)
 		{
 			SetMetaData(SpecifierContext, "FriendlyName", Value);
 		}
 
 		[UhtSpecifier(Extends = UhtTableNames.Default, ValueType = UhtSpecifierValueType.None, When = UhtSpecifierWhen.Immediate)]
+		[SuppressMessage("CodeQuality", "IDE0051:Remove unused private members", Justification = "Attribute accessed method")]
 		private static void BlueprintInternalUseOnlySpecifier(UhtSpecifierContext SpecifierContext)
 		{
 			SetMetaData(SpecifierContext, "BlueprintInternalUseOnly", true);
@@ -108,6 +113,7 @@ namespace EpicGames.UHT.Parsers
 		}
 
 		[UhtSpecifier(Extends = UhtTableNames.Default, ValueType = UhtSpecifierValueType.None, When = UhtSpecifierWhen.Immediate)]
+		[SuppressMessage("CodeQuality", "IDE0051:Remove unused private members", Justification = "Attribute accessed method")]
 		private static void BlueprintInternalUseOnlyHierarchicalSpecifier(UhtSpecifierContext SpecifierContext)
 		{
 			SetMetaData(SpecifierContext, "BlueprintInternalUseOnlyHierarchical", true);
@@ -116,12 +122,14 @@ namespace EpicGames.UHT.Parsers
 		}
 
 		[UhtSpecifier(Extends = UhtTableNames.Default, ValueType = UhtSpecifierValueType.None, When = UhtSpecifierWhen.Immediate)]
+		[SuppressMessage("CodeQuality", "IDE0051:Remove unused private members", Justification = "Attribute accessed method")]
 		private static void BlueprintTypeSpecifier(UhtSpecifierContext SpecifierContext)
 		{
 			SetMetaData(SpecifierContext, "BlueprintType", true);
 		}
 
 		[UhtSpecifier(Extends = UhtTableNames.Default, ValueType = UhtSpecifierValueType.None, When = UhtSpecifierWhen.Immediate)]
+		[SuppressMessage("CodeQuality", "IDE0051:Remove unused private members", Justification = "Attribute accessed method")]
 		private static void NotBlueprintTypeSpecifier(UhtSpecifierContext SpecifierContext)
 		{
 			SetMetaData(SpecifierContext, "NotBlueprintType", true);
@@ -129,6 +137,7 @@ namespace EpicGames.UHT.Parsers
 		}
 
 		[UhtSpecifier(Extends = UhtTableNames.Default, ValueType = UhtSpecifierValueType.None, When = UhtSpecifierWhen.Immediate)]
+		[SuppressMessage("CodeQuality", "IDE0051:Remove unused private members", Justification = "Attribute accessed method")]
 		private static void BlueprintableSpecifier(UhtSpecifierContext SpecifierContext)
 		{
 			SetMetaData(SpecifierContext, "IsBlueprintBase", true);
@@ -136,12 +145,14 @@ namespace EpicGames.UHT.Parsers
 		}
 
 		[UhtSpecifier(Extends = UhtTableNames.Default, ValueType = UhtSpecifierValueType.None, When = UhtSpecifierWhen.Immediate)]
+		[SuppressMessage("CodeQuality", "IDE0051:Remove unused private members", Justification = "Attribute accessed method")]
 		private static void CallInEditorSpecifier(UhtSpecifierContext SpecifierContext)
 		{
 			SetMetaData(SpecifierContext, "CallInEditor", true);
 		}
 
 		[UhtSpecifier(Extends = UhtTableNames.Default, ValueType = UhtSpecifierValueType.None, When = UhtSpecifierWhen.Immediate)]
+		[SuppressMessage("CodeQuality", "IDE0051:Remove unused private members", Justification = "Attribute accessed method")]
 		private static void NotBlueprintableSpecifier(UhtSpecifierContext SpecifierContext)
 		{
 			SetMetaData(SpecifierContext, "IsBlueprintBase", false);
@@ -149,30 +160,35 @@ namespace EpicGames.UHT.Parsers
 		}
 
 		[UhtSpecifier(Extends = UhtTableNames.Default, ValueType = UhtSpecifierValueType.String, When = UhtSpecifierWhen.Immediate)]
+		[SuppressMessage("CodeQuality", "IDE0051:Remove unused private members", Justification = "Attribute accessed method")]
 		private static void CategorySpecifier(UhtSpecifierContext SpecifierContext, StringView Value)
 		{
 			SetMetaData(SpecifierContext, "Category", Value);
 		}
 
 		[UhtSpecifier(Extends = UhtTableNames.Default, ValueType = UhtSpecifierValueType.None, When = UhtSpecifierWhen.Immediate)]
+		[SuppressMessage("CodeQuality", "IDE0051:Remove unused private members", Justification = "Attribute accessed method")]
 		private static void ExperimentalSpecifier(UhtSpecifierContext SpecifierContext)
 		{
 			SetMetaData(SpecifierContext, "DevelopmentStatus", "Experimental");
 		}
 
 		[UhtSpecifier(Extends = UhtTableNames.Default, ValueType = UhtSpecifierValueType.None, When = UhtSpecifierWhen.Immediate)]
+		[SuppressMessage("CodeQuality", "IDE0051:Remove unused private members", Justification = "Attribute accessed method")]
 		private static void EarlyAccessPreviewSpecifier(UhtSpecifierContext SpecifierContext)
 		{
 			SetMetaData(SpecifierContext, "DevelopmentStatus", "EarlyAccess");
 		}
 
 		[UhtSpecifier(Extends = UhtTableNames.Default, ValueType = UhtSpecifierValueType.None, When = UhtSpecifierWhen.Immediate)]
+		[SuppressMessage("CodeQuality", "IDE0051:Remove unused private members", Justification = "Attribute accessed method")]
 		private static void DocumentationPolicySpecifier(UhtSpecifierContext SpecifierContext)
 		{
 			SetMetaData(SpecifierContext, UhtNames.DocumentationPolicy, "Strict");
 		}
 
 		[UhtSpecifier(Extends = UhtTableNames.Default, ValueType = UhtSpecifierValueType.String, When = UhtSpecifierWhen.Immediate)]
+		[SuppressMessage("CodeQuality", "IDE0051:Remove unused private members", Justification = "Attribute accessed method")]
 		private static void SparseClassDataTypeSpecifier(UhtSpecifierContext SpecifierContext, StringView Value)
 		{
 			SetMetaData(SpecifierContext, "SparseClassDataType", Value);
@@ -184,6 +200,7 @@ namespace EpicGames.UHT.Parsers
 		[UhtSpecifierValidator(Name = "UIMax", Extends = UhtTableNames.Default)]
 		[UhtSpecifierValidator(Name = "ClampMin", Extends = UhtTableNames.Default)]
 		[UhtSpecifierValidator(Name = "ClampMax", Extends = UhtTableNames.Default)]
+		[SuppressMessage("CodeQuality", "IDE0051:Remove unused private members", Justification = "Attribute accessed method")]
 		private static void ValidateNumeric(UhtType Type, UhtMetaData MetaData, UhtMetaDataKey Key, StringView Value)
 		{
 			if (!UhtFCString.IsNumeric(Value.Span))
@@ -193,6 +210,7 @@ namespace EpicGames.UHT.Parsers
 		}
 
 		[UhtSpecifierValidator(Extends = UhtTableNames.Default)]
+		[SuppressMessage("CodeQuality", "IDE0051:Remove unused private members", Justification = "Attribute accessed method")]
 		private static void DevelopmentStatusSpecifierValidator(UhtType Type, UhtMetaData MetaData, UhtMetaDataKey Key, StringView Value)
 		{
 			string[] AllowedValues = { "EarlyAccess", "Experimental" };
@@ -207,6 +225,7 @@ namespace EpicGames.UHT.Parsers
 		}
 
 		[UhtSpecifierValidator(Extends = UhtTableNames.Default)]
+		[SuppressMessage("CodeQuality", "IDE0051:Remove unused private members", Justification = "Attribute accessed method")]
 		private static void DocumentationPolicySpecifierValidator(UhtType Type, UhtMetaData MetaData, UhtMetaDataKey Key, StringView Value)
 		{
 			const string StrictValue = "Strict";
@@ -217,6 +236,7 @@ namespace EpicGames.UHT.Parsers
 		}
 
 		[UhtSpecifierValidator(Extends = UhtTableNames.Default)]
+		[SuppressMessage("CodeQuality", "IDE0051:Remove unused private members", Justification = "Attribute accessed method")]
 		private static void UnitsSpecifierValidator(UhtType Type, UhtMetaData MetaData, UhtMetaDataKey Key, StringView Value)
 		{
 			// Check for numeric property
diff --git a/Engine/Source/Programs/Shared/EpicGames.UHT/Specifiers/UhtEnumSpecifiers.cs b/Engine/Source/Programs/Shared/EpicGames.UHT/Specifiers/UhtEnumSpecifiers.cs
--- a/Engine/Source/Programs/Shared/EpicGames.UHT/Specifiers/UhtEnumSpecifiers.cs
+++ b/Engine/Source/Programs/Shared/EpicGames.UHT/Specifiers/UhtEnumSpecifiers.cs
@@ -1,5 +1,6 @@
 // Copyright Epic Games, Inc. All Rights Reserved.
 
+using System.Diagnostics.CodeAnalysis;
 using EpicGames.Core;
 using EpicGames.UHT.Tables;
 using EpicGames.UHT.Types;
@@ -15,6 +16,7 @@ namespace EpicGames.UHT.Parsers
 	public static class UhtEnumSpecifiers
 	{
 		[UhtSpecifier(Extends = UhtTableNames.Enum, ValueType = UhtSpecifierValueType.Legacy)]
+		[SuppressMessage("CodeQuality", "IDE0051:Remove unused private members", Justification = "Attribute accessed method")]
 		private static void FlagsSpecifier(UhtSpecifierContext SpecifierContext)
 		{
 			UhtEnum Enum = (UhtEnum)SpecifierContext.Scope.ScopeType;
diff --git a/Engine/Source/Programs/Shared/EpicGames.UHT/Specifiers/UhtFunctionSpecifiers.cs b/Engine/Source/Programs/Shared/EpicGames.UHT/Specifiers/UhtFunctionSpecifiers.cs
--- a/Engine/Source/Programs/Shared/EpicGames.UHT/Specifiers/UhtFunctionSpecifiers.cs
+++ b/Engine/Source/Programs/Shared/EpicGames.UHT/Specifiers/UhtFunctionSpecifiers.cs
@@ -7,6 +7,7 @@ using EpicGames.UHT.Types;
 using EpicGames.UHT.Utils;
 using System;
 using System.Collections.Generic;
+using System.Diagnostics.CodeAnalysis;
 
 namespace EpicGames.UHT.Parsers
 {
@@ -17,6 +18,7 @@ namespace EpicGames.UHT.Parsers
 	public static class UhtFunctionSpecifiers
 	{
 		[UhtSpecifier(Extends = UhtTableNames.Function, ValueType = UhtSpecifierValueType.Legacy)]
+		[SuppressMessage("CodeQuality", "IDE0051:Remove unused private members", Justification = "Attribute accessed method")]
 		private static void BlueprintAuthorityOnlySpecifier(UhtSpecifierContext SpecifierContext)
 		{
 			UhtFunction Function = (UhtFunction)SpecifierContext.Scope.ScopeType;
@@ -24,6 +26,7 @@ namespace EpicGames.UHT.Parsers
 		}
 
 		[UhtSpecifier(Extends = UhtTableNames.Function, ValueType = UhtSpecifierValueType.Legacy)]
+		[SuppressMessage("CodeQuality", "IDE0051:Remove unused private members", Justification = "Attribute accessed method")]
 		private static void BlueprintCallableSpecifier(UhtSpecifierContext SpecifierContext)
 		{
 			UhtFunction Function = (UhtFunction)SpecifierContext.Scope.ScopeType;
@@ -31,6 +34,7 @@ namespace EpicGames.UHT.Parsers
 		}
 
 		[UhtSpecifier(Extends = UhtTableNames.Function, ValueType = UhtSpecifierValueType.Legacy)]
+		[SuppressMessage("CodeQuality", "IDE0051:Remove unused private members", Justification = "Attribute accessed method")]
 		private static void BlueprintCosmeticSpecifier(UhtSpecifierContext SpecifierContext)
 		{
 			UhtFunction Function = (UhtFunction)SpecifierContext.Scope.ScopeType;
@@ -38,6 +42,7 @@ namespace EpicGames.UHT.Parsers
 		}
 
 		[UhtSpecifier(Extends = UhtTableNames.Function, ValueType = UhtSpecifierValueType.Legacy)]
+		[SuppressMessage("CodeQuality", "IDE0051:Remove unused private members", Justification = "Attribute accessed method")]
 		private static void BlueprintGetterSpecifier(UhtSpecifierContext SpecifierContext)
 		{
 			UhtFunctionParser Function = (UhtFunctionParser)SpecifierContext.Scope.ScopeType;
@@ -54,6 +59,7 @@ namespace EpicGames.UHT.Parsers
 		}
 
 		[UhtSpecifier(Extends = UhtTableNames.Function, ValueType = UhtSpecifierValueType.Legacy)]
+		[SuppressMessage("CodeQuality", "IDE0051:Remove unused private members", Justification = "Attribute accessed method")]
 		private static void BlueprintImplementableEventSpecifier(UhtSpecifierContext SpecifierContext)
 		{
 			UhtFunctionParser Function = (UhtFunctionParser)SpecifierContext.Scope.ScopeType;
@@ -82,6 +88,7 @@ namespace EpicGames.UHT.Parsers
 		}
 
 		[UhtSpecifier(Extends = UhtTableNames.Function, ValueType = UhtSpecifierValueType.Legacy)]
+		[SuppressMessage("CodeQuality", "IDE0051:Remove unused private members", Justification = "Attribute accessed method")]
 		private static void BlueprintNativeEventSpecifier(UhtSpecifierContext SpecifierContext)
 		{
 			UhtFunctionParser Function = (UhtFunctionParser)SpecifierContext.Scope.ScopeType;
@@ -109,6 +116,7 @@ namespace EpicGames.UHT.Parsers
 		}
 
 		[UhtSpecifier(Extends = UhtTableNames.Function, ValueType = UhtSpecifierValueType.OptionalString)]
+		[SuppressMessage("CodeQuality", "IDE0051:Remove unused private members", Justification = "Attribute accessed method")]
 		private static void BlueprintPureSpecifier(UhtSpecifierContext SpecifierContext, StringView? Value)
 		{
 			UhtFunction Function = (UhtFunction)SpecifierContext.Scope.ScopeType;
@@ -127,6 +135,7 @@ namespace EpicGames.UHT.Parsers
 		}
 
 		[UhtSpecifier(Extends = UhtTableNames.Function, ValueType = UhtSpecifierValueType.Legacy)]
+		[SuppressMessage("CodeQuality", "IDE0051:Remove unused private members", Justification = "Attribute accessed method")]
 		private static void BlueprintSetterSpecifier(UhtSpecifierContext SpecifierContext)
 		{
 			UhtFunctionParser Function = (UhtFunctionParser)SpecifierContext.Scope.ScopeType;
@@ -142,6 +151,7 @@ namespace EpicGames.UHT.Parsers
 		}
 
 		[UhtSpecifier(Extends = UhtTableNames.Function, ValueType = UhtSpecifierValueType.OptionalString)]
+		[SuppressMessage("CodeQuality", "IDE0051:Remove unused private members", Justification = "Attribute accessed method")]
 		private static void ClientSpecifier(UhtSpecifierContext SpecifierContext, StringView? Value)
 		{
 			UhtFunction Function = (UhtFunction)SpecifierContext.Scope.ScopeType;
@@ -166,6 +176,7 @@ namespace EpicGames.UHT.Parsers
 		}
 
 		[UhtSpecifier(Extends = UhtTableNames.Function, ValueType = UhtSpecifierValueType.Legacy)]
+		[SuppressMessage("CodeQuality", "IDE0051:Remove unused private members", Justification = "Attribute accessed method")]
 		private static void CustomThunkSpecifier(UhtSpecifierContext SpecifierContext)
 		{
 			UhtFunction Function = (UhtFunction)SpecifierContext.Scope.ScopeType;
@@ -173,6 +184,7 @@ namespace EpicGames.UHT.Parsers
 		}
 
 		[UhtSpecifier(Extends = UhtTableNames.Function, ValueType = UhtSpecifierValueType.Legacy)]
+		[SuppressMessage("CodeQuality", "IDE0051:Remove unused private members", Justification = "Attribute accessed method")]
 		private static void ExecSpecifier(UhtSpecifierContext SpecifierContext)
 		{
 			UhtFunction Function = (UhtFunction)SpecifierContext.Scope.ScopeType;
@@ -184,6 +196,7 @@ namespace EpicGames.UHT.Parsers
 		}
 
 		[UhtSpecifier(Extends = UhtTableNames.Function, ValueType = UhtSpecifierValueType.Legacy)]
+		[SuppressMessage("CodeQuality", "IDE0051:Remove unused private members", Justification = "Attribute accessed method")]
 		private static void NetMulticastSpecifier(UhtSpecifierContext SpecifierContext)
 		{
 			UhtFunction Function = (UhtFunction)SpecifierContext.Scope.ScopeType;
@@ -197,6 +210,7 @@ namespace EpicGames.UHT.Parsers
 		}
 
 		[UhtSpecifier(Extends = UhtTableNames.Function, ValueType = UhtSpecifierValueType.Legacy)]
+		[SuppressMessage("CodeQuality", "IDE0051:Remove unused private members", Justification = "Attribute accessed method")]
 		private static void ReliableSpecifier(UhtSpecifierContext SpecifierContext)
 		{
 			UhtFunction Function = (UhtFunction)SpecifierContext.Scope.ScopeType;
@@ -204,6 +218,7 @@ namespace EpicGames.UHT.Parsers
 		}
 
 		[UhtSpecifier(Extends = UhtTableNames.Function, ValueType = UhtSpecifierValueType.Legacy)]
+		[SuppressMessage("CodeQuality", "IDE0051:Remove unused private members", Justification = "Attribute accessed method")]
 		private static void SealedEventSpecifier(UhtSpecifierContext SpecifierContext)
 		{
 			UhtFunction Function = (UhtFunction)SpecifierContext.Scope.ScopeType;
@@ -211,6 +226,7 @@ namespace EpicGames.UHT.Parsers
 		}
 
 		[UhtSpecifier(Extends = UhtTableNames.Function, ValueType = UhtSpecifierValueType.OptionalString)]
+		[SuppressMessage("CodeQuality", "IDE0051:Remove unused private members", Justification = "Attribute accessed method")]
 		private static void ServerSpecifier(UhtSpecifierContext SpecifierContext, StringView? Value)
 		{
 			UhtFunction Function = (UhtFunction)SpecifierContext.Scope.ScopeType;
@@ -235,6 +251,7 @@ namespace EpicGames.UHT.Parsers
 		}
 
 		[UhtSpecifier(Extends = UhtTableNames.Function, ValueType = UhtSpecifierValueType.OptionalEqualsKeyValuePairList)]
+		[SuppressMessage("CodeQuality", "IDE0051:Remove unused private members", Justification = "Attribute accessed method")]
 		private static void ServiceRequestSpecifier(UhtSpecifierContext SpecifierContext, List<KeyValuePair<StringView, StringView>> Value)
 		{
 			UhtFunctionParser Function = (UhtFunctionParser)SpecifierContext.Scope.ScopeType;
@@ -258,6 +275,7 @@ namespace EpicGames.UHT.Parsers
 		}
 
 		[UhtSpecifier(Extends = UhtTableNames.Function, ValueType = UhtSpecifierValueType.OptionalEqualsKeyValuePairList)]
+		[SuppressMessage("CodeQuality", "IDE0051:Remove unused private members", Justification = "Attribute accessed method")]
 		private static void ServiceResponseSpecifier(UhtSpecifierContext SpecifierContext, List<KeyValuePair<StringView, StringView>> Value)
 		{
 			UhtFunctionParser Function = (UhtFunctionParser)SpecifierContext.Scope.ScopeType;
@@ -281,6 +299,7 @@ namespace EpicGames.UHT.Parsers
 		}
 
 		[UhtSpecifier(Extends = UhtTableNames.Function, ValueType = UhtSpecifierValueType.Legacy)]
+		[SuppressMessage("CodeQuality", "IDE0051:Remove unused private members", Justification = "Attribute accessed method")]
 		private static void UnreliableSpecifier(UhtSpecifierContext SpecifierContext)
 		{
 			UhtFunction Function = (UhtFunction)SpecifierContext.Scope.ScopeType;
@@ -288,6 +307,7 @@ namespace EpicGames.UHT.Parsers
 		}
 
 		[UhtSpecifier(Extends = UhtTableNames.Function, ValueType = UhtSpecifierValueType.StringList)]
+		[SuppressMessage("CodeQuality", "IDE0051:Remove unused private members", Justification = "Attribute accessed method")]
 		private static void WithValidationSpecifier(UhtSpecifierContext SpecifierContext, List<StringView>? Value)
 		{
 			UhtFunction Function = (UhtFunction)SpecifierContext.Scope.ScopeType;
@@ -300,6 +320,7 @@ namespace EpicGames.UHT.Parsers
 		}
 
 		[UhtSpecifier(Extends = UhtTableNames.Function, ValueType = UhtSpecifierValueType.None)]
+		[SuppressMessage("CodeQuality", "IDE0051:Remove unused private members", Justification = "Attribute accessed method")]
 		private static void FieldNotifySpecifier(UhtSpecifierContext SpecifierContext)
 		{
 			UhtFunction Function = (UhtFunction)SpecifierContext.Scope.ScopeType;
@@ -307,6 +328,7 @@ namespace EpicGames.UHT.Parsers
 		}
 
 		[UhtSpecifierValidator(Extends = UhtTableNames.Function)]
+		[SuppressMessage("CodeQuality", "IDE0051:Remove unused private members", Justification = "Attribute accessed method")]
 		private static void BlueprintProtectedSpecifierValidator(UhtType Type, UhtMetaData MetaData, UhtMetaDataKey Key, StringView Value)
 		{
 			UhtFunction Function = (UhtFunction)Type;
@@ -334,6 +356,7 @@ namespace EpicGames.UHT.Parsers
 		}
 
 		[UhtSpecifierValidator(Extends = UhtTableNames.Function)]
+		[SuppressMessage("CodeQuality", "IDE0051:Remove unused private members", Justification = "Attribute accessed method")]
 		private static void CommutativeAssociativeBinaryOperatorSpecifierValidator(UhtType Type, UhtMetaData MetaData, UhtMetaDataKey Key, StringView Value)
 		{
 			UhtFunction Function = (UhtFunction)Type;
@@ -348,6 +371,7 @@ namespace EpicGames.UHT.Parsers
 
 		[UhtSpecifierValidator(Name = "ExpandBoolAsExecs", Extends = UhtTableNames.Function)]
 		[UhtSpecifierValidator(Name = "ExpandEnumAsExecs", Extends = UhtTableNames.Function)]
+		[SuppressMessage("CodeQuality", "IDE0051:Remove unused private members", Justification = "Attribute accessed method")]
 		private static void ExpandsSpecifierValidator(UhtType Type, UhtMetaData MetaData, UhtMetaDataKey Key, StringView Value)
 		{
 			// multiple entry parsing in the same format as eg SetParam.
diff --git a/Engine/Source/Programs/Shared/EpicGames.UHT/Specifiers/UhtInterfaceClassSpecifiers.cs b/Engine/Source/Programs/Shared/EpicGames.UHT/Specifiers/UhtInterfaceClassSpecifiers.cs
--- a/Engine/Source/Programs/Shared/EpicGames.UHT/Specifiers/UhtInterfaceClassSpecifiers.cs
+++ b/Engine/Source/Programs/Shared/EpicGames.UHT/Specifiers/UhtInterfaceClassSpecifiers.cs
@@ -1,5 +1,6 @@
 // Copyright Epic Games, Inc. All Rights Reserved.
 
+using System.Diagnostics.CodeAnalysis;
 using EpicGames.Core;
 using EpicGames.UHT.Tables;
 using EpicGames.UHT.Types;
@@ -11,12 +12,14 @@ namespace EpicGames.UHT.Parsers
 	class UhtInterfaceClassSpecifiers
 	{
 		[UhtSpecifier(Extends = UhtTableNames.Interface, ValueType = UhtSpecifierValueType.Legacy)]
+		[SuppressMessage("CodeQuality", "IDE0051:Remove unused private members", Justification = "Attribute accessed method")]
 		private static void DependsOnSpecifier(UhtSpecifierContext SpecifierContext)
 		{
 			throw new UhtException(SpecifierContext.MessageSite, $"The dependsOn specifier is deprecated. Please use #include \"ClassHeaderFilename.h\" instead.");
 		}
 
 		[UhtSpecifier(Extends = UhtTableNames.Interface, ValueType = UhtSpecifierValueType.Legacy)]
+		[SuppressMessage("CodeQuality", "IDE0051:Remove unused private members", Justification = "Attribute accessed method")]
 		private static void MinimalAPISpecifier(UhtSpecifierContext SpecifierContext)
 		{
 			UhtClass Class = (UhtClass)SpecifierContext.Scope.ScopeType;
@@ -24,6 +27,7 @@ namespace EpicGames.UHT.Parsers
 		}
 
 		[UhtSpecifier(Extends = UhtTableNames.Interface, ValueType = UhtSpecifierValueType.Legacy)]
+		[SuppressMessage("CodeQuality", "IDE0051:Remove unused private members", Justification = "Attribute accessed method")]
 		private static void ConversionRootSpecifier(UhtSpecifierContext SpecifierContext)
 		{
 			UhtClass Class = (UhtClass)SpecifierContext.Scope.ScopeType;
diff --git a/Engine/Source/Programs/Shared/EpicGames.UHT/Specifiers/UhtPropertyArgumentSpecifiers.cs b/Engine/Source/Programs/Shared/EpicGames.UHT/Specifiers/UhtPropertyArgumentSpecifiers.cs
--- a/Engine/Source/Programs/Shared/EpicGames.UHT/Specifiers/UhtPropertyArgumentSpecifiers.cs
+++ b/Engine/Source/Programs/Shared/EpicGames.UHT/Specifiers/UhtPropertyArgumentSpecifiers.cs
@@ -1,5 +1,6 @@
 // Copyright Epic Games, Inc. All Rights Reserved.
 
+using System.Diagnostics.CodeAnalysis;
 using EpicGames.Core;
 using EpicGames.UHT.Tables;
 using EpicGames.UHT.Types;
@@ -16,6 +17,7 @@ namespace EpicGames.UHT.Parsers
 	{
 		#region Argument Property Specifiers
 		[UhtSpecifier(Extends = UhtTableNames.PropertyArgument, ValueType = UhtSpecifierValueType.Legacy)]
+		[SuppressMessage("CodeQuality", "IDE0051:Remove unused private members", Justification = "Attribute accessed method")]
 		private static void ConstSpecifier(UhtSpecifierContext SpecifierContext)
 		{
 			UhtPropertySpecifierContext Context = (UhtPropertySpecifierContext)SpecifierContext;
@@ -23,6 +25,7 @@ namespace EpicGames.UHT.Parsers
 		}
 
 		[UhtSpecifier(Extends = UhtTableNames.PropertyArgument, ValueType = UhtSpecifierValueType.Legacy)]
+		[SuppressMessage("CodeQuality", "IDE0051:Remove unused private members", Justification = "Attribute accessed method")]
 		private static void RefSpecifier(UhtSpecifierContext SpecifierContext)
 		{
 			UhtPropertySpecifierContext Context = (UhtPropertySpecifierContext)SpecifierContext;
@@ -30,6 +33,7 @@ namespace EpicGames.UHT.Parsers
 		}
 
 		[UhtSpecifier(Extends = UhtTableNames.PropertyArgument, ValueType = UhtSpecifierValueType.Legacy)]
+		[SuppressMessage("CodeQuality", "IDE0051:Remove unused private members", Justification = "Attribute accessed method")]
 		private static void NotReplicatedSpecifier(UhtSpecifierContext SpecifierContext)
 		{
 			UhtPropertySpecifierContext Context = (UhtPropertySpecifierContext)SpecifierContext;
diff --git a/Engine/Source/Programs/Shared/EpicGames.UHT/Specifiers/UhtPropertyMemberSpecifiers.cs b/Engine/Source/Programs/Shared/EpicGames.UHT/Specifiers/UhtPropertyMemberSpecifiers.cs
--- a/Engine/Source/Programs/Shared/EpicGames.UHT/Specifiers/UhtPropertyMemberSpecifiers.cs
+++ b/Engine/Source/Programs/Shared/EpicGames.UHT/Specifiers/UhtPropertyMemberSpecifiers.cs
@@ -5,6 +5,7 @@ using EpicGames.UHT.Tables;
 using EpicGames.UHT.Types;
 using EpicGames.UHT.Utils;
 using System;
+using System.Diagnostics.CodeAnalysis;
 
 namespace EpicGames.UHT.Parsers
 {
@@ -16,6 +17,7 @@ namespace EpicGames.UHT.Parsers
 	public static class UhtPropertyMemberSpecifiers
 	{
 		[UhtSpecifier(Extends = UhtTableNames.PropertyMember, ValueType = UhtSpecifierValueType.Legacy)]
+		[SuppressMessage("CodeQuality", "IDE0051:Remove unused private members", Justification = "Attribute accessed method")]
 		private static void EditAnywhereSpecifier(UhtSpecifierContext SpecifierContext)
 		{
 			UhtPropertySpecifierContext Context = (UhtPropertySpecifierContext)SpecifierContext;
@@ -28,6 +30,7 @@ namespace EpicGames.UHT.Parsers
 		}
 
 		[UhtSpecifier(Extends = UhtTableNames.PropertyMember, ValueType = UhtSpecifierValueType.Legacy)]
+		[SuppressMessage("CodeQuality", "IDE0051:Remove unused private members", Justification = "Attribute accessed method")]
 		private static void EditInstanceOnlySpecifier(UhtSpecifierContext SpecifierContext)
 		{
 			UhtPropertySpecifierContext Context = (UhtPropertySpecifierContext)SpecifierContext;
@@ -40,6 +43,7 @@ namespace EpicGames.UHT.Parsers
 		}
 
 		[UhtSpecifier(Extends = UhtTableNames.PropertyMember, ValueType = UhtSpecifierValueType.Legacy)]
+		[SuppressMessage("CodeQuality", "IDE0051:Remove unused private members", Justification = "Attribute accessed method")]
 		private static void EditDefaultsOnlySpecifier(UhtSpecifierContext SpecifierContext)
 		{
 			UhtPropertySpecifierContext Context = (UhtPropertySpecifierContext)SpecifierContext;
@@ -52,6 +56,7 @@ namespace EpicGames.UHT.Parsers
 		}
 
 		[UhtSpecifier(Extends = UhtTableNames.PropertyMember, ValueType = UhtSpecifierValueType.Legacy)]
+		[SuppressMessage("CodeQuality", "IDE0051:Remove unused private members", Justification = "Attribute accessed method")]
 		private static void VisibleAnywhereSpecifier(UhtSpecifierContext SpecifierContext)
 		{
 			UhtPropertySpecifierContext Context = (UhtPropertySpecifierContext)SpecifierContext;
@@ -64,6 +69,7 @@ namespace EpicGames.UHT.Parsers
 		}
 
 		[UhtSpecifier(Extends = UhtTableNames.PropertyMember, ValueType = UhtSpecifierValueType.Legacy)]
+		[SuppressMessage("CodeQuality", "IDE0051:Remove unused private members", Justification = "Attribute accessed method")]
 		private static void VisibleInstanceOnlySpecifier(UhtSpecifierContext SpecifierContext)
 		{
 			UhtPropertySpecifierContext Context = (UhtPropertySpecifierContext)SpecifierContext;
@@ -76,6 +82,7 @@ namespace EpicGames.UHT.Parsers
 		}
 
 		[UhtSpecifier(Extends = UhtTableNames.PropertyMember, ValueType = UhtSpecifierValueType.Legacy)]
+		[SuppressMessage("CodeQuality", "IDE0051:Remove unused private members", Justification = "Attribute accessed method")]
 		private static void VisibleDefaultsOnlySpecifier(UhtSpecifierContext SpecifierContext)
 		{
 			UhtPropertySpecifierContext Context = (UhtPropertySpecifierContext)SpecifierContext;
@@ -88,6 +95,7 @@ namespace EpicGames.UHT.Parsers
 		}
 
 		[UhtSpecifier(Extends = UhtTableNames.PropertyMember, ValueType = UhtSpecifierValueType.Legacy)]
+		[SuppressMessage("CodeQuality", "IDE0051:Remove unused private members", Justification = "Attribute accessed method")]
 		private static void BlueprintReadWriteSpecifier(UhtSpecifierContext SpecifierContext)
 		{
 			UhtPropertySpecifierContext Context = (UhtPropertySpecifierContext)SpecifierContext;
@@ -113,6 +121,7 @@ namespace EpicGames.UHT.Parsers
 		}
 
 		[UhtSpecifier(Extends = UhtTableNames.PropertyMember, ValueType = UhtSpecifierValueType.Legacy)]
+		[SuppressMessage("CodeQuality", "IDE0051:Remove unused private members", Justification = "Attribute accessed method")]
 		private static void BlueprintReadOnlySpecifier(UhtSpecifierContext SpecifierContext)
 		{
 			UhtPropertySpecifierContext Context = (UhtPropertySpecifierContext)SpecifierContext;
@@ -138,6 +147,7 @@ namespace EpicGames.UHT.Parsers
 		}
 
 		[UhtSpecifier(Extends = UhtTableNames.PropertyMember, ValueType = UhtSpecifierValueType.SingleString)]
+		[SuppressMessage("CodeQuality", "IDE0051:Remove unused private members", Justification = "Attribute accessed method")]
 		private static void BlueprintSetterSpecifier(UhtSpecifierContext SpecifierContext, StringView Value)
 		{
 			UhtPropertySpecifierContext Context = (UhtPropertySpecifierContext)SpecifierContext;
@@ -158,6 +168,7 @@ namespace EpicGames.UHT.Parsers
 		}
 
 		[UhtSpecifier(Extends = UhtTableNames.PropertyMember, ValueType = UhtSpecifierValueType.SingleString)]
+		[SuppressMessage("CodeQuality", "IDE0051:Remove unused private members", Justification = "Attribute accessed method")]
 		private static void BlueprintGetterSpecifier(UhtSpecifierContext SpecifierContext, StringView Value)
 		{
 			UhtPropertySpecifierContext Context = (UhtPropertySpecifierContext)SpecifierContext;
@@ -173,6 +184,7 @@ namespace EpicGames.UHT.Parsers
 		}
 
 		[UhtSpecifier(Extends = UhtTableNames.PropertyMember, ValueType = UhtSpecifierValueType.Legacy)]
+		[SuppressMessage("CodeQuality", "IDE0051:Remove unused private members", Justification = "Attribute accessed method")]
 		private static void ConfigSpecifier(UhtSpecifierContext SpecifierContext)
 		{
 			UhtPropertySpecifierContext Context = (UhtPropertySpecifierContext)SpecifierContext;
@@ -180,6 +192,7 @@ namespace EpicGames.UHT.Parsers
 		}
 
 		[UhtSpecifier(Extends = UhtTableNames.PropertyMember, ValueType = UhtSpecifierValueType.Legacy)]
+		[SuppressMessage("CodeQuality", "IDE0051:Remove unused private members", Justification = "Attribute accessed method")]
 		private static void GlobalConfigSpecifier(UhtSpecifierContext SpecifierContext)
 		{
 			UhtPropertySpecifierContext Context = (UhtPropertySpecifierContext)SpecifierContext;
@@ -187,6 +200,7 @@ namespace EpicGames.UHT.Parsers
 		}
 
 		[UhtSpecifier(Extends = UhtTableNames.PropertyMember, ValueType = UhtSpecifierValueType.Legacy)]
+		[SuppressMessage("CodeQuality", "IDE0051:Remove unused private members", Justification = "Attribute accessed method")]
 		private static void LocalizedSpecifier(UhtSpecifierContext SpecifierContext)
 		{
 			UhtPropertySpecifierContext Context = (UhtPropertySpecifierContext)SpecifierContext;
@@ -194,6 +208,7 @@ namespace EpicGames.UHT.Parsers
 		}
 
 		[UhtSpecifier(Extends = UhtTableNames.PropertyMember, ValueType = UhtSpecifierValueType.Legacy)]
+		[SuppressMessage("CodeQuality", "IDE0051:Remove unused private members", Justification = "Attribute accessed method")]
 		private static void TransientSpecifier(UhtSpecifierContext SpecifierContext)
 		{
 			UhtPropertySpecifierContext Context = (UhtPropertySpecifierContext)SpecifierContext;
@@ -201,6 +216,7 @@ namespace EpicGames.UHT.Parsers
 		}
 
 		[UhtSpecifier(Extends = UhtTableNames.PropertyMember, ValueType = UhtSpecifierValueType.Legacy)]
+		[SuppressMessage("CodeQuality", "IDE0051:Remove unused private members", Justification = "Attribute accessed method")]
 		private static void DuplicateTransientSpecifier(UhtSpecifierContext SpecifierContext)
 		{
 			UhtPropertySpecifierContext Context = (UhtPropertySpecifierContext)SpecifierContext;
@@ -208,6 +224,7 @@ namespace EpicGames.UHT.Parsers
 		}
 
 		[UhtSpecifier(Extends = UhtTableNames.PropertyMember, ValueType = UhtSpecifierValueType.Legacy)]
+		[SuppressMessage("CodeQuality", "IDE0051:Remove unused private members", Justification = "Attribute accessed method")]
 		private static void TextExportTransientSpecifier(UhtSpecifierContext SpecifierContext)
 		{
 			UhtPropertySpecifierContext Context = (UhtPropertySpecifierContext)SpecifierContext;
@@ -215,6 +232,7 @@ namespace EpicGames.UHT.Parsers
 		}
 
 		[UhtSpecifier(Extends = UhtTableNames.PropertyMember, ValueType = UhtSpecifierValueType.Legacy)]
+		[SuppressMessage("CodeQuality", "IDE0051:Remove unused private members", Justification = "Attribute accessed method")]
 		private static void NonPIETransientSpecifier(UhtSpecifierContext SpecifierContext)
 		{
 			UhtPropertySpecifierContext Context = (UhtPropertySpecifierContext)SpecifierContext;
@@ -223,6 +241,7 @@ namespace EpicGames.UHT.Parsers
 		}
 
 		[UhtSpecifier(Extends = UhtTableNames.PropertyMember, ValueType = UhtSpecifierValueType.Legacy)]
+		[SuppressMessage("CodeQuality", "IDE0051:Remove unused private members", Justification = "Attribute accessed method")]
 		private static void NonPIEDuplicateTransientSpecifier(UhtSpecifierContext SpecifierContext)
 		{
 			UhtPropertySpecifierContext Context = (UhtPropertySpecifierContext)SpecifierContext;
@@ -230,6 +249,7 @@ namespace EpicGames.UHT.Parsers
 		}
 
 		[UhtSpecifier(Extends = UhtTableNames.PropertyMember, ValueType = UhtSpecifierValueType.Legacy)]
+		[SuppressMessage("CodeQuality", "IDE0051:Remove unused private members", Justification = "Attribute accessed method")]
 		private static void ExportSpecifier(UhtSpecifierContext SpecifierContext)
 		{
 			UhtPropertySpecifierContext Context = (UhtPropertySpecifierContext)SpecifierContext;
@@ -237,6 +257,7 @@ namespace EpicGames.UHT.Parsers
 		}
 
 		[UhtSpecifier(Extends = UhtTableNames.PropertyMember, ValueType = UhtSpecifierValueType.Legacy)]
+		[SuppressMessage("CodeQuality", "IDE0051:Remove unused private members", Justification = "Attribute accessed method")]
 		private static void EditInlineSpecifier(UhtSpecifierContext SpecifierContext)
 		{
 			UhtPropertySpecifierContext Context = (UhtPropertySpecifierContext)SpecifierContext;
@@ -244,6 +265,7 @@ namespace EpicGames.UHT.Parsers
 		}
 
 		[UhtSpecifier(Extends = UhtTableNames.PropertyMember, ValueType = UhtSpecifierValueType.Legacy)]
+		[SuppressMessage("CodeQuality", "IDE0051:Remove unused private members", Justification = "Attribute accessed method")]
 		private static void NoClearSpecifier(UhtSpecifierContext SpecifierContext)
 		{
 			UhtPropertySpecifierContext Context = (UhtPropertySpecifierContext)SpecifierContext;
@@ -251,6 +273,7 @@ namespace EpicGames.UHT.Parsers
 		}
 
 		[UhtSpecifier(Extends = UhtTableNames.PropertyMember, ValueType = UhtSpecifierValueType.Legacy)]
+		[SuppressMessage("CodeQuality", "IDE0051:Remove unused private members", Justification = "Attribute accessed method")]
 		private static void EditFixedSizeSpecifier(UhtSpecifierContext SpecifierContext)
 		{
 			UhtPropertySpecifierContext Context = (UhtPropertySpecifierContext)SpecifierContext;
@@ -258,6 +281,7 @@ namespace EpicGames.UHT.Parsers
 		}
 
 		[UhtSpecifier(Extends = UhtTableNames.PropertyMember, ValueType = UhtSpecifierValueType.Legacy)]
+		[SuppressMessage("CodeQuality", "IDE0051:Remove unused private members", Justification = "Attribute accessed method")]
 		private static void ReplicatedSpecifier(UhtSpecifierContext SpecifierContext)
 		{
 			UhtPropertySpecifierContext Context = (UhtPropertySpecifierContext)SpecifierContext;
@@ -269,6 +293,7 @@ namespace EpicGames.UHT.Parsers
 		}
 
 		[UhtSpecifier(Extends = UhtTableNames.PropertyMember, ValueType = UhtSpecifierValueType.SingleString)]
+		[SuppressMessage("CodeQuality", "IDE0051:Remove unused private members", Justification = "Attribute accessed method")]
 		private static void ReplicatedUsingSpecifier(UhtSpecifierContext SpecifierContext, StringView Value)
 		{
 			UhtPropertySpecifierContext Context = (UhtPropertySpecifierContext)SpecifierContext;
@@ -287,6 +312,7 @@ namespace EpicGames.UHT.Parsers
 		}
 
 		[UhtSpecifier(Extends = UhtTableNames.PropertyMember, ValueType = UhtSpecifierValueType.Legacy)]
+		[SuppressMessage("CodeQuality", "IDE0051:Remove unused private members", Justification = "Attribute accessed method")]
 		private static void NotReplicatedSpecifier(UhtSpecifierContext SpecifierContext)
 		{
 			UhtPropertySpecifierContext Context = (UhtPropertySpecifierContext)SpecifierContext;
@@ -298,6 +324,7 @@ namespace EpicGames.UHT.Parsers
 		}
 
 		[UhtSpecifier(Extends = UhtTableNames.PropertyMember, ValueType = UhtSpecifierValueType.Legacy)]
+		[SuppressMessage("CodeQuality", "IDE0051:Remove unused private members", Justification = "Attribute accessed method")]
 		private static void RepRetrySpecifier(UhtSpecifierContext SpecifierContext)
 		{
 			UhtPropertySpecifierContext Context = (UhtPropertySpecifierContext)SpecifierContext;
@@ -305,6 +332,7 @@ namespace EpicGames.UHT.Parsers
 		}
 
 		[UhtSpecifier(Extends = UhtTableNames.PropertyMember, ValueType = UhtSpecifierValueType.Legacy)]
+		[SuppressMessage("CodeQuality", "IDE0051:Remove unused private members", Justification = "Attribute accessed method")]
 		private static void InterpSpecifier(UhtSpecifierContext SpecifierContext)
 		{
 			UhtPropertySpecifierContext Context = (UhtPropertySpecifierContext)SpecifierContext;
@@ -312,6 +340,7 @@ namespace EpicGames.UHT.Parsers
 		}
 
 		[UhtSpecifier(Extends = UhtTableNames.PropertyMember, ValueType = UhtSpecifierValueType.Legacy)]
+		[SuppressMessage("CodeQuality", "IDE0051:Remove unused private members", Justification = "Attribute accessed method")]
 		private static void NonTransactionalSpecifier(UhtSpecifierContext SpecifierContext)
 		{
 			UhtPropertySpecifierContext Context = (UhtPropertySpecifierContext)SpecifierContext;
@@ -319,6 +348,7 @@ namespace EpicGames.UHT.Parsers
 		}
 
 		[UhtSpecifier(Extends = UhtTableNames.PropertyMember, ValueType = UhtSpecifierValueType.Legacy)]
+		[SuppressMessage("CodeQuality", "IDE0051:Remove unused private members", Justification = "Attribute accessed method")]
 		private static void InstancedSpecifier(UhtSpecifierContext SpecifierContext)
 		{
 			UhtPropertySpecifierContext Context = (UhtPropertySpecifierContext)SpecifierContext;
@@ -327,6 +357,7 @@ namespace EpicGames.UHT.Parsers
 		}
 
 		[UhtSpecifier(Extends = UhtTableNames.PropertyMember, ValueType = UhtSpecifierValueType.Legacy)]
+		[SuppressMessage("CodeQuality", "IDE0051:Remove unused private members", Justification = "Attribute accessed method")]
 		private static void BlueprintAssignableSpecifier(UhtSpecifierContext SpecifierContext)
 		{
 			UhtPropertySpecifierContext Context = (UhtPropertySpecifierContext)SpecifierContext;
@@ -334,6 +365,7 @@ namespace EpicGames.UHT.Parsers
 		}
 
 		[UhtSpecifier(Extends = UhtTableNames.PropertyMember, ValueType = UhtSpecifierValueType.Legacy)]
+		[SuppressMessage("CodeQuality", "IDE0051:Remove unused private members", Justification = "Attribute accessed method")]
 		private static void BlueprintCallableSpecifier(UhtSpecifierContext SpecifierContext)
 		{
 			UhtPropertySpecifierContext Context = (UhtPropertySpecifierContext)SpecifierContext;
@@ -341,6 +373,7 @@ namespace EpicGames.UHT.Parsers
 		}
 
 		[UhtSpecifier(Extends = UhtTableNames.PropertyMember, ValueType = UhtSpecifierValueType.Legacy)]
+		[SuppressMessage("CodeQuality", "IDE0051:Remove unused private members", Justification = "Attribute accessed method")]
 		private static void BlueprintAuthorityOnlySpecifier(UhtSpecifierContext SpecifierContext)
 		{
 			UhtPropertySpecifierContext Context = (UhtPropertySpecifierContext)SpecifierContext;
@@ -348,6 +381,7 @@ namespace EpicGames.UHT.Parsers
 		}
 
 		[UhtSpecifier(Extends = UhtTableNames.PropertyMember, ValueType = UhtSpecifierValueType.Legacy)]
+		[SuppressMessage("CodeQuality", "IDE0051:Remove unused private members", Justification = "Attribute accessed method")]
 		private static void AssetRegistrySearchableSpecifier(UhtSpecifierContext SpecifierContext)
 		{
 			UhtPropertySpecifierContext Context = (UhtPropertySpecifierContext)SpecifierContext;
@@ -355,6 +389,7 @@ namespace EpicGames.UHT.Parsers
 		}
 
 		[UhtSpecifier(Extends = UhtTableNames.PropertyMember, ValueType = UhtSpecifierValueType.Legacy)]
+		[SuppressMessage("CodeQuality", "IDE0051:Remove unused private members", Justification = "Attribute accessed method")]
 		private static void SimpleDisplaySpecifier(UhtSpecifierContext SpecifierContext)
 		{
 			UhtPropertySpecifierContext Context = (UhtPropertySpecifierContext)SpecifierContext;
@@ -362,6 +397,7 @@ namespace EpicGames.UHT.Parsers
 		}
 
 		[UhtSpecifier(Extends = UhtTableNames.PropertyMember, ValueType = UhtSpecifierValueType.Legacy)]
+		[SuppressMessage("CodeQuality", "IDE0051:Remove unused private members", Justification = "Attribute accessed method")]
 		private static void AdvancedDisplaySpecifier(UhtSpecifierContext SpecifierContext)
 		{
 			UhtPropertySpecifierContext Context = (UhtPropertySpecifierContext)SpecifierContext;
@@ -369,6 +405,7 @@ namespace EpicGames.UHT.Parsers
 		}
 
 		[UhtSpecifier(Extends = UhtTableNames.PropertyMember, ValueType = UhtSpecifierValueType.Legacy)]
+		[SuppressMessage("CodeQuality", "IDE0051:Remove unused private members", Justification = "Attribute accessed method")]
 		private static void SaveGameSpecifier(UhtSpecifierContext SpecifierContext)
 		{
 			UhtPropertySpecifierContext Context = (UhtPropertySpecifierContext)SpecifierContext;
@@ -376,6 +413,7 @@ namespace EpicGames.UHT.Parsers
 		}
 
 		[UhtSpecifier(Extends = UhtTableNames.PropertyMember, ValueType = UhtSpecifierValueType.Legacy)]
+		[SuppressMessage("CodeQuality", "IDE0051:Remove unused private members", Justification = "Attribute accessed method")]
 		private static void SkipSerializationSpecifier(UhtSpecifierContext SpecifierContext)
 		{
 			UhtPropertySpecifierContext Context = (UhtPropertySpecifierContext)SpecifierContext;
@@ -383,6 +421,7 @@ namespace EpicGames.UHT.Parsers
 		}
 
 		[UhtSpecifier(Extends = UhtTableNames.PropertyMember, ValueType = UhtSpecifierValueType.OptionalString)]
+		[SuppressMessage("CodeQuality", "IDE0051:Remove unused private members", Justification = "Attribute accessed method")]
 		private static void GetterSpecifier(UhtSpecifierContext SpecifierContext, StringView? Value)
 		{
 			UhtPropertySpecifierContext Context = (UhtPropertySpecifierContext)SpecifierContext;
@@ -409,6 +448,7 @@ namespace EpicGames.UHT.Parsers
 		}
 
 		[UhtSpecifier(Extends = UhtTableNames.PropertyMember, ValueType = UhtSpecifierValueType.OptionalString)]
+		[SuppressMessage("CodeQuality", "IDE0051:Remove unused private members", Justification = "Attribute accessed method")]
 		private static void SetterSpecifier(UhtSpecifierContext SpecifierContext, StringView? Value)
 		{
 			UhtPropertySpecifierContext Context = (UhtPropertySpecifierContext)SpecifierContext;
@@ -435,6 +475,7 @@ namespace EpicGames.UHT.Parsers
 		}
 
 		[UhtSpecifier(Extends = UhtTableNames.PropertyMember, ValueType = UhtSpecifierValueType.None)]
+		[SuppressMessage("CodeQuality", "IDE0051:Remove unused private members", Justification = "Attribute accessed method")]
 		private static void FieldNotifySpecifier(UhtSpecifierContext SpecifierContext)
 		{
 			UhtPropertySpecifierContext Context = (UhtPropertySpecifierContext)SpecifierContext;
diff --git a/Engine/Source/Programs/Shared/EpicGames.UHT/Specifiers/UhtScriptStructSpecifiers.cs b/Engine/Source/Programs/Shared/EpicGames.UHT/Specifiers/UhtScriptStructSpecifiers.cs
--- a/Engine/Source/Programs/Shared/EpicGames.UHT/Specifiers/UhtScriptStructSpecifiers.cs
+++ b/Engine/Source/Programs/Shared/EpicGames.UHT/Specifiers/UhtScriptStructSpecifiers.cs
@@ -1,5 +1,6 @@
 ﻿// Copyright Epic Games, Inc. All Rights Reserved.
 
+using System.Diagnostics.CodeAnalysis;
 using EpicGames.Core;
 using EpicGames.UHT.Tables;
 using EpicGames.UHT.Types;
@@ -14,6 +15,7 @@ namespace EpicGames.UHT.Parsers
 	public static class UhtScriptStructSpecifiers
 	{
 		[UhtSpecifier(Extends = UhtTableNames.ScriptStruct, ValueType = UhtSpecifierValueType.Legacy)]
+		[SuppressMessage("CodeQuality", "IDE0051:Remove unused private members", Justification = "Attribute accessed method")]
 		private static void NoExportSpecifier(UhtSpecifierContext SpecifierContext)
 		{
 			UhtScriptStruct ScriptStruct = (UhtScriptStruct)SpecifierContext.Scope.ScopeType;
@@ -22,6 +24,7 @@ namespace EpicGames.UHT.Parsers
 		}
 
 		[UhtSpecifier(Extends = UhtTableNames.ScriptStruct, ValueType = UhtSpecifierValueType.Legacy)]
+		[SuppressMessage("CodeQuality", "IDE0051:Remove unused private members", Justification = "Attribute accessed method")]
 		private static void AtomicSpecifier(UhtSpecifierContext SpecifierContext)
 		{
 			UhtScriptStruct ScriptStruct = (UhtScriptStruct)SpecifierContext.Scope.ScopeType;
@@ -29,6 +32,7 @@ namespace EpicGames.UHT.Parsers
 		}
 
 		[UhtSpecifier(Extends = UhtTableNames.ScriptStruct, ValueType = UhtSpecifierValueType.Legacy)]
+		[SuppressMessage("CodeQuality", "IDE0051:Remove unused private members", Justification = "Attribute accessed method")]
 		private static void ImmutableSpecifier(UhtSpecifierContext SpecifierContext)
 		{
 			UhtScriptStruct ScriptStruct = (UhtScriptStruct)SpecifierContext.Scope.ScopeType;
@@ -36,6 +40,7 @@ namespace EpicGames.UHT.Parsers
 		}
 
 		[UhtSpecifier(Extends = UhtTableNames.ScriptStruct, ValueType = UhtSpecifierValueType.Legacy)]
+		[SuppressMessage("CodeQuality", "IDE0051:Remove unused private members", Justification = "Attribute accessed method")]
 		private static void HasDefaultsSpecifier(UhtSpecifierContext SpecifierContext)
 		{
 			if (!SpecifierContext.Scope.HeaderParser.HeaderFile.bIsNoExportTypes)
@@ -47,6 +52,7 @@ namespace EpicGames.UHT.Parsers
 		}
 
 		[UhtSpecifier(Extends = UhtTableNames.ScriptStruct, ValueType = UhtSpecifierValueType.Legacy)]
+		[SuppressMessage("CodeQuality", "IDE0051:Remove unused private members", Justification = "Attribute accessed method")]
 		private static void HasNoOpConstructorSpecifier(UhtSpecifierContext SpecifierContext)
 		{
 			if (!SpecifierContext.Scope.HeaderParser.HeaderFile.bIsNoExportTypes)
@@ -58,6 +64,7 @@ namespace EpicGames.UHT.Parsers
 		}
 
 		[UhtSpecifier(Extends = UhtTableNames.ScriptStruct, ValueType = UhtSpecifierValueType.Legacy)]
+		[SuppressMessage("CodeQuality", "IDE0051:Remove unused private members", Justification = "Attribute accessed method")]
 		private static void IsAlwaysAccessibleSpecifier(UhtSpecifierContext SpecifierContext)
 		{
 			if (!SpecifierContext.Scope.HeaderParser.HeaderFile.bIsNoExportTypes)
@@ -69,6 +76,7 @@ namespace EpicGames.UHT.Parsers
 		}
 
 		[UhtSpecifier(Extends = UhtTableNames.ScriptStruct, ValueType = UhtSpecifierValueType.Legacy)]
+		[SuppressMessage("CodeQuality", "IDE0051:Remove unused private members", Justification = "Attribute accessed method")]
 		private static void IsCoreTypeSpecifier(UhtSpecifierContext SpecifierContext)
 		{
 			if (!SpecifierContext.Scope.HeaderParser.HeaderFile.bIsNoExportTypes)
diff --git a/Engine/Source/Programs/Shared/EpicGames.UHT/Types/Properties/UhtArrayProperty.cs b/Engine/Source/Programs/Shared/EpicGames.UHT/Types/Properties/UhtArrayProperty.cs
--- a/Engine/Source/Programs/Shared/EpicGames.UHT/Types/Properties/UhtArrayProperty.cs
+++ b/Engine/Source/Programs/Shared/EpicGames.UHT/Types/Properties/UhtArrayProperty.cs
@@ -6,6 +6,7 @@ using EpicGames.UHT.Tables;
 using EpicGames.UHT.Tokenizer;
 using EpicGames.UHT.Utils;
 using System.Collections.Generic;
+using System.Diagnostics.CodeAnalysis;
 using System.Text;
 
 namespace EpicGames.UHT.Types
@@ -208,6 +209,7 @@ namespace EpicGames.UHT.Types
 
 #region Keyword
 		[UhtPropertyType(Keyword = "TArray")]
+		[SuppressMessage("CodeQuality", "IDE0051:Remove unused private members", Justification = "Attribute accessed method")]
 		private static UhtProperty? ArrayProperty(UhtPropertyResolvePhase ResolvePhase, UhtPropertySettings PropertySettings, IUhtTokenReader TokenReader, UhtToken MatchedToken)
 		{
 			using (var TokenContext = new UhtMessageContext(TokenReader, "TArray"))
diff --git a/Engine/Source/Programs/Shared/EpicGames.UHT/Types/Properties/UhtBoolProperty.cs b/Engine/Source/Programs/Shared/EpicGames.UHT/Types/Properties/UhtBoolProperty.cs
--- a/Engine/Source/Programs/Shared/EpicGames.UHT/Types/Properties/UhtBoolProperty.cs
+++ b/Engine/Source/Programs/Shared/EpicGames.UHT/Types/Properties/UhtBoolProperty.cs
@@ -4,6 +4,7 @@ using EpicGames.Core;
 using EpicGames.UHT.Tables;
 using EpicGames.UHT.Tokenizer;
 using EpicGames.UHT.Utils;
+using System.Diagnostics.CodeAnalysis;
 using System.Text;
 
 namespace EpicGames.UHT.Types
@@ -257,6 +258,7 @@ namespace EpicGames.UHT.Types
 
 		#region Keyword
 		[UhtPropertyType(Keyword = "bool", Options = UhtPropertyTypeOptions.Simple | UhtPropertyTypeOptions.Immediate)]
+		[SuppressMessage("CodeQuality", "IDE0051:Remove unused private members", Justification = "Attribute accessed method")]
 		private static UhtProperty? BoolProperty(UhtPropertyResolvePhase ResolvePhase, UhtPropertySettings PropertySettings, IUhtTokenReader TokenReader, UhtToken MatchedToken)
 		{
 			if (PropertySettings.bIsBitfield)
diff --git a/Engine/Source/Programs/Shared/EpicGames.UHT/Types/Properties/UhtByteProperty.cs b/Engine/Source/Programs/Shared/EpicGames.UHT/Types/Properties/UhtByteProperty.cs
--- a/Engine/Source/Programs/Shared/EpicGames.UHT/Types/Properties/UhtByteProperty.cs
+++ b/Engine/Source/Programs/Shared/EpicGames.UHT/Types/Properties/UhtByteProperty.cs
@@ -5,6 +5,7 @@ using EpicGames.UHT.Tables;
 using EpicGames.UHT.Tokenizer;
 using EpicGames.UHT.Utils;
 using System.Collections.Generic;
+using System.Diagnostics.CodeAnalysis;
 using System.Text;
 using System.Text.Json.Serialization;
 
@@ -157,6 +158,7 @@ namespace EpicGames.UHT.Types
 
 		#region Keyword
 		[UhtPropertyType(Keyword = "uint8", Options = UhtPropertyTypeOptions.Simple | UhtPropertyTypeOptions.Immediate)]
+		[SuppressMessage("CodeQuality", "IDE0051:Remove unused private members", Justification = "Attribute accessed method")]
 		private static UhtProperty? UInt8Property(UhtPropertyResolvePhase ResolvePhase, UhtPropertySettings PropertySettings, IUhtTokenReader TokenReader, UhtToken MatchedToken)
 		{
 			if (PropertySettings.bIsBitfield)
diff --git a/Engine/Source/Programs/Shared/EpicGames.UHT/Types/Properties/UhtClassProperty.cs b/Engine/Source/Programs/Shared/EpicGames.UHT/Types/Properties/UhtClassProperty.cs
--- a/Engine/Source/Programs/Shared/EpicGames.UHT/Types/Properties/UhtClassProperty.cs
+++ b/Engine/Source/Programs/Shared/EpicGames.UHT/Types/Properties/UhtClassProperty.cs
@@ -4,6 +4,7 @@ using EpicGames.Core;
 using EpicGames.UHT.Tables;
 using EpicGames.UHT.Tokenizer;
 using EpicGames.UHT.Utils;
+using System.Diagnostics.CodeAnalysis;
 using System.Text;
 
 namespace EpicGames.UHT.Types
@@ -146,6 +147,7 @@ namespace EpicGames.UHT.Types
 
 		#region Keyword
 		[UhtPropertyType(Keyword = "TSubclassOf")]
+		[SuppressMessage("CodeQuality", "IDE0051:Remove unused private members", Justification = "Attribute accessed method")]
 		private static UhtProperty? SubclassOfProperty(UhtPropertyResolvePhase ResolvePhase, UhtPropertySettings PropertySettings, IUhtTokenReader TokenReader, UhtToken MatchedToken)
 		{
 			UhtClass? MetaClass = UhtObjectPropertyBase.ParseTemplateClass(PropertySettings, TokenReader, MatchedToken);
diff --git a/Engine/Source/Programs/Shared/EpicGames.UHT/Types/Properties/UhtDoubleProperty.cs b/Engine/Source/Programs/Shared/EpicGames.UHT/Types/Properties/UhtDoubleProperty.cs
--- a/Engine/Source/Programs/Shared/EpicGames.UHT/Types/Properties/UhtDoubleProperty.cs
+++ b/Engine/Source/Programs/Shared/EpicGames.UHT/Types/Properties/UhtDoubleProperty.cs
@@ -3,6 +3,7 @@
 using EpicGames.UHT.Tables;
 using EpicGames.UHT.Tokenizer;
 using EpicGames.UHT.Utils;
+using System.Diagnostics.CodeAnalysis;
 using System.Text;
 
 namespace EpicGames.UHT.Types
@@ -59,6 +60,7 @@ namespace EpicGames.UHT.Types
 
 		#region Keyword
 		[UhtPropertyType(Keyword = "double", Options = UhtPropertyTypeOptions.Simple | UhtPropertyTypeOptions.Immediate)]
+		[SuppressMessage("CodeQuality", "IDE0051:Remove unused private members", Justification = "Attribute accessed method")]
 		private static UhtProperty? DoubleProperty(UhtPropertyResolvePhase ResolvePhase, UhtPropertySettings PropertySettings, IUhtTokenReader TokenReader, UhtToken MatchedToken)
 		{
 			return new UhtDoubleProperty(PropertySettings);
diff --git a/Engine/Source/Programs/Shared/EpicGames.UHT/Types/Properties/UhtEnumProperty.cs b/Engine/Source/Programs/Shared/EpicGames.UHT/Types/Properties/UhtEnumProperty.cs
--- a/Engine/Source/Programs/Shared/EpicGames.UHT/Types/Properties/UhtEnumProperty.cs
+++ b/Engine/Source/Programs/Shared/EpicGames.UHT/Types/Properties/UhtEnumProperty.cs
@@ -5,6 +5,7 @@ using EpicGames.UHT.Tables;
 using EpicGames.UHT.Tokenizer;
 using EpicGames.UHT.Utils;
 using System.Collections.Generic;
+using System.Diagnostics.CodeAnalysis;
 using System.Text;
 using System.Text.Json.Serialization;
 
@@ -332,6 +333,7 @@ namespace EpicGames.UHT.Types
 
 		#region Keyword
 		[UhtPropertyType(Keyword = "TEnumAsByte")]
+		[SuppressMessage("CodeQuality", "IDE0051:Remove unused private members", Justification = "Attribute accessed method")]
 		private static UhtProperty? EnumProperty(UhtPropertyResolvePhase ResolvePhase, UhtPropertySettings PropertySettings, IUhtTokenReader TokenReader, UhtToken MatchedToken)
 		{
 			UhtType Outer = PropertySettings.Outer;
diff --git a/Engine/Source/Programs/Shared/EpicGames.UHT/Types/Properties/UhtFieldPathProperty.cs b/Engine/Source/Programs/Shared/EpicGames.UHT/Types/Properties/UhtFieldPathProperty.cs
--- a/Engine/Source/Programs/Shared/EpicGames.UHT/Types/Properties/UhtFieldPathProperty.cs
+++ b/Engine/Source/Programs/Shared/EpicGames.UHT/Types/Properties/UhtFieldPathProperty.cs
@@ -4,6 +4,7 @@ using EpicGames.Core;
 using EpicGames.UHT.Tables;
 using EpicGames.UHT.Tokenizer;
 using EpicGames.UHT.Utils;
+using System.Diagnostics.CodeAnalysis;
 using System.Text;
 
 namespace EpicGames.UHT.Types
@@ -105,6 +106,7 @@ namespace EpicGames.UHT.Types
 
 		#region Keyword
 		[UhtPropertyType(Keyword = "TFieldPath", Options = UhtPropertyTypeOptions.Immediate)]
+		[SuppressMessage("CodeQuality", "IDE0051:Remove unused private members", Justification = "Attribute accessed method")]
 		private static UhtProperty? FieldPathProperty(UhtPropertyResolvePhase ResolvePhase, UhtPropertySettings PropertySettings, IUhtTokenReader TokenReader, UhtToken MatchedToken)
 		{
 			using (var TokenContext = new UhtMessageContext(TokenReader, "TFieldPath"))
diff --git a/Engine/Source/Programs/Shared/EpicGames.UHT/Types/Properties/UhtFloatProperty.cs b/Engine/Source/Programs/Shared/EpicGames.UHT/Types/Properties/UhtFloatProperty.cs
--- a/Engine/Source/Programs/Shared/EpicGames.UHT/Types/Properties/UhtFloatProperty.cs
+++ b/Engine/Source/Programs/Shared/EpicGames.UHT/Types/Properties/UhtFloatProperty.cs
@@ -3,6 +3,7 @@
 using EpicGames.UHT.Tables;
 using EpicGames.UHT.Tokenizer;
 using EpicGames.UHT.Utils;
+using System.Diagnostics.CodeAnalysis;
 using System.Text;
 
 namespace EpicGames.UHT.Types
@@ -58,6 +59,7 @@ namespace EpicGames.UHT.Types
 
 		#region Keywords
 		[UhtPropertyType(Keyword = "float", Options = UhtPropertyTypeOptions.Simple | UhtPropertyTypeOptions.Immediate)]
+		[SuppressMessage("CodeQuality", "IDE0051:Remove unused private members", Justification = "Attribute accessed method")]
 		private static UhtProperty? FloatProperty(UhtPropertyResolvePhase ResolvePhase, UhtPropertySettings PropertySettings, IUhtTokenReader TokenReader, UhtToken MatchedToken)
 		{
 			return new UhtFloatProperty(PropertySettings);
diff --git a/Engine/Source/Programs/Shared/EpicGames.UHT/Types/Properties/UhtInt16Property.cs b/Engine/Source/Programs/Shared/EpicGames.UHT/Types/Properties/UhtInt16Property.cs
--- a/Engine/Source/Programs/Shared/EpicGames.UHT/Types/Properties/UhtInt16Property.cs
+++ b/Engine/Source/Programs/Shared/EpicGames.UHT/Types/Properties/UhtInt16Property.cs
@@ -3,6 +3,7 @@
 using EpicGames.UHT.Tables;
 using EpicGames.UHT.Tokenizer;
 using EpicGames.UHT.Utils;
+using System.Diagnostics.CodeAnalysis;
 using System.Text;
 
 namespace EpicGames.UHT.Types
@@ -57,6 +58,7 @@ namespace EpicGames.UHT.Types
 
 		#region Keyword
 		[UhtPropertyType(Keyword = "int16", Options = UhtPropertyTypeOptions.Simple | UhtPropertyTypeOptions.Immediate)]
+		[SuppressMessage("CodeQuality", "IDE0051:Remove unused private members", Justification = "Attribute accessed method")]
 		private static UhtProperty? Int16Property(UhtPropertyResolvePhase ResolvePhase, UhtPropertySettings PropertySettings, IUhtTokenReader TokenReader, UhtToken MatchedToken)
 		{
 			return new UhtInt16Property(PropertySettings, UhtPropertyIntType.Sized);
diff --git a/Engine/Source/Programs/Shared/EpicGames.UHT/Types/Properties/UhtInt64Property.cs b/Engine/Source/Programs/Shared/EpicGames.UHT/Types/Properties/UhtInt64Property.cs
--- a/Engine/Source/Programs/Shared/EpicGames.UHT/Types/Properties/UhtInt64Property.cs
+++ b/Engine/Source/Programs/Shared/EpicGames.UHT/Types/Properties/UhtInt64Property.cs
@@ -3,6 +3,7 @@
 using EpicGames.UHT.Tables;
 using EpicGames.UHT.Tokenizer;
 using EpicGames.UHT.Utils;
+using System.Diagnostics.CodeAnalysis;
 using System.Text;
 
 namespace EpicGames.UHT.Types
@@ -59,6 +60,7 @@ namespace EpicGames.UHT.Types
 
 		#region Keyword
 		[UhtPropertyType(Keyword = "int64", Options = UhtPropertyTypeOptions.Simple | UhtPropertyTypeOptions.Immediate)]
+		[SuppressMessage("CodeQuality", "IDE0051:Remove unused private members", Justification = "Attribute accessed method")]
 		private static UhtProperty? Int64Property(UhtPropertyResolvePhase ResolvePhase, UhtPropertySettings PropertySettings, IUhtTokenReader TokenReader, UhtToken MatchedToken)
 		{
 			return new UhtInt64Property(PropertySettings, UhtPropertyIntType.Sized);
diff --git a/Engine/Source/Programs/Shared/EpicGames.UHT/Types/Properties/UhtInt8Property.cs b/Engine/Source/Programs/Shared/EpicGames.UHT/Types/Properties/UhtInt8Property.cs
--- a/Engine/Source/Programs/Shared/EpicGames.UHT/Types/Properties/UhtInt8Property.cs
+++ b/Engine/Source/Programs/Shared/EpicGames.UHT/Types/Properties/UhtInt8Property.cs
@@ -3,6 +3,7 @@
 using EpicGames.UHT.Tables;
 using EpicGames.UHT.Tokenizer;
 using EpicGames.UHT.Utils;
+using System.Diagnostics.CodeAnalysis;
 using System.Text;
 
 namespace EpicGames.UHT.Types
@@ -57,6 +58,7 @@ namespace EpicGames.UHT.Types
 
 		#region Keyword
 		[UhtPropertyType(Keyword = "int8", Options = UhtPropertyTypeOptions.Simple | UhtPropertyTypeOptions.Immediate)]
+		[SuppressMessage("CodeQuality", "IDE0051:Remove unused private members", Justification = "Attribute accessed method")]
 		private static UhtProperty? Int8Property(UhtPropertyResolvePhase ResolvePhase, UhtPropertySettings PropertySettings, IUhtTokenReader TokenReader, UhtToken MatchedToken)
 		{
 			return new UhtInt8Property(PropertySettings, UhtPropertyIntType.Sized);
diff --git a/Engine/Source/Programs/Shared/EpicGames.UHT/Types/Properties/UhtIntProperty.cs b/Engine/Source/Programs/Shared/EpicGames.UHT/Types/Properties/UhtIntProperty.cs
--- a/Engine/Source/Programs/Shared/EpicGames.UHT/Types/Properties/UhtIntProperty.cs
+++ b/Engine/Source/Programs/Shared/EpicGames.UHT/Types/Properties/UhtIntProperty.cs
@@ -3,6 +3,7 @@
 using EpicGames.UHT.Tables;
 using EpicGames.UHT.Tokenizer;
 using EpicGames.UHT.Utils;
+using System.Diagnostics.CodeAnalysis;
 using System.Text;
 
 namespace EpicGames.UHT.Types
@@ -62,18 +63,21 @@ namespace EpicGames.UHT.Types
 
 		#region Keywords
 		[UhtPropertyType(Keyword = "int32", Options = UhtPropertyTypeOptions.Simple | UhtPropertyTypeOptions.Immediate)]
+		[SuppressMessage("CodeQuality", "IDE0051:Remove unused private members", Justification = "Attribute accessed method")]
 		private static UhtProperty? Int32Property(UhtPropertyResolvePhase ResolvePhase, UhtPropertySettings PropertySettings, IUhtTokenReader TokenReader, UhtToken MatchedToken)
 		{
 			return new UhtIntProperty(PropertySettings, UhtPropertyIntType.Sized);
 		}
 
 		[UhtPropertyType(Keyword = "int", Options = UhtPropertyTypeOptions.Simple | UhtPropertyTypeOptions.Immediate)]
+		[SuppressMessage("CodeQuality", "IDE0051:Remove unused private members", Justification = "Attribute accessed method")]
 		private static UhtProperty? IntProperty(UhtPropertyResolvePhase ResolvePhase, UhtPropertySettings PropertySettings, IUhtTokenReader TokenReader, UhtToken MatchedToken)
 		{
 			return new UhtIntProperty(PropertySettings, UhtPropertyIntType.Unsized);
 		}
 
 		[UhtPropertyType(Keyword = "signed", Options = UhtPropertyTypeOptions.Immediate)]
+		[SuppressMessage("CodeQuality", "IDE0051:Remove unused private members", Justification = "Attribute accessed method")]
 		private static UhtProperty? SignedProperty(UhtPropertyResolvePhase ResolvePhase, UhtPropertySettings PropertySettings, IUhtTokenReader TokenReader, UhtToken MatchedToken)
 		{
 			TokenReader
diff --git a/Engine/Source/Programs/Shared/EpicGames.UHT/Types/Properties/UhtInterfaceProperty.cs b/Engine/Source/Programs/Shared/EpicGames.UHT/Types/Properties/UhtInterfaceProperty.cs
--- a/Engine/Source/Programs/Shared/EpicGames.UHT/Types/Properties/UhtInterfaceProperty.cs
+++ b/Engine/Source/Programs/Shared/EpicGames.UHT/Types/Properties/UhtInterfaceProperty.cs
@@ -177,12 +177,14 @@ namespace EpicGames.UHT.Types
 
 		#region Keywords
 		[UhtPropertyType(Keyword = "FScriptInterface", Options = UhtPropertyTypeOptions.Simple)] // This can't be immediate due to the reference to UInterface
+		[SuppressMessage("CodeQuality", "IDE0051:Remove unused private members", Justification = "Attribute accessed method")]
 		private static UhtProperty? FScriptInterfaceProperty(UhtPropertyResolvePhase ResolvePhase, UhtPropertySettings PropertySettings, IUhtTokenReader TokenReader, UhtToken MatchedToken)
 		{
 			return new UhtInterfaceProperty(PropertySettings, PropertySettings.Outer.Session.IInterface);
 		}
 
 		[UhtPropertyType(Keyword = "TScriptInterface")]
+		[SuppressMessage("CodeQuality", "IDE0051:Remove unused private members", Justification = "Attribute accessed method")]
 		private static UhtProperty? TScriptInterfaceProperty(UhtPropertyResolvePhase ResolvePhase, UhtPropertySettings PropertySettings, IUhtTokenReader TokenReader, UhtToken MatchedToken)
 		{
 			UhtClass? PropertyClass = UhtObjectPropertyBase.ParseTemplateObject(PropertySettings, TokenReader, MatchedToken, false);
diff --git a/Engine/Source/Programs/Shared/EpicGames.UHT/Types/Properties/UhtLargeWorldCoordinatesRealProperty.cs b/Engine/Source/Programs/Shared/EpicGames.UHT/Types/Properties/UhtLargeWorldCoordinatesRealProperty.cs
--- a/Engine/Source/Programs/Shared/EpicGames.UHT/Types/Properties/UhtLargeWorldCoordinatesRealProperty.cs
+++ b/Engine/Source/Programs/Shared/EpicGames.UHT/Types/Properties/UhtLargeWorldCoordinatesRealProperty.cs
@@ -3,6 +3,7 @@
 using EpicGames.UHT.Tables;
 using EpicGames.UHT.Tokenizer;
 using EpicGames.UHT.Utils;
+using System.Diagnostics.CodeAnalysis;
 using System.Text;
 
 namespace EpicGames.UHT.Types
@@ -60,6 +61,7 @@ namespace EpicGames.UHT.Types
 
 		#region Keyword
 		[UhtPropertyType(Keyword = "FLargeWorldCoordinatesReal", Options = UhtPropertyTypeOptions.Simple | UhtPropertyTypeOptions.Immediate)]
+		[SuppressMessage("CodeQuality", "IDE0051:Remove unused private members", Justification = "Attribute accessed method")]
 		private static UhtProperty? LargeWorldCoordinatesRealProperty(UhtPropertyResolvePhase ResolvePhase, UhtPropertySettings PropertySettings, IUhtTokenReader TokenReader, UhtToken MatchedToken)
 		{
 			if (!PropertySettings.Outer.HeaderFile.bIsNoExportTypes)
diff --git a/Engine/Source/Programs/Shared/EpicGames.UHT/Types/Properties/UhtLazyObjectPtrProperty.cs b/Engine/Source/Programs/Shared/EpicGames.UHT/Types/Properties/UhtLazyObjectPtrProperty.cs
--- a/Engine/Source/Programs/Shared/EpicGames.UHT/Types/Properties/UhtLazyObjectPtrProperty.cs
+++ b/Engine/Source/Programs/Shared/EpicGames.UHT/Types/Properties/UhtLazyObjectPtrProperty.cs
@@ -4,6 +4,7 @@ using EpicGames.Core;
 using EpicGames.UHT.Tables;
 using EpicGames.UHT.Tokenizer;
 using EpicGames.UHT.Utils;
+using System.Diagnostics.CodeAnalysis;
 using System.Text;
 
 namespace EpicGames.UHT.Types
@@ -106,6 +107,7 @@ namespace EpicGames.UHT.Types
 
 		#region Keyword
 		[UhtPropertyType(Keyword = "TLazyObjectPtr")]
+		[SuppressMessage("CodeQuality", "IDE0051:Remove unused private members", Justification = "Attribute accessed method")]
 		private static UhtProperty? LazyObjectPtrProperty(UhtPropertyResolvePhase ResolvePhase, UhtPropertySettings PropertySettings, IUhtTokenReader TokenReader, UhtToken MatchedToken)
 		{
 			UhtClass? PropertyClass = ParseTemplateObject(PropertySettings, TokenReader, MatchedToken, true);
diff --git a/Engine/Source/Programs/Shared/EpicGames.UHT/Types/Properties/UhtMapProperty.cs b/Engine/Source/Programs/Shared/EpicGames.UHT/Types/Properties/UhtMapProperty.cs
--- a/Engine/Source/Programs/Shared/EpicGames.UHT/Types/Properties/UhtMapProperty.cs
+++ b/Engine/Source/Programs/Shared/EpicGames.UHT/Types/Properties/UhtMapProperty.cs
@@ -6,6 +6,7 @@ using EpicGames.UHT.Tables;
 using EpicGames.UHT.Tokenizer;
 using EpicGames.UHT.Utils;
 using System.Collections.Generic;
+using System.Diagnostics.CodeAnalysis;
 using System.Text;
 
 namespace EpicGames.UHT.Types
@@ -256,6 +257,7 @@ namespace EpicGames.UHT.Types
 
 		#region Keyword
 		[UhtPropertyType(Keyword = "TMap")]
+		[SuppressMessage("CodeQuality", "IDE0051:Remove unused private members", Justification = "Attribute accessed method")]
 		private static UhtProperty? MapProperty(UhtPropertyResolvePhase ResolvePhase, UhtPropertySettings PropertySettings, IUhtTokenReader TokenReader, UhtToken MatchedToken)
 		{
 			using (var TokenContext = new UhtMessageContext(TokenReader, "TMap"))
diff --git a/Engine/Source/Programs/Shared/EpicGames.UHT/Types/Properties/UhtNameProperty.cs b/Engine/Source/Programs/Shared/EpicGames.UHT/Types/Properties/UhtNameProperty.cs
--- a/Engine/Source/Programs/Shared/EpicGames.UHT/Types/Properties/UhtNameProperty.cs
+++ b/Engine/Source/Programs/Shared/EpicGames.UHT/Types/Properties/UhtNameProperty.cs
@@ -4,6 +4,7 @@ using EpicGames.Core;
 using EpicGames.UHT.Tables;
 using EpicGames.UHT.Tokenizer;
 using EpicGames.UHT.Utils;
+using System.Diagnostics.CodeAnalysis;
 using System.Text;
 
 namespace EpicGames.UHT.Types
@@ -89,6 +90,7 @@ namespace EpicGames.UHT.Types
 
 		#region Keyword
 		[UhtPropertyType(Keyword = "FName", Options = UhtPropertyTypeOptions.Simple | UhtPropertyTypeOptions.Immediate)]
+		[SuppressMessage("CodeQuality", "IDE0051:Remove unused private members", Justification = "Attribute accessed method")]
 		private static UhtProperty? NameProperty(UhtPropertyResolvePhase ResolvePhase, UhtPropertySettings PropertySettings, IUhtTokenReader TokenReader, UhtToken MatchedToken)
 		{
 			return new UhtNameProperty(PropertySettings);
diff --git a/Engine/Source/Programs/Shared/EpicGames.UHT/Types/Properties/UhtObjectPropertyBase.cs b/Engine/Source/Programs/Shared/EpicGames.UHT/Types/Properties/UhtObjectPropertyBase.cs
--- a/Engine/Source/Programs/Shared/EpicGames.UHT/Types/Properties/UhtObjectPropertyBase.cs
+++ b/Engine/Source/Programs/Shared/EpicGames.UHT/Types/Properties/UhtObjectPropertyBase.cs
@@ -4,6 +4,7 @@ using EpicGames.Core;
 using EpicGames.UHT.Tables;
 using EpicGames.UHT.Tokenizer;
 using EpicGames.UHT.Utils;
+using System;
 using System.Collections.Generic;
 using System.Diagnostics.CodeAnalysis;
 using System.Text;
@@ -250,7 +251,9 @@ namespace EpicGames.UHT.Types
 			}
 
 			string Type = TokenReader.GetStringView(TypeStartPos, TokenReader.InputPos - TypeStartPos).ToString();
-			Type = Type.Replace("\n", " ").Replace("\r", "").Replace("\t", " ");
+			Type = Type.Replace("\n", " ", StringComparison.Ordinal);
+			Type = Type.Replace("\r", "", StringComparison.Ordinal);
+			Type = Type.Replace("\t", " ", StringComparison.Ordinal);
 
 			switch (Behavior)
 			{
diff --git a/Engine/Source/Programs/Shared/EpicGames.UHT/Types/Properties/UhtObjectPtrProperty.cs b/Engine/Source/Programs/Shared/EpicGames.UHT/Types/Properties/UhtObjectPtrProperty.cs
--- a/Engine/Source/Programs/Shared/EpicGames.UHT/Types/Properties/UhtObjectPtrProperty.cs
+++ b/Engine/Source/Programs/Shared/EpicGames.UHT/Types/Properties/UhtObjectPtrProperty.cs
@@ -4,6 +4,7 @@ using EpicGames.Core;
 using EpicGames.UHT.Tables;
 using EpicGames.UHT.Tokenizer;
 using EpicGames.UHT.Utils;
+using System.Diagnostics.CodeAnalysis;
 using System.Text;
 
 namespace EpicGames.UHT.Types
@@ -101,6 +102,7 @@ namespace EpicGames.UHT.Types
 
 		#region Keyword
 		[UhtPropertyType(Keyword = "TObjectPtr")]
+		[SuppressMessage("CodeQuality", "IDE0051:Remove unused private members", Justification = "Attribute accessed method")]
 		private static UhtProperty? ObjectPtrProperty(UhtPropertyResolvePhase ResolvePhase, UhtPropertySettings PropertySettings, IUhtTokenReader TokenReader, UhtToken MatchedToken)
 		{
 			UhtSession Session = PropertySettings.Outer.Session;
diff --git a/Engine/Source/Programs/Shared/EpicGames.UHT/Types/Properties/UhtSetProperty.cs b/Engine/Source/Programs/Shared/EpicGames.UHT/Types/Properties/UhtSetProperty.cs
--- a/Engine/Source/Programs/Shared/EpicGames.UHT/Types/Properties/UhtSetProperty.cs
+++ b/Engine/Source/Programs/Shared/EpicGames.UHT/Types/Properties/UhtSetProperty.cs
@@ -6,6 +6,7 @@ using EpicGames.UHT.Tables;
 using EpicGames.UHT.Tokenizer;
 using EpicGames.UHT.Utils;
 using System.Collections.Generic;
+using System.Diagnostics.CodeAnalysis;
 using System.Text;
 
 namespace EpicGames.UHT.Types
@@ -207,6 +208,7 @@ namespace EpicGames.UHT.Types
 
 		#region Keyword
 		[UhtPropertyType(Keyword = "TSet")]
+		[SuppressMessage("CodeQuality", "IDE0051:Remove unused private members", Justification = "Attribute accessed method")]
 		private static UhtProperty? SetProperty(UhtPropertyResolvePhase ResolvePhase, UhtPropertySettings PropertySettings, IUhtTokenReader TokenReader, UhtToken MatchedToken)
 		{
 			using (var TokenContext = new UhtMessageContext(TokenReader, "TSet"))
diff --git a/Engine/Source/Programs/Shared/EpicGames.UHT/Types/Properties/UhtSoftClassProperty.cs b/Engine/Source/Programs/Shared/EpicGames.UHT/Types/Properties/UhtSoftClassProperty.cs
--- a/Engine/Source/Programs/Shared/EpicGames.UHT/Types/Properties/UhtSoftClassProperty.cs
+++ b/Engine/Source/Programs/Shared/EpicGames.UHT/Types/Properties/UhtSoftClassProperty.cs
@@ -3,6 +3,7 @@
 using EpicGames.UHT.Tables;
 using EpicGames.UHT.Tokenizer;
 using EpicGames.UHT.Utils;
+using System.Diagnostics.CodeAnalysis;
 using System.Text;
 
 namespace EpicGames.UHT.Types
@@ -79,6 +80,7 @@ namespace EpicGames.UHT.Types
 
 		#region Keyword
 		[UhtPropertyType(Keyword = "TSoftClassPtr")]
+		[SuppressMessage("CodeQuality", "IDE0051:Remove unused private members", Justification = "Attribute accessed method")]
 		private static UhtProperty? SoftClassPtrProperty(UhtPropertyResolvePhase ResolvePhase, UhtPropertySettings PropertySettings, IUhtTokenReader TokenReader, UhtToken MatchedToken)
 		{
 			UhtClass? MetaClass = ParseTemplateClass(PropertySettings, TokenReader, MatchedToken);
diff --git a/Engine/Source/Programs/Shared/EpicGames.UHT/Types/Properties/UhtSoftObjectProperty.cs b/Engine/Source/Programs/Shared/EpicGames.UHT/Types/Properties/UhtSoftObjectProperty.cs
--- a/Engine/Source/Programs/Shared/EpicGames.UHT/Types/Properties/UhtSoftObjectProperty.cs
+++ b/Engine/Source/Programs/Shared/EpicGames.UHT/Types/Properties/UhtSoftObjectProperty.cs
@@ -4,6 +4,7 @@ using EpicGames.Core;
 using EpicGames.UHT.Tables;
 using EpicGames.UHT.Tokenizer;
 using EpicGames.UHT.Utils;
+using System.Diagnostics.CodeAnalysis;
 using System.Text;
 
 namespace EpicGames.UHT.Types
@@ -87,6 +88,7 @@ namespace EpicGames.UHT.Types
 
 		#region Keyword
 		[UhtPropertyType(Keyword = "TSoftObjectPtr")]
+		[SuppressMessage("CodeQuality", "IDE0051:Remove unused private members", Justification = "Attribute accessed method")]
 		private static UhtProperty? SoftObjectPtrProperty(UhtPropertyResolvePhase ResolvePhase, UhtPropertySettings PropertySettings, IUhtTokenReader TokenReader, UhtToken MatchedToken)
 		{
 			UhtClass? PropertyClass = UhtObjectPropertyBase.ParseTemplateObject(PropertySettings, TokenReader, MatchedToken, true);
diff --git a/Engine/Source/Programs/Shared/EpicGames.UHT/Types/Properties/UhtStrProperty.cs b/Engine/Source/Programs/Shared/EpicGames.UHT/Types/Properties/UhtStrProperty.cs
--- a/Engine/Source/Programs/Shared/EpicGames.UHT/Types/Properties/UhtStrProperty.cs
+++ b/Engine/Source/Programs/Shared/EpicGames.UHT/Types/Properties/UhtStrProperty.cs
@@ -4,6 +4,7 @@ using EpicGames.Core;
 using EpicGames.UHT.Tables;
 using EpicGames.UHT.Tokenizer;
 using EpicGames.UHT.Utils;
+using System.Diagnostics.CodeAnalysis;
 using System.Text;
 
 namespace EpicGames.UHT.Types
@@ -102,6 +103,7 @@ namespace EpicGames.UHT.Types
 
 		[UhtPropertyType(Keyword = "FString")]
 		[UhtPropertyType(Keyword = "FMemoryImageString", Options = UhtPropertyTypeOptions.Immediate)]
+		[SuppressMessage("CodeQuality", "IDE0051:Remove unused private members", Justification = "Attribute accessed method")]
 		private static UhtProperty? StrProperty(UhtPropertyResolvePhase ResolvePhase, UhtPropertySettings PropertySettings, IUhtTokenReader TokenReader, UhtToken MatchedToken)
 		{
 			if (!TokenReader.SkipExpectedType(MatchedToken.Value, PropertySettings.PropertyCategory == UhtPropertyCategory.Member))
diff --git a/Engine/Source/Programs/Shared/EpicGames.UHT/Types/Properties/UhtStructProperty.cs b/Engine/Source/Programs/Shared/EpicGames.UHT/Types/Properties/UhtStructProperty.cs
--- a/Engine/Source/Programs/Shared/EpicGames.UHT/Types/Properties/UhtStructProperty.cs
+++ b/Engine/Source/Programs/Shared/EpicGames.UHT/Types/Properties/UhtStructProperty.cs
@@ -5,6 +5,7 @@ using EpicGames.UHT.Tables;
 using EpicGames.UHT.Tokenizer;
 using EpicGames.UHT.Utils;
 using System.Collections.Generic;
+using System.Diagnostics.CodeAnalysis;
 using System.Text;
 using System.Text.Json.Serialization;
 
@@ -207,6 +208,7 @@ namespace EpicGames.UHT.Types
 
 		#region Structure default value sanitizers
 		[UhtStructDefaultValue(Name = "FVector")]
+		[SuppressMessage("CodeQuality", "IDE0051:Remove unused private members", Justification = "Attribute accessed method")]
 		private static bool VectorStructDefaultValue(UhtStructProperty Property, IUhtTokenReader DefaultValueReader, StringBuilder InnerDefaultValue)
 		{
 			string Format = "{0:F6},{1:F6},{2:F6}";
@@ -247,6 +249,7 @@ namespace EpicGames.UHT.Types
 		}
 
 		[UhtStructDefaultValue(Name = "FRotator")]
+		[SuppressMessage("CodeQuality", "IDE0051:Remove unused private members", Justification = "Attribute accessed method")]
 		private static bool RotatorStructDefaultValue(UhtStructProperty Property, IUhtTokenReader DefaultValueReader, StringBuilder InnerDefaultValue)
 		{
 			DefaultValueReader.Require("FRotator");
@@ -279,6 +282,7 @@ namespace EpicGames.UHT.Types
 		}
 
 		[UhtStructDefaultValue(Name = "FVector2D")]
+		[SuppressMessage("CodeQuality", "IDE0051:Remove unused private members", Justification = "Attribute accessed method")]
 		private static bool Vector2DStructDefaultValue(UhtStructProperty Property, IUhtTokenReader DefaultValueReader, StringBuilder InnerDefaultValue)
 		{
 			string Format = "(X={0:F3},Y={1:F3})";
@@ -312,6 +316,7 @@ namespace EpicGames.UHT.Types
 		}
 
 		[UhtStructDefaultValue(Name = "FLinearColor")]
+		[SuppressMessage("CodeQuality", "IDE0051:Remove unused private members", Justification = "Attribute accessed method")]
 		private static bool LinearColorStructDefaultValue(UhtStructProperty Property, IUhtTokenReader DefaultValueReader, StringBuilder InnerDefaultValue)
 		{
 			string Format = "(R={0:F6},G={1:F6},B={2:F6},A={3:F6})";
@@ -358,6 +363,7 @@ namespace EpicGames.UHT.Types
 		}
 
 		[UhtStructDefaultValue(Name = "FColor")]
+		[SuppressMessage("CodeQuality", "IDE0051:Remove unused private members", Justification = "Attribute accessed method")]
 		private static bool ColorStructDefaultValue(UhtStructProperty Property, IUhtTokenReader DefaultValueReader, StringBuilder InnerDefaultValue)
 		{
 			string Format = "(R={0},G={1},B={2},A={3})";
@@ -404,6 +410,7 @@ namespace EpicGames.UHT.Types
 		}
 
 		[UhtStructDefaultValue(Options = UhtStructDefaultValueOptions.Default)]
+		[SuppressMessage("CodeQuality", "IDE0051:Remove unused private members", Justification = "Attribute accessed method")]
 		private static bool DefaultStructDefaultValue(UhtStructProperty Property, IUhtTokenReader DefaultValueReader, StringBuilder InnerDefaultValue)
 		{
 			DefaultValueReader
diff --git a/Engine/Source/Programs/Shared/EpicGames.UHT/Types/Properties/UhtTextProperty.cs b/Engine/Source/Programs/Shared/EpicGames.UHT/Types/Properties/UhtTextProperty.cs
--- a/Engine/Source/Programs/Shared/EpicGames.UHT/Types/Properties/UhtTextProperty.cs
+++ b/Engine/Source/Programs/Shared/EpicGames.UHT/Types/Properties/UhtTextProperty.cs
@@ -5,6 +5,7 @@ using EpicGames.UHT.Tables;
 using EpicGames.UHT.Tokenizer;
 using EpicGames.UHT.Utils;
 using System;
+using System.Diagnostics.CodeAnalysis;
 using System.Text;
 
 namespace EpicGames.UHT.Types
@@ -116,18 +117,21 @@ namespace EpicGames.UHT.Types
 
 		#region Parsing keywords and default parsers		
 		[UhtPropertyType(Keyword = "FText", Options = UhtPropertyTypeOptions.Simple | UhtPropertyTypeOptions.Immediate)]
+		[SuppressMessage("CodeQuality", "IDE0051:Remove unused private members", Justification = "Attribute accessed method")]
 		private static UhtProperty? TextProperty(UhtPropertyResolvePhase ResolvePhase, UhtPropertySettings PropertySettings, IUhtTokenReader TokenReader, UhtToken MatchedToken)
 		{
 			return new UhtTextProperty(PropertySettings);
 		}
 
 		[UhtPropertyType(Keyword = "Text", Options = UhtPropertyTypeOptions.Simple | UhtPropertyTypeOptions.Immediate | UhtPropertyTypeOptions.CaseInsensitive)]
+		[SuppressMessage("CodeQuality", "IDE0051:Remove unused private members", Justification = "Attribute accessed method")]
 		private static UhtProperty? MissingPrefixTextProperty(UhtPropertyResolvePhase ResolvePhase, UhtPropertySettings PropertySettings, IUhtTokenReader TokenReader, UhtToken MatchedToken)
 		{
 			throw new UhtException(TokenReader, "'Text' is missing a prefix, expecting 'FText'");
 		}
 
 		[UhtLocTextDefaultValue(Name = "INVTEXT")]
+		[SuppressMessage("CodeQuality", "IDE0051:Remove unused private members", Justification = "Attribute accessed method")]
 		private static bool InvTextDefaultValue(UhtTextProperty Property, IUhtTokenReader DefaultValueReader, ref UhtToken MacroToken, StringBuilder InnerDefaultValue)
 		{
 			DefaultValueReader.Require('(');
@@ -138,6 +142,7 @@ namespace EpicGames.UHT.Types
 		}
 
 		[UhtLocTextDefaultValue(Name = "LOCTEXT")]
+		[SuppressMessage("CodeQuality", "IDE0051:Remove unused private members", Justification = "Attribute accessed method")]
 		private static bool LocTextDefaultValue(UhtTextProperty Property, IUhtTokenReader DefaultValueReader, ref UhtToken MacroToken, StringBuilder InnerDefaultValue)
 		{
 			Property.LogError($"LOCTEXT default parameter values are not supported; use NSLOCTEXT instead: {Property.SourceName}");
@@ -145,6 +150,7 @@ namespace EpicGames.UHT.Types
 		}
 
 		[UhtLocTextDefaultValue(Name = "NSLOCTEXT")]
+		[SuppressMessage("CodeQuality", "IDE0051:Remove unused private members", Justification = "Attribute accessed method")]
 		private static bool NsLocTextDefaultValue(UhtTextProperty Property, IUhtTokenReader DefaultValueReader, ref UhtToken MacroToken, StringBuilder InnerDefaultValue)
 		{
 			DefaultValueReader.Require('(');
@@ -178,6 +184,7 @@ namespace EpicGames.UHT.Types
 		}
 
 		[UhtLocTextDefaultValue(Name = "LOCTABLE")]
+		[SuppressMessage("CodeQuality", "IDE0051:Remove unused private members", Justification = "Attribute accessed method")]
 		private static bool LocTableDefaultValue(UhtTextProperty Property, IUhtTokenReader DefaultValueReader, ref UhtToken MacroToken, StringBuilder InnerDefaultValue)
 		{
 			DefaultValueReader.Require('(');
@@ -192,6 +199,7 @@ namespace EpicGames.UHT.Types
 		}
 
 		[UhtLocTextDefaultValue(Name = "LOCGEN_FORMAT_NAMED")]
+		[SuppressMessage("CodeQuality", "IDE0051:Remove unused private members", Justification = "Attribute accessed method")]
 		private static bool LocGenFormatNamedDefaultValue(UhtTextProperty Property, IUhtTokenReader DefaultValueReader, ref UhtToken MacroToken, StringBuilder InnerDefaultValue)
 		{
 			InnerDefaultValue.Append("LOCGEN_FORMAT_NAMED(");
@@ -220,6 +228,7 @@ namespace EpicGames.UHT.Types
 		}
 
 		[UhtLocTextDefaultValue(Name = "LOCGEN_FORMAT_ORDERED")]
+		[SuppressMessage("CodeQuality", "IDE0051:Remove unused private members", Justification = "Attribute accessed method")]
 		private static bool LocGenFormatOrderedDefaultValue(UhtTextProperty Property, IUhtTokenReader DefaultValueReader, ref UhtToken MacroToken, StringBuilder InnerDefaultValue)
 		{
 			InnerDefaultValue.Append("LOCGEN_FORMAT_ORDERED(");
@@ -246,54 +255,63 @@ namespace EpicGames.UHT.Types
 		}
 
 		[UhtLocTextDefaultValue(Name = "LOCGEN_NUMBER")]
+		[SuppressMessage("CodeQuality", "IDE0051:Remove unused private members", Justification = "Attribute accessed method")]
 		private static bool LocGenNumberDefaultValue(UhtTextProperty Property, IUhtTokenReader DefaultValueReader, ref UhtToken MacroToken, StringBuilder InnerDefaultValue)
 		{
 			return LocGenNumberOrPercentDefaultValue(Property, DefaultValueReader, ref MacroToken, InnerDefaultValue, NumberStyle.None);
 		}
 
 		[UhtLocTextDefaultValue(Name = "LOCGEN_NUMBER_GROUPED")]
+		[SuppressMessage("CodeQuality", "IDE0051:Remove unused private members", Justification = "Attribute accessed method")]
 		private static bool LocGenNumberGroupedDefaultValue(UhtTextProperty Property, IUhtTokenReader DefaultValueReader, ref UhtToken MacroToken, StringBuilder InnerDefaultValue)
 		{
 			return LocGenNumberOrPercentDefaultValue(Property, DefaultValueReader, ref MacroToken, InnerDefaultValue, NumberStyle.Grouped);
 		}
 
 		[UhtLocTextDefaultValue(Name = "LOCGEN_NUMBER_UNGROUPED")]
+		[SuppressMessage("CodeQuality", "IDE0051:Remove unused private members", Justification = "Attribute accessed method")]
 		private static bool LocGenNumberUngroupedDefaultValue(UhtTextProperty Property, IUhtTokenReader DefaultValueReader, ref UhtToken MacroToken, StringBuilder InnerDefaultValue)
 		{
 			return LocGenNumberOrPercentDefaultValue(Property, DefaultValueReader, ref MacroToken, InnerDefaultValue, NumberStyle.Ungrouped);
 		}
 
 		[UhtLocTextDefaultValue(Name = "LOCGEN_NUMBER_CUSTOM")]
+		[SuppressMessage("CodeQuality", "IDE0051:Remove unused private members", Justification = "Attribute accessed method")]
 		private static bool LocGenNumberCustomDefaultValue(UhtTextProperty Property, IUhtTokenReader DefaultValueReader, ref UhtToken MacroToken, StringBuilder InnerDefaultValue)
 		{
 			return LocGenNumberOrPercentDefaultValue(Property, DefaultValueReader, ref MacroToken, InnerDefaultValue, NumberStyle.Custom);
 		}
 
 		[UhtLocTextDefaultValue(Name = "LOCGEN_PERCENT")]
+		[SuppressMessage("CodeQuality", "IDE0051:Remove unused private members", Justification = "Attribute accessed method")]
 		private static bool LocGenPercentDefaultValue(UhtTextProperty Property, IUhtTokenReader DefaultValueReader, ref UhtToken MacroToken, StringBuilder InnerDefaultValue)
 		{
 			return LocGenNumberOrPercentDefaultValue(Property, DefaultValueReader, ref MacroToken, InnerDefaultValue, NumberStyle.None);
 		}
 
 		[UhtLocTextDefaultValue(Name = "LOCGEN_PERCENT_GROUPED")]
+		[SuppressMessage("CodeQuality", "IDE0051:Remove unused private members", Justification = "Attribute accessed method")]
 		private static bool LocGenPercentGroupedDefaultValue(UhtTextProperty Property, IUhtTokenReader DefaultValueReader, ref UhtToken MacroToken, StringBuilder InnerDefaultValue)
 		{
 			return LocGenNumberOrPercentDefaultValue(Property, DefaultValueReader, ref MacroToken, InnerDefaultValue, NumberStyle.Grouped);
 		}
 
 		[UhtLocTextDefaultValue(Name = "LOCGEN_PERCENT_UNGROUPED")]
+		[SuppressMessage("CodeQuality", "IDE0051:Remove unused private members", Justification = "Attribute accessed method")]
 		private static bool LocGenPercentUngroupedDefaultValue(UhtTextProperty Property, IUhtTokenReader DefaultValueReader, ref UhtToken MacroToken, StringBuilder InnerDefaultValue)
 		{
 			return LocGenNumberOrPercentDefaultValue(Property, DefaultValueReader, ref MacroToken, InnerDefaultValue, NumberStyle.Ungrouped);
 		}
 
 		[UhtLocTextDefaultValue(Name = "LOCGEN_PERCENT_CUSTOM")]
+		[SuppressMessage("CodeQuality", "IDE0051:Remove unused private members", Justification = "Attribute accessed method")]
 		private static bool LocGenPercentCustomDefaultValue(UhtTextProperty Property, IUhtTokenReader DefaultValueReader, ref UhtToken MacroToken, StringBuilder InnerDefaultValue)
 		{
 			return LocGenNumberOrPercentDefaultValue(Property, DefaultValueReader, ref MacroToken, InnerDefaultValue, NumberStyle.Custom);
 		}
 
 		[UhtLocTextDefaultValue(Name = "LOCGEN_CURRENCY")]
+		[SuppressMessage("CodeQuality", "IDE0051:Remove unused private members", Justification = "Attribute accessed method")]
 		private static bool LocGenCurrencyDefaultValue(UhtTextProperty Property, IUhtTokenReader DefaultValueReader, ref UhtToken MacroToken, StringBuilder InnerDefaultValue)
 		{
 			DefaultValueReader.Require('(');
@@ -317,60 +335,70 @@ namespace EpicGames.UHT.Types
 		}
 
 		[UhtLocTextDefaultValue(Name = "LOCGEN_DATE_LOCAL")]
+		[SuppressMessage("CodeQuality", "IDE0051:Remove unused private members", Justification = "Attribute accessed method")]
 		private static bool LocGenDateLocalDefaultValue(UhtTextProperty Property, IUhtTokenReader DefaultValueReader, ref UhtToken MacroToken, StringBuilder InnerDefaultValue)
 		{
 			return LocGenDateTimeDefaultValue(Property, DefaultValueReader, ref MacroToken, InnerDefaultValue, true, false, false, false);
 		}
 
 		[UhtLocTextDefaultValue(Name = "LOCGEN_DATE_UTC")]
+		[SuppressMessage("CodeQuality", "IDE0051:Remove unused private members", Justification = "Attribute accessed method")]
 		private static bool LocGenDateUtcDefaultValue(UhtTextProperty Property, IUhtTokenReader DefaultValueReader, ref UhtToken MacroToken, StringBuilder InnerDefaultValue)
 		{
 			return LocGenDateTimeDefaultValue(Property, DefaultValueReader, ref MacroToken, InnerDefaultValue, true, false, true, false);
 		}
 
 		[UhtLocTextDefaultValue(Name = "LOCGEN_TIME_LOCAL")]
+		[SuppressMessage("CodeQuality", "IDE0051:Remove unused private members", Justification = "Attribute accessed method")]
 		private static bool LocGenTimeLocalDefaultValue(UhtTextProperty Property, IUhtTokenReader DefaultValueReader, ref UhtToken MacroToken, StringBuilder InnerDefaultValue)
 		{
 			return LocGenDateTimeDefaultValue(Property, DefaultValueReader, ref MacroToken, InnerDefaultValue, false, true, false, false);
 		}
 
 		[UhtLocTextDefaultValue(Name = "LOCGEN_TIME_UTC")]
+		[SuppressMessage("CodeQuality", "IDE0051:Remove unused private members", Justification = "Attribute accessed method")]
 		private static bool LocGenTimeUtcDefaultValue(UhtTextProperty Property, IUhtTokenReader DefaultValueReader, ref UhtToken MacroToken, StringBuilder InnerDefaultValue)
 		{
 			return LocGenDateTimeDefaultValue(Property, DefaultValueReader, ref MacroToken, InnerDefaultValue, false, true, true, false);
 		}
 
 		[UhtLocTextDefaultValue(Name = "LOCGEN_DATETIME_LOCAL")]
+		[SuppressMessage("CodeQuality", "IDE0051:Remove unused private members", Justification = "Attribute accessed method")]
 		private static bool LocGenDateTimeLocalDefaultValue(UhtTextProperty Property, IUhtTokenReader DefaultValueReader, ref UhtToken MacroToken, StringBuilder InnerDefaultValue)
 		{
 			return LocGenDateTimeDefaultValue(Property, DefaultValueReader, ref MacroToken, InnerDefaultValue, true, true, false, false);
 		}
 
 		[UhtLocTextDefaultValue(Name = "LOCGEN_DATETIME_UTC")]
+		[SuppressMessage("CodeQuality", "IDE0051:Remove unused private members", Justification = "Attribute accessed method")]
 		private static bool LocGenDateTimeUtcDefaultValue(UhtTextProperty Property, IUhtTokenReader DefaultValueReader, ref UhtToken MacroToken, StringBuilder InnerDefaultValue)
 		{
 			return LocGenDateTimeDefaultValue(Property, DefaultValueReader, ref MacroToken, InnerDefaultValue, true, true, true, false);
 		}
 
 		[UhtLocTextDefaultValue(Name = "LOCGEN_DATETIME_CUSTOM_LOCAL")]
+		[SuppressMessage("CodeQuality", "IDE0051:Remove unused private members", Justification = "Attribute accessed method")]
 		private static bool LocGenDateTimeCustomLocalDefaultValue(UhtTextProperty Property, IUhtTokenReader DefaultValueReader, ref UhtToken MacroToken, StringBuilder InnerDefaultValue)
 		{
 			return LocGenDateTimeDefaultValue(Property, DefaultValueReader, ref MacroToken, InnerDefaultValue, true, true, false, true);
 		}
 
 		[UhtLocTextDefaultValue(Name = "LOCGEN_DATETIME_CUSTOM_UTC")]
+		[SuppressMessage("CodeQuality", "IDE0051:Remove unused private members", Justification = "Attribute accessed method")]
 		private static bool LocGenDateTimeCustomUtcDefaultValue(UhtTextProperty Property, IUhtTokenReader DefaultValueReader, ref UhtToken MacroToken, StringBuilder InnerDefaultValue)
 		{
 			return LocGenDateTimeDefaultValue(Property, DefaultValueReader, ref MacroToken, InnerDefaultValue, true, true, true, true);
 		}
 
 		[UhtLocTextDefaultValue(Name = "LOCGEN_TOUPPER")]
+		[SuppressMessage("CodeQuality", "IDE0051:Remove unused private members", Justification = "Attribute accessed method")]
 		private static bool LocGenToUpperDefaultValue(UhtTextProperty Property, IUhtTokenReader DefaultValueReader, ref UhtToken MacroToken, StringBuilder InnerDefaultValue)
 		{
 			return LocGenTransformDefaultValue(Property, DefaultValueReader, ref MacroToken, InnerDefaultValue);
 		}
 
 		[UhtLocTextDefaultValue(Name = "LOCGEN_TOLOWER")]
+		[SuppressMessage("CodeQuality", "IDE0051:Remove unused private members", Justification = "Attribute accessed method")]
 		private static bool LocGenToLowerDefaultValue(UhtTextProperty Property, IUhtTokenReader DefaultValueReader, ref UhtToken MacroToken, StringBuilder InnerDefaultValue)
 		{
 			return LocGenTransformDefaultValue(Property, DefaultValueReader, ref MacroToken, InnerDefaultValue);
@@ -481,11 +509,6 @@ namespace EpicGames.UHT.Types
 			Custom,
 		}
 
-		private static bool LocGenNumberOfPercentBoolean(IUhtTokenReader DefaultValueReader, StringBuilder InnerDefaultValue)
-		{
-			return false;
-		}
-
 		private static bool LocGenNumberOrPercentDefaultValue(UhtTextProperty Property, IUhtTokenReader DefaultValueReader, ref UhtToken MacroToken, StringBuilder InnerDefaultValue, NumberStyle NumberStyle)
 		{
 			InnerDefaultValue.Append(MacroToken.Value.ToString());
diff --git a/Engine/Source/Programs/Shared/EpicGames.UHT/Types/Properties/UhtUInt16Property.cs b/Engine/Source/Programs/Shared/EpicGames.UHT/Types/Properties/UhtUInt16Property.cs
--- a/Engine/Source/Programs/Shared/EpicGames.UHT/Types/Properties/UhtUInt16Property.cs
+++ b/Engine/Source/Programs/Shared/EpicGames.UHT/Types/Properties/UhtUInt16Property.cs
@@ -3,6 +3,7 @@
 using EpicGames.UHT.Tables;
 using EpicGames.UHT.Tokenizer;
 using EpicGames.UHT.Utils;
+using System.Diagnostics.CodeAnalysis;
 using System.Text;
 
 namespace EpicGames.UHT.Types
@@ -58,6 +59,7 @@ namespace EpicGames.UHT.Types
 
 		#region Keyword
 		[UhtPropertyType(Keyword = "uint16", Options = UhtPropertyTypeOptions.Simple | UhtPropertyTypeOptions.Immediate)]
+		[SuppressMessage("CodeQuality", "IDE0051:Remove unused private members", Justification = "Attribute accessed method")]
 		private static UhtProperty? UInt16Property(UhtPropertyResolvePhase ResolvePhase, UhtPropertySettings PropertySettings, IUhtTokenReader TokenReader, UhtToken MatchedToken)
 		{
 			if (PropertySettings.bIsBitfield)
diff --git a/Engine/Source/Programs/Shared/EpicGames.UHT/Types/Properties/UhtUInt32Property.cs b/Engine/Source/Programs/Shared/EpicGames.UHT/Types/Properties/UhtUInt32Property.cs
--- a/Engine/Source/Programs/Shared/EpicGames.UHT/Types/Properties/UhtUInt32Property.cs
+++ b/Engine/Source/Programs/Shared/EpicGames.UHT/Types/Properties/UhtUInt32Property.cs
@@ -3,6 +3,7 @@
 using EpicGames.UHT.Tables;
 using EpicGames.UHT.Tokenizer;
 using EpicGames.UHT.Utils;
+using System.Diagnostics.CodeAnalysis;
 using System.Text;
 
 namespace EpicGames.UHT.Types
@@ -61,6 +62,7 @@ namespace EpicGames.UHT.Types
 
 		#region Keyword
 		[UhtPropertyType(Keyword = "uint32", Options = UhtPropertyTypeOptions.Simple | UhtPropertyTypeOptions.Immediate)]
+		[SuppressMessage("CodeQuality", "IDE0051:Remove unused private members", Justification = "Attribute accessed method")]
 		private static UhtProperty? UInt32Property(UhtPropertyResolvePhase ResolvePhase, UhtPropertySettings PropertySettings, IUhtTokenReader TokenReader, UhtToken MatchedToken)
 		{
 			if (PropertySettings.bIsBitfield)
@@ -74,6 +76,7 @@ namespace EpicGames.UHT.Types
 		}
 
 		[UhtPropertyType(Keyword = "unsigned", Options = UhtPropertyTypeOptions.Immediate)]
+		[SuppressMessage("CodeQuality", "IDE0051:Remove unused private members", Justification = "Attribute accessed method")]
 		private static UhtProperty? UnsignedProperty(UhtPropertyResolvePhase ResolvePhase, UhtPropertySettings PropertySettings, IUhtTokenReader TokenReader, UhtToken MatchedToken)
 		{
 			TokenReader
diff --git a/Engine/Source/Programs/Shared/EpicGames.UHT/Types/Properties/UhtUInt64Property.cs b/Engine/Source/Programs/Shared/EpicGames.UHT/Types/Properties/UhtUInt64Property.cs
--- a/Engine/Source/Programs/Shared/EpicGames.UHT/Types/Properties/UhtUInt64Property.cs
+++ b/Engine/Source/Programs/Shared/EpicGames.UHT/Types/Properties/UhtUInt64Property.cs
@@ -3,6 +3,7 @@
 using EpicGames.UHT.Tables;
 using EpicGames.UHT.Tokenizer;
 using EpicGames.UHT.Utils;
+using System.Diagnostics.CodeAnalysis;
 using System.Text;
 
 namespace EpicGames.UHT.Types
@@ -58,6 +59,7 @@ namespace EpicGames.UHT.Types
 
 		#region Keyword
 		[UhtPropertyType(Keyword = "uint64", Options = UhtPropertyTypeOptions.Simple | UhtPropertyTypeOptions.Immediate)]
+		[SuppressMessage("CodeQuality", "IDE0051:Remove unused private members", Justification = "Attribute accessed method")]
 		private static UhtProperty? UInt64Property(UhtPropertyResolvePhase ResolvePhase, UhtPropertySettings PropertySettings, IUhtTokenReader TokenReader, UhtToken MatchedToken)
 		{
 			if (PropertySettings.bIsBitfield)
diff --git a/Engine/Source/Programs/Shared/EpicGames.UHT/Types/Properties/UhtVoidProperty.cs b/Engine/Source/Programs/Shared/EpicGames.UHT/Types/Properties/UhtVoidProperty.cs
--- a/Engine/Source/Programs/Shared/EpicGames.UHT/Types/Properties/UhtVoidProperty.cs
+++ b/Engine/Source/Programs/Shared/EpicGames.UHT/Types/Properties/UhtVoidProperty.cs
@@ -4,6 +4,7 @@ using EpicGames.UHT.Tables;
 using EpicGames.UHT.Tokenizer;
 using EpicGames.UHT.Utils;
 using System;
+using System.Diagnostics.CodeAnalysis;
 using System.Text;
 
 namespace EpicGames.UHT.Types
@@ -70,6 +71,7 @@ namespace EpicGames.UHT.Types
 
 		#region Keyword
 		[UhtPropertyType(Keyword = "void", Options = UhtPropertyTypeOptions.Simple | UhtPropertyTypeOptions.Immediate)]
+		[SuppressMessage("CodeQuality", "IDE0051:Remove unused private members", Justification = "Attribute accessed method")]
 		private static UhtProperty? VoidProperty(UhtPropertyResolvePhase ResolvePhase, UhtPropertySettings PropertySettings, IUhtTokenReader TokenReader, UhtToken MatchedToken)
 		{
 			if (PropertySettings.PropertyCategory != UhtPropertyCategory.Return)
diff --git a/Engine/Source/Programs/Shared/EpicGames.UHT/Types/Properties/UhtWeakObjectPtrProperty.cs b/Engine/Source/Programs/Shared/EpicGames.UHT/Types/Properties/UhtWeakObjectPtrProperty.cs
--- a/Engine/Source/Programs/Shared/EpicGames.UHT/Types/Properties/UhtWeakObjectPtrProperty.cs
+++ b/Engine/Source/Programs/Shared/EpicGames.UHT/Types/Properties/UhtWeakObjectPtrProperty.cs
@@ -4,6 +4,7 @@ using EpicGames.Core;
 using EpicGames.UHT.Tables;
 using EpicGames.UHT.Tokenizer;
 using EpicGames.UHT.Utils;
+using System.Diagnostics.CodeAnalysis;
 using System.Text;
 
 namespace EpicGames.UHT.Types
@@ -109,6 +110,7 @@ namespace EpicGames.UHT.Types
 		}
 
 		[UhtPropertyType(Keyword = "TWeakObjectPtr")]
+		[SuppressMessage("CodeQuality", "IDE0051:Remove unused private members", Justification = "Attribute accessed method")]
 		private static UhtProperty? WeakObjectProperty(UhtPropertyResolvePhase ResolvePhase, UhtPropertySettings PropertySettings, IUhtTokenReader TokenReader, UhtToken MatchedToken)
 		{
 			UhtClass? PropertyClass = ParseTemplateObject(PropertySettings, TokenReader, MatchedToken, true);
@@ -121,6 +123,7 @@ namespace EpicGames.UHT.Types
 		}
 
 		[UhtPropertyType(Keyword = "TAutoWeakObjectPtr")]
+		[SuppressMessage("CodeQuality", "IDE0051:Remove unused private members", Justification = "Attribute accessed method")]
 		private static UhtProperty? AutoWeakObjectProperty(UhtPropertyResolvePhase ResolvePhase, UhtPropertySettings PropertySettings, IUhtTokenReader TokenReader, UhtToken MatchedToken)
 		{
 			UhtClass? PropertyClass = ParseTemplateObject(PropertySettings, TokenReader, MatchedToken, true);
diff --git a/Engine/Source/Programs/Shared/EpicGames.UHT/Types/UhtEnum.cs b/Engine/Source/Programs/Shared/EpicGames.UHT/Types/UhtEnum.cs
--- a/Engine/Source/Programs/Shared/EpicGames.UHT/Types/UhtEnum.cs
+++ b/Engine/Source/Programs/Shared/EpicGames.UHT/Types/UhtEnum.cs
@@ -206,7 +206,7 @@ namespace EpicGames.UHT.Types
 			string EnumName = this.EnumValues[NameIndex].Name;
 			if (this.CppForm != UhtEnumCppForm.Regular)
 			{
-				int ScopeIndex = EnumName.IndexOf("::");
+				int ScopeIndex = EnumName.IndexOf("::", StringComparison.Ordinal);
 				if (ScopeIndex >= 0)
 				{
 					EnumName = EnumName.Substring(ScopeIndex + 2);
@@ -258,7 +258,7 @@ namespace EpicGames.UHT.Types
 		/// <returns>Reconstructed enum name</returns>
 		private string CleanEnumValueName(string Name)
 		{
-			int LastColons = Name.LastIndexOf("::");
+			int LastColons = Name.LastIndexOf("::", StringComparison.Ordinal);
 			return LastColons == -1 ? GetFullEnumName(Name) : GetFullEnumName(Name.Substring(LastColons + 2));		
 		}
 
diff --git a/Engine/Source/Programs/Shared/EpicGames.UHT/Types/UhtMetaData.cs b/Engine/Source/Programs/Shared/EpicGames.UHT/Types/UhtMetaData.cs
--- a/Engine/Source/Programs/Shared/EpicGames.UHT/Types/UhtMetaData.cs
+++ b/Engine/Source/Programs/Shared/EpicGames.UHT/Types/UhtMetaData.cs
@@ -110,7 +110,7 @@ namespace EpicGames.UHT.Types
 		/// <returns>Hash code</returns>
 		public int GetHashCode([DisallowNull] UhtMetaDataKey obj)
 		{
-			return obj.Name.GetHashCode() ^ obj.Name.GetHashCode(StringComparison.OrdinalIgnoreCase);
+			return obj.Index.GetHashCode() ^ obj.Name.GetHashCode(StringComparison.OrdinalIgnoreCase);
 		}
 	}
 
diff --git a/Engine/Source/Programs/Shared/EpicGames.UHT/Types/UhtPackage.cs b/Engine/Source/Programs/Shared/EpicGames.UHT/Types/UhtPackage.cs
--- a/Engine/Source/Programs/Shared/EpicGames.UHT/Types/UhtPackage.cs
+++ b/Engine/Source/Programs/Shared/EpicGames.UHT/Types/UhtPackage.cs
@@ -58,7 +58,7 @@ namespace EpicGames.UHT.Types
 				switch (this.Module.ModuleType)
 				{
 					case UHTModuleType.Program:
-						return this.Module.BaseDirectory.Replace('\\', '/').StartsWith(Unreal.EngineDirectory.FullName.Replace('\\', '/'));
+						return this.Module.BaseDirectory.Replace('\\', '/').StartsWith(Unreal.EngineDirectory.FullName.Replace('\\', '/'), StringComparison.Ordinal);
 					case UHTModuleType.EngineRuntime:
 					case UHTModuleType.EngineUncooked:
 					case UHTModuleType.EngineDeveloper:
@@ -81,7 +81,7 @@ namespace EpicGames.UHT.Types
 		/// True if the package is a plugin
 		/// </summary>
 		[JsonIgnore]
-		public bool bIsPlugin => this.Module.BaseDirectory.Replace('\\', '/').Contains("/Plugins/");
+		public bool bIsPlugin => this.Module.BaseDirectory.Replace('\\', '/').Contains("/Plugins/", StringComparison.Ordinal);
 
 		/// <summary>
 		/// Short name of the package (without the /Script/)
diff --git a/Engine/Source/Programs/Shared/EpicGames.UHT/Types/UhtScriptStruct.cs b/Engine/Source/Programs/Shared/EpicGames.UHT/Types/UhtScriptStruct.cs
--- a/Engine/Source/Programs/Shared/EpicGames.UHT/Types/UhtScriptStruct.cs
+++ b/Engine/Source/Programs/Shared/EpicGames.UHT/Types/UhtScriptStruct.cs
@@ -280,7 +280,7 @@ namespace EpicGames.UHT.Types
 		{
 			string String = TypeOriginal(bCastType);
 
-			int LesserPos = String.IndexOf('<');
+			int LesserPos = String.IndexOf('<', StringComparison.Ordinal);
 			if (LesserPos >= 0)
 			{
 				return String.Substring(0, LesserPos);
@@ -300,7 +300,7 @@ namespace EpicGames.UHT.Types
 		{
 			string String = TypeOriginal(bCastType);
 
-			int LesserPos = String.IndexOf('<');
+			int LesserPos = String.IndexOf('<', StringComparison.Ordinal);
 			if (LesserPos >= 0)
 			{
 				return String.Substring(LesserPos);
@@ -319,7 +319,7 @@ namespace EpicGames.UHT.Types
 		public string TypeConstRef(bool bCastType = false)
 		{
 			string String = TypeNoRef(bCastType);
-			if (String.StartsWith("T") || String.StartsWith("F"))
+			if (String.StartsWith("T", StringComparison.Ordinal) || String.StartsWith("F", StringComparison.Ordinal))
 			{
 				return $"const {String}&";
 			}
@@ -348,7 +348,7 @@ namespace EpicGames.UHT.Types
 		public string TypeNoRef(bool bCastType = false)
 		{
 			string String = TypeOriginal(bCastType);
-			if (String.EndsWith("&"))
+			if (String.EndsWith("&", StringComparison.Ordinal))
 			{
 				return String.Substring(0, String.Length - 1);
 			}
diff --git a/Engine/Source/Programs/Shared/EpicGames.UHT/Types/UhtType.cs b/Engine/Source/Programs/Shared/EpicGames.UHT/Types/UhtType.cs
--- a/Engine/Source/Programs/Shared/EpicGames.UHT/Types/UhtType.cs
+++ b/Engine/Source/Programs/Shared/EpicGames.UHT/Types/UhtType.cs
@@ -1232,7 +1232,7 @@ namespace EpicGames.UHT.Types
 			}
 			else
 			{
-				return NativeToolTip.TrimEnd().Replace("@see", "See:");
+				return NativeToolTip.TrimEnd().Replace("@see", "See:", StringComparison.Ordinal);
 			}
 		}
 
diff --git a/Engine/Source/Programs/Shared/EpicGames.UHT/Utils/UhtFCString.cs b/Engine/Source/Programs/Shared/EpicGames.UHT/Utils/UhtFCString.cs
--- a/Engine/Source/Programs/Shared/EpicGames.UHT/Utils/UhtFCString.cs
+++ b/Engine/Source/Programs/Shared/EpicGames.UHT/Utils/UhtFCString.cs
@@ -468,7 +468,7 @@ namespace EpicGames.UHT.Utils
 		public static string TabsToSpaces(String Input, int TabSpacing, bool bEmulateCrBug)
 		{
 			// If we have any tab characters, then we need to convert them to spaces
-			int TabIndex = Input.IndexOf('\t');
+			int TabIndex = Input.IndexOf('\t', StringComparison.Ordinal);
 			if (TabIndex == -1)
 			{
 				return Input;
diff --git a/Engine/Source/Programs/Shared/EpicGames.UHT/Utils/UhtSession.cs b/Engine/Source/Programs/Shared/EpicGames.UHT/Utils/UhtSession.cs
--- a/Engine/Source/Programs/Shared/EpicGames.UHT/Utils/UhtSession.cs
+++ b/Engine/Source/Programs/Shared/EpicGames.UHT/Utils/UhtSession.cs
@@ -1174,7 +1174,7 @@ namespace EpicGames.UHT.Utils
 				}
 			}
 
-			if (!Name.Contains("::") && this.ShortEnumValueLookup.TryGetValue(Name, out UhtEnum? Enum))
+			if (!Name.Contains("::", StringComparison.Ordinal) && this.ShortEnumValueLookup.TryGetValue(Name, out UhtEnum? Enum))
 			{
 				return Enum;
 			}
@@ -2063,7 +2063,7 @@ namespace EpicGames.UHT.Utils
 					this.FullEnumValueLookup.Add(Value.Name, new EnumAndValue { Enum = Enum, Value = Value.Value });
 					if (bAddShortNames)
 					{
-						if (Value.Name.StartsWith(CheckName))
+						if (Value.Name.StartsWith(CheckName, StringComparison.Ordinal))
 						{
 							this.ShortEnumValueLookup.TryAdd(Value.Name.Substring(CheckName.Length), Enum);
 						}
diff --git a/Engine/Source/Programs/Shared/EpicGames.UHT/Utils/UhtTables.cs b/Engine/Source/Programs/Shared/EpicGames.UHT/Utils/UhtTables.cs
--- a/Engine/Source/Programs/Shared/EpicGames.UHT/Utils/UhtTables.cs
+++ b/Engine/Source/Programs/Shared/EpicGames.UHT/Utils/UhtTables.cs
@@ -148,7 +148,7 @@ namespace EpicGames.UHT.Utils
 			string Name = InName ?? string.Empty;
 			if (string.IsNullOrEmpty(Name))
 			{
-				if (MethodInfo.Name.EndsWith(Suffix))
+				if (MethodInfo.Name.EndsWith(Suffix, StringComparison.Ordinal))
 				{
 					Name = MethodInfo.Name.Substring(0, MethodInfo.Name.Length - Suffix.Length);
 				}
