More code style fixes

#rnx

diff --git a/Engine/Source/Programs/Shared/EpicGames.UHT/.editorconfig b/Engine/Source/Programs/Shared/EpicGames.UHT/.editorconfig
--- a/Engine/Source/Programs/Shared/EpicGames.UHT/.editorconfig
+++ b/Engine/Source/Programs/Shared/EpicGames.UHT/.editorconfig
@@ -5,9 +5,6 @@
 # Default severity for all analyzer diagnostics
 dotnet_analyzer_diagnostic.severity = none
 
-# IDE0060: Remove unused parameter
-dotnet_diagnostic.IDE0060.severity = none
-
 # CA1051: Do not declare visible instance fields
 dotnet_diagnostic.CA1051.severity = none
 
diff --git a/Engine/Source/Programs/Shared/EpicGames.UHT/Exporters/CodeGen/UhtHeaderCodeGeneratorCppFile.cs b/Engine/Source/Programs/Shared/EpicGames.UHT/Exporters/CodeGen/UhtHeaderCodeGeneratorCppFile.cs
--- a/Engine/Source/Programs/Shared/EpicGames.UHT/Exporters/CodeGen/UhtHeaderCodeGeneratorCppFile.cs
+++ b/Engine/Source/Programs/Shared/EpicGames.UHT/Exporters/CodeGen/UhtHeaderCodeGeneratorCppFile.cs
@@ -1654,14 +1654,14 @@ namespace EpicGames.UHT.Exporters.CodeGen
 				{
 					Builder.Append("\tDEFINE_FUNCTION(").AppendClassSourceNameOrInterfaceName(Class).Append("::").Append(Function.UnMarshalAndCallName).Append(")\r\n");
 					Builder.Append("\t{\r\n");
-					AppendFunctionThunk(Builder, Class, Function);
+					AppendFunctionThunk(Builder, Function);
 					Builder.Append("\t}\r\n");
 				}
 			}
 			return Builder;
 		}
 
-		private static StringBuilder AppendFunctionThunk(StringBuilder Builder, UhtClass Class, UhtFunction Function)
+		private static StringBuilder AppendFunctionThunk(StringBuilder Builder, UhtFunction Function)
 		{
 			// Export the GET macro for the parameters
 			foreach (UhtProperty Parameter in Function.ParameterProperties.Span)
diff --git a/Engine/Source/Programs/Shared/EpicGames.UHT/Exporters/CodeGen/UhtHeaderCodeGeneratorHFile.cs b/Engine/Source/Programs/Shared/EpicGames.UHT/Exporters/CodeGen/UhtHeaderCodeGeneratorHFile.cs
--- a/Engine/Source/Programs/Shared/EpicGames.UHT/Exporters/CodeGen/UhtHeaderCodeGeneratorHFile.cs
+++ b/Engine/Source/Programs/Shared/EpicGames.UHT/Exporters/CodeGen/UhtHeaderCodeGeneratorHFile.cs
@@ -599,17 +599,17 @@ namespace EpicGames.UHT.Exporters.CodeGen
 
 			using (UhtMacroCreator Macro = new UhtMacroCreator(Builder, this, Class, bEditorOnly ? EditorOnlyRpcWrappersMacroSuffix : RpcWrappersMacroSuffix))
 			{
-				AppendAutogeneratedBlueprintFunctionDeclarations(Builder, Class, ReversedFunctions, bEditorOnly);
-				AppendRpcWrappers(Builder, Class, ReversedFunctions, bEditorOnly);
+				AppendAutogeneratedBlueprintFunctionDeclarations(Builder, ReversedFunctions, bEditorOnly);
+				AppendRpcWrappers(Builder, ReversedFunctions, bEditorOnly);
 			}
 
 			using (UhtMacroCreator Macro = new UhtMacroCreator(Builder, this, Class, bEditorOnly ? EditorOnlyRpcWrappersNoPureDeclsMacroSuffix : RpcWrappersNoPureDeclsMacroSuffix))
 			{
 				if (Class.GeneratedCodeVersion <= EGeneratedCodeVersion.V1)
 				{
-					AppendAutogeneratedBlueprintFunctionDeclarationsOnlyNotDeclared(Builder, Class, ReversedFunctions, bEditorOnly);
+					AppendAutogeneratedBlueprintFunctionDeclarationsOnlyNotDeclared(Builder, ReversedFunctions, bEditorOnly);
 				}
-				AppendRpcWrappers(Builder, Class, ReversedFunctions, bEditorOnly);
+				AppendRpcWrappers(Builder, ReversedFunctions, bEditorOnly);
 			}
 
 			if (bEditorOnly)
@@ -649,7 +649,7 @@ namespace EpicGames.UHT.Exporters.CodeGen
 			return Builder;
 		}
 
-		private StringBuilder AppendAutogeneratedBlueprintFunctionDeclarations(StringBuilder Builder, UhtClass Class, List<UhtFunction> ReversedFunctions, bool bEditorOnly)
+		private StringBuilder AppendAutogeneratedBlueprintFunctionDeclarations(StringBuilder Builder, List<UhtFunction> ReversedFunctions, bool bEditorOnly)
 		{
 			foreach (UhtFunction Function in ReversedFunctions)
 			{
@@ -663,7 +663,7 @@ namespace EpicGames.UHT.Exporters.CodeGen
 			return Builder;
 		}
 
-		private StringBuilder AppendAutogeneratedBlueprintFunctionDeclarationsOnlyNotDeclared(StringBuilder Builder, UhtClass Class, List<UhtFunction> ReversedFunctions, bool bEditorOnly)
+		private StringBuilder AppendAutogeneratedBlueprintFunctionDeclarationsOnlyNotDeclared(StringBuilder Builder, List<UhtFunction> ReversedFunctions, bool bEditorOnly)
 		{
 			foreach (UhtFunction Function in ReversedFunctions)
 			{
@@ -696,7 +696,7 @@ namespace EpicGames.UHT.Exporters.CodeGen
 			return Builder;
 		}
 
-		private static StringBuilder AppendRpcWrappers(StringBuilder Builder, UhtClass Class, List<UhtFunction> ReversedFunctions, bool bEditorOnly)
+		private static StringBuilder AppendRpcWrappers(StringBuilder Builder, List<UhtFunction> ReversedFunctions, bool bEditorOnly)
 		{
 			bool bFirst = true;
 			foreach (UhtFunction Function in ReversedFunctions)
@@ -834,7 +834,7 @@ namespace EpicGames.UHT.Exporters.CodeGen
 			{
 				AppendConstructorDefinition(Builder, Class, Api);
 				AppendVTableHelperCtorAndCaller(Builder, Class, Api);
-				AppendDefaultConstructorCallDefinition(Builder, Class, Api);
+				AppendDefaultConstructorCallDefinition(Builder, Class);
 			}
 			return Builder;
 		}
@@ -925,9 +925,8 @@ namespace EpicGames.UHT.Exporters.CodeGen
 		/// </summary>
 		/// <param name="Builder">Output builder</param>
 		/// <param name="Class">Class being exported</param>
-		/// <param name="Api">API text to be used</param>
 		/// <returns>Output builder</returns>
-		private static StringBuilder AppendDefaultConstructorCallDefinition(StringBuilder Builder, UhtClass Class, string Api)
+		private static StringBuilder AppendDefaultConstructorCallDefinition(StringBuilder Builder, UhtClass Class)
 		{
 			if (Class.ClassExportFlags.HasAnyFlags(UhtClassExportFlags.HasObjectInitializerConstructor))
 			{
@@ -1103,8 +1102,9 @@ namespace EpicGames.UHT.Exporters.CodeGen
 		/// Appends preprocessor string to emit GENERATED_U*_BODY() macro is deprecated.
 		/// </summary>
 		/// <param name="Builder">Output builder</param>
-		/// <param name="MacroName">Name ofthe macro to deprecate</param>
+		/// <param name="MacroName">Name of the macro to deprecate</param>
 		/// <returns>Output builder</returns>
+		[System.Diagnostics.CodeAnalysis.SuppressMessage("Style", "IDE0060:Remove unused parameter", Justification = "Referencing code commented out")]
 		private static StringBuilder AppendGeneratedMacroDeprecationWarning(StringBuilder Builder, string MacroName)
 		{
 			// Deprecation warning is disabled right now. After people get familiar with the new macro it should be re-enabled.
@@ -1142,7 +1142,7 @@ namespace EpicGames.UHT.Exporters.CodeGen
 			Builder.Append("\ttypedef ").Append(Class.SourceName).Append(" UClassType; \\\r\n");
 			Builder.Append("\ttypedef ").Append(InterfaceSourceName).Append(" ThisClass; \\\r\n");
 
-			AppendInterfaceCallFunctions(Builder, Class, CallbackFunctions);
+			AppendInterfaceCallFunctions(Builder, CallbackFunctions);
 
 			// we'll need a way to get to the UObject portion of a native interface, so that we can safely pass native interfaces
 			// to script VM functions
@@ -1266,7 +1266,7 @@ namespace EpicGames.UHT.Exporters.CodeGen
 			return Builder;
 		}
 
-		private StringBuilder AppendInterfaceCallFunctions(StringBuilder Builder, UhtClass Class, List<UhtFunction> CallbackFunctions)
+		private StringBuilder AppendInterfaceCallFunctions(StringBuilder Builder, List<UhtFunction> CallbackFunctions)
 		{
 			const string ExtraArg = "UObject* O";
 			const string ConstExtraArg = "const UObject* O";
diff --git a/Engine/Source/Programs/Shared/EpicGames.UHT/Parsers/UhtClassParser.cs b/Engine/Source/Programs/Shared/EpicGames.UHT/Parsers/UhtClassParser.cs
--- a/Engine/Source/Programs/Shared/EpicGames.UHT/Parsers/UhtClassParser.cs
+++ b/Engine/Source/Programs/Shared/EpicGames.UHT/Parsers/UhtClassParser.cs
@@ -150,7 +150,7 @@ namespace EpicGames.UHT.Parsers
 					MergeCategories();
 
 					SetAndValidateWithinClass(ResolvePhase);
-					SetAndValidateConfigName(ResolvePhase);
+					SetAndValidateConfigName();
 
 					// Copy all of the lists back to the meta data
 					this.MetaData.AddIfNotEmpty(UhtNames.ClassGroupNames, this.ClassGroupNames);
@@ -286,7 +286,7 @@ namespace EpicGames.UHT.Parsers
 			}
 		}
 
-		private void SetAndValidateConfigName(UhtResolvePhase ResolvePhase)
+		private void SetAndValidateConfigName()
 		{
 			// Since this flag is computed in this method, we have to re-propagate the flag from the super
 			// just in case they were defined in this source file.
@@ -367,6 +367,7 @@ namespace EpicGames.UHT.Parsers
 		#region Keywords
 		[UhtKeyword(Extends = UhtTableNames.Global)]
 		[SuppressMessage("CodeQuality", "IDE0051:Remove unused private members", Justification = "Attribute accessed method")]
+		[SuppressMessage("Style", "IDE0060:Remove unused parameter", Justification = "Attribute accessed method")]
 		private static UhtParseResult UCLASSKeyword(UhtParsingScope TopScope, UhtParsingScope ActionScope, ref UhtToken Token)
 		{
 			return ParseUClass(TopScope, ref Token);
@@ -375,6 +376,7 @@ namespace EpicGames.UHT.Parsers
 		[UhtKeyword(Extends = UhtTableNames.Class)]
 		[UhtKeyword(Extends = UhtTableNames.Class, Keyword = "GENERATED_BODY")]
 		[SuppressMessage("CodeQuality", "IDE0051:Remove unused private members", Justification = "Attribute accessed method")]
+		[SuppressMessage("Style", "IDE0060:Remove unused parameter", Justification = "Attribute accessed method")]
 		private static UhtParseResult GENERATED_UCLASS_BODYKeyword(UhtParsingScope TopScope, UhtParsingScope ActionScope, ref UhtToken Token)
 		{
 			UhtClass Class = (UhtClass)TopScope.ScopeType;
@@ -405,7 +407,7 @@ namespace EpicGames.UHT.Parsers
 			{
 				const string ScopeName = "class";
 
-				using (var TokenContext = new UhtMessageContext(TopScope.TokenReader, ScopeName))
+				using (var TokenContext = new UhtMessageContext(ScopeName))
 				{
 					// Parse the specifiers
 					UhtSpecifierContext SpecifierContext = new UhtSpecifierContext(TopScope, TopScope.TokenReader, Class.MetaData);
diff --git a/Engine/Source/Programs/Shared/EpicGames.UHT/Parsers/UhtEnumParser.cs b/Engine/Source/Programs/Shared/EpicGames.UHT/Parsers/UhtEnumParser.cs
--- a/Engine/Source/Programs/Shared/EpicGames.UHT/Parsers/UhtEnumParser.cs
+++ b/Engine/Source/Programs/Shared/EpicGames.UHT/Parsers/UhtEnumParser.cs
@@ -27,6 +27,7 @@ namespace EpicGames.UHT.Parsers
 		#region Keywords
 		[UhtKeyword(Extends = UhtTableNames.Global)]
 		[SuppressMessage("CodeQuality", "IDE0051:Remove unused private members", Justification = "Attribute accessed method")]
+		[SuppressMessage("Style", "IDE0060:Remove unused parameter", Justification = "Attribute accessed method")]
 		private static UhtParseResult UENUMKeyword(UhtParsingScope TopScope, UhtParsingScope ActionScope, ref UhtToken Token)
 		{
 			return ParseUEnum(TopScope, Token);
@@ -40,7 +41,7 @@ namespace EpicGames.UHT.Parsers
 			{
 				const string ScopeName = "UENUM";
 
-				using (var TokenContext = new UhtMessageContext(TopScope.TokenReader, ScopeName))
+				using (var TokenContext = new UhtMessageContext(ScopeName))
 				{
 
 					// Parse the specifiers
@@ -67,7 +68,7 @@ namespace EpicGames.UHT.Parsers
 
 					Enum.SourceName = EnumToken.Value.ToString();
 
-					Specifiers.ParseFieldMetaData(ScopeName);
+					Specifiers.ParseFieldMetaData();
 					Specifiers.ParseDeferred();
 
 					if (Enum.Outer != null)
@@ -216,7 +217,7 @@ namespace EpicGames.UHT.Parsers
 
 						// check for metadata on this enum value
 						SpecifierContext.MetaNameIndex = EnumIndex;
-						Specifiers.ParseFieldMetaData("Enum value");
+						Specifiers.ParseFieldMetaData();
 					});
 
 					// Trailing brace and semicolon for the enum
diff --git a/Engine/Source/Programs/Shared/EpicGames.UHT/Parsers/UhtFunctionParser.cs b/Engine/Source/Programs/Shared/EpicGames.UHT/Parsers/UhtFunctionParser.cs
--- a/Engine/Source/Programs/Shared/EpicGames.UHT/Parsers/UhtFunctionParser.cs
+++ b/Engine/Source/Programs/Shared/EpicGames.UHT/Parsers/UhtFunctionParser.cs
@@ -183,6 +183,7 @@ namespace EpicGames.UHT.Parsers
 		[UhtKeyword(Extends = UhtTableNames.Global)]
 		[UhtKeyword(Extends = UhtTableNames.Class)]
 		[SuppressMessage("CodeQuality", "IDE0051:Remove unused private members", Justification = "Attribute accessed method")]
+		[SuppressMessage("Style", "IDE0060:Remove unused parameter", Justification = "Attribute accessed method")]
 		private static UhtParseResult UDELEGATEKeyword(UhtParsingScope TopScope, UhtParsingScope ActionScope, ref UhtToken Token)
 		{
 			return ParseUDelegate(TopScope, Token, true);
@@ -192,6 +193,7 @@ namespace EpicGames.UHT.Parsers
 		[UhtKeyword(Extends = UhtTableNames.Interface)]
 		[UhtKeyword(Extends = UhtTableNames.NativeInterface)]
 		[SuppressMessage("CodeQuality", "IDE0051:Remove unused private members", Justification = "Attribute accessed method")]
+		[SuppressMessage("Style", "IDE0060:Remove unused parameter", Justification = "Attribute accessed method")]
 		private static UhtParseResult UFUNCTIONKeyword(UhtParsingScope TopScope, UhtParsingScope ActionScope, ref UhtToken Token)
 		{
 			return ParseUFunction(TopScope, Token);
@@ -199,6 +201,7 @@ namespace EpicGames.UHT.Parsers
 
 		[UhtKeywordCatchAll(Extends = UhtTableNames.Global)]
 		[SuppressMessage("CodeQuality", "IDE0051:Remove unused private members", Justification = "Attribute accessed method")]
+		[SuppressMessage("Style", "IDE0060:Remove unused parameter", Justification = "Attribute accessed method")]
 		private static UhtParseResult ParseCatchAllKeyword(UhtParsingScope TopScope, ref UhtToken Token)
 		{
 			if (UhtFunctionParser.IsValidateDelegateDeclaration(Token))
@@ -217,7 +220,7 @@ namespace EpicGames.UHT.Parsers
 			{
 				const string ScopeName = "delegate declaration";
 
-				using (var TokenContext = new UhtMessageContext(TopScope.TokenReader, ScopeName))
+				using (var TokenContext = new UhtMessageContext(ScopeName))
 				{
 					TopScope.AddModuleRelativePathToMetaData();
 
@@ -368,7 +371,7 @@ namespace EpicGames.UHT.Parsers
 					Function.MacroLineNumber = TopScope.TokenReader.InputLine;
 
 					// Try parsing metadata for the function
-					Specifiers.ParseFieldMetaData(Function.SourceName);
+					Specifiers.ParseFieldMetaData();
 
 					TopScope.AddFormattedCommentsAsTooltipMetaData();
 
@@ -389,7 +392,7 @@ namespace EpicGames.UHT.Parsers
 				UhtClass OuterClass = (UhtClass)OuterClassScope.ScopeType;
 				string ScopeName = "function";
 
-				using (var TokenContext = new UhtMessageContext(TopScope.TokenReader, ScopeName))
+				using (var TokenContext = new UhtMessageContext(ScopeName))
 				{
 					TopScope.AddModuleRelativePathToMetaData();
 
@@ -563,7 +566,7 @@ namespace EpicGames.UHT.Parsers
 						}
 
 						// Try parsing metadata for the function
-						SpecifierParser.ParseFieldMetaData(Function.EngineName);
+						SpecifierParser.ParseFieldMetaData();
 
 						// COMPATIBILITY-TODO - Try to pull any comment following the declaration
 						TopScope.TokenReader.PeekToken();
diff --git a/Engine/Source/Programs/Shared/EpicGames.UHT/Parsers/UhtHeaderFileParser.cs b/Engine/Source/Programs/Shared/EpicGames.UHT/Parsers/UhtHeaderFileParser.cs
--- a/Engine/Source/Programs/Shared/EpicGames.UHT/Parsers/UhtHeaderFileParser.cs
+++ b/Engine/Source/Programs/Shared/EpicGames.UHT/Parsers/UhtHeaderFileParser.cs
@@ -139,6 +139,7 @@ namespace EpicGames.UHT.Parsers
 		[UhtKeyword(Extends = UhtTableNames.ClassBase, Keyword = "public")]
 		[UhtKeyword(Extends = UhtTableNames.ScriptStruct, Keyword = "public")]
 		[SuppressMessage("CodeQuality", "IDE0051:Remove unused private members", Justification = "Attribute accessed method")]
+		[SuppressMessage("Style", "IDE0060:Remove unused parameter", Justification = "Attribute accessed method")]
 		private static UhtParseResult PublicKeyword(UhtParsingScope TopScope, UhtParsingScope ActionScope, ref UhtToken Token)
 		{
 			return SetAccessSpecifier(TopScope, UhtAccessSpecifier.Public);
@@ -147,6 +148,7 @@ namespace EpicGames.UHT.Parsers
 		[UhtKeyword(Extends = UhtTableNames.ClassBase, Keyword = "protected")]
 		[UhtKeyword(Extends = UhtTableNames.ScriptStruct, Keyword = "protected")]
 		[SuppressMessage("CodeQuality", "IDE0051:Remove unused private members", Justification = "Attribute accessed method")]
+		[SuppressMessage("Style", "IDE0060:Remove unused parameter", Justification = "Attribute accessed method")]
 		private static UhtParseResult ProtectedKeyword(UhtParsingScope TopScope, UhtParsingScope ActionScope, ref UhtToken Token)
 		{
 			return SetAccessSpecifier(TopScope, UhtAccessSpecifier.Protected);
@@ -155,6 +157,7 @@ namespace EpicGames.UHT.Parsers
 		[UhtKeyword(Extends = UhtTableNames.ClassBase, Keyword = "private")]
 		[UhtKeyword(Extends = UhtTableNames.ScriptStruct, Keyword = "private")]
 		[SuppressMessage("CodeQuality", "IDE0051:Remove unused private members", Justification = "Attribute accessed method")]
+		[SuppressMessage("Style", "IDE0060:Remove unused parameter", Justification = "Attribute accessed method")]
 		private static UhtParseResult PrivateKeyword(UhtParsingScope TopScope, UhtParsingScope ActionScope, ref UhtToken Token)
 		{
 			return SetAccessSpecifier(TopScope, UhtAccessSpecifier.Private);
@@ -358,7 +361,7 @@ namespace EpicGames.UHT.Parsers
 		public bool ParsePreprocessorDirective(ref UhtToken Token, bool bIsBeingIncluded, out bool bClearComments, out bool bIllegalContentsCheck)
 		{
 			bClearComments = true;
-			if (ParseDirectiveInternal(Token, bIsBeingIncluded))
+			if (ParseDirectiveInternal(bIsBeingIncluded))
 			{
 				bClearComments = ClearCommentsCompilerDirective();
 			}
@@ -613,10 +616,9 @@ namespace EpicGames.UHT.Parsers
 		/// <summary>
 		/// Parse a preprocessor directive.
 		/// </summary>
-		/// <param name="Token">Token that started the parse</param>
 		/// <param name="bIsBeingIncluded">If true, then this directive is in an active block</param>
 		/// <returns>True if we should check to see if tokenizer should clear comments</returns>
-		private bool ParseDirectiveInternal(UhtToken Token, bool bIsBeingIncluded)
+		private bool ParseDirectiveInternal(bool bIsBeingIncluded)
 		{
 			bool bCheckClearComments = false;
 
diff --git a/Engine/Source/Programs/Shared/EpicGames.UHT/Parsers/UhtInterfaceClassParser.cs b/Engine/Source/Programs/Shared/EpicGames.UHT/Parsers/UhtInterfaceClassParser.cs
--- a/Engine/Source/Programs/Shared/EpicGames.UHT/Parsers/UhtInterfaceClassParser.cs
+++ b/Engine/Source/Programs/Shared/EpicGames.UHT/Parsers/UhtInterfaceClassParser.cs
@@ -79,6 +79,7 @@ namespace EpicGames.UHT.Parsers
 		#region Keywords
 		[UhtKeyword(Extends = UhtTableNames.Global)]
 		[SuppressMessage("CodeQuality", "IDE0051:Remove unused private members", Justification = "Attribute accessed method")]
+		[SuppressMessage("Style", "IDE0060:Remove unused parameter", Justification = "Attribute accessed method")]
 		private static UhtParseResult UINTERFACEKeyword(UhtParsingScope TopScope, UhtParsingScope ActionScope, ref UhtToken Token)
 		{
 			return ParseUInterface(TopScope, ref Token);
@@ -87,6 +88,7 @@ namespace EpicGames.UHT.Parsers
 		[UhtKeyword(Extends = UhtTableNames.Interface)]
 		[UhtKeyword(Extends = UhtTableNames.Interface, Keyword = "GENERATED_BODY")]
 		[SuppressMessage("CodeQuality", "IDE0051:Remove unused private members", Justification = "Attribute accessed method")]
+		[SuppressMessage("Style", "IDE0060:Remove unused parameter", Justification = "Attribute accessed method")]
 		private static UhtParseResult GENERATED_UINTERFACE_BODYKeyword(UhtParsingScope TopScope, UhtParsingScope ActionScope, ref UhtToken Token)
 		{
 			UhtClass Class = (UhtClass)TopScope.ScopeType;
@@ -113,7 +115,7 @@ namespace EpicGames.UHT.Parsers
 			{
 				const string ScopeName = "interface";
 
-				using (var TokenContext = new UhtMessageContext(TopScope.TokenReader, ScopeName))
+				using (var TokenContext = new UhtMessageContext(ScopeName))
 				{
 					// Parse the specifiers
 					UhtSpecifierContext SpecifierContext = new UhtSpecifierContext(TopScope, TopScope.TokenReader, Class.MetaData);
diff --git a/Engine/Source/Programs/Shared/EpicGames.UHT/Parsers/UhtNativeInterfaceClassParser.cs b/Engine/Source/Programs/Shared/EpicGames.UHT/Parsers/UhtNativeInterfaceClassParser.cs
--- a/Engine/Source/Programs/Shared/EpicGames.UHT/Parsers/UhtNativeInterfaceClassParser.cs
+++ b/Engine/Source/Programs/Shared/EpicGames.UHT/Parsers/UhtNativeInterfaceClassParser.cs
@@ -68,13 +68,14 @@ namespace EpicGames.UHT.Parsers
 		#region Keywords
 		[UhtKeyword(Extends = UhtTableNames.Global, Keyword = "class", DisableUsageError = true)]
 		[SuppressMessage("CodeQuality", "IDE0051:Remove unused private members", Justification = "Attribute accessed method")]
+		[SuppressMessage("Style", "IDE0060:Remove unused parameter", Justification = "Attribute accessed method")]
 		private static UhtParseResult ClassKeyword(UhtParsingScope TopScope, UhtParsingScope ActionScope, ref UhtToken Token)
 		{
 			using (var SaveState = new UhtTokenSaveState(TopScope.TokenReader))
 			{
 				UhtToken SourceName = new UhtToken();
 				UhtToken SuperName = new UhtToken();
-				if (TryParseIInterface(TopScope, ref Token, out SourceName, out SuperName))
+				if (TryParseIInterface(TopScope, out SourceName, out SuperName))
 				{
 					SaveState.AbandonState();
 					ParseIInterface(TopScope, ref Token, SourceName, SuperName);
@@ -87,6 +88,7 @@ namespace EpicGames.UHT.Parsers
 		[UhtKeyword(Extends = UhtTableNames.NativeInterface)]
 		[UhtKeyword(Extends = UhtTableNames.NativeInterface, Keyword = "GENERATED_BODY")]
 		[SuppressMessage("CodeQuality", "IDE0051:Remove unused private members", Justification = "Attribute accessed method")]
+		[SuppressMessage("Style", "IDE0060:Remove unused parameter", Justification = "Attribute accessed method")]
 		private static UhtParseResult GENERATED_IINTERFACE_BODYKeyword(UhtParsingScope TopScope, UhtParsingScope ActionScope, ref UhtToken Token)
 		{
 			UhtClass Class = (UhtClass)TopScope.ScopeType;
@@ -105,7 +107,7 @@ namespace EpicGames.UHT.Parsers
 		}
 		#endregion
 
-		private static bool TryParseIInterface(UhtParsingScope ParentScope, ref UhtToken Token, out UhtToken SourceName, out UhtToken SuperName)
+		private static bool TryParseIInterface(UhtParsingScope ParentScope, out UhtToken SourceName, out UhtToken SuperName)
 		{
 			IUhtTokenReader TokenReader = ParentScope.TokenReader;
 
@@ -176,7 +178,7 @@ namespace EpicGames.UHT.Parsers
 			{
 				const string ScopeName = "native interface";
 
-				using (var TokenContext = new UhtMessageContext(TopScope.TokenReader, ScopeName))
+				using (var TokenContext = new UhtMessageContext(ScopeName))
 				{
 					TopScope.HeaderParser.ParseStatements('{', '}', false);
 					TokenReader.Require(';');
diff --git a/Engine/Source/Programs/Shared/EpicGames.UHT/Parsers/UhtPropertyParser.cs b/Engine/Source/Programs/Shared/EpicGames.UHT/Parsers/UhtPropertyParser.cs
--- a/Engine/Source/Programs/Shared/EpicGames.UHT/Parsers/UhtPropertyParser.cs
+++ b/Engine/Source/Programs/Shared/EpicGames.UHT/Parsers/UhtPropertyParser.cs
@@ -509,7 +509,7 @@ namespace EpicGames.UHT.Parsers
 			this.CurrentTypeTokens = new List<UhtToken>();
 			this.CurrentTemplateDepth = 0;
 
-			using (var TokenContext = new UhtMessageContext(this.TokenReader, this))
+			using (var TokenContext = new UhtMessageContext(this))
 			{
 				ParseInternal(Delegate);
 			}
@@ -1106,7 +1106,7 @@ namespace EpicGames.UHT.Parsers
 			if (NewProperty.PropertyCategory == UhtPropertyCategory.Member)
 			{
 				UhtSpecifierParser Specifiers = this.TopScope.HeaderParser.GetSpecifierParser(this.SpecifierContext, NewProperty.SourceName, this.TopScope.Session.GetSpecifierTable(UhtTableNames.PropertyMember));
-				Specifiers.ParseFieldMetaData(NewProperty.SourceName);
+				Specifiers.ParseFieldMetaData();
 				this.TopScope.AddFormattedCommentsAsTooltipMetaData(NewProperty);
 			}
 
@@ -1190,6 +1190,7 @@ namespace EpicGames.UHT.Parsers
 		[UhtKeyword(Extends = UhtTableNames.Class)]
 		[UhtKeyword(Extends = UhtTableNames.ScriptStruct)]
 		[SuppressMessage("CodeQuality", "IDE0051:Remove unused private members", Justification = "Attribute accessed method")]
+		[SuppressMessage("Style", "IDE0060:Remove unused parameter", Justification = "Attribute accessed method")]
 		private static UhtParseResult UPROPERTYKeyword(UhtParsingScope TopScope, UhtParsingScope ActionScope, ref UhtToken Token)
 		{
 			UhtPropertyParseOptions Options = UhtPropertyParseOptions.ParseLayoutMacro | UhtPropertyParseOptions.List | UhtPropertyParseOptions.AddModuleRelativePath;
@@ -1222,6 +1223,7 @@ namespace EpicGames.UHT.Parsers
 	{
 		[UhtPropertyType(Options = UhtPropertyTypeOptions.Default)]
 		[SuppressMessage("CodeQuality", "IDE0051:Remove unused private members", Justification = "Attribute accessed method")]
+		[SuppressMessage("Style", "IDE0060:Remove unused parameter", Justification = "Attribute accessed method")]
 		private static UhtProperty? DefaultProperty(UhtPropertyResolvePhase ResolvePhase, UhtPropertySettings PropertySettings, IUhtTokenReader TokenReader, UhtToken MatchedToken)
 		{
 			UhtSession Session = PropertySettings.Outer.Session;
diff --git a/Engine/Source/Programs/Shared/EpicGames.UHT/Parsers/UhtScriptStructParser.cs b/Engine/Source/Programs/Shared/EpicGames.UHT/Parsers/UhtScriptStructParser.cs
--- a/Engine/Source/Programs/Shared/EpicGames.UHT/Parsers/UhtScriptStructParser.cs
+++ b/Engine/Source/Programs/Shared/EpicGames.UHT/Parsers/UhtScriptStructParser.cs
@@ -73,6 +73,7 @@ namespace EpicGames.UHT.Parsers
 		#region Keywords
 		[UhtKeyword(Extends = UhtTableNames.Global)]
 		[SuppressMessage("CodeQuality", "IDE0051:Remove unused private members", Justification = "Attribute accessed method")]
+		[SuppressMessage("Style", "IDE0060:Remove unused parameter", Justification = "Attribute accessed method")]
 		private static UhtParseResult USTRUCTKeyword(UhtParsingScope TopScope, UhtParsingScope ActionScope, ref UhtToken Token)
 		{
 			return ParseUScriptStruct(TopScope, Token);
@@ -81,6 +82,7 @@ namespace EpicGames.UHT.Parsers
 		[UhtKeyword(Extends = UhtTableNames.ScriptStruct)]
 		[UhtKeyword(Extends = UhtTableNames.ScriptStruct, Keyword = "GENERATED_BODY")]
 		[SuppressMessage("CodeQuality", "IDE0051:Remove unused private members", Justification = "Attribute accessed method")]
+		[SuppressMessage("Style", "IDE0060:Remove unused parameter", Justification = "Attribute accessed method")]
 		private static UhtParseResult GENERATED_USTRUCT_BODYKeyword(UhtParsingScope TopScope, UhtParsingScope ActionScope, ref UhtToken Token)
 		{
 			UhtScriptStruct ScriptStruct = (UhtScriptStruct)TopScope.ScopeType;
@@ -105,6 +107,7 @@ namespace EpicGames.UHT.Parsers
 
 		[UhtKeyword(Extends = UhtTableNames.ScriptStruct)]
 		[SuppressMessage("CodeQuality", "IDE0051:Remove unused private members", Justification = "Attribute accessed method")]
+		[SuppressMessage("Style", "IDE0060:Remove unused parameter", Justification = "Attribute accessed method")]
 		private static UhtParseResult RIGVM_METHODKeyword(UhtParsingScope TopScope, UhtParsingScope ActionScope, ref UhtToken Token)
 		{
 			ParseRigVM(TopScope);
@@ -119,7 +122,7 @@ namespace EpicGames.UHT.Parsers
 			{
 				const string ScopeName = "struct";
 
-				using (var TokenContext = new UhtMessageContext(TopScope.TokenReader, ScopeName))
+				using (var TokenContext = new UhtMessageContext(ScopeName))
 				{
 					// Parse the specifiers
 					UhtSpecifierContext SpecifierContext = new UhtSpecifierContext(TopScope, TopScope.TokenReader, ScriptStruct.MetaData);
@@ -198,7 +201,7 @@ namespace EpicGames.UHT.Parsers
 		{
 			UhtScriptStruct ScriptStruct = (UhtScriptStruct)TopScope.ScopeType;
 
-			using (var TokenContext = new UhtMessageContext(TopScope.TokenReader, "RIGVM_METHOD"))
+			using (var TokenContext = new UhtMessageContext("RIGVM_METHOD"))
 			{
 
 				// Create the RigVM information if it doesn't already exist
diff --git a/Engine/Source/Programs/Shared/EpicGames.UHT/Parsers/UhtSpecifierParser.cs b/Engine/Source/Programs/Shared/EpicGames.UHT/Parsers/UhtSpecifierParser.cs
--- a/Engine/Source/Programs/Shared/EpicGames.UHT/Parsers/UhtSpecifierParser.cs
+++ b/Engine/Source/Programs/Shared/EpicGames.UHT/Parsers/UhtSpecifierParser.cs
@@ -97,7 +97,7 @@ namespace EpicGames.UHT.Parsers
 			this.bIsParsingFieldMetaData = false;
 			this.SpecifierContext.MetaData.LineNumber = TokenReader.InputLine;
 
-			using (var TokenContext = new UhtMessageContext(this.TokenReader, this))
+			using (var TokenContext = new UhtMessageContext(this))
 			{
 				this.TokenReader.RequireList('(', ')', ',', false, this.ParseAction);
 			}
@@ -107,14 +107,13 @@ namespace EpicGames.UHT.Parsers
 		/// <summary>
 		/// Parse field meta data
 		/// </summary>
-		/// <param name="Context">User facing context for error messages</param>
 		/// <returns>Specifier parser</returns>
-		public UhtSpecifierParser ParseFieldMetaData(StringView Context)
+		public UhtSpecifierParser ParseFieldMetaData()
 		{
 			this.TokenReader = this.SpecifierContext.Scope.TokenReader;
 			this.bIsParsingFieldMetaData = true;
 
-			using (var TokenContext = new UhtMessageContext(this.TokenReader, this))
+			using (var TokenContext = new UhtMessageContext(this))
 			{
 				if (TokenReader.TryOptional("UMETA"))
 				{
@@ -130,7 +129,7 @@ namespace EpicGames.UHT.Parsers
 		}
 
 		/// <summary>
-		/// Parse any deffered specifiers
+		/// Parse any deferred specifiers
 		/// </summary>
 		public void ParseDeferred()
 		{
diff --git a/Engine/Source/Programs/Shared/EpicGames.UHT/Specifiers/UhtDefaultSpecifiers.cs b/Engine/Source/Programs/Shared/EpicGames.UHT/Specifiers/UhtDefaultSpecifiers.cs
--- a/Engine/Source/Programs/Shared/EpicGames.UHT/Specifiers/UhtDefaultSpecifiers.cs
+++ b/Engine/Source/Programs/Shared/EpicGames.UHT/Specifiers/UhtDefaultSpecifiers.cs
@@ -82,6 +82,7 @@ namespace EpicGames.UHT.Parsers
 		#region Specifiers
 		[UhtSpecifier(Extends = UhtTableNames.Default, ValueType = UhtSpecifierValueType.KeyValuePairList, When = UhtSpecifierWhen.Immediate)]
 		[SuppressMessage("CodeQuality", "IDE0051:Remove unused private members", Justification = "Attribute accessed method")]
+		[SuppressMessage("Style", "IDE0060:Remove unused parameter", Justification = "Attribute accessed method")]
 		private static void MetaSpecifier(UhtSpecifierContext SpecifierContext, List<KeyValuePair<StringView, StringView>> Value)
 		{
 			foreach (KeyValuePair<StringView, StringView> KVP in (List<KeyValuePair<StringView, StringView>>)Value)
@@ -92,6 +93,8 @@ namespace EpicGames.UHT.Parsers
 
 		[UhtSpecifier(Extends = UhtTableNames.Default, ValueType = UhtSpecifierValueType.String, When = UhtSpecifierWhen.Immediate)]
 		[SuppressMessage("CodeQuality", "IDE0051:Remove unused private members", Justification = "Attribute accessed method")]
+		[SuppressMessage("Style", "IDE0060:Remove unused parameter", Justification = "Attribute accessed method")]
+		[SuppressMessage("Style", "IDE0060:Remove unused parameter", Justification = "Attribute accessed method")]
 		private static void DisplayNameSpecifier(UhtSpecifierContext SpecifierContext, StringView Value)
 		{
 			SetMetaData(SpecifierContext, "DisplayName", Value);
@@ -99,6 +102,7 @@ namespace EpicGames.UHT.Parsers
 
 		[UhtSpecifier(Extends = UhtTableNames.Default, ValueType = UhtSpecifierValueType.String, When = UhtSpecifierWhen.Immediate)]
 		[SuppressMessage("CodeQuality", "IDE0051:Remove unused private members", Justification = "Attribute accessed method")]
+		[SuppressMessage("Style", "IDE0060:Remove unused parameter", Justification = "Attribute accessed method")]
 		private static void FriendlyNameSpecifier(UhtSpecifierContext SpecifierContext, StringView Value)
 		{
 			SetMetaData(SpecifierContext, "FriendlyName", Value);
@@ -106,6 +110,7 @@ namespace EpicGames.UHT.Parsers
 
 		[UhtSpecifier(Extends = UhtTableNames.Default, ValueType = UhtSpecifierValueType.None, When = UhtSpecifierWhen.Immediate)]
 		[SuppressMessage("CodeQuality", "IDE0051:Remove unused private members", Justification = "Attribute accessed method")]
+		[SuppressMessage("Style", "IDE0060:Remove unused parameter", Justification = "Attribute accessed method")]
 		private static void BlueprintInternalUseOnlySpecifier(UhtSpecifierContext SpecifierContext)
 		{
 			SetMetaData(SpecifierContext, "BlueprintInternalUseOnly", true);
@@ -114,6 +119,7 @@ namespace EpicGames.UHT.Parsers
 
 		[UhtSpecifier(Extends = UhtTableNames.Default, ValueType = UhtSpecifierValueType.None, When = UhtSpecifierWhen.Immediate)]
 		[SuppressMessage("CodeQuality", "IDE0051:Remove unused private members", Justification = "Attribute accessed method")]
+		[SuppressMessage("Style", "IDE0060:Remove unused parameter", Justification = "Attribute accessed method")]
 		private static void BlueprintInternalUseOnlyHierarchicalSpecifier(UhtSpecifierContext SpecifierContext)
 		{
 			SetMetaData(SpecifierContext, "BlueprintInternalUseOnlyHierarchical", true);
@@ -123,6 +129,7 @@ namespace EpicGames.UHT.Parsers
 
 		[UhtSpecifier(Extends = UhtTableNames.Default, ValueType = UhtSpecifierValueType.None, When = UhtSpecifierWhen.Immediate)]
 		[SuppressMessage("CodeQuality", "IDE0051:Remove unused private members", Justification = "Attribute accessed method")]
+		[SuppressMessage("Style", "IDE0060:Remove unused parameter", Justification = "Attribute accessed method")]
 		private static void BlueprintTypeSpecifier(UhtSpecifierContext SpecifierContext)
 		{
 			SetMetaData(SpecifierContext, "BlueprintType", true);
@@ -130,6 +137,7 @@ namespace EpicGames.UHT.Parsers
 
 		[UhtSpecifier(Extends = UhtTableNames.Default, ValueType = UhtSpecifierValueType.None, When = UhtSpecifierWhen.Immediate)]
 		[SuppressMessage("CodeQuality", "IDE0051:Remove unused private members", Justification = "Attribute accessed method")]
+		[SuppressMessage("Style", "IDE0060:Remove unused parameter", Justification = "Attribute accessed method")]
 		private static void NotBlueprintTypeSpecifier(UhtSpecifierContext SpecifierContext)
 		{
 			SetMetaData(SpecifierContext, "NotBlueprintType", true);
@@ -138,6 +146,7 @@ namespace EpicGames.UHT.Parsers
 
 		[UhtSpecifier(Extends = UhtTableNames.Default, ValueType = UhtSpecifierValueType.None, When = UhtSpecifierWhen.Immediate)]
 		[SuppressMessage("CodeQuality", "IDE0051:Remove unused private members", Justification = "Attribute accessed method")]
+		[SuppressMessage("Style", "IDE0060:Remove unused parameter", Justification = "Attribute accessed method")]
 		private static void BlueprintableSpecifier(UhtSpecifierContext SpecifierContext)
 		{
 			SetMetaData(SpecifierContext, "IsBlueprintBase", true);
@@ -146,6 +155,7 @@ namespace EpicGames.UHT.Parsers
 
 		[UhtSpecifier(Extends = UhtTableNames.Default, ValueType = UhtSpecifierValueType.None, When = UhtSpecifierWhen.Immediate)]
 		[SuppressMessage("CodeQuality", "IDE0051:Remove unused private members", Justification = "Attribute accessed method")]
+		[SuppressMessage("Style", "IDE0060:Remove unused parameter", Justification = "Attribute accessed method")]
 		private static void CallInEditorSpecifier(UhtSpecifierContext SpecifierContext)
 		{
 			SetMetaData(SpecifierContext, "CallInEditor", true);
@@ -153,6 +163,7 @@ namespace EpicGames.UHT.Parsers
 
 		[UhtSpecifier(Extends = UhtTableNames.Default, ValueType = UhtSpecifierValueType.None, When = UhtSpecifierWhen.Immediate)]
 		[SuppressMessage("CodeQuality", "IDE0051:Remove unused private members", Justification = "Attribute accessed method")]
+		[SuppressMessage("Style", "IDE0060:Remove unused parameter", Justification = "Attribute accessed method")]
 		private static void NotBlueprintableSpecifier(UhtSpecifierContext SpecifierContext)
 		{
 			SetMetaData(SpecifierContext, "IsBlueprintBase", false);
@@ -161,6 +172,7 @@ namespace EpicGames.UHT.Parsers
 
 		[UhtSpecifier(Extends = UhtTableNames.Default, ValueType = UhtSpecifierValueType.String, When = UhtSpecifierWhen.Immediate)]
 		[SuppressMessage("CodeQuality", "IDE0051:Remove unused private members", Justification = "Attribute accessed method")]
+		[SuppressMessage("Style", "IDE0060:Remove unused parameter", Justification = "Attribute accessed method")]
 		private static void CategorySpecifier(UhtSpecifierContext SpecifierContext, StringView Value)
 		{
 			SetMetaData(SpecifierContext, "Category", Value);
@@ -168,6 +180,7 @@ namespace EpicGames.UHT.Parsers
 
 		[UhtSpecifier(Extends = UhtTableNames.Default, ValueType = UhtSpecifierValueType.None, When = UhtSpecifierWhen.Immediate)]
 		[SuppressMessage("CodeQuality", "IDE0051:Remove unused private members", Justification = "Attribute accessed method")]
+		[SuppressMessage("Style", "IDE0060:Remove unused parameter", Justification = "Attribute accessed method")]
 		private static void ExperimentalSpecifier(UhtSpecifierContext SpecifierContext)
 		{
 			SetMetaData(SpecifierContext, "DevelopmentStatus", "Experimental");
@@ -175,6 +188,7 @@ namespace EpicGames.UHT.Parsers
 
 		[UhtSpecifier(Extends = UhtTableNames.Default, ValueType = UhtSpecifierValueType.None, When = UhtSpecifierWhen.Immediate)]
 		[SuppressMessage("CodeQuality", "IDE0051:Remove unused private members", Justification = "Attribute accessed method")]
+		[SuppressMessage("Style", "IDE0060:Remove unused parameter", Justification = "Attribute accessed method")]
 		private static void EarlyAccessPreviewSpecifier(UhtSpecifierContext SpecifierContext)
 		{
 			SetMetaData(SpecifierContext, "DevelopmentStatus", "EarlyAccess");
@@ -182,6 +196,7 @@ namespace EpicGames.UHT.Parsers
 
 		[UhtSpecifier(Extends = UhtTableNames.Default, ValueType = UhtSpecifierValueType.None, When = UhtSpecifierWhen.Immediate)]
 		[SuppressMessage("CodeQuality", "IDE0051:Remove unused private members", Justification = "Attribute accessed method")]
+		[SuppressMessage("Style", "IDE0060:Remove unused parameter", Justification = "Attribute accessed method")]
 		private static void DocumentationPolicySpecifier(UhtSpecifierContext SpecifierContext)
 		{
 			SetMetaData(SpecifierContext, UhtNames.DocumentationPolicy, "Strict");
@@ -189,6 +204,7 @@ namespace EpicGames.UHT.Parsers
 
 		[UhtSpecifier(Extends = UhtTableNames.Default, ValueType = UhtSpecifierValueType.String, When = UhtSpecifierWhen.Immediate)]
 		[SuppressMessage("CodeQuality", "IDE0051:Remove unused private members", Justification = "Attribute accessed method")]
+		[SuppressMessage("Style", "IDE0060:Remove unused parameter", Justification = "Attribute accessed method")]
 		private static void SparseClassDataTypeSpecifier(UhtSpecifierContext SpecifierContext, StringView Value)
 		{
 			SetMetaData(SpecifierContext, "SparseClassDataType", Value);
@@ -201,6 +217,7 @@ namespace EpicGames.UHT.Parsers
 		[UhtSpecifierValidator(Name = "ClampMin", Extends = UhtTableNames.Default)]
 		[UhtSpecifierValidator(Name = "ClampMax", Extends = UhtTableNames.Default)]
 		[SuppressMessage("CodeQuality", "IDE0051:Remove unused private members", Justification = "Attribute accessed method")]
+		[SuppressMessage("Style", "IDE0060:Remove unused parameter", Justification = "Attribute accessed method")]
 		private static void ValidateNumeric(UhtType Type, UhtMetaData MetaData, UhtMetaDataKey Key, StringView Value)
 		{
 			if (!UhtFCString.IsNumeric(Value.Span))
@@ -211,6 +228,7 @@ namespace EpicGames.UHT.Parsers
 
 		[UhtSpecifierValidator(Extends = UhtTableNames.Default)]
 		[SuppressMessage("CodeQuality", "IDE0051:Remove unused private members", Justification = "Attribute accessed method")]
+		[SuppressMessage("Style", "IDE0060:Remove unused parameter", Justification = "Attribute accessed method")]
 		private static void DevelopmentStatusSpecifierValidator(UhtType Type, UhtMetaData MetaData, UhtMetaDataKey Key, StringView Value)
 		{
 			string[] AllowedValues = { "EarlyAccess", "Experimental" };
@@ -226,6 +244,7 @@ namespace EpicGames.UHT.Parsers
 
 		[UhtSpecifierValidator(Extends = UhtTableNames.Default)]
 		[SuppressMessage("CodeQuality", "IDE0051:Remove unused private members", Justification = "Attribute accessed method")]
+		[SuppressMessage("Style", "IDE0060:Remove unused parameter", Justification = "Attribute accessed method")]
 		private static void DocumentationPolicySpecifierValidator(UhtType Type, UhtMetaData MetaData, UhtMetaDataKey Key, StringView Value)
 		{
 			const string StrictValue = "Strict";
@@ -237,6 +256,7 @@ namespace EpicGames.UHT.Parsers
 
 		[UhtSpecifierValidator(Extends = UhtTableNames.Default)]
 		[SuppressMessage("CodeQuality", "IDE0051:Remove unused private members", Justification = "Attribute accessed method")]
+		[SuppressMessage("Style", "IDE0060:Remove unused parameter", Justification = "Attribute accessed method")]
 		private static void UnitsSpecifierValidator(UhtType Type, UhtMetaData MetaData, UhtMetaDataKey Key, StringView Value)
 		{
 			// Check for numeric property
diff --git a/Engine/Source/Programs/Shared/EpicGames.UHT/Specifiers/UhtFunctionSpecifiers.cs b/Engine/Source/Programs/Shared/EpicGames.UHT/Specifiers/UhtFunctionSpecifiers.cs
--- a/Engine/Source/Programs/Shared/EpicGames.UHT/Specifiers/UhtFunctionSpecifiers.cs
+++ b/Engine/Source/Programs/Shared/EpicGames.UHT/Specifiers/UhtFunctionSpecifiers.cs
@@ -19,6 +19,7 @@ namespace EpicGames.UHT.Parsers
 	{
 		[UhtSpecifier(Extends = UhtTableNames.Function, ValueType = UhtSpecifierValueType.Legacy)]
 		[SuppressMessage("CodeQuality", "IDE0051:Remove unused private members", Justification = "Attribute accessed method")]
+		[SuppressMessage("Style", "IDE0060:Remove unused parameter", Justification = "Attribute accessed method")]
 		private static void BlueprintAuthorityOnlySpecifier(UhtSpecifierContext SpecifierContext)
 		{
 			UhtFunction Function = (UhtFunction)SpecifierContext.Scope.ScopeType;
@@ -27,6 +28,7 @@ namespace EpicGames.UHT.Parsers
 
 		[UhtSpecifier(Extends = UhtTableNames.Function, ValueType = UhtSpecifierValueType.Legacy)]
 		[SuppressMessage("CodeQuality", "IDE0051:Remove unused private members", Justification = "Attribute accessed method")]
+		[SuppressMessage("Style", "IDE0060:Remove unused parameter", Justification = "Attribute accessed method")]
 		private static void BlueprintCallableSpecifier(UhtSpecifierContext SpecifierContext)
 		{
 			UhtFunction Function = (UhtFunction)SpecifierContext.Scope.ScopeType;
@@ -35,6 +37,7 @@ namespace EpicGames.UHT.Parsers
 
 		[UhtSpecifier(Extends = UhtTableNames.Function, ValueType = UhtSpecifierValueType.Legacy)]
 		[SuppressMessage("CodeQuality", "IDE0051:Remove unused private members", Justification = "Attribute accessed method")]
+		[SuppressMessage("Style", "IDE0060:Remove unused parameter", Justification = "Attribute accessed method")]
 		private static void BlueprintCosmeticSpecifier(UhtSpecifierContext SpecifierContext)
 		{
 			UhtFunction Function = (UhtFunction)SpecifierContext.Scope.ScopeType;
@@ -43,6 +46,7 @@ namespace EpicGames.UHT.Parsers
 
 		[UhtSpecifier(Extends = UhtTableNames.Function, ValueType = UhtSpecifierValueType.Legacy)]
 		[SuppressMessage("CodeQuality", "IDE0051:Remove unused private members", Justification = "Attribute accessed method")]
+		[SuppressMessage("Style", "IDE0060:Remove unused parameter", Justification = "Attribute accessed method")]
 		private static void BlueprintGetterSpecifier(UhtSpecifierContext SpecifierContext)
 		{
 			UhtFunctionParser Function = (UhtFunctionParser)SpecifierContext.Scope.ScopeType;
@@ -60,6 +64,7 @@ namespace EpicGames.UHT.Parsers
 
 		[UhtSpecifier(Extends = UhtTableNames.Function, ValueType = UhtSpecifierValueType.Legacy)]
 		[SuppressMessage("CodeQuality", "IDE0051:Remove unused private members", Justification = "Attribute accessed method")]
+		[SuppressMessage("Style", "IDE0060:Remove unused parameter", Justification = "Attribute accessed method")]
 		private static void BlueprintImplementableEventSpecifier(UhtSpecifierContext SpecifierContext)
 		{
 			UhtFunctionParser Function = (UhtFunctionParser)SpecifierContext.Scope.ScopeType;
@@ -89,6 +94,7 @@ namespace EpicGames.UHT.Parsers
 
 		[UhtSpecifier(Extends = UhtTableNames.Function, ValueType = UhtSpecifierValueType.Legacy)]
 		[SuppressMessage("CodeQuality", "IDE0051:Remove unused private members", Justification = "Attribute accessed method")]
+		[SuppressMessage("Style", "IDE0060:Remove unused parameter", Justification = "Attribute accessed method")]
 		private static void BlueprintNativeEventSpecifier(UhtSpecifierContext SpecifierContext)
 		{
 			UhtFunctionParser Function = (UhtFunctionParser)SpecifierContext.Scope.ScopeType;
@@ -117,6 +123,7 @@ namespace EpicGames.UHT.Parsers
 
 		[UhtSpecifier(Extends = UhtTableNames.Function, ValueType = UhtSpecifierValueType.OptionalString)]
 		[SuppressMessage("CodeQuality", "IDE0051:Remove unused private members", Justification = "Attribute accessed method")]
+		[SuppressMessage("Style", "IDE0060:Remove unused parameter", Justification = "Attribute accessed method")]
 		private static void BlueprintPureSpecifier(UhtSpecifierContext SpecifierContext, StringView? Value)
 		{
 			UhtFunction Function = (UhtFunction)SpecifierContext.Scope.ScopeType;
@@ -136,6 +143,7 @@ namespace EpicGames.UHT.Parsers
 
 		[UhtSpecifier(Extends = UhtTableNames.Function, ValueType = UhtSpecifierValueType.Legacy)]
 		[SuppressMessage("CodeQuality", "IDE0051:Remove unused private members", Justification = "Attribute accessed method")]
+		[SuppressMessage("Style", "IDE0060:Remove unused parameter", Justification = "Attribute accessed method")]
 		private static void BlueprintSetterSpecifier(UhtSpecifierContext SpecifierContext)
 		{
 			UhtFunctionParser Function = (UhtFunctionParser)SpecifierContext.Scope.ScopeType;
@@ -152,6 +160,7 @@ namespace EpicGames.UHT.Parsers
 
 		[UhtSpecifier(Extends = UhtTableNames.Function, ValueType = UhtSpecifierValueType.OptionalString)]
 		[SuppressMessage("CodeQuality", "IDE0051:Remove unused private members", Justification = "Attribute accessed method")]
+		[SuppressMessage("Style", "IDE0060:Remove unused parameter", Justification = "Attribute accessed method")]
 		private static void ClientSpecifier(UhtSpecifierContext SpecifierContext, StringView? Value)
 		{
 			UhtFunction Function = (UhtFunction)SpecifierContext.Scope.ScopeType;
@@ -177,6 +186,7 @@ namespace EpicGames.UHT.Parsers
 
 		[UhtSpecifier(Extends = UhtTableNames.Function, ValueType = UhtSpecifierValueType.Legacy)]
 		[SuppressMessage("CodeQuality", "IDE0051:Remove unused private members", Justification = "Attribute accessed method")]
+		[SuppressMessage("Style", "IDE0060:Remove unused parameter", Justification = "Attribute accessed method")]
 		private static void CustomThunkSpecifier(UhtSpecifierContext SpecifierContext)
 		{
 			UhtFunction Function = (UhtFunction)SpecifierContext.Scope.ScopeType;
@@ -185,6 +195,7 @@ namespace EpicGames.UHT.Parsers
 
 		[UhtSpecifier(Extends = UhtTableNames.Function, ValueType = UhtSpecifierValueType.Legacy)]
 		[SuppressMessage("CodeQuality", "IDE0051:Remove unused private members", Justification = "Attribute accessed method")]
+		[SuppressMessage("Style", "IDE0060:Remove unused parameter", Justification = "Attribute accessed method")]
 		private static void ExecSpecifier(UhtSpecifierContext SpecifierContext)
 		{
 			UhtFunction Function = (UhtFunction)SpecifierContext.Scope.ScopeType;
@@ -197,6 +208,7 @@ namespace EpicGames.UHT.Parsers
 
 		[UhtSpecifier(Extends = UhtTableNames.Function, ValueType = UhtSpecifierValueType.Legacy)]
 		[SuppressMessage("CodeQuality", "IDE0051:Remove unused private members", Justification = "Attribute accessed method")]
+		[SuppressMessage("Style", "IDE0060:Remove unused parameter", Justification = "Attribute accessed method")]
 		private static void NetMulticastSpecifier(UhtSpecifierContext SpecifierContext)
 		{
 			UhtFunction Function = (UhtFunction)SpecifierContext.Scope.ScopeType;
@@ -211,6 +223,7 @@ namespace EpicGames.UHT.Parsers
 
 		[UhtSpecifier(Extends = UhtTableNames.Function, ValueType = UhtSpecifierValueType.Legacy)]
 		[SuppressMessage("CodeQuality", "IDE0051:Remove unused private members", Justification = "Attribute accessed method")]
+		[SuppressMessage("Style", "IDE0060:Remove unused parameter", Justification = "Attribute accessed method")]
 		private static void ReliableSpecifier(UhtSpecifierContext SpecifierContext)
 		{
 			UhtFunction Function = (UhtFunction)SpecifierContext.Scope.ScopeType;
@@ -219,6 +232,7 @@ namespace EpicGames.UHT.Parsers
 
 		[UhtSpecifier(Extends = UhtTableNames.Function, ValueType = UhtSpecifierValueType.Legacy)]
 		[SuppressMessage("CodeQuality", "IDE0051:Remove unused private members", Justification = "Attribute accessed method")]
+		[SuppressMessage("Style", "IDE0060:Remove unused parameter", Justification = "Attribute accessed method")]
 		private static void SealedEventSpecifier(UhtSpecifierContext SpecifierContext)
 		{
 			UhtFunction Function = (UhtFunction)SpecifierContext.Scope.ScopeType;
@@ -227,6 +241,7 @@ namespace EpicGames.UHT.Parsers
 
 		[UhtSpecifier(Extends = UhtTableNames.Function, ValueType = UhtSpecifierValueType.OptionalString)]
 		[SuppressMessage("CodeQuality", "IDE0051:Remove unused private members", Justification = "Attribute accessed method")]
+		[SuppressMessage("Style", "IDE0060:Remove unused parameter", Justification = "Attribute accessed method")]
 		private static void ServerSpecifier(UhtSpecifierContext SpecifierContext, StringView? Value)
 		{
 			UhtFunction Function = (UhtFunction)SpecifierContext.Scope.ScopeType;
@@ -252,6 +267,7 @@ namespace EpicGames.UHT.Parsers
 
 		[UhtSpecifier(Extends = UhtTableNames.Function, ValueType = UhtSpecifierValueType.OptionalEqualsKeyValuePairList)]
 		[SuppressMessage("CodeQuality", "IDE0051:Remove unused private members", Justification = "Attribute accessed method")]
+		[SuppressMessage("Style", "IDE0060:Remove unused parameter", Justification = "Attribute accessed method")]
 		private static void ServiceRequestSpecifier(UhtSpecifierContext SpecifierContext, List<KeyValuePair<StringView, StringView>> Value)
 		{
 			UhtFunctionParser Function = (UhtFunctionParser)SpecifierContext.Scope.ScopeType;
@@ -276,6 +292,7 @@ namespace EpicGames.UHT.Parsers
 
 		[UhtSpecifier(Extends = UhtTableNames.Function, ValueType = UhtSpecifierValueType.OptionalEqualsKeyValuePairList)]
 		[SuppressMessage("CodeQuality", "IDE0051:Remove unused private members", Justification = "Attribute accessed method")]
+		[SuppressMessage("Style", "IDE0060:Remove unused parameter", Justification = "Attribute accessed method")]
 		private static void ServiceResponseSpecifier(UhtSpecifierContext SpecifierContext, List<KeyValuePair<StringView, StringView>> Value)
 		{
 			UhtFunctionParser Function = (UhtFunctionParser)SpecifierContext.Scope.ScopeType;
@@ -300,6 +317,7 @@ namespace EpicGames.UHT.Parsers
 
 		[UhtSpecifier(Extends = UhtTableNames.Function, ValueType = UhtSpecifierValueType.Legacy)]
 		[SuppressMessage("CodeQuality", "IDE0051:Remove unused private members", Justification = "Attribute accessed method")]
+		[SuppressMessage("Style", "IDE0060:Remove unused parameter", Justification = "Attribute accessed method")]
 		private static void UnreliableSpecifier(UhtSpecifierContext SpecifierContext)
 		{
 			UhtFunction Function = (UhtFunction)SpecifierContext.Scope.ScopeType;
@@ -308,6 +326,7 @@ namespace EpicGames.UHT.Parsers
 
 		[UhtSpecifier(Extends = UhtTableNames.Function, ValueType = UhtSpecifierValueType.StringList)]
 		[SuppressMessage("CodeQuality", "IDE0051:Remove unused private members", Justification = "Attribute accessed method")]
+		[SuppressMessage("Style", "IDE0060:Remove unused parameter", Justification = "Attribute accessed method")]
 		private static void WithValidationSpecifier(UhtSpecifierContext SpecifierContext, List<StringView>? Value)
 		{
 			UhtFunction Function = (UhtFunction)SpecifierContext.Scope.ScopeType;
@@ -321,6 +340,7 @@ namespace EpicGames.UHT.Parsers
 
 		[UhtSpecifier(Extends = UhtTableNames.Function, ValueType = UhtSpecifierValueType.None)]
 		[SuppressMessage("CodeQuality", "IDE0051:Remove unused private members", Justification = "Attribute accessed method")]
+		[SuppressMessage("Style", "IDE0060:Remove unused parameter", Justification = "Attribute accessed method")]
 		private static void FieldNotifySpecifier(UhtSpecifierContext SpecifierContext)
 		{
 			UhtFunction Function = (UhtFunction)SpecifierContext.Scope.ScopeType;
@@ -329,6 +349,7 @@ namespace EpicGames.UHT.Parsers
 
 		[UhtSpecifierValidator(Extends = UhtTableNames.Function)]
 		[SuppressMessage("CodeQuality", "IDE0051:Remove unused private members", Justification = "Attribute accessed method")]
+		[SuppressMessage("Style", "IDE0060:Remove unused parameter", Justification = "Attribute accessed method")]
 		private static void BlueprintProtectedSpecifierValidator(UhtType Type, UhtMetaData MetaData, UhtMetaDataKey Key, StringView Value)
 		{
 			UhtFunction Function = (UhtFunction)Type;
@@ -357,6 +378,7 @@ namespace EpicGames.UHT.Parsers
 
 		[UhtSpecifierValidator(Extends = UhtTableNames.Function)]
 		[SuppressMessage("CodeQuality", "IDE0051:Remove unused private members", Justification = "Attribute accessed method")]
+		[SuppressMessage("Style", "IDE0060:Remove unused parameter", Justification = "Attribute accessed method")]
 		private static void CommutativeAssociativeBinaryOperatorSpecifierValidator(UhtType Type, UhtMetaData MetaData, UhtMetaDataKey Key, StringView Value)
 		{
 			UhtFunction Function = (UhtFunction)Type;
@@ -372,6 +394,7 @@ namespace EpicGames.UHT.Parsers
 		[UhtSpecifierValidator(Name = "ExpandBoolAsExecs", Extends = UhtTableNames.Function)]
 		[UhtSpecifierValidator(Name = "ExpandEnumAsExecs", Extends = UhtTableNames.Function)]
 		[SuppressMessage("CodeQuality", "IDE0051:Remove unused private members", Justification = "Attribute accessed method")]
+		[SuppressMessage("Style", "IDE0060:Remove unused parameter", Justification = "Attribute accessed method")]
 		private static void ExpandsSpecifierValidator(UhtType Type, UhtMetaData MetaData, UhtMetaDataKey Key, StringView Value)
 		{
 			// multiple entry parsing in the same format as eg SetParam.
diff --git a/Engine/Source/Programs/Shared/EpicGames.UHT/Tables/UhtEngineClassTable.cs b/Engine/Source/Programs/Shared/EpicGames.UHT/Tables/UhtEngineClassTable.cs
--- a/Engine/Source/Programs/Shared/EpicGames.UHT/Tables/UhtEngineClassTable.cs
+++ b/Engine/Source/Programs/Shared/EpicGames.UHT/Tables/UhtEngineClassTable.cs
@@ -69,9 +69,8 @@ namespace EpicGames.UHT.Tables
 		/// <summary>
 		/// Add an entry to the table
 		/// </summary>
-		/// <param name="Type">Type associated with the given attribute</param>
 		/// <param name="EngineClassAttribute">The located attribute</param>
-		public void OnEngineClassAttribute(Type Type, UhtEngineClassAttribute EngineClassAttribute)
+		public void OnEngineClassAttribute(UhtEngineClassAttribute EngineClassAttribute)
 		{
 			if (string.IsNullOrEmpty(EngineClassAttribute.Name))
 			{
diff --git a/Engine/Source/Programs/Shared/EpicGames.UHT/Tables/UhtLocTextDefaultValueTable.cs b/Engine/Source/Programs/Shared/EpicGames.UHT/Tables/UhtLocTextDefaultValueTable.cs
--- a/Engine/Source/Programs/Shared/EpicGames.UHT/Tables/UhtLocTextDefaultValueTable.cs
+++ b/Engine/Source/Programs/Shared/EpicGames.UHT/Tables/UhtLocTextDefaultValueTable.cs
@@ -68,11 +68,10 @@ namespace EpicGames.UHT.Tables
 		/// <summary>
 		/// Handle a loctext default value attribute
 		/// </summary>
-		/// <param name="Type">Containing type</param>
 		/// <param name="MethodInfo">Method info</param>
 		/// <param name="LocTextDefaultValueAttribute">Defining attribute</param>
 		/// <exception cref="UhtIceException">Thrown if the attribute isn't properly defined</exception>
-		public void OnLocTextDefaultValueAttribute(Type Type, MethodInfo MethodInfo, UhtLocTextDefaultValueAttribute LocTextDefaultValueAttribute)
+		public void OnLocTextDefaultValueAttribute(MethodInfo MethodInfo, UhtLocTextDefaultValueAttribute LocTextDefaultValueAttribute)
 		{
 			if (string.IsNullOrEmpty(LocTextDefaultValueAttribute.Name))
 			{
diff --git a/Engine/Source/Programs/Shared/EpicGames.UHT/Tables/UhtPropertyTypeTable.cs b/Engine/Source/Programs/Shared/EpicGames.UHT/Tables/UhtPropertyTypeTable.cs
--- a/Engine/Source/Programs/Shared/EpicGames.UHT/Tables/UhtPropertyTypeTable.cs
+++ b/Engine/Source/Programs/Shared/EpicGames.UHT/Tables/UhtPropertyTypeTable.cs
@@ -189,11 +189,10 @@ namespace EpicGames.UHT.Tables
 		/// <summary>
 		/// Handle a property type attribute
 		/// </summary>
-		/// <param name="Type">Containing type</param>
 		/// <param name="MethodInfo">Method info</param>
 		/// <param name="PropertyTypeAttribute">Attribute</param>
 		/// <exception cref="UhtIceException">Thrown if the property type isn't properly defined.</exception>
-		public void OnPropertyTypeAttribute(Type Type, MethodInfo MethodInfo, UhtPropertyTypeAttribute PropertyTypeAttribute)
+		public void OnPropertyTypeAttribute(MethodInfo MethodInfo, UhtPropertyTypeAttribute PropertyTypeAttribute)
 		{
 			if (string.IsNullOrEmpty(PropertyTypeAttribute.Keyword) && !PropertyTypeAttribute.Options.HasAnyFlags(UhtPropertyTypeOptions.Default))
 			{
diff --git a/Engine/Source/Programs/Shared/EpicGames.UHT/Tables/UhtStructDefaultValueTable.cs b/Engine/Source/Programs/Shared/EpicGames.UHT/Tables/UhtStructDefaultValueTable.cs
--- a/Engine/Source/Programs/Shared/EpicGames.UHT/Tables/UhtStructDefaultValueTable.cs
+++ b/Engine/Source/Programs/Shared/EpicGames.UHT/Tables/UhtStructDefaultValueTable.cs
@@ -148,11 +148,10 @@ namespace EpicGames.UHT.Tables
 		/// <summary>
 		/// Handle a structure default value sanitizer attribute
 		/// </summary>
-		/// <param name="Type">Containing type</param>
 		/// <param name="MethodInfo">Method information</param>
 		/// <param name="StructDefaultValueAttribute">Found attribute</param>
 		/// <exception cref="UhtIceException">Thrown if the attribute isn't property defined</exception>
-		public void OnStructDefaultValueAttribute(Type Type, MethodInfo MethodInfo, UhtStructDefaultValueAttribute StructDefaultValueAttribute)
+		public void OnStructDefaultValueAttribute(MethodInfo MethodInfo, UhtStructDefaultValueAttribute StructDefaultValueAttribute)
 		{
 			if (string.IsNullOrEmpty(StructDefaultValueAttribute.Name) && !StructDefaultValueAttribute.Options.HasAnyFlags(UhtStructDefaultValueOptions.Default))
 			{
diff --git a/Engine/Source/Programs/Shared/EpicGames.UHT/Tokenizer/UhtTokenReaderIntegerExtensions.cs b/Engine/Source/Programs/Shared/EpicGames.UHT/Tokenizer/UhtTokenReaderIntegerExtensions.cs
--- a/Engine/Source/Programs/Shared/EpicGames.UHT/Tokenizer/UhtTokenReaderIntegerExtensions.cs
+++ b/Engine/Source/Programs/Shared/EpicGames.UHT/Tokenizer/UhtTokenReaderIntegerExtensions.cs
@@ -40,9 +40,8 @@ namespace EpicGames.UHT.Tokenizer
 		/// Get the next token as an integer.  If the next token is not an integer, an exception is thrown
 		/// </summary>
 		/// <param name="TokenReader">Token reader</param>
-		/// <param name="ExceptionContext">If not null, an exception will be thrown with the given text as part of the message.</param>
 		/// <returns>The value of the constant</returns>
-		public static IUhtTokenReader OptionalConstInt(this IUhtTokenReader TokenReader, object? ExceptionContext = null)
+		public static IUhtTokenReader OptionalConstInt(this IUhtTokenReader TokenReader)
 		{
 			TokenReader.TryOptionalConstInt(out int _);
 			return TokenReader;
@@ -101,9 +100,8 @@ namespace EpicGames.UHT.Tokenizer
 		/// Get the next token as an integer.  If the next token is not an integer, an exception is thrown
 		/// </summary>
 		/// <param name="TokenReader">Token reader</param>
-		/// <param name="ExceptionContext">If not null, an exception will be thrown with the given text as part of the message.</param>
 		/// <returns>The value of the constant</returns>
-		public static IUhtTokenReader OptionalConstLong(this IUhtTokenReader TokenReader, object? ExceptionContext = null)
+		public static IUhtTokenReader OptionalConstLong(this IUhtTokenReader TokenReader)
 		{
 			ref UhtToken Token = ref TokenReader.PeekToken();
 
diff --git a/Engine/Source/Programs/Shared/EpicGames.UHT/Tokenizer/UhtTokenReplayReader.cs b/Engine/Source/Programs/Shared/EpicGames.UHT/Tokenizer/UhtTokenReplayReader.cs
--- a/Engine/Source/Programs/Shared/EpicGames.UHT/Tokenizer/UhtTokenReplayReader.cs
+++ b/Engine/Source/Programs/Shared/EpicGames.UHT/Tokenizer/UhtTokenReplayReader.cs
@@ -43,8 +43,8 @@ namespace EpicGames.UHT.Tokenizer
 		public UhtTokenReplayReader(IUhtMessageSite MessageSite, ReadOnlyMemory<char> Data, ReadOnlyMemory<UhtToken> Tokens, UhtTokenType EndTokenType)
 		{
 			this.MessageSiteInternal = MessageSite;
-			this.Tokens = Array.Empty<UhtToken>().AsMemory();
-			this.Data = Array.Empty<char>().AsMemory();
+			this.Tokens = Tokens;
+			this.Data = Data;
 			this.EndTokenType = EndTokenType;
 			this.CurrentToken = new UhtToken(EndTokenType);
 		}
diff --git a/Engine/Source/Programs/Shared/EpicGames.UHT/Types/Properties/UhtArrayProperty.cs b/Engine/Source/Programs/Shared/EpicGames.UHT/Types/Properties/UhtArrayProperty.cs
--- a/Engine/Source/Programs/Shared/EpicGames.UHT/Types/Properties/UhtArrayProperty.cs
+++ b/Engine/Source/Programs/Shared/EpicGames.UHT/Types/Properties/UhtArrayProperty.cs
@@ -212,7 +212,7 @@ namespace EpicGames.UHT.Types
 		[SuppressMessage("CodeQuality", "IDE0051:Remove unused private members", Justification = "Attribute accessed method")]
 		private static UhtProperty? ArrayProperty(UhtPropertyResolvePhase ResolvePhase, UhtPropertySettings PropertySettings, IUhtTokenReader TokenReader, UhtToken MatchedToken)
 		{
-			using (var TokenContext = new UhtMessageContext(TokenReader, "TArray"))
+			using (var TokenContext = new UhtMessageContext("TArray"))
 			{
 				if (!TokenReader.SkipExpectedType(MatchedToken.Value, PropertySettings.PropertyCategory == UhtPropertyCategory.Member))
 				{
diff --git a/Engine/Source/Programs/Shared/EpicGames.UHT/Types/Properties/UhtBoolProperty.cs b/Engine/Source/Programs/Shared/EpicGames.UHT/Types/Properties/UhtBoolProperty.cs
--- a/Engine/Source/Programs/Shared/EpicGames.UHT/Types/Properties/UhtBoolProperty.cs
+++ b/Engine/Source/Programs/Shared/EpicGames.UHT/Types/Properties/UhtBoolProperty.cs
@@ -259,6 +259,7 @@ namespace EpicGames.UHT.Types
 		#region Keyword
 		[UhtPropertyType(Keyword = "bool", Options = UhtPropertyTypeOptions.Simple | UhtPropertyTypeOptions.Immediate)]
 		[SuppressMessage("CodeQuality", "IDE0051:Remove unused private members", Justification = "Attribute accessed method")]
+		[SuppressMessage("Style", "IDE0060:Remove unused parameter", Justification = "Attribute accessed method")]
 		private static UhtProperty? BoolProperty(UhtPropertyResolvePhase ResolvePhase, UhtPropertySettings PropertySettings, IUhtTokenReader TokenReader, UhtToken MatchedToken)
 		{
 			if (PropertySettings.bIsBitfield)
diff --git a/Engine/Source/Programs/Shared/EpicGames.UHT/Types/Properties/UhtByteProperty.cs b/Engine/Source/Programs/Shared/EpicGames.UHT/Types/Properties/UhtByteProperty.cs
--- a/Engine/Source/Programs/Shared/EpicGames.UHT/Types/Properties/UhtByteProperty.cs
+++ b/Engine/Source/Programs/Shared/EpicGames.UHT/Types/Properties/UhtByteProperty.cs
@@ -159,6 +159,7 @@ namespace EpicGames.UHT.Types
 		#region Keyword
 		[UhtPropertyType(Keyword = "uint8", Options = UhtPropertyTypeOptions.Simple | UhtPropertyTypeOptions.Immediate)]
 		[SuppressMessage("CodeQuality", "IDE0051:Remove unused private members", Justification = "Attribute accessed method")]
+		[SuppressMessage("Style", "IDE0060:Remove unused parameter", Justification = "Attribute accessed method")]
 		private static UhtProperty? UInt8Property(UhtPropertyResolvePhase ResolvePhase, UhtPropertySettings PropertySettings, IUhtTokenReader TokenReader, UhtToken MatchedToken)
 		{
 			if (PropertySettings.bIsBitfield)
diff --git a/Engine/Source/Programs/Shared/EpicGames.UHT/Types/Properties/UhtClassProperty.cs b/Engine/Source/Programs/Shared/EpicGames.UHT/Types/Properties/UhtClassProperty.cs
--- a/Engine/Source/Programs/Shared/EpicGames.UHT/Types/Properties/UhtClassProperty.cs
+++ b/Engine/Source/Programs/Shared/EpicGames.UHT/Types/Properties/UhtClassProperty.cs
@@ -148,6 +148,7 @@ namespace EpicGames.UHT.Types
 		#region Keyword
 		[UhtPropertyType(Keyword = "TSubclassOf")]
 		[SuppressMessage("CodeQuality", "IDE0051:Remove unused private members", Justification = "Attribute accessed method")]
+		[SuppressMessage("Style", "IDE0060:Remove unused parameter", Justification = "Attribute accessed method")]
 		private static UhtProperty? SubclassOfProperty(UhtPropertyResolvePhase ResolvePhase, UhtPropertySettings PropertySettings, IUhtTokenReader TokenReader, UhtToken MatchedToken)
 		{
 			UhtClass? MetaClass = UhtObjectPropertyBase.ParseTemplateClass(PropertySettings, TokenReader, MatchedToken);
diff --git a/Engine/Source/Programs/Shared/EpicGames.UHT/Types/Properties/UhtDoubleProperty.cs b/Engine/Source/Programs/Shared/EpicGames.UHT/Types/Properties/UhtDoubleProperty.cs
--- a/Engine/Source/Programs/Shared/EpicGames.UHT/Types/Properties/UhtDoubleProperty.cs
+++ b/Engine/Source/Programs/Shared/EpicGames.UHT/Types/Properties/UhtDoubleProperty.cs
@@ -61,6 +61,7 @@ namespace EpicGames.UHT.Types
 		#region Keyword
 		[UhtPropertyType(Keyword = "double", Options = UhtPropertyTypeOptions.Simple | UhtPropertyTypeOptions.Immediate)]
 		[SuppressMessage("CodeQuality", "IDE0051:Remove unused private members", Justification = "Attribute accessed method")]
+		[SuppressMessage("Style", "IDE0060:Remove unused parameter", Justification = "Attribute accessed method")]
 		private static UhtProperty? DoubleProperty(UhtPropertyResolvePhase ResolvePhase, UhtPropertySettings PropertySettings, IUhtTokenReader TokenReader, UhtToken MatchedToken)
 		{
 			return new UhtDoubleProperty(PropertySettings);
diff --git a/Engine/Source/Programs/Shared/EpicGames.UHT/Types/Properties/UhtEnumProperty.cs b/Engine/Source/Programs/Shared/EpicGames.UHT/Types/Properties/UhtEnumProperty.cs
--- a/Engine/Source/Programs/Shared/EpicGames.UHT/Types/Properties/UhtEnumProperty.cs
+++ b/Engine/Source/Programs/Shared/EpicGames.UHT/Types/Properties/UhtEnumProperty.cs
@@ -334,11 +334,12 @@ namespace EpicGames.UHT.Types
 		#region Keyword
 		[UhtPropertyType(Keyword = "TEnumAsByte")]
 		[SuppressMessage("CodeQuality", "IDE0051:Remove unused private members", Justification = "Attribute accessed method")]
+		[SuppressMessage("Style", "IDE0060:Remove unused parameter", Justification = "Attribute accessed method")]
 		private static UhtProperty? EnumProperty(UhtPropertyResolvePhase ResolvePhase, UhtPropertySettings PropertySettings, IUhtTokenReader TokenReader, UhtToken MatchedToken)
 		{
 			UhtType Outer = PropertySettings.Outer;
 			UhtEnum? Enum = null;
-			using (var TokenContext = new UhtMessageContext(TokenReader, "TEnumAsByte"))
+			using (var TokenContext = new UhtMessageContext("TEnumAsByte"))
 			{
 				TokenReader
 					.Require("TEnumAsByte")
diff --git a/Engine/Source/Programs/Shared/EpicGames.UHT/Types/Properties/UhtFieldPathProperty.cs b/Engine/Source/Programs/Shared/EpicGames.UHT/Types/Properties/UhtFieldPathProperty.cs
--- a/Engine/Source/Programs/Shared/EpicGames.UHT/Types/Properties/UhtFieldPathProperty.cs
+++ b/Engine/Source/Programs/Shared/EpicGames.UHT/Types/Properties/UhtFieldPathProperty.cs
@@ -107,9 +107,10 @@ namespace EpicGames.UHT.Types
 		#region Keyword
 		[UhtPropertyType(Keyword = "TFieldPath", Options = UhtPropertyTypeOptions.Immediate)]
 		[SuppressMessage("CodeQuality", "IDE0051:Remove unused private members", Justification = "Attribute accessed method")]
+		[SuppressMessage("Style", "IDE0060:Remove unused parameter", Justification = "Attribute accessed method")]
 		private static UhtProperty? FieldPathProperty(UhtPropertyResolvePhase ResolvePhase, UhtPropertySettings PropertySettings, IUhtTokenReader TokenReader, UhtToken MatchedToken)
 		{
-			using (var TokenContext = new UhtMessageContext(TokenReader, "TFieldPath"))
+			using (var TokenContext = new UhtMessageContext("TFieldPath"))
 			{
 				UhtToken Identifier = new UhtToken();
 				TokenReader
diff --git a/Engine/Source/Programs/Shared/EpicGames.UHT/Types/Properties/UhtFloatProperty.cs b/Engine/Source/Programs/Shared/EpicGames.UHT/Types/Properties/UhtFloatProperty.cs
--- a/Engine/Source/Programs/Shared/EpicGames.UHT/Types/Properties/UhtFloatProperty.cs
+++ b/Engine/Source/Programs/Shared/EpicGames.UHT/Types/Properties/UhtFloatProperty.cs
@@ -60,6 +60,7 @@ namespace EpicGames.UHT.Types
 		#region Keywords
 		[UhtPropertyType(Keyword = "float", Options = UhtPropertyTypeOptions.Simple | UhtPropertyTypeOptions.Immediate)]
 		[SuppressMessage("CodeQuality", "IDE0051:Remove unused private members", Justification = "Attribute accessed method")]
+		[SuppressMessage("Style", "IDE0060:Remove unused parameter", Justification = "Attribute accessed method")]
 		private static UhtProperty? FloatProperty(UhtPropertyResolvePhase ResolvePhase, UhtPropertySettings PropertySettings, IUhtTokenReader TokenReader, UhtToken MatchedToken)
 		{
 			return new UhtFloatProperty(PropertySettings);
diff --git a/Engine/Source/Programs/Shared/EpicGames.UHT/Types/Properties/UhtInt16Property.cs b/Engine/Source/Programs/Shared/EpicGames.UHT/Types/Properties/UhtInt16Property.cs
--- a/Engine/Source/Programs/Shared/EpicGames.UHT/Types/Properties/UhtInt16Property.cs
+++ b/Engine/Source/Programs/Shared/EpicGames.UHT/Types/Properties/UhtInt16Property.cs
@@ -59,6 +59,7 @@ namespace EpicGames.UHT.Types
 		#region Keyword
 		[UhtPropertyType(Keyword = "int16", Options = UhtPropertyTypeOptions.Simple | UhtPropertyTypeOptions.Immediate)]
 		[SuppressMessage("CodeQuality", "IDE0051:Remove unused private members", Justification = "Attribute accessed method")]
+		[SuppressMessage("Style", "IDE0060:Remove unused parameter", Justification = "Attribute accessed method")]
 		private static UhtProperty? Int16Property(UhtPropertyResolvePhase ResolvePhase, UhtPropertySettings PropertySettings, IUhtTokenReader TokenReader, UhtToken MatchedToken)
 		{
 			return new UhtInt16Property(PropertySettings, UhtPropertyIntType.Sized);
diff --git a/Engine/Source/Programs/Shared/EpicGames.UHT/Types/Properties/UhtInt64Property.cs b/Engine/Source/Programs/Shared/EpicGames.UHT/Types/Properties/UhtInt64Property.cs
--- a/Engine/Source/Programs/Shared/EpicGames.UHT/Types/Properties/UhtInt64Property.cs
+++ b/Engine/Source/Programs/Shared/EpicGames.UHT/Types/Properties/UhtInt64Property.cs
@@ -61,6 +61,7 @@ namespace EpicGames.UHT.Types
 		#region Keyword
 		[UhtPropertyType(Keyword = "int64", Options = UhtPropertyTypeOptions.Simple | UhtPropertyTypeOptions.Immediate)]
 		[SuppressMessage("CodeQuality", "IDE0051:Remove unused private members", Justification = "Attribute accessed method")]
+		[SuppressMessage("Style", "IDE0060:Remove unused parameter", Justification = "Attribute accessed method")]
 		private static UhtProperty? Int64Property(UhtPropertyResolvePhase ResolvePhase, UhtPropertySettings PropertySettings, IUhtTokenReader TokenReader, UhtToken MatchedToken)
 		{
 			return new UhtInt64Property(PropertySettings, UhtPropertyIntType.Sized);
diff --git a/Engine/Source/Programs/Shared/EpicGames.UHT/Types/Properties/UhtInt8Property.cs b/Engine/Source/Programs/Shared/EpicGames.UHT/Types/Properties/UhtInt8Property.cs
--- a/Engine/Source/Programs/Shared/EpicGames.UHT/Types/Properties/UhtInt8Property.cs
+++ b/Engine/Source/Programs/Shared/EpicGames.UHT/Types/Properties/UhtInt8Property.cs
@@ -59,6 +59,7 @@ namespace EpicGames.UHT.Types
 		#region Keyword
 		[UhtPropertyType(Keyword = "int8", Options = UhtPropertyTypeOptions.Simple | UhtPropertyTypeOptions.Immediate)]
 		[SuppressMessage("CodeQuality", "IDE0051:Remove unused private members", Justification = "Attribute accessed method")]
+		[SuppressMessage("Style", "IDE0060:Remove unused parameter", Justification = "Attribute accessed method")]
 		private static UhtProperty? Int8Property(UhtPropertyResolvePhase ResolvePhase, UhtPropertySettings PropertySettings, IUhtTokenReader TokenReader, UhtToken MatchedToken)
 		{
 			return new UhtInt8Property(PropertySettings, UhtPropertyIntType.Sized);
diff --git a/Engine/Source/Programs/Shared/EpicGames.UHT/Types/Properties/UhtIntProperty.cs b/Engine/Source/Programs/Shared/EpicGames.UHT/Types/Properties/UhtIntProperty.cs
--- a/Engine/Source/Programs/Shared/EpicGames.UHT/Types/Properties/UhtIntProperty.cs
+++ b/Engine/Source/Programs/Shared/EpicGames.UHT/Types/Properties/UhtIntProperty.cs
@@ -64,6 +64,7 @@ namespace EpicGames.UHT.Types
 		#region Keywords
 		[UhtPropertyType(Keyword = "int32", Options = UhtPropertyTypeOptions.Simple | UhtPropertyTypeOptions.Immediate)]
 		[SuppressMessage("CodeQuality", "IDE0051:Remove unused private members", Justification = "Attribute accessed method")]
+		[SuppressMessage("Style", "IDE0060:Remove unused parameter", Justification = "Attribute accessed method")]
 		private static UhtProperty? Int32Property(UhtPropertyResolvePhase ResolvePhase, UhtPropertySettings PropertySettings, IUhtTokenReader TokenReader, UhtToken MatchedToken)
 		{
 			return new UhtIntProperty(PropertySettings, UhtPropertyIntType.Sized);
@@ -71,6 +72,7 @@ namespace EpicGames.UHT.Types
 
 		[UhtPropertyType(Keyword = "int", Options = UhtPropertyTypeOptions.Simple | UhtPropertyTypeOptions.Immediate)]
 		[SuppressMessage("CodeQuality", "IDE0051:Remove unused private members", Justification = "Attribute accessed method")]
+		[SuppressMessage("Style", "IDE0060:Remove unused parameter", Justification = "Attribute accessed method")]
 		private static UhtProperty? IntProperty(UhtPropertyResolvePhase ResolvePhase, UhtPropertySettings PropertySettings, IUhtTokenReader TokenReader, UhtToken MatchedToken)
 		{
 			return new UhtIntProperty(PropertySettings, UhtPropertyIntType.Unsized);
@@ -78,6 +80,7 @@ namespace EpicGames.UHT.Types
 
 		[UhtPropertyType(Keyword = "signed", Options = UhtPropertyTypeOptions.Immediate)]
 		[SuppressMessage("CodeQuality", "IDE0051:Remove unused private members", Justification = "Attribute accessed method")]
+		[SuppressMessage("Style", "IDE0060:Remove unused parameter", Justification = "Attribute accessed method")]
 		private static UhtProperty? SignedProperty(UhtPropertyResolvePhase ResolvePhase, UhtPropertySettings PropertySettings, IUhtTokenReader TokenReader, UhtToken MatchedToken)
 		{
 			TokenReader
diff --git a/Engine/Source/Programs/Shared/EpicGames.UHT/Types/Properties/UhtInterfaceProperty.cs b/Engine/Source/Programs/Shared/EpicGames.UHT/Types/Properties/UhtInterfaceProperty.cs
--- a/Engine/Source/Programs/Shared/EpicGames.UHT/Types/Properties/UhtInterfaceProperty.cs
+++ b/Engine/Source/Programs/Shared/EpicGames.UHT/Types/Properties/UhtInterfaceProperty.cs
@@ -178,6 +178,7 @@ namespace EpicGames.UHT.Types
 		#region Keywords
 		[UhtPropertyType(Keyword = "FScriptInterface", Options = UhtPropertyTypeOptions.Simple)] // This can't be immediate due to the reference to UInterface
 		[SuppressMessage("CodeQuality", "IDE0051:Remove unused private members", Justification = "Attribute accessed method")]
+		[SuppressMessage("Style", "IDE0060:Remove unused parameter", Justification = "Attribute accessed method")]
 		private static UhtProperty? FScriptInterfaceProperty(UhtPropertyResolvePhase ResolvePhase, UhtPropertySettings PropertySettings, IUhtTokenReader TokenReader, UhtToken MatchedToken)
 		{
 			return new UhtInterfaceProperty(PropertySettings, PropertySettings.Outer.Session.IInterface);
@@ -185,6 +186,7 @@ namespace EpicGames.UHT.Types
 
 		[UhtPropertyType(Keyword = "TScriptInterface")]
 		[SuppressMessage("CodeQuality", "IDE0051:Remove unused private members", Justification = "Attribute accessed method")]
+		[SuppressMessage("Style", "IDE0060:Remove unused parameter", Justification = "Attribute accessed method")]
 		private static UhtProperty? TScriptInterfaceProperty(UhtPropertyResolvePhase ResolvePhase, UhtPropertySettings PropertySettings, IUhtTokenReader TokenReader, UhtToken MatchedToken)
 		{
 			UhtClass? PropertyClass = UhtObjectPropertyBase.ParseTemplateObject(PropertySettings, TokenReader, MatchedToken, false);
diff --git a/Engine/Source/Programs/Shared/EpicGames.UHT/Types/Properties/UhtLargeWorldCoordinatesRealProperty.cs b/Engine/Source/Programs/Shared/EpicGames.UHT/Types/Properties/UhtLargeWorldCoordinatesRealProperty.cs
--- a/Engine/Source/Programs/Shared/EpicGames.UHT/Types/Properties/UhtLargeWorldCoordinatesRealProperty.cs
+++ b/Engine/Source/Programs/Shared/EpicGames.UHT/Types/Properties/UhtLargeWorldCoordinatesRealProperty.cs
@@ -62,6 +62,7 @@ namespace EpicGames.UHT.Types
 		#region Keyword
 		[UhtPropertyType(Keyword = "FLargeWorldCoordinatesReal", Options = UhtPropertyTypeOptions.Simple | UhtPropertyTypeOptions.Immediate)]
 		[SuppressMessage("CodeQuality", "IDE0051:Remove unused private members", Justification = "Attribute accessed method")]
+		[SuppressMessage("Style", "IDE0060:Remove unused parameter", Justification = "Attribute accessed method")]
 		private static UhtProperty? LargeWorldCoordinatesRealProperty(UhtPropertyResolvePhase ResolvePhase, UhtPropertySettings PropertySettings, IUhtTokenReader TokenReader, UhtToken MatchedToken)
 		{
 			if (!PropertySettings.Outer.HeaderFile.bIsNoExportTypes)
diff --git a/Engine/Source/Programs/Shared/EpicGames.UHT/Types/Properties/UhtLazyObjectPtrProperty.cs b/Engine/Source/Programs/Shared/EpicGames.UHT/Types/Properties/UhtLazyObjectPtrProperty.cs
--- a/Engine/Source/Programs/Shared/EpicGames.UHT/Types/Properties/UhtLazyObjectPtrProperty.cs
+++ b/Engine/Source/Programs/Shared/EpicGames.UHT/Types/Properties/UhtLazyObjectPtrProperty.cs
@@ -108,6 +108,7 @@ namespace EpicGames.UHT.Types
 		#region Keyword
 		[UhtPropertyType(Keyword = "TLazyObjectPtr")]
 		[SuppressMessage("CodeQuality", "IDE0051:Remove unused private members", Justification = "Attribute accessed method")]
+		[SuppressMessage("Style", "IDE0060:Remove unused parameter", Justification = "Attribute accessed method")]
 		private static UhtProperty? LazyObjectPtrProperty(UhtPropertyResolvePhase ResolvePhase, UhtPropertySettings PropertySettings, IUhtTokenReader TokenReader, UhtToken MatchedToken)
 		{
 			UhtClass? PropertyClass = ParseTemplateObject(PropertySettings, TokenReader, MatchedToken, true);
diff --git a/Engine/Source/Programs/Shared/EpicGames.UHT/Types/Properties/UhtMapProperty.cs b/Engine/Source/Programs/Shared/EpicGames.UHT/Types/Properties/UhtMapProperty.cs
--- a/Engine/Source/Programs/Shared/EpicGames.UHT/Types/Properties/UhtMapProperty.cs
+++ b/Engine/Source/Programs/Shared/EpicGames.UHT/Types/Properties/UhtMapProperty.cs
@@ -260,7 +260,7 @@ namespace EpicGames.UHT.Types
 		[SuppressMessage("CodeQuality", "IDE0051:Remove unused private members", Justification = "Attribute accessed method")]
 		private static UhtProperty? MapProperty(UhtPropertyResolvePhase ResolvePhase, UhtPropertySettings PropertySettings, IUhtTokenReader TokenReader, UhtToken MatchedToken)
 		{
-			using (var TokenContext = new UhtMessageContext(TokenReader, "TMap"))
+			using (var TokenContext = new UhtMessageContext("TMap"))
 			{
 				if (!TokenReader.SkipExpectedType(MatchedToken.Value, PropertySettings.PropertyCategory == UhtPropertyCategory.Member))
 				{
diff --git a/Engine/Source/Programs/Shared/EpicGames.UHT/Types/Properties/UhtNameProperty.cs b/Engine/Source/Programs/Shared/EpicGames.UHT/Types/Properties/UhtNameProperty.cs
--- a/Engine/Source/Programs/Shared/EpicGames.UHT/Types/Properties/UhtNameProperty.cs
+++ b/Engine/Source/Programs/Shared/EpicGames.UHT/Types/Properties/UhtNameProperty.cs
@@ -91,6 +91,7 @@ namespace EpicGames.UHT.Types
 		#region Keyword
 		[UhtPropertyType(Keyword = "FName", Options = UhtPropertyTypeOptions.Simple | UhtPropertyTypeOptions.Immediate)]
 		[SuppressMessage("CodeQuality", "IDE0051:Remove unused private members", Justification = "Attribute accessed method")]
+		[SuppressMessage("Style", "IDE0060:Remove unused parameter", Justification = "Attribute accessed method")]
 		private static UhtProperty? NameProperty(UhtPropertyResolvePhase ResolvePhase, UhtPropertySettings PropertySettings, IUhtTokenReader TokenReader, UhtToken MatchedToken)
 		{
 			return new UhtNameProperty(PropertySettings);
diff --git a/Engine/Source/Programs/Shared/EpicGames.UHT/Types/Properties/UhtObjectPtrProperty.cs b/Engine/Source/Programs/Shared/EpicGames.UHT/Types/Properties/UhtObjectPtrProperty.cs
--- a/Engine/Source/Programs/Shared/EpicGames.UHT/Types/Properties/UhtObjectPtrProperty.cs
+++ b/Engine/Source/Programs/Shared/EpicGames.UHT/Types/Properties/UhtObjectPtrProperty.cs
@@ -103,6 +103,7 @@ namespace EpicGames.UHT.Types
 		#region Keyword
 		[UhtPropertyType(Keyword = "TObjectPtr")]
 		[SuppressMessage("CodeQuality", "IDE0051:Remove unused private members", Justification = "Attribute accessed method")]
+		[SuppressMessage("Style", "IDE0060:Remove unused parameter", Justification = "Attribute accessed method")]
 		private static UhtProperty? ObjectPtrProperty(UhtPropertyResolvePhase ResolvePhase, UhtPropertySettings PropertySettings, IUhtTokenReader TokenReader, UhtToken MatchedToken)
 		{
 			UhtSession Session = PropertySettings.Outer.Session;
diff --git a/Engine/Source/Programs/Shared/EpicGames.UHT/Types/Properties/UhtSetProperty.cs b/Engine/Source/Programs/Shared/EpicGames.UHT/Types/Properties/UhtSetProperty.cs
--- a/Engine/Source/Programs/Shared/EpicGames.UHT/Types/Properties/UhtSetProperty.cs
+++ b/Engine/Source/Programs/Shared/EpicGames.UHT/Types/Properties/UhtSetProperty.cs
@@ -211,7 +211,7 @@ namespace EpicGames.UHT.Types
 		[SuppressMessage("CodeQuality", "IDE0051:Remove unused private members", Justification = "Attribute accessed method")]
 		private static UhtProperty? SetProperty(UhtPropertyResolvePhase ResolvePhase, UhtPropertySettings PropertySettings, IUhtTokenReader TokenReader, UhtToken MatchedToken)
 		{
-			using (var TokenContext = new UhtMessageContext(TokenReader, "TSet"))
+			using (var TokenContext = new UhtMessageContext("TSet"))
 			{
 				if (!TokenReader.SkipExpectedType(MatchedToken.Value, PropertySettings.PropertyCategory == UhtPropertyCategory.Member))
 				{
diff --git a/Engine/Source/Programs/Shared/EpicGames.UHT/Types/Properties/UhtSoftClassProperty.cs b/Engine/Source/Programs/Shared/EpicGames.UHT/Types/Properties/UhtSoftClassProperty.cs
--- a/Engine/Source/Programs/Shared/EpicGames.UHT/Types/Properties/UhtSoftClassProperty.cs
+++ b/Engine/Source/Programs/Shared/EpicGames.UHT/Types/Properties/UhtSoftClassProperty.cs
@@ -81,6 +81,7 @@ namespace EpicGames.UHT.Types
 		#region Keyword
 		[UhtPropertyType(Keyword = "TSoftClassPtr")]
 		[SuppressMessage("CodeQuality", "IDE0051:Remove unused private members", Justification = "Attribute accessed method")]
+		[SuppressMessage("Style", "IDE0060:Remove unused parameter", Justification = "Attribute accessed method")]
 		private static UhtProperty? SoftClassPtrProperty(UhtPropertyResolvePhase ResolvePhase, UhtPropertySettings PropertySettings, IUhtTokenReader TokenReader, UhtToken MatchedToken)
 		{
 			UhtClass? MetaClass = ParseTemplateClass(PropertySettings, TokenReader, MatchedToken);
diff --git a/Engine/Source/Programs/Shared/EpicGames.UHT/Types/Properties/UhtSoftObjectProperty.cs b/Engine/Source/Programs/Shared/EpicGames.UHT/Types/Properties/UhtSoftObjectProperty.cs
--- a/Engine/Source/Programs/Shared/EpicGames.UHT/Types/Properties/UhtSoftObjectProperty.cs
+++ b/Engine/Source/Programs/Shared/EpicGames.UHT/Types/Properties/UhtSoftObjectProperty.cs
@@ -89,6 +89,7 @@ namespace EpicGames.UHT.Types
 		#region Keyword
 		[UhtPropertyType(Keyword = "TSoftObjectPtr")]
 		[SuppressMessage("CodeQuality", "IDE0051:Remove unused private members", Justification = "Attribute accessed method")]
+		[SuppressMessage("Style", "IDE0060:Remove unused parameter", Justification = "Attribute accessed method")]
 		private static UhtProperty? SoftObjectPtrProperty(UhtPropertyResolvePhase ResolvePhase, UhtPropertySettings PropertySettings, IUhtTokenReader TokenReader, UhtToken MatchedToken)
 		{
 			UhtClass? PropertyClass = UhtObjectPropertyBase.ParseTemplateObject(PropertySettings, TokenReader, MatchedToken, true);
diff --git a/Engine/Source/Programs/Shared/EpicGames.UHT/Types/Properties/UhtStrProperty.cs b/Engine/Source/Programs/Shared/EpicGames.UHT/Types/Properties/UhtStrProperty.cs
--- a/Engine/Source/Programs/Shared/EpicGames.UHT/Types/Properties/UhtStrProperty.cs
+++ b/Engine/Source/Programs/Shared/EpicGames.UHT/Types/Properties/UhtStrProperty.cs
@@ -104,6 +104,7 @@ namespace EpicGames.UHT.Types
 		[UhtPropertyType(Keyword = "FString")]
 		[UhtPropertyType(Keyword = "FMemoryImageString", Options = UhtPropertyTypeOptions.Immediate)]
 		[SuppressMessage("CodeQuality", "IDE0051:Remove unused private members", Justification = "Attribute accessed method")]
+		[SuppressMessage("Style", "IDE0060:Remove unused parameter", Justification = "Attribute accessed method")]
 		private static UhtProperty? StrProperty(UhtPropertyResolvePhase ResolvePhase, UhtPropertySettings PropertySettings, IUhtTokenReader TokenReader, UhtToken MatchedToken)
 		{
 			if (!TokenReader.SkipExpectedType(MatchedToken.Value, PropertySettings.PropertyCategory == UhtPropertyCategory.Member))
diff --git a/Engine/Source/Programs/Shared/EpicGames.UHT/Types/Properties/UhtStructProperty.cs b/Engine/Source/Programs/Shared/EpicGames.UHT/Types/Properties/UhtStructProperty.cs
--- a/Engine/Source/Programs/Shared/EpicGames.UHT/Types/Properties/UhtStructProperty.cs
+++ b/Engine/Source/Programs/Shared/EpicGames.UHT/Types/Properties/UhtStructProperty.cs
@@ -209,6 +209,7 @@ namespace EpicGames.UHT.Types
 		#region Structure default value sanitizers
 		[UhtStructDefaultValue(Name = "FVector")]
 		[SuppressMessage("CodeQuality", "IDE0051:Remove unused private members", Justification = "Attribute accessed method")]
+		[SuppressMessage("Style", "IDE0060:Remove unused parameter", Justification = "Attribute accessed method")]
 		private static bool VectorStructDefaultValue(UhtStructProperty Property, IUhtTokenReader DefaultValueReader, StringBuilder InnerDefaultValue)
 		{
 			string Format = "{0:F6},{1:F6},{2:F6}";
@@ -250,6 +251,7 @@ namespace EpicGames.UHT.Types
 
 		[UhtStructDefaultValue(Name = "FRotator")]
 		[SuppressMessage("CodeQuality", "IDE0051:Remove unused private members", Justification = "Attribute accessed method")]
+		[SuppressMessage("Style", "IDE0060:Remove unused parameter", Justification = "Attribute accessed method")]
 		private static bool RotatorStructDefaultValue(UhtStructProperty Property, IUhtTokenReader DefaultValueReader, StringBuilder InnerDefaultValue)
 		{
 			DefaultValueReader.Require("FRotator");
@@ -283,6 +285,7 @@ namespace EpicGames.UHT.Types
 
 		[UhtStructDefaultValue(Name = "FVector2D")]
 		[SuppressMessage("CodeQuality", "IDE0051:Remove unused private members", Justification = "Attribute accessed method")]
+		[SuppressMessage("Style", "IDE0060:Remove unused parameter", Justification = "Attribute accessed method")]
 		private static bool Vector2DStructDefaultValue(UhtStructProperty Property, IUhtTokenReader DefaultValueReader, StringBuilder InnerDefaultValue)
 		{
 			string Format = "(X={0:F3},Y={1:F3})";
@@ -317,6 +320,7 @@ namespace EpicGames.UHT.Types
 
 		[UhtStructDefaultValue(Name = "FLinearColor")]
 		[SuppressMessage("CodeQuality", "IDE0051:Remove unused private members", Justification = "Attribute accessed method")]
+		[SuppressMessage("Style", "IDE0060:Remove unused parameter", Justification = "Attribute accessed method")]
 		private static bool LinearColorStructDefaultValue(UhtStructProperty Property, IUhtTokenReader DefaultValueReader, StringBuilder InnerDefaultValue)
 		{
 			string Format = "(R={0:F6},G={1:F6},B={2:F6},A={3:F6})";
@@ -364,6 +368,7 @@ namespace EpicGames.UHT.Types
 
 		[UhtStructDefaultValue(Name = "FColor")]
 		[SuppressMessage("CodeQuality", "IDE0051:Remove unused private members", Justification = "Attribute accessed method")]
+		[SuppressMessage("Style", "IDE0060:Remove unused parameter", Justification = "Attribute accessed method")]
 		private static bool ColorStructDefaultValue(UhtStructProperty Property, IUhtTokenReader DefaultValueReader, StringBuilder InnerDefaultValue)
 		{
 			string Format = "(R={0},G={1},B={2},A={3})";
@@ -411,6 +416,7 @@ namespace EpicGames.UHT.Types
 
 		[UhtStructDefaultValue(Options = UhtStructDefaultValueOptions.Default)]
 		[SuppressMessage("CodeQuality", "IDE0051:Remove unused private members", Justification = "Attribute accessed method")]
+		[SuppressMessage("Style", "IDE0060:Remove unused parameter", Justification = "Attribute accessed method")]
 		private static bool DefaultStructDefaultValue(UhtStructProperty Property, IUhtTokenReader DefaultValueReader, StringBuilder InnerDefaultValue)
 		{
 			DefaultValueReader
diff --git a/Engine/Source/Programs/Shared/EpicGames.UHT/Types/Properties/UhtTextProperty.cs b/Engine/Source/Programs/Shared/EpicGames.UHT/Types/Properties/UhtTextProperty.cs
--- a/Engine/Source/Programs/Shared/EpicGames.UHT/Types/Properties/UhtTextProperty.cs
+++ b/Engine/Source/Programs/Shared/EpicGames.UHT/Types/Properties/UhtTextProperty.cs
@@ -118,6 +118,7 @@ namespace EpicGames.UHT.Types
 		#region Parsing keywords and default parsers		
 		[UhtPropertyType(Keyword = "FText", Options = UhtPropertyTypeOptions.Simple | UhtPropertyTypeOptions.Immediate)]
 		[SuppressMessage("CodeQuality", "IDE0051:Remove unused private members", Justification = "Attribute accessed method")]
+		[SuppressMessage("Style", "IDE0060:Remove unused parameter", Justification = "Attribute accessed method")]
 		private static UhtProperty? TextProperty(UhtPropertyResolvePhase ResolvePhase, UhtPropertySettings PropertySettings, IUhtTokenReader TokenReader, UhtToken MatchedToken)
 		{
 			return new UhtTextProperty(PropertySettings);
@@ -125,6 +126,7 @@ namespace EpicGames.UHT.Types
 
 		[UhtPropertyType(Keyword = "Text", Options = UhtPropertyTypeOptions.Simple | UhtPropertyTypeOptions.Immediate | UhtPropertyTypeOptions.CaseInsensitive)]
 		[SuppressMessage("CodeQuality", "IDE0051:Remove unused private members", Justification = "Attribute accessed method")]
+		[SuppressMessage("Style", "IDE0060:Remove unused parameter", Justification = "Attribute accessed method")]
 		private static UhtProperty? MissingPrefixTextProperty(UhtPropertyResolvePhase ResolvePhase, UhtPropertySettings PropertySettings, IUhtTokenReader TokenReader, UhtToken MatchedToken)
 		{
 			throw new UhtException(TokenReader, "'Text' is missing a prefix, expecting 'FText'");
@@ -132,6 +134,7 @@ namespace EpicGames.UHT.Types
 
 		[UhtLocTextDefaultValue(Name = "INVTEXT")]
 		[SuppressMessage("CodeQuality", "IDE0051:Remove unused private members", Justification = "Attribute accessed method")]
+		[SuppressMessage("Style", "IDE0060:Remove unused parameter", Justification = "Attribute accessed method")]
 		private static bool InvTextDefaultValue(UhtTextProperty Property, IUhtTokenReader DefaultValueReader, ref UhtToken MacroToken, StringBuilder InnerDefaultValue)
 		{
 			DefaultValueReader.Require('(');
@@ -143,6 +146,7 @@ namespace EpicGames.UHT.Types
 
 		[UhtLocTextDefaultValue(Name = "LOCTEXT")]
 		[SuppressMessage("CodeQuality", "IDE0051:Remove unused private members", Justification = "Attribute accessed method")]
+		[SuppressMessage("Style", "IDE0060:Remove unused parameter", Justification = "Attribute accessed method")]
 		private static bool LocTextDefaultValue(UhtTextProperty Property, IUhtTokenReader DefaultValueReader, ref UhtToken MacroToken, StringBuilder InnerDefaultValue)
 		{
 			Property.LogError($"LOCTEXT default parameter values are not supported; use NSLOCTEXT instead: {Property.SourceName}");
@@ -151,6 +155,7 @@ namespace EpicGames.UHT.Types
 
 		[UhtLocTextDefaultValue(Name = "NSLOCTEXT")]
 		[SuppressMessage("CodeQuality", "IDE0051:Remove unused private members", Justification = "Attribute accessed method")]
+		[SuppressMessage("Style", "IDE0060:Remove unused parameter", Justification = "Attribute accessed method")]
 		private static bool NsLocTextDefaultValue(UhtTextProperty Property, IUhtTokenReader DefaultValueReader, ref UhtToken MacroToken, StringBuilder InnerDefaultValue)
 		{
 			DefaultValueReader.Require('(');
@@ -185,6 +190,7 @@ namespace EpicGames.UHT.Types
 
 		[UhtLocTextDefaultValue(Name = "LOCTABLE")]
 		[SuppressMessage("CodeQuality", "IDE0051:Remove unused private members", Justification = "Attribute accessed method")]
+		[SuppressMessage("Style", "IDE0060:Remove unused parameter", Justification = "Attribute accessed method")]
 		private static bool LocTableDefaultValue(UhtTextProperty Property, IUhtTokenReader DefaultValueReader, ref UhtToken MacroToken, StringBuilder InnerDefaultValue)
 		{
 			DefaultValueReader.Require('(');
@@ -200,6 +206,7 @@ namespace EpicGames.UHT.Types
 
 		[UhtLocTextDefaultValue(Name = "LOCGEN_FORMAT_NAMED")]
 		[SuppressMessage("CodeQuality", "IDE0051:Remove unused private members", Justification = "Attribute accessed method")]
+		[SuppressMessage("Style", "IDE0060:Remove unused parameter", Justification = "Attribute accessed method")]
 		private static bool LocGenFormatNamedDefaultValue(UhtTextProperty Property, IUhtTokenReader DefaultValueReader, ref UhtToken MacroToken, StringBuilder InnerDefaultValue)
 		{
 			InnerDefaultValue.Append("LOCGEN_FORMAT_NAMED(");
@@ -229,6 +236,7 @@ namespace EpicGames.UHT.Types
 
 		[UhtLocTextDefaultValue(Name = "LOCGEN_FORMAT_ORDERED")]
 		[SuppressMessage("CodeQuality", "IDE0051:Remove unused private members", Justification = "Attribute accessed method")]
+		[SuppressMessage("Style", "IDE0060:Remove unused parameter", Justification = "Attribute accessed method")]
 		private static bool LocGenFormatOrderedDefaultValue(UhtTextProperty Property, IUhtTokenReader DefaultValueReader, ref UhtToken MacroToken, StringBuilder InnerDefaultValue)
 		{
 			InnerDefaultValue.Append("LOCGEN_FORMAT_ORDERED(");
@@ -256,62 +264,71 @@ namespace EpicGames.UHT.Types
 
 		[UhtLocTextDefaultValue(Name = "LOCGEN_NUMBER")]
 		[SuppressMessage("CodeQuality", "IDE0051:Remove unused private members", Justification = "Attribute accessed method")]
+		[SuppressMessage("Style", "IDE0060:Remove unused parameter", Justification = "Attribute accessed method")]
 		private static bool LocGenNumberDefaultValue(UhtTextProperty Property, IUhtTokenReader DefaultValueReader, ref UhtToken MacroToken, StringBuilder InnerDefaultValue)
 		{
-			return LocGenNumberOrPercentDefaultValue(Property, DefaultValueReader, ref MacroToken, InnerDefaultValue, NumberStyle.None);
+			return LocGenNumberOrPercentDefaultValue(DefaultValueReader, ref MacroToken, InnerDefaultValue, NumberStyle.None);
 		}
 
 		[UhtLocTextDefaultValue(Name = "LOCGEN_NUMBER_GROUPED")]
 		[SuppressMessage("CodeQuality", "IDE0051:Remove unused private members", Justification = "Attribute accessed method")]
+		[SuppressMessage("Style", "IDE0060:Remove unused parameter", Justification = "Attribute accessed method")]
 		private static bool LocGenNumberGroupedDefaultValue(UhtTextProperty Property, IUhtTokenReader DefaultValueReader, ref UhtToken MacroToken, StringBuilder InnerDefaultValue)
 		{
-			return LocGenNumberOrPercentDefaultValue(Property, DefaultValueReader, ref MacroToken, InnerDefaultValue, NumberStyle.Grouped);
+			return LocGenNumberOrPercentDefaultValue(DefaultValueReader, ref MacroToken, InnerDefaultValue, NumberStyle.Grouped);
 		}
 
 		[UhtLocTextDefaultValue(Name = "LOCGEN_NUMBER_UNGROUPED")]
 		[SuppressMessage("CodeQuality", "IDE0051:Remove unused private members", Justification = "Attribute accessed method")]
+		[SuppressMessage("Style", "IDE0060:Remove unused parameter", Justification = "Attribute accessed method")]
 		private static bool LocGenNumberUngroupedDefaultValue(UhtTextProperty Property, IUhtTokenReader DefaultValueReader, ref UhtToken MacroToken, StringBuilder InnerDefaultValue)
 		{
-			return LocGenNumberOrPercentDefaultValue(Property, DefaultValueReader, ref MacroToken, InnerDefaultValue, NumberStyle.Ungrouped);
+			return LocGenNumberOrPercentDefaultValue(DefaultValueReader, ref MacroToken, InnerDefaultValue, NumberStyle.Ungrouped);
 		}
 
 		[UhtLocTextDefaultValue(Name = "LOCGEN_NUMBER_CUSTOM")]
 		[SuppressMessage("CodeQuality", "IDE0051:Remove unused private members", Justification = "Attribute accessed method")]
+		[SuppressMessage("Style", "IDE0060:Remove unused parameter", Justification = "Attribute accessed method")]
 		private static bool LocGenNumberCustomDefaultValue(UhtTextProperty Property, IUhtTokenReader DefaultValueReader, ref UhtToken MacroToken, StringBuilder InnerDefaultValue)
 		{
-			return LocGenNumberOrPercentDefaultValue(Property, DefaultValueReader, ref MacroToken, InnerDefaultValue, NumberStyle.Custom);
+			return LocGenNumberOrPercentDefaultValue(DefaultValueReader, ref MacroToken, InnerDefaultValue, NumberStyle.Custom);
 		}
 
 		[UhtLocTextDefaultValue(Name = "LOCGEN_PERCENT")]
 		[SuppressMessage("CodeQuality", "IDE0051:Remove unused private members", Justification = "Attribute accessed method")]
+		[SuppressMessage("Style", "IDE0060:Remove unused parameter", Justification = "Attribute accessed method")]
 		private static bool LocGenPercentDefaultValue(UhtTextProperty Property, IUhtTokenReader DefaultValueReader, ref UhtToken MacroToken, StringBuilder InnerDefaultValue)
 		{
-			return LocGenNumberOrPercentDefaultValue(Property, DefaultValueReader, ref MacroToken, InnerDefaultValue, NumberStyle.None);
+			return LocGenNumberOrPercentDefaultValue(DefaultValueReader, ref MacroToken, InnerDefaultValue, NumberStyle.None);
 		}
 
 		[UhtLocTextDefaultValue(Name = "LOCGEN_PERCENT_GROUPED")]
 		[SuppressMessage("CodeQuality", "IDE0051:Remove unused private members", Justification = "Attribute accessed method")]
+		[SuppressMessage("Style", "IDE0060:Remove unused parameter", Justification = "Attribute accessed method")]
 		private static bool LocGenPercentGroupedDefaultValue(UhtTextProperty Property, IUhtTokenReader DefaultValueReader, ref UhtToken MacroToken, StringBuilder InnerDefaultValue)
 		{
-			return LocGenNumberOrPercentDefaultValue(Property, DefaultValueReader, ref MacroToken, InnerDefaultValue, NumberStyle.Grouped);
+			return LocGenNumberOrPercentDefaultValue(DefaultValueReader, ref MacroToken, InnerDefaultValue, NumberStyle.Grouped);
 		}
 
 		[UhtLocTextDefaultValue(Name = "LOCGEN_PERCENT_UNGROUPED")]
 		[SuppressMessage("CodeQuality", "IDE0051:Remove unused private members", Justification = "Attribute accessed method")]
+		[SuppressMessage("Style", "IDE0060:Remove unused parameter", Justification = "Attribute accessed method")]
 		private static bool LocGenPercentUngroupedDefaultValue(UhtTextProperty Property, IUhtTokenReader DefaultValueReader, ref UhtToken MacroToken, StringBuilder InnerDefaultValue)
 		{
-			return LocGenNumberOrPercentDefaultValue(Property, DefaultValueReader, ref MacroToken, InnerDefaultValue, NumberStyle.Ungrouped);
+			return LocGenNumberOrPercentDefaultValue(DefaultValueReader, ref MacroToken, InnerDefaultValue, NumberStyle.Ungrouped);
 		}
 
 		[UhtLocTextDefaultValue(Name = "LOCGEN_PERCENT_CUSTOM")]
 		[SuppressMessage("CodeQuality", "IDE0051:Remove unused private members", Justification = "Attribute accessed method")]
+		[SuppressMessage("Style", "IDE0060:Remove unused parameter", Justification = "Attribute accessed method")]
 		private static bool LocGenPercentCustomDefaultValue(UhtTextProperty Property, IUhtTokenReader DefaultValueReader, ref UhtToken MacroToken, StringBuilder InnerDefaultValue)
 		{
-			return LocGenNumberOrPercentDefaultValue(Property, DefaultValueReader, ref MacroToken, InnerDefaultValue, NumberStyle.Custom);
+			return LocGenNumberOrPercentDefaultValue(DefaultValueReader, ref MacroToken, InnerDefaultValue, NumberStyle.Custom);
 		}
 
 		[UhtLocTextDefaultValue(Name = "LOCGEN_CURRENCY")]
 		[SuppressMessage("CodeQuality", "IDE0051:Remove unused private members", Justification = "Attribute accessed method")]
+		[SuppressMessage("Style", "IDE0060:Remove unused parameter", Justification = "Attribute accessed method")]
 		private static bool LocGenCurrencyDefaultValue(UhtTextProperty Property, IUhtTokenReader DefaultValueReader, ref UhtToken MacroToken, StringBuilder InnerDefaultValue)
 		{
 			DefaultValueReader.Require('(');
@@ -336,62 +353,71 @@ namespace EpicGames.UHT.Types
 
 		[UhtLocTextDefaultValue(Name = "LOCGEN_DATE_LOCAL")]
 		[SuppressMessage("CodeQuality", "IDE0051:Remove unused private members", Justification = "Attribute accessed method")]
+		[SuppressMessage("Style", "IDE0060:Remove unused parameter", Justification = "Attribute accessed method")]
 		private static bool LocGenDateLocalDefaultValue(UhtTextProperty Property, IUhtTokenReader DefaultValueReader, ref UhtToken MacroToken, StringBuilder InnerDefaultValue)
 		{
-			return LocGenDateTimeDefaultValue(Property, DefaultValueReader, ref MacroToken, InnerDefaultValue, true, false, false, false);
+			return LocGenDateTimeDefaultValue(DefaultValueReader, ref MacroToken, InnerDefaultValue, true, false, false, false);
 		}
 
 		[UhtLocTextDefaultValue(Name = "LOCGEN_DATE_UTC")]
 		[SuppressMessage("CodeQuality", "IDE0051:Remove unused private members", Justification = "Attribute accessed method")]
+		[SuppressMessage("Style", "IDE0060:Remove unused parameter", Justification = "Attribute accessed method")]
 		private static bool LocGenDateUtcDefaultValue(UhtTextProperty Property, IUhtTokenReader DefaultValueReader, ref UhtToken MacroToken, StringBuilder InnerDefaultValue)
 		{
-			return LocGenDateTimeDefaultValue(Property, DefaultValueReader, ref MacroToken, InnerDefaultValue, true, false, true, false);
+			return LocGenDateTimeDefaultValue(DefaultValueReader, ref MacroToken, InnerDefaultValue, true, false, true, false);
 		}
 
 		[UhtLocTextDefaultValue(Name = "LOCGEN_TIME_LOCAL")]
 		[SuppressMessage("CodeQuality", "IDE0051:Remove unused private members", Justification = "Attribute accessed method")]
+		[SuppressMessage("Style", "IDE0060:Remove unused parameter", Justification = "Attribute accessed method")]
 		private static bool LocGenTimeLocalDefaultValue(UhtTextProperty Property, IUhtTokenReader DefaultValueReader, ref UhtToken MacroToken, StringBuilder InnerDefaultValue)
 		{
-			return LocGenDateTimeDefaultValue(Property, DefaultValueReader, ref MacroToken, InnerDefaultValue, false, true, false, false);
+			return LocGenDateTimeDefaultValue(DefaultValueReader, ref MacroToken, InnerDefaultValue, false, true, false, false);
 		}
 
 		[UhtLocTextDefaultValue(Name = "LOCGEN_TIME_UTC")]
 		[SuppressMessage("CodeQuality", "IDE0051:Remove unused private members", Justification = "Attribute accessed method")]
+		[SuppressMessage("Style", "IDE0060:Remove unused parameter", Justification = "Attribute accessed method")]
 		private static bool LocGenTimeUtcDefaultValue(UhtTextProperty Property, IUhtTokenReader DefaultValueReader, ref UhtToken MacroToken, StringBuilder InnerDefaultValue)
 		{
-			return LocGenDateTimeDefaultValue(Property, DefaultValueReader, ref MacroToken, InnerDefaultValue, false, true, true, false);
+			return LocGenDateTimeDefaultValue(DefaultValueReader, ref MacroToken, InnerDefaultValue, false, true, true, false);
 		}
 
 		[UhtLocTextDefaultValue(Name = "LOCGEN_DATETIME_LOCAL")]
 		[SuppressMessage("CodeQuality", "IDE0051:Remove unused private members", Justification = "Attribute accessed method")]
+		[SuppressMessage("Style", "IDE0060:Remove unused parameter", Justification = "Attribute accessed method")]
 		private static bool LocGenDateTimeLocalDefaultValue(UhtTextProperty Property, IUhtTokenReader DefaultValueReader, ref UhtToken MacroToken, StringBuilder InnerDefaultValue)
 		{
-			return LocGenDateTimeDefaultValue(Property, DefaultValueReader, ref MacroToken, InnerDefaultValue, true, true, false, false);
+			return LocGenDateTimeDefaultValue(DefaultValueReader, ref MacroToken, InnerDefaultValue, true, true, false, false);
 		}
 
 		[UhtLocTextDefaultValue(Name = "LOCGEN_DATETIME_UTC")]
 		[SuppressMessage("CodeQuality", "IDE0051:Remove unused private members", Justification = "Attribute accessed method")]
+		[SuppressMessage("Style", "IDE0060:Remove unused parameter", Justification = "Attribute accessed method")]
 		private static bool LocGenDateTimeUtcDefaultValue(UhtTextProperty Property, IUhtTokenReader DefaultValueReader, ref UhtToken MacroToken, StringBuilder InnerDefaultValue)
 		{
-			return LocGenDateTimeDefaultValue(Property, DefaultValueReader, ref MacroToken, InnerDefaultValue, true, true, true, false);
+			return LocGenDateTimeDefaultValue(DefaultValueReader, ref MacroToken, InnerDefaultValue, true, true, true, false);
 		}
 
 		[UhtLocTextDefaultValue(Name = "LOCGEN_DATETIME_CUSTOM_LOCAL")]
 		[SuppressMessage("CodeQuality", "IDE0051:Remove unused private members", Justification = "Attribute accessed method")]
+		[SuppressMessage("Style", "IDE0060:Remove unused parameter", Justification = "Attribute accessed method")]
 		private static bool LocGenDateTimeCustomLocalDefaultValue(UhtTextProperty Property, IUhtTokenReader DefaultValueReader, ref UhtToken MacroToken, StringBuilder InnerDefaultValue)
 		{
-			return LocGenDateTimeDefaultValue(Property, DefaultValueReader, ref MacroToken, InnerDefaultValue, true, true, false, true);
+			return LocGenDateTimeDefaultValue(DefaultValueReader, ref MacroToken, InnerDefaultValue, true, true, false, true);
 		}
 
 		[UhtLocTextDefaultValue(Name = "LOCGEN_DATETIME_CUSTOM_UTC")]
 		[SuppressMessage("CodeQuality", "IDE0051:Remove unused private members", Justification = "Attribute accessed method")]
+		[SuppressMessage("Style", "IDE0060:Remove unused parameter", Justification = "Attribute accessed method")]
 		private static bool LocGenDateTimeCustomUtcDefaultValue(UhtTextProperty Property, IUhtTokenReader DefaultValueReader, ref UhtToken MacroToken, StringBuilder InnerDefaultValue)
 		{
-			return LocGenDateTimeDefaultValue(Property, DefaultValueReader, ref MacroToken, InnerDefaultValue, true, true, true, true);
+			return LocGenDateTimeDefaultValue(DefaultValueReader, ref MacroToken, InnerDefaultValue, true, true, true, true);
 		}
 
 		[UhtLocTextDefaultValue(Name = "LOCGEN_TOUPPER")]
 		[SuppressMessage("CodeQuality", "IDE0051:Remove unused private members", Justification = "Attribute accessed method")]
+		[SuppressMessage("Style", "IDE0060:Remove unused parameter", Justification = "Attribute accessed method")]
 		private static bool LocGenToUpperDefaultValue(UhtTextProperty Property, IUhtTokenReader DefaultValueReader, ref UhtToken MacroToken, StringBuilder InnerDefaultValue)
 		{
 			return LocGenTransformDefaultValue(Property, DefaultValueReader, ref MacroToken, InnerDefaultValue);
@@ -399,6 +425,7 @@ namespace EpicGames.UHT.Types
 
 		[UhtLocTextDefaultValue(Name = "LOCGEN_TOLOWER")]
 		[SuppressMessage("CodeQuality", "IDE0051:Remove unused private members", Justification = "Attribute accessed method")]
+		[SuppressMessage("Style", "IDE0060:Remove unused parameter", Justification = "Attribute accessed method")]
 		private static bool LocGenToLowerDefaultValue(UhtTextProperty Property, IUhtTokenReader DefaultValueReader, ref UhtToken MacroToken, StringBuilder InnerDefaultValue)
 		{
 			return LocGenTransformDefaultValue(Property, DefaultValueReader, ref MacroToken, InnerDefaultValue);
@@ -509,7 +536,7 @@ namespace EpicGames.UHT.Types
 			Custom,
 		}
 
-		private static bool LocGenNumberOrPercentDefaultValue(UhtTextProperty Property, IUhtTokenReader DefaultValueReader, ref UhtToken MacroToken, StringBuilder InnerDefaultValue, NumberStyle NumberStyle)
+		private static bool LocGenNumberOrPercentDefaultValue(IUhtTokenReader DefaultValueReader, ref UhtToken MacroToken, StringBuilder InnerDefaultValue, NumberStyle NumberStyle)
 		{
 			InnerDefaultValue.Append(MacroToken.Value.ToString());
 			DefaultValueReader.Require('(');
@@ -628,7 +655,7 @@ namespace EpicGames.UHT.Types
 			InnerDefaultValue.Append(", EDateTimeStyle::").Append(Identifier);
 		}
 
-		private static bool LocGenDateTimeDefaultValue(UhtTextProperty Property, IUhtTokenReader DefaultValueReader, ref UhtToken MacroToken, StringBuilder InnerDefaultValue, bool bDate, bool bTime, bool bUtc, bool bCustom)
+		private static bool LocGenDateTimeDefaultValue(IUhtTokenReader DefaultValueReader, ref UhtToken MacroToken, StringBuilder InnerDefaultValue, bool bDate, bool bTime, bool bUtc, bool bCustom)
 		{
 			DefaultValueReader.Require('(');
 			InnerDefaultValue.Append(MacroToken.Value).Append('(');
diff --git a/Engine/Source/Programs/Shared/EpicGames.UHT/Types/Properties/UhtUInt16Property.cs b/Engine/Source/Programs/Shared/EpicGames.UHT/Types/Properties/UhtUInt16Property.cs
--- a/Engine/Source/Programs/Shared/EpicGames.UHT/Types/Properties/UhtUInt16Property.cs
+++ b/Engine/Source/Programs/Shared/EpicGames.UHT/Types/Properties/UhtUInt16Property.cs
@@ -60,6 +60,7 @@ namespace EpicGames.UHT.Types
 		#region Keyword
 		[UhtPropertyType(Keyword = "uint16", Options = UhtPropertyTypeOptions.Simple | UhtPropertyTypeOptions.Immediate)]
 		[SuppressMessage("CodeQuality", "IDE0051:Remove unused private members", Justification = "Attribute accessed method")]
+		[SuppressMessage("Style", "IDE0060:Remove unused parameter", Justification = "Attribute accessed method")]
 		private static UhtProperty? UInt16Property(UhtPropertyResolvePhase ResolvePhase, UhtPropertySettings PropertySettings, IUhtTokenReader TokenReader, UhtToken MatchedToken)
 		{
 			if (PropertySettings.bIsBitfield)
diff --git a/Engine/Source/Programs/Shared/EpicGames.UHT/Types/Properties/UhtUInt32Property.cs b/Engine/Source/Programs/Shared/EpicGames.UHT/Types/Properties/UhtUInt32Property.cs
--- a/Engine/Source/Programs/Shared/EpicGames.UHT/Types/Properties/UhtUInt32Property.cs
+++ b/Engine/Source/Programs/Shared/EpicGames.UHT/Types/Properties/UhtUInt32Property.cs
@@ -63,6 +63,7 @@ namespace EpicGames.UHT.Types
 		#region Keyword
 		[UhtPropertyType(Keyword = "uint32", Options = UhtPropertyTypeOptions.Simple | UhtPropertyTypeOptions.Immediate)]
 		[SuppressMessage("CodeQuality", "IDE0051:Remove unused private members", Justification = "Attribute accessed method")]
+		[SuppressMessage("Style", "IDE0060:Remove unused parameter", Justification = "Attribute accessed method")]
 		private static UhtProperty? UInt32Property(UhtPropertyResolvePhase ResolvePhase, UhtPropertySettings PropertySettings, IUhtTokenReader TokenReader, UhtToken MatchedToken)
 		{
 			if (PropertySettings.bIsBitfield)
@@ -77,6 +78,7 @@ namespace EpicGames.UHT.Types
 
 		[UhtPropertyType(Keyword = "unsigned", Options = UhtPropertyTypeOptions.Immediate)]
 		[SuppressMessage("CodeQuality", "IDE0051:Remove unused private members", Justification = "Attribute accessed method")]
+		[SuppressMessage("Style", "IDE0060:Remove unused parameter", Justification = "Attribute accessed method")]
 		private static UhtProperty? UnsignedProperty(UhtPropertyResolvePhase ResolvePhase, UhtPropertySettings PropertySettings, IUhtTokenReader TokenReader, UhtToken MatchedToken)
 		{
 			TokenReader
diff --git a/Engine/Source/Programs/Shared/EpicGames.UHT/Types/Properties/UhtUInt64Property.cs b/Engine/Source/Programs/Shared/EpicGames.UHT/Types/Properties/UhtUInt64Property.cs
--- a/Engine/Source/Programs/Shared/EpicGames.UHT/Types/Properties/UhtUInt64Property.cs
+++ b/Engine/Source/Programs/Shared/EpicGames.UHT/Types/Properties/UhtUInt64Property.cs
@@ -60,6 +60,7 @@ namespace EpicGames.UHT.Types
 		#region Keyword
 		[UhtPropertyType(Keyword = "uint64", Options = UhtPropertyTypeOptions.Simple | UhtPropertyTypeOptions.Immediate)]
 		[SuppressMessage("CodeQuality", "IDE0051:Remove unused private members", Justification = "Attribute accessed method")]
+		[SuppressMessage("Style", "IDE0060:Remove unused parameter", Justification = "Attribute accessed method")]
 		private static UhtProperty? UInt64Property(UhtPropertyResolvePhase ResolvePhase, UhtPropertySettings PropertySettings, IUhtTokenReader TokenReader, UhtToken MatchedToken)
 		{
 			if (PropertySettings.bIsBitfield)
diff --git a/Engine/Source/Programs/Shared/EpicGames.UHT/Types/Properties/UhtVoidProperty.cs b/Engine/Source/Programs/Shared/EpicGames.UHT/Types/Properties/UhtVoidProperty.cs
--- a/Engine/Source/Programs/Shared/EpicGames.UHT/Types/Properties/UhtVoidProperty.cs
+++ b/Engine/Source/Programs/Shared/EpicGames.UHT/Types/Properties/UhtVoidProperty.cs
@@ -72,6 +72,7 @@ namespace EpicGames.UHT.Types
 		#region Keyword
 		[UhtPropertyType(Keyword = "void", Options = UhtPropertyTypeOptions.Simple | UhtPropertyTypeOptions.Immediate)]
 		[SuppressMessage("CodeQuality", "IDE0051:Remove unused private members", Justification = "Attribute accessed method")]
+		[SuppressMessage("Style", "IDE0060:Remove unused parameter", Justification = "Attribute accessed method")]
 		private static UhtProperty? VoidProperty(UhtPropertyResolvePhase ResolvePhase, UhtPropertySettings PropertySettings, IUhtTokenReader TokenReader, UhtToken MatchedToken)
 		{
 			if (PropertySettings.PropertyCategory != UhtPropertyCategory.Return)
diff --git a/Engine/Source/Programs/Shared/EpicGames.UHT/Types/Properties/UhtWeakObjectPtrProperty.cs b/Engine/Source/Programs/Shared/EpicGames.UHT/Types/Properties/UhtWeakObjectPtrProperty.cs
--- a/Engine/Source/Programs/Shared/EpicGames.UHT/Types/Properties/UhtWeakObjectPtrProperty.cs
+++ b/Engine/Source/Programs/Shared/EpicGames.UHT/Types/Properties/UhtWeakObjectPtrProperty.cs
@@ -111,6 +111,7 @@ namespace EpicGames.UHT.Types
 
 		[UhtPropertyType(Keyword = "TWeakObjectPtr")]
 		[SuppressMessage("CodeQuality", "IDE0051:Remove unused private members", Justification = "Attribute accessed method")]
+		[SuppressMessage("Style", "IDE0060:Remove unused parameter", Justification = "Attribute accessed method")]
 		private static UhtProperty? WeakObjectProperty(UhtPropertyResolvePhase ResolvePhase, UhtPropertySettings PropertySettings, IUhtTokenReader TokenReader, UhtToken MatchedToken)
 		{
 			UhtClass? PropertyClass = ParseTemplateObject(PropertySettings, TokenReader, MatchedToken, true);
@@ -124,6 +125,7 @@ namespace EpicGames.UHT.Types
 
 		[UhtPropertyType(Keyword = "TAutoWeakObjectPtr")]
 		[SuppressMessage("CodeQuality", "IDE0051:Remove unused private members", Justification = "Attribute accessed method")]
+		[SuppressMessage("Style", "IDE0060:Remove unused parameter", Justification = "Attribute accessed method")]
 		private static UhtProperty? AutoWeakObjectProperty(UhtPropertyResolvePhase ResolvePhase, UhtPropertySettings PropertySettings, IUhtTokenReader TokenReader, UhtToken MatchedToken)
 		{
 			UhtClass? PropertyClass = ParseTemplateObject(PropertySettings, TokenReader, MatchedToken, true);
diff --git a/Engine/Source/Programs/Shared/EpicGames.UHT/Utils/UhtException.cs b/Engine/Source/Programs/Shared/EpicGames.UHT/Utils/UhtException.cs
--- a/Engine/Source/Programs/Shared/EpicGames.UHT/Utils/UhtException.cs
+++ b/Engine/Source/Programs/Shared/EpicGames.UHT/Utils/UhtException.cs
@@ -84,18 +84,17 @@ namespace EpicGames.UHT.Utils
 		public UhtTokenException(IUhtMessageSite MessageSite, UhtToken Got, object? Expected, object? ExtraContext = null)
 		{
 			string Message = Expected != null 
-				? $"Found {UhtMessage.FormatContext(Got)} when expecting {UhtMessage.FormatContext(Expected)}{FormatExtraContext(MessageSite, ExtraContext)}"
-				: $"Found {UhtMessage.FormatContext(Got)}{FormatExtraContext(MessageSite, ExtraContext)}";
+				? $"Found {UhtMessage.FormatContext(Got)} when expecting {UhtMessage.FormatContext(Expected)}{FormatExtraContext(ExtraContext)}"
+				: $"Found {UhtMessage.FormatContext(Got)}{FormatExtraContext(ExtraContext)}";
 			this.UhtMessage = UhtMessage.MakeMessage(UhtMessageType.Error, MessageSite.MessageSource, null, Got.InputLine, Message);
 		}
 
 		/// <summary>
 		/// Format any extra context supplied by the caller or the message site
 		/// </summary>
-		/// <param name="MessageSite">Message site that may contain extra context</param>
 		/// <param name="ExtraContext">Additional caller supplied context</param>
 		/// <returns></returns>
-		private static string FormatExtraContext(IUhtMessageSite MessageSite, object? ExtraContext = null)
+		private static string FormatExtraContext(object? ExtraContext = null)
 		{
 			StringBuilder Builder = new StringBuilder(" while parsing ");
 			int StartingLength = Builder.Length;
diff --git a/Engine/Source/Programs/Shared/EpicGames.UHT/Utils/UhtMessage.cs b/Engine/Source/Programs/Shared/EpicGames.UHT/Utils/UhtMessage.cs
--- a/Engine/Source/Programs/Shared/EpicGames.UHT/Utils/UhtMessage.cs
+++ b/Engine/Source/Programs/Shared/EpicGames.UHT/Utils/UhtMessage.cs
@@ -540,9 +540,8 @@ namespace EpicGames.UHT.Utils
 		/// <summary>
 		/// Construct a new entry
 		/// </summary>
-		/// <param name="MessageSite">Destination message site</param>
 		/// <param name="ExtraContext">Extra context to be added</param>
-		public UhtMessageContext(IUhtMessageSite MessageSite, object? ExtraContext)
+		public UhtMessageContext(object? ExtraContext)
 		{
 			this.Stack = UhtTlsMessageExtraContext.GetTls();
 			if (this.Stack != null)
diff --git a/Engine/Source/Programs/Shared/EpicGames.UHT/Utils/UhtTables.cs b/Engine/Source/Programs/Shared/EpicGames.UHT/Utils/UhtTables.cs
--- a/Engine/Source/Programs/Shared/EpicGames.UHT/Utils/UhtTables.cs
+++ b/Engine/Source/Programs/Shared/EpicGames.UHT/Utils/UhtTables.cs
@@ -608,7 +608,7 @@ namespace EpicGames.UHT.Utils
 					}
 					else if (ClassAttribute is UhtEngineClassAttribute EngineClassAttribute)
 					{
-						this.EngineClassTable.OnEngineClassAttribute(Type, EngineClassAttribute);
+						this.EngineClassTable.OnEngineClassAttribute(EngineClassAttribute);
 					}
 				}
 			}
@@ -662,11 +662,11 @@ namespace EpicGames.UHT.Utils
 					}
 					else if (MethodAttribute is UhtPropertyTypeAttribute PropertyTypeAttribute)
 					{
-						this.PropertyTypeTable.OnPropertyTypeAttribute(Type, MethodInfo, PropertyTypeAttribute);
+						this.PropertyTypeTable.OnPropertyTypeAttribute(MethodInfo, PropertyTypeAttribute);
 					}
 					else if (MethodAttribute is UhtStructDefaultValueAttribute StructDefaultValueAttribute)
 					{
-						this.StructDefaultValueTable.OnStructDefaultValueAttribute(Type, MethodInfo, StructDefaultValueAttribute);
+						this.StructDefaultValueTable.OnStructDefaultValueAttribute(MethodInfo, StructDefaultValueAttribute);
 					}
 					else if (MethodAttribute is UhtExporterAttribute ExporterAttribute)
 					{
@@ -674,7 +674,7 @@ namespace EpicGames.UHT.Utils
 					}
 					else if (MethodAttribute is UhtLocTextDefaultValueAttribute LocTextDefaultValueAttribute)
 					{
-						this.LocTextDefaultValueTable.OnLocTextDefaultValueAttribute(Type, MethodInfo, LocTextDefaultValueAttribute);
+						this.LocTextDefaultValueTable.OnLocTextDefaultValueAttribute(MethodInfo, LocTextDefaultValueAttribute);
 					}
 				}
 			}
