BuildGraph: Allow specifying annotations directly on the node declaration.

diff --git a/Engine/Source/Programs/Shared/EpicGames.BuildGraph/BgSchema.cs b/Engine/Source/Programs/Shared/EpicGames.BuildGraph/BgSchema.cs
--- a/Engine/Source/Programs/Shared/EpicGames.BuildGraph/BgSchema.cs
+++ b/Engine/Source/Programs/Shared/EpicGames.BuildGraph/BgSchema.cs
@@ -560,6 +560,7 @@ namespace EpicGames.BuildGraph
 			extension.Attributes.Add(CreateSchemaAttribute("If", ScriptSchemaStandardType.BalancedString, XmlSchemaUse.Optional));
 			extension.Attributes.Add(CreateSchemaAttribute("RunEarly", ScriptSchemaStandardType.Boolean, XmlSchemaUse.Optional));
 			extension.Attributes.Add(CreateSchemaAttribute("NotifyOnWarnings", ScriptSchemaStandardType.Boolean, XmlSchemaUse.Optional));
+			extension.Attributes.Add(CreateSchemaAttribute("Annotations", ScriptSchemaStandardType.BalancedString, XmlSchemaUse.Optional));
 
 			XmlSchemaComplexContent contentModel = new XmlSchemaComplexContent();
 			contentModel.Content = extension;
diff --git a/Engine/Source/Programs/Shared/EpicGames.BuildGraph/BgScriptReader.cs b/Engine/Source/Programs/Shared/EpicGames.BuildGraph/BgScriptReader.cs
--- a/Engine/Source/Programs/Shared/EpicGames.BuildGraph/BgScriptReader.cs
+++ b/Engine/Source/Programs/Shared/EpicGames.BuildGraph/BgScriptReader.cs
@@ -1140,6 +1140,57 @@ namespace EpicGames.BuildGraph
 			return value.Split(new char[] { separator }).Select(x => x.Trim()).Where(x => x.Length > 0).ToArray();
 		}
 
+		/// <summary>
+		/// Parse a map of annotations
+		/// </summary>
+		/// <param name="element"></param>
+		/// <param name="name"></param>
+		/// <returns></returns>
+		protected Dictionary<string, string> ReadAnnotationsAttribute(BgScriptElement element, string name)
+		{
+			string[] pairs = ReadListAttribute(element, name);
+			return ParseAnnotations(element, pairs);
+		}
+
+		/// <summary>
+		/// Parse a map of annotations
+		/// </summary>
+		/// <param name="element"></param>
+		/// <param name="pairs"></param>
+		/// <returns></returns>
+		private Dictionary<string, string> ParseAnnotations(BgScriptElement element, string[] pairs)
+		{
+			// Find the annotations to apply
+			Dictionary<string, string> pairMap = new Dictionary<string, string>(StringComparer.OrdinalIgnoreCase);
+			foreach (string pair in pairs)
+			{
+				if (!String.IsNullOrWhiteSpace(pair))
+				{
+					int idx = pair.IndexOf('=');
+					if (idx < 0)
+					{
+						LogError(element, "Invalid annotation '{0}'", pair);
+						continue;
+					}
+
+					string key = pair.Substring(0, idx).Trim();
+					if (!Regex.IsMatch(key, @"[a-zA-Z0-9_\.]+"))
+					{
+						LogError(element, "Invalid annotation key '{0}'", pair);
+						continue;
+					}
+					if (pairMap.ContainsKey(key))
+					{
+						LogError(element, "Annotation key '{0}' was specified twice", key);
+						continue;
+					}
+
+					pairMap.Add(key, pair.Substring(idx + 1).Trim());
+				}
+			}
+			return pairMap;
+		}
+
 		/// <summary>
 		/// Reads an attribute from the given XML element, expands any properties in it, and parses it as a boolean.
 		/// </summary>
@@ -1777,6 +1828,7 @@ namespace EpicGames.BuildGraph
 				string[] tokenFileNames = ReadListAttribute(element, "Token");
 				bool bRunEarly = ReadBooleanAttribute(element, "RunEarly", false);
 				bool bNotifyOnWarnings = ReadBooleanAttribute(element, "NotifyOnWarnings", true);
+				Dictionary<string, string> annotations = ReadAnnotationsAttribute(element, "Annotations");
 
 				// Resolve all the inputs we depend on
 				HashSet<BgNodeOutput> inputs = ResolveInputReferences(element, requiresNames);
@@ -1844,6 +1896,10 @@ namespace EpicGames.BuildGraph
 					BgNode newNode = new BgNode(name, inputs.ToArray(), validOutputNames.ToArray(), inputDependencies.ToArray(), orderDependencies.ToArray(), requiredTokens.ToArray());
 					newNode.RunEarly = bRunEarly;
 					newNode.NotifyOnWarnings = bNotifyOnWarnings;
+					foreach ((string key, string value) in annotations)
+					{
+						newNode.Annotations[key] = value;
+					}
 					_graph.NameToNode.Add(name, newNode);
 
 					// Register all the output tags in the global name table.
@@ -2016,7 +2072,7 @@ namespace EpicGames.BuildGraph
 				string[] targetNames = ReadListAttribute(element, "Targets");
 				string[] exceptNames = ReadListAttribute(element, "Except");
 				string[] individualNodeNames = ReadListAttribute(element, "Nodes");
-				string[] pairs = ReadListAttribute(element, "Values");
+				Dictionary<string, string> annotations = ReadAnnotationsAttribute(element, "Values");
 
 				// Find the list of targets which are included, and recurse through all their dependencies
 				HashSet<BgNode> nodes = new HashSet<BgNode>();
@@ -2044,39 +2100,10 @@ namespace EpicGames.BuildGraph
 					nodes.ExceptWith(exceptNodes);
 				}
 
-				// Find the annotations to apply
-				Dictionary<string, string> pairMap = new Dictionary<string, string>(StringComparer.OrdinalIgnoreCase);
-				foreach (string pair in pairs)
-				{
-					if (!String.IsNullOrWhiteSpace(pair))
-					{
-						int idx = pair.IndexOf('=');
-						if (idx < 0)
-						{
-							LogError(element, "Invalid annotation '{0}'", pair);
-							continue;
-						}
-
-						string key = pair.Substring(0, idx).Trim();
-						if (!Regex.IsMatch(key, @"[a-zA-Z0-9_\.]+"))
-						{
-							LogError(element, "Invalid annotation key '{0}'", pair);
-							continue;
-						}
-						if (pairMap.ContainsKey(key))
-						{
-							LogError(element, "Annotation key '{0}' was specified twice", key);
-							continue;
-						}
-
-						pairMap.Add(key, pair.Substring(idx + 1).Trim());
-					}
-				}
-
 				// Update all the referenced nodes with the settings
 				foreach (BgNode node in nodes)
 				{
-					foreach ((string key, string value) in pairMap)
+					foreach ((string key, string value) in annotations)
 					{
 						node.Annotations[key] = value;
 					}
