More coding style issues fixed

#rnx

diff --git a/Engine/Source/Programs/Shared/EpicGames.UHT/.editorconfig b/Engine/Source/Programs/Shared/EpicGames.UHT/.editorconfig
--- a/Engine/Source/Programs/Shared/EpicGames.UHT/.editorconfig
+++ b/Engine/Source/Programs/Shared/EpicGames.UHT/.editorconfig
@@ -5,9 +5,6 @@
 # Default severity for all analyzer diagnostics
 dotnet_analyzer_diagnostic.severity = none
 
-# CA1051: Do not declare visible instance fields
-dotnet_diagnostic.CA1051.severity = none
-
 # CA1045: Do not pass types by reference
 dotnet_diagnostic.CA1045.severity = none
 
diff --git a/Engine/Source/Programs/Shared/EpicGames.UHT/Parsers/UhtClassBaseParser.cs b/Engine/Source/Programs/Shared/EpicGames.UHT/Parsers/UhtClassBaseParser.cs
--- a/Engine/Source/Programs/Shared/EpicGames.UHT/Parsers/UhtClassBaseParser.cs
+++ b/Engine/Source/Programs/Shared/EpicGames.UHT/Parsers/UhtClassBaseParser.cs
@@ -4,6 +4,7 @@ using EpicGames.Core;
 using EpicGames.UHT.Tokenizer;
 using EpicGames.UHT.Types;
 using System.Collections.Generic;
+using System.Text.Json.Serialization;
 
 namespace EpicGames.UHT.Parsers
 {
@@ -17,12 +18,15 @@ namespace EpicGames.UHT.Parsers
 		/// <summary>
 		/// The super class identifier
 		/// </summary>
-		public UhtToken SuperIdentifier;
+		[JsonIgnore]
+		public UhtToken SuperIdentifier { get; set; }
 
 		/// <summary>
 		/// List of base identifiers
 		/// </summary>
-		public List<UhtToken[]>? BaseIdentifiers = null;
+		[JsonIgnore]
+		[System.Diagnostics.CodeAnalysis.SuppressMessage("Usage", "CA2227:Collection properties should be read only", Justification = "")]
+		public List<UhtToken[]>? BaseIdentifiers { get; set; } = null;
 
 		/// <summary>
 		/// Construct a new base class parser
@@ -40,7 +44,7 @@ namespace EpicGames.UHT.Parsers
 			switch (ResolvePhase)
 			{
 				case UhtResolvePhase.Bases:
-					BindAndResolveSuper(ref this.SuperIdentifier, UhtFindOptions.Class);
+					BindAndResolveSuper(this.SuperIdentifier, UhtFindOptions.Class);
 					BindAndResolveBases(this.BaseIdentifiers, UhtFindOptions.Class);
 
 					// Force the MatchedSerializers on for anything being exported
diff --git a/Engine/Source/Programs/Shared/EpicGames.UHT/Parsers/UhtClassParser.cs b/Engine/Source/Programs/Shared/EpicGames.UHT/Parsers/UhtClassParser.cs
--- a/Engine/Source/Programs/Shared/EpicGames.UHT/Parsers/UhtClassParser.cs
+++ b/Engine/Source/Programs/Shared/EpicGames.UHT/Parsers/UhtClassParser.cs
@@ -21,57 +21,57 @@ namespace EpicGames.UHT.Parsers
 		/// <summary>
 		/// Engine class flags removed
 		/// </summary>
-		public EClassFlags RemovedClassFlags = EClassFlags.None;
+		public EClassFlags RemovedClassFlags { get; set; } = EClassFlags.None;
 
 		/// <summary>
 		/// Class within identifier
 		/// </summary>
-		public string ClassWithinIdentifier = String.Empty;
+		public string ClassWithinIdentifier { get; set; } = String.Empty;
 
 		/// <summary>
 		/// Collection of show categories
 		/// </summary>
-		public List<string> ShowCategories = new List<string>();
+		public List<string> ShowCategories { get; } = new List<string>();
 
 		/// <summary>
 		/// Collection of hide categories
 		/// </summary>
-		public List<string> HideCategories = new List<string>();
+		public List<string> HideCategories { get; } = new List<string>();
 
 		/// <summary>
 		/// Collection of auto expand categories
 		/// </summary>
-		public List<string> AutoExpandCategories = new List<string>();
+		public List<string> AutoExpandCategories { get; } = new List<string>();
 
 		/// <summary>
 		/// Collection of auto collapse categories
 		/// </summary>
-		public List<string> AutoCollapseCategories = new List<string>();
+		public List<string> AutoCollapseCategories { get; } = new List<string>();
 
 		/// <summary>
 		/// Collection of prioritize categories
 		/// </summary>
-		public List<string> PrioritizeCategories = new List<string>();
+		public List<string> PrioritizeCategories { get; } = new List<string>();
 
 		/// <summary>
 		/// Collection of show functions
 		/// </summary>
-		public List<string> ShowFunctions = new List<string>();
+		public List<string> ShowFunctions { get; } = new List<string>();
 
 		/// <summary>
 		/// Collection of hide functions
 		/// </summary>
-		public List<string> HideFunctions = new List<string>();
+		public List<string> HideFunctions { get; } = new List<string>();
 
 		/// <summary>
 		/// Sparse class data types
 		/// </summary>
-		public List<string> SparseClassDataTypes = new List<string>();
+		public List<string> SparseClassDataTypes { get; } = new List<string>();
 
 		/// <summary>
 		/// Class group names
 		/// </summary>
-		public List<string> ClassGroupNames = new List<string>();
+		public List<string> ClassGroupNames { get; } = new List<string>();
 
 		/// <summary>
 		/// Add the given class flags
@@ -238,7 +238,8 @@ namespace EpicGames.UHT.Parsers
 			}
 
 			// Replace the show categories
-			this.ShowCategories = OutShowCategories;
+			this.ShowCategories.Clear();
+			this.ShowCategories.AddRange(OutShowCategories);
 		}
 
 		private void SetAndValidateWithinClass(UhtResolvePhase ResolvePhase)
@@ -411,7 +412,7 @@ namespace EpicGames.UHT.Parsers
 				{
 					// Parse the specifiers
 					UhtSpecifierContext SpecifierContext = new UhtSpecifierContext(TopScope, TopScope.TokenReader, Class.MetaData);
-					UhtSpecifierParser Specifiers = TopScope.HeaderParser.GetSpecifierParser(SpecifierContext, ScopeName, ParentScope.Session.GetSpecifierTable(UhtTableNames.Class));
+					UhtSpecifierParser Specifiers = TopScope.HeaderParser.GetCachedSpecifierParser(SpecifierContext, ScopeName, ParentScope.Session.GetSpecifierTable(UhtTableNames.Class));
 					Specifiers.ParseSpecifiers();
 					Class.PrologLineNumber = TopScope.TokenReader.InputLine;
 					Class.ClassFlags |= EClassFlags.Native;
@@ -445,7 +446,11 @@ namespace EpicGames.UHT.Parsers
 					TopScope.TokenReader.Optional("final");
 
 					// Parse the inheritance
-					UhtParserHelpers.ParseInheritance(TopScope.TokenReader, TopScope.Session.Config!, out Class.SuperIdentifier, out Class.BaseIdentifiers);
+					UhtToken SuperIdentifier;
+					List<UhtToken[]>? BaseIdentifiers;
+					UhtParserHelpers.ParseInheritance(TopScope.TokenReader, TopScope.Session.Config!, out SuperIdentifier, out BaseIdentifiers);
+					Class.SuperIdentifier = SuperIdentifier;
+					Class.BaseIdentifiers = BaseIdentifiers;
 
 					if (APIMacroToken)
 					{
diff --git a/Engine/Source/Programs/Shared/EpicGames.UHT/Parsers/UhtEnumParser.cs b/Engine/Source/Programs/Shared/EpicGames.UHT/Parsers/UhtEnumParser.cs
--- a/Engine/Source/Programs/Shared/EpicGames.UHT/Parsers/UhtEnumParser.cs
+++ b/Engine/Source/Programs/Shared/EpicGames.UHT/Parsers/UhtEnumParser.cs
@@ -46,7 +46,7 @@ namespace EpicGames.UHT.Parsers
 
 					// Parse the specifiers
 					UhtSpecifierContext SpecifierContext = new UhtSpecifierContext(TopScope, TopScope.TokenReader, Enum.MetaData);
-					UhtSpecifierParser Specifiers = TopScope.HeaderParser.GetSpecifierParser(SpecifierContext, ScopeName, ParentScope.Session.GetSpecifierTable(UhtTableNames.Enum));
+					UhtSpecifierParser Specifiers = TopScope.HeaderParser.GetCachedSpecifierParser(SpecifierContext, ScopeName, ParentScope.Session.GetSpecifierTable(UhtTableNames.Enum));
 					Specifiers.ParseSpecifiers();
 
 					// Read the name and the CPP type
diff --git a/Engine/Source/Programs/Shared/EpicGames.UHT/Parsers/UhtFunctionParser.cs b/Engine/Source/Programs/Shared/EpicGames.UHT/Parsers/UhtFunctionParser.cs
--- a/Engine/Source/Programs/Shared/EpicGames.UHT/Parsers/UhtFunctionParser.cs
+++ b/Engine/Source/Programs/Shared/EpicGames.UHT/Parsers/UhtFunctionParser.cs
@@ -109,7 +109,7 @@ namespace EpicGames.UHT.Parsers
 		/// <summary>
 		/// True if the function specifier has a getter/setter specified
 		/// </summary>
-		public bool bSawPropertyAccessor = false;
+		public bool bSawPropertyAccessor { get; set; } = false;
 
 		/// <inheritdoc/>
 		protected override bool ResolveSelf(UhtResolvePhase ResolvePhase)
@@ -225,7 +225,7 @@ namespace EpicGames.UHT.Parsers
 					TopScope.AddModuleRelativePathToMetaData();
 
 					UhtSpecifierContext SpecifierContext = new UhtSpecifierContext(TopScope, TopScope.TokenReader, Function.MetaData);
-					UhtSpecifierParser Specifiers = TopScope.HeaderParser.GetSpecifierParser(SpecifierContext, ScopeName, ParentScope.Session.GetSpecifierTable(UhtTableNames.Function));
+					UhtSpecifierParser Specifiers = TopScope.HeaderParser.GetCachedSpecifierParser(SpecifierContext, ScopeName, ParentScope.Session.GetSpecifierTable(UhtTableNames.Function));
 
 					// If this is a UDELEGATE, parse the specifiers
 					StringView DelegateMacro = new StringView();
@@ -289,7 +289,7 @@ namespace EpicGames.UHT.Parsers
 					UhtProperty? ReturnValueProperty = null;
 					if (bHasReturnValue)
 					{
-						TopScope.HeaderParser.GetPropertyParser(TopScope).Parse(EPropertyFlags.None,
+						TopScope.HeaderParser.GetCachedPropertyParser(TopScope).Parse(EPropertyFlags.None,
 							Function.GetPropertyParseOptions(true), UhtParsePropertyDeclarationStyle.None, UhtPropertyCategory.Return,
 							(UhtParsingScope TopScope, UhtProperty Property, ref UhtToken NameToken, UhtLayoutMacroType LayoutMacroType) =>
 							{
@@ -397,7 +397,7 @@ namespace EpicGames.UHT.Parsers
 					TopScope.AddModuleRelativePathToMetaData();
 
 					UhtSpecifierContext SpecifierContext = new UhtSpecifierContext(TopScope, TopScope.TokenReader, Function.MetaData);
-					UhtSpecifierParser SpecifierParser = TopScope.HeaderParser.GetSpecifierParser(SpecifierContext, ScopeName, ParentScope.Session.GetSpecifierTable(UhtTableNames.Function));
+					UhtSpecifierParser SpecifierParser = TopScope.HeaderParser.GetCachedSpecifierParser(SpecifierContext, ScopeName, ParentScope.Session.GetSpecifierTable(UhtTableNames.Function));
 					SpecifierParser.ParseSpecifiers();
 
 					if (!OuterClass.ClassFlags.HasAnyFlags(EClassFlags.Native))
@@ -515,7 +515,7 @@ namespace EpicGames.UHT.Parsers
 						// Get return type.  C++ style functions always have a return value type, even if it's void
 						UhtToken FuncNameToken = new UhtToken();
 						UhtProperty? ReturnValueProperty = null;
-						TopScope.HeaderParser.GetPropertyParser(TopScope).Parse(EPropertyFlags.None,
+						TopScope.HeaderParser.GetCachedPropertyParser(TopScope).Parse(EPropertyFlags.None,
 							Function.GetPropertyParseOptions(true), UhtParsePropertyDeclarationStyle.None, UhtPropertyCategory.Return,
 							(UhtParsingScope TopScope, UhtProperty Property, ref UhtToken NameToken, UhtLayoutMacroType LayoutMacroType) =>
 							{
@@ -646,7 +646,7 @@ namespace EpicGames.UHT.Parsers
 
 			TopScope.TokenReader.RequireList(')', ',', false, () =>
 			{
-				TopScope.HeaderParser.GetPropertyParser(TopScope).Parse(DisallowFlags, Options, UhtParsePropertyDeclarationStyle.None, PropertyCategory,
+				TopScope.HeaderParser.GetCachedPropertyParser(TopScope).Parse(DisallowFlags, Options, UhtParsePropertyDeclarationStyle.None, PropertyCategory,
 					(UhtParsingScope TopScope, UhtProperty Property, ref UhtToken NameToken, UhtLayoutMacroType LayoutMacroType) =>
 					{
 						Property.PropertyFlags |= EPropertyFlags.Parm;
diff --git a/Engine/Source/Programs/Shared/EpicGames.UHT/Parsers/UhtHeaderFileParser.cs b/Engine/Source/Programs/Shared/EpicGames.UHT/Parsers/UhtHeaderFileParser.cs
--- a/Engine/Source/Programs/Shared/EpicGames.UHT/Parsers/UhtHeaderFileParser.cs
+++ b/Engine/Source/Programs/Shared/EpicGames.UHT/Parsers/UhtHeaderFileParser.cs
@@ -186,12 +186,12 @@ namespace EpicGames.UHT.Parsers
 		/// <summary>
 		/// Header file being parsed
 		/// </summary>
-		public readonly UhtHeaderFile HeaderFile;
+		public UhtHeaderFile HeaderFile { get; }
 
 		/// <summary>
 		/// Token reader for the header
 		/// </summary>
-		public readonly IUhtTokenReader TokenReader;
+		public IUhtTokenReader TokenReader { get; }
 
 		/// <summary>
 		/// If true, this header file belongs to the engine
@@ -201,18 +201,18 @@ namespace EpicGames.UHT.Parsers
 		/// <summary>
 		/// If true, the inclusion of the generated header file was seen
 		/// </summary>
-		public bool bSpottedAutogeneratedHeaderInclude = false;
+		public bool bSpottedAutogeneratedHeaderInclude { get; set; } = false;
 
 		/// <summary>
 		/// For a given header file, we share a common specifier parser to reduce the number of allocations.
 		/// Before the parser can be reused, the ParseDeferred method must be called to dispatch that list.
 		/// </summary>
-		private UhtSpecifierParser? SpecifierParser = null;
+		private UhtSpecifierParser? SpecifierParser { get; set; } = null;
 
 		/// <summary>
 		/// For a given header file, we share a common property parser to reduce the number of allocations.
 		/// </summary>
-		public UhtPropertyParser? PropertyParser = null;
+		public UhtPropertyParser? PropertyParser { get; set; } = null;
 
 		/// <summary>
 		/// Stack of current #if states
@@ -307,7 +307,7 @@ namespace EpicGames.UHT.Parsers
 		/// <param name="Context">User facing context</param>
 		/// <param name="Table">Specifier table</param>
 		/// <returns>Specifier parser</returns>
-		public UhtSpecifierParser GetSpecifierParser(UhtSpecifierContext SpecifierContext, StringView Context, UhtSpecifierTable Table)
+		public UhtSpecifierParser GetCachedSpecifierParser(UhtSpecifierContext SpecifierContext, StringView Context, UhtSpecifierTable Table)
 		{
 			if (this.SpecifierParser == null)
 			{
@@ -325,7 +325,7 @@ namespace EpicGames.UHT.Parsers
 		/// </summary>
 		/// <param name="TopScope">Top scope requesting parser</param>
 		/// <returns>Property parser</returns>
-		public UhtPropertyParser GetPropertyParser(UhtParsingScope TopScope)
+		public UhtPropertyParser GetCachedPropertyParser(UhtParsingScope TopScope)
 		{
 			if (this.PropertyParser == null)
 			{
diff --git a/Engine/Source/Programs/Shared/EpicGames.UHT/Parsers/UhtInterfaceClassParser.cs b/Engine/Source/Programs/Shared/EpicGames.UHT/Parsers/UhtInterfaceClassParser.cs
--- a/Engine/Source/Programs/Shared/EpicGames.UHT/Parsers/UhtInterfaceClassParser.cs
+++ b/Engine/Source/Programs/Shared/EpicGames.UHT/Parsers/UhtInterfaceClassParser.cs
@@ -119,7 +119,7 @@ namespace EpicGames.UHT.Parsers
 				{
 					// Parse the specifiers
 					UhtSpecifierContext SpecifierContext = new UhtSpecifierContext(TopScope, TopScope.TokenReader, Class.MetaData);
-					UhtSpecifierParser Specifiers = TopScope.HeaderParser.GetSpecifierParser(SpecifierContext, ScopeName, ParentScope.Session.GetSpecifierTable(UhtTableNames.Interface));
+					UhtSpecifierParser Specifiers = TopScope.HeaderParser.GetCachedSpecifierParser(SpecifierContext, ScopeName, ParentScope.Session.GetSpecifierTable(UhtTableNames.Interface));
 					Specifiers.ParseSpecifiers();
 					Class.PrologLineNumber = TopScope.TokenReader.InputLine;
 					Class.ClassFlags |= EClassFlags.Native | EClassFlags.Interface | EClassFlags.Abstract;
diff --git a/Engine/Source/Programs/Shared/EpicGames.UHT/Parsers/UhtParsingScope.cs b/Engine/Source/Programs/Shared/EpicGames.UHT/Parsers/UhtParsingScope.cs
--- a/Engine/Source/Programs/Shared/EpicGames.UHT/Parsers/UhtParsingScope.cs
+++ b/Engine/Source/Programs/Shared/EpicGames.UHT/Parsers/UhtParsingScope.cs
@@ -20,32 +20,32 @@ namespace EpicGames.UHT.Parsers
 		/// <summary>
 		/// Header file parser
 		/// </summary>
-		public readonly UhtHeaderFileParser HeaderParser;
+		public UhtHeaderFileParser HeaderParser { get; }
 
 		/// <summary>
 		/// Token reader
 		/// </summary>
-		public readonly IUhtTokenReader TokenReader;
+		public IUhtTokenReader TokenReader { get; }
 
 		/// <summary>
 		/// Parent scope
 		/// </summary>
-		public readonly UhtParsingScope? ParentScope;
+		public UhtParsingScope? ParentScope { get; }
 
 		/// <summary>
 		/// Type being parsed
 		/// </summary>
-		public readonly UhtType ScopeType;
+		public UhtType ScopeType { get; }
 
 		/// <summary>
 		/// Keyword table for the scope
 		/// </summary>
-		public readonly UhtKeywordTable ScopeKeywordTable;
+		public UhtKeywordTable ScopeKeywordTable { get; }
 
 		/// <summary>
 		/// Current access specifier
 		/// </summary>
-		public UhtAccessSpecifier AccessSpecifier = UhtAccessSpecifier.Public;
+		public UhtAccessSpecifier AccessSpecifier { get; set; } = UhtAccessSpecifier.Public;
 
 		/// <summary>
 		/// Current session
diff --git a/Engine/Source/Programs/Shared/EpicGames.UHT/Parsers/UhtPropertyParser.cs b/Engine/Source/Programs/Shared/EpicGames.UHT/Parsers/UhtPropertyParser.cs
--- a/Engine/Source/Programs/Shared/EpicGames.UHT/Parsers/UhtPropertyParser.cs
+++ b/Engine/Source/Programs/Shared/EpicGames.UHT/Parsers/UhtPropertyParser.cs
@@ -316,27 +316,27 @@ namespace EpicGames.UHT.Parsers
 		/// <summary>
 		/// The property settings being parsed
 		/// </summary>
-		public UhtPropertySettings PropertySettings;
+		public UhtPropertySettings PropertySettings { get; set; }
 
 		/// <summary>
 		/// If true, editor specifier seen
 		/// </summary>
-		public bool bSeenEditSpecifier = false;
+		public bool bSeenEditSpecifier { get; set; } = false;
 
 		/// <summary>
 		/// If true, blueprint write specifier seen
 		/// </summary>
-		public bool bSeenBlueprintWriteSpecifier = false;
+		public bool bSeenBlueprintWriteSpecifier { get; set; } = false;
 
 		/// <summary>
 		/// If true, blueprint readonly specifier seen
 		/// </summary>
-		public bool bSeenBlueprintReadOnlySpecifier = false;
+		public bool bSeenBlueprintReadOnlySpecifier { get; set; } = false;
 
 		/// <summary>
 		/// If true, blueprint getter specifier seen
 		/// </summary>
-		public bool bSeenBlueprintGetterSpecifier = false;
+		public bool bSeenBlueprintGetterSpecifier { get; set; } = false;
 
 		/// <summary>
 		/// Construct a new property specifier context
@@ -368,12 +368,12 @@ namespace EpicGames.UHT.Parsers
 		/// <summary>
 		/// Collection of type tokens 
 		/// </summary>
-		public ReadOnlyMemory<UhtToken> TypeTokens;
+		public ReadOnlyMemory<UhtToken> TypeTokens { get; set; }
 
 		/// <summary>
 		/// Property settings being parsed
 		/// </summary>
-		public UhtPropertySettings PropertySettings;
+		public UhtPropertySettings PropertySettings { get; set; }
 
 		/// <summary>
 		/// Construct a new property to be resolved
@@ -767,7 +767,7 @@ namespace EpicGames.UHT.Parsers
 
 			bool bIsParamList = this.SpecifierContext.PropertySettings.PropertyCategory != UhtPropertyCategory.Member && this.TokenReader.TryOptional("UPARAM");
 
-			UhtSpecifierParser Specifiers = this.TopScope.HeaderParser.GetSpecifierParser(this.SpecifierContext, "Variable", 
+			UhtSpecifierParser Specifiers = this.TopScope.HeaderParser.GetCachedSpecifierParser(this.SpecifierContext, "Variable", 
 				bIsParamList ? this.TopScope.Session.GetSpecifierTable(UhtTableNames.PropertyArgument) : this.TopScope.Session.GetSpecifierTable(UhtTableNames.PropertyMember));
 			if (DeclarationStyle == UhtParsePropertyDeclarationStyle.UPROPERTY || bIsParamList)
 			{
@@ -1105,7 +1105,7 @@ namespace EpicGames.UHT.Parsers
 			// Try gathering metadata for member fields
 			if (NewProperty.PropertyCategory == UhtPropertyCategory.Member)
 			{
-				UhtSpecifierParser Specifiers = this.TopScope.HeaderParser.GetSpecifierParser(this.SpecifierContext, NewProperty.SourceName, this.TopScope.Session.GetSpecifierTable(UhtTableNames.PropertyMember));
+				UhtSpecifierParser Specifiers = this.TopScope.HeaderParser.GetCachedSpecifierParser(this.SpecifierContext, NewProperty.SourceName, this.TopScope.Session.GetSpecifierTable(UhtTableNames.PropertyMember));
 				Specifiers.ParseFieldMetaData();
 				this.TopScope.AddFormattedCommentsAsTooltipMetaData(NewProperty);
 			}
@@ -1194,7 +1194,7 @@ namespace EpicGames.UHT.Parsers
 		private static UhtParseResult UPROPERTYKeyword(UhtParsingScope TopScope, UhtParsingScope ActionScope, ref UhtToken Token)
 		{
 			UhtPropertyParseOptions Options = UhtPropertyParseOptions.ParseLayoutMacro | UhtPropertyParseOptions.List | UhtPropertyParseOptions.AddModuleRelativePath;
-			TopScope.HeaderParser.GetPropertyParser(TopScope).Parse(EPropertyFlags.ParmFlags, Options, UhtParsePropertyDeclarationStyle.UPROPERTY, UhtPropertyCategory.Member, PropertyDelegate);
+			TopScope.HeaderParser.GetCachedPropertyParser(TopScope).Parse(EPropertyFlags.ParmFlags, Options, UhtParsePropertyDeclarationStyle.UPROPERTY, UhtPropertyCategory.Member, PropertyDelegate);
 			TopScope.TokenReader.Require(';');
 			return UhtParseResult.Handled;
 		}
diff --git a/Engine/Source/Programs/Shared/EpicGames.UHT/Parsers/UhtScriptStructParser.cs b/Engine/Source/Programs/Shared/EpicGames.UHT/Parsers/UhtScriptStructParser.cs
--- a/Engine/Source/Programs/Shared/EpicGames.UHT/Parsers/UhtScriptStructParser.cs
+++ b/Engine/Source/Programs/Shared/EpicGames.UHT/Parsers/UhtScriptStructParser.cs
@@ -7,6 +7,7 @@ using EpicGames.UHT.Types;
 using EpicGames.UHT.Utils;
 using System.Collections.Generic;
 using System.Diagnostics.CodeAnalysis;
+using System.Text.Json.Serialization;
 
 namespace EpicGames.UHT.Parsers
 {
@@ -20,12 +21,15 @@ namespace EpicGames.UHT.Parsers
 		/// <summary>
 		/// Super identifier
 		/// </summary>
-		public UhtToken SuperIdentifier;
+		[JsonIgnore]
+		public UhtToken SuperIdentifier { get; set; }
 
 		/// <summary>
 		/// Base identifiers
 		/// </summary>
-		public List<UhtToken[]>? BaseIdentifiers = null;
+		[JsonIgnore]
+		[SuppressMessage("Usage", "CA2227:Collection properties should be read only", Justification = "<Pending>")]
+		public List<UhtToken[]>? BaseIdentifiers { get; set; } = null;
 
 		/// <summary>
 		/// Construct a new instance
@@ -44,7 +48,7 @@ namespace EpicGames.UHT.Parsers
 			switch (ResolvePhase)
 			{
 				case UhtResolvePhase.Bases:
-					BindAndResolveSuper(ref this.SuperIdentifier, UhtFindOptions.ScriptStruct);
+					BindAndResolveSuper(this.SuperIdentifier, UhtFindOptions.ScriptStruct);
 
 					// if we have a base struct, propagate inherited struct flags now
 					UhtScriptStruct? SuperScriptStruct = this.SuperScriptStruct;
@@ -126,7 +130,7 @@ namespace EpicGames.UHT.Parsers
 				{
 					// Parse the specifiers
 					UhtSpecifierContext SpecifierContext = new UhtSpecifierContext(TopScope, TopScope.TokenReader, ScriptStruct.MetaData);
-					UhtSpecifierParser Specifiers = TopScope.HeaderParser.GetSpecifierParser(SpecifierContext, ScopeName, ParentScope.Session.GetSpecifierTable(UhtTableNames.ScriptStruct));
+					UhtSpecifierParser Specifiers = TopScope.HeaderParser.GetCachedSpecifierParser(SpecifierContext, ScopeName, ParentScope.Session.GetSpecifierTable(UhtTableNames.ScriptStruct));
 					Specifiers.ParseSpecifiers();
 
 					// Consume the struct specifier
@@ -159,7 +163,11 @@ namespace EpicGames.UHT.Parsers
 					}
 
 					// Parse the inheritance
-					UhtParserHelpers.ParseInheritance(TopScope.TokenReader, TopScope.Session.Config!, out ScriptStruct.SuperIdentifier, out ScriptStruct.BaseIdentifiers);
+					UhtToken SuperIdentifier;
+					List<UhtToken[]>? BaseIdentifiers;
+					UhtParserHelpers.ParseInheritance(TopScope.TokenReader, TopScope.Session.Config!, out SuperIdentifier, out BaseIdentifiers);
+					ScriptStruct.SuperIdentifier = SuperIdentifier;
+					ScriptStruct.BaseIdentifiers = BaseIdentifiers;
 
 					// Add the comments here for compatibility with old UHT
 					//COMPATIBILITY-TODO - Move this back to where the AddModuleRelativePathToMetaData is called.
diff --git a/Engine/Source/Programs/Shared/EpicGames.UHT/Tables/UhtEngineClassTable.cs b/Engine/Source/Programs/Shared/EpicGames.UHT/Tables/UhtEngineClassTable.cs
--- a/Engine/Source/Programs/Shared/EpicGames.UHT/Tables/UhtEngineClassTable.cs
+++ b/Engine/Source/Programs/Shared/EpicGames.UHT/Tables/UhtEngineClassTable.cs
@@ -18,12 +18,12 @@ namespace EpicGames.UHT.Tables
 		/// <summary>
 		/// The name of the engine class excluding any prefix
 		/// </summary>
-		public string? Name = null;
+		public string? Name { get; set; } = null;
 
 		/// <summary>
 		/// If true, this class is a property
 		/// </summary>
-		public bool IsProperty = false;
+		public bool IsProperty { get; set; } = false;
 	}
 
 	/// <summary>
@@ -34,12 +34,12 @@ namespace EpicGames.UHT.Tables
 		/// <summary>
 		/// The name of the engine class excluding any prefix
 		/// </summary>
-		public string Name;
+		public string Name { get; set; }
 
 		/// <summary>
 		/// If true, this class is a property
 		/// </summary>
-		public bool bIsProperty;
+		public bool bIsProperty { get; set; }
 	}
 
 	/// <summary>
diff --git a/Engine/Source/Programs/Shared/EpicGames.UHT/Tables/UhtExporterTable.cs b/Engine/Source/Programs/Shared/EpicGames.UHT/Tables/UhtExporterTable.cs
--- a/Engine/Source/Programs/Shared/EpicGames.UHT/Tables/UhtExporterTable.cs
+++ b/Engine/Source/Programs/Shared/EpicGames.UHT/Tables/UhtExporterTable.cs
@@ -85,37 +85,37 @@ namespace EpicGames.UHT.Tables
 		/// <summary>
 		/// Name of the exporter
 		/// </summary>
-		public string Name = string.Empty;
+		public string Name { get; set; } = string.Empty;
 
 		/// <summary>
 		/// Description of the export.  Used to display help
 		/// </summary>
-		public string Description = string.Empty;
+		public string Description { get; set; } = string.Empty;
 
 		/// <summary>
 		/// Exporters in plugins need to specify a module name
 		/// </summary>
-		public string ModuleName = string.Empty;
+		public string ModuleName { get; set; } = string.Empty;
 
 		/// <summary>
 		/// Exporter options
 		/// </summary>
-		public UhtExporterOptions Options = UhtExporterOptions.None;
+		public UhtExporterOptions Options { get; set; } = UhtExporterOptions.None;
 
 		/// <summary>
 		/// Collection of filters used to delete old cpp files
 		/// </summary>
-		public string[]? CppFilters = null;
+		public string[]? CppFilters { get; set; }
 
 		/// <summary>
 		/// Collection of filters used to delete old h files
 		/// </summary>
-		public string[]? HeaderFilters = null;
+		public string[]? HeaderFilters { get; set; }
 
 		/// <summary>
 		/// Collection of filters for other file types
 		/// </summary>
-		public string[]? OtherFilters = null;
+		public string[]? OtherFilters { get; set; }
 	}
 
 	/// <summary>
@@ -127,42 +127,59 @@ namespace EpicGames.UHT.Tables
 		/// <summary>
 		/// Name of the exporter
 		/// </summary>
-		public string Name;
+		public string Name { get; }
 
 		/// <summary>
 		/// Description of the export.  Used to display help
 		/// </summary>
-		public string Description;
+		public string Description { get; }
 
 		/// <summary>
 		/// Exporters in plugins need to specify a module name
 		/// </summary>
-		public string ModuleName;
+		public string ModuleName { get; }
 
 		/// <summary>
 		/// Exporter options
 		/// </summary>
-		public UhtExporterOptions Options;
+		public UhtExporterOptions Options { get; }
 
 		/// <summary>
 		/// Delegate to invoke to start export
 		/// </summary>
-		public UhtExporterDelegate Delegate;
+		public UhtExporterDelegate Delegate { get; }
 
 		/// <summary>
 		/// Collection of filters used to delete old cpp files
 		/// </summary>
-		public string[]? CppFilters;
+		public IReadOnlyList<string> CppFilters { get; }
 
 		/// <summary>
 		/// Collection of filters used to delete old h files
 		/// </summary>
-		public string[]? HeaderFilters;
+		public IReadOnlyList<string> HeaderFilters { get; }
 
 		/// <summary>
 		/// Collection of filters for other file types
 		/// </summary>
-		public string[]? OtherFilters;
+		public IReadOnlyList<string> OtherFilters { get; }
+
+		/// <summary>
+		/// Construct an exporter table instance
+		/// </summary>
+		/// <param name="Attribute">Source attribute</param>
+		/// <param name="Delegate">Delegate to invoke</param>
+		public UhtExporter(UhtExporterAttribute Attribute, UhtExporterDelegate Delegate)
+		{
+			this.Name = Attribute.Name;
+			this.Description = Attribute.Description;
+			this.ModuleName = Attribute.ModuleName;
+			this.Options = Attribute.Options;
+			this.Delegate = Delegate;
+			this.CppFilters = Attribute.CppFilters != null ? new List<string>(Attribute.CppFilters) : new List<string>();
+			this.HeaderFilters = Attribute.HeaderFilters != null ? new List<string>(Attribute.HeaderFilters) : new List<string>();
+			this.OtherFilters = Attribute.OtherFilters != null ? new List<string>(Attribute.OtherFilters) : new List<string>();
+		}
 	}
 
 	/// <summary>
@@ -206,18 +223,7 @@ namespace EpicGames.UHT.Tables
 				}
 			}
 
-			UhtExporter ExporterValue = new UhtExporter
-			{
-				Name = ExporterAttribute.Name,
-				Description = ExporterAttribute.Description,
-				ModuleName = ExporterAttribute.ModuleName,
-				Delegate = (UhtExporterDelegate)Delegate.CreateDelegate(typeof(UhtExporterDelegate), MethodInfo),
-				Options = ExporterAttribute.Options,
-				CppFilters = ExporterAttribute.CppFilters,
-				HeaderFilters = ExporterAttribute.HeaderFilters,
-				OtherFilters = ExporterAttribute.OtherFilters,
-			};
-
+			UhtExporter ExporterValue = new UhtExporter(ExporterAttribute, (UhtExporterDelegate)Delegate.CreateDelegate(typeof(UhtExporterDelegate), MethodInfo));
 			this.ExporterValues.Add(ExporterAttribute.Name, ExporterValue);
 		}
 
diff --git a/Engine/Source/Programs/Shared/EpicGames.UHT/Tables/UhtKeywordTable.cs b/Engine/Source/Programs/Shared/EpicGames.UHT/Tables/UhtKeywordTable.cs
--- a/Engine/Source/Programs/Shared/EpicGames.UHT/Tables/UhtKeywordTable.cs
+++ b/Engine/Source/Programs/Shared/EpicGames.UHT/Tables/UhtKeywordTable.cs
@@ -21,27 +21,27 @@ namespace EpicGames.UHT.Tables
 		/// <summary>
 		/// Keyword table/scope being extended
 		/// </summary>
-		public string? Extends;
+		public string? Extends { get; set; }
 
 		/// <summary>
 		/// Name of the keyword
 		/// </summary>
-		public string? Keyword = null;
+		public string? Keyword { get; set; } = null;
 
 		/// <summary>
 		/// Text to be displayed to the user when referencing this keyword
 		/// </summary>
-		public string? AllowText = null;
+		public string? AllowText { get; set; } = null;
 
 		/// <summary>
 		/// If true, this applies to all scopes
 		/// </summary>
-		public bool AllScopes = false;
+		public bool AllScopes { get; set; } = false;
 
 		/// <summary>
 		/// If true, do not include in usage errors
 		/// </summary>
-		public bool DisableUsageError = false;
+		public bool DisableUsageError { get; set; } = false;
 	}
 
 	/// <summary>
@@ -54,7 +54,7 @@ namespace EpicGames.UHT.Tables
 		/// <summary>
 		/// Table/scope to be extended
 		/// </summary>
-		public string? Extends;
+		public string? Extends { get; set; }
 	}
 
 	/// <summary>
@@ -83,27 +83,27 @@ namespace EpicGames.UHT.Tables
 		/// <summary>
 		/// Name of the keyword
 		/// </summary>
-		public readonly string Name;
+		public string Name { get; }
 
 		/// <summary>
 		/// Delegate to invoke
 		/// </summary>
-		public readonly UhtKeywordDelegate Delegate;
+		public UhtKeywordDelegate Delegate { get; }
 
 		/// <summary>
 		/// Text to be displayed to the user when referencing this keyword
 		/// </summary>
-		public readonly string? AllowText;
+		public string? AllowText { get; }
 
 		/// <summary>
 		/// If true, this applies to all scopes
 		/// </summary>
-		public readonly bool bAllScopes;
+		public bool bAllScopes { get; }
 
 		/// <summary>
 		/// If true, do not include in usage errors
 		/// </summary>
-		public readonly bool bDisableUsageError;
+		public bool bDisableUsageError { get; }
 
 		/// <summary>
 		/// Construct a new keyword
@@ -139,7 +139,7 @@ namespace EpicGames.UHT.Tables
 		/// <summary>
 		/// List of catch-alls associated with this table
 		/// </summary>
-		public List<UhtKeywordCatchAllDelegate> CatchAlls = new List<UhtKeywordCatchAllDelegate>();
+		public List<UhtKeywordCatchAllDelegate> CatchAlls { get; } = new List<UhtKeywordCatchAllDelegate>();
 
 		/// <summary>
 		/// Construct a new keyword table
diff --git a/Engine/Source/Programs/Shared/EpicGames.UHT/Tables/UhtLocTextDefaultValueTable.cs b/Engine/Source/Programs/Shared/EpicGames.UHT/Tables/UhtLocTextDefaultValueTable.cs
--- a/Engine/Source/Programs/Shared/EpicGames.UHT/Tables/UhtLocTextDefaultValueTable.cs
+++ b/Engine/Source/Programs/Shared/EpicGames.UHT/Tables/UhtLocTextDefaultValueTable.cs
@@ -32,7 +32,7 @@ namespace EpicGames.UHT.Tables
 		/// <summary>
 		/// Name of the sanitizer (i.e. LOCTEXT, NSLOCTEXT, ...)
 		/// </summary>
-		public string? Name;
+		public string? Name { get; set; }
 	}
 
 	/// <summary>
@@ -43,7 +43,7 @@ namespace EpicGames.UHT.Tables
 		/// <summary>
 		/// Delegate to invoke
 		/// </summary>
-		public UhtLocTextDefaultValueDelegate Delegate;
+		public UhtLocTextDefaultValueDelegate Delegate { get; set; }
 	}
 
 	/// <summary>
diff --git a/Engine/Source/Programs/Shared/EpicGames.UHT/Tables/UhtPropertyTypeTable.cs b/Engine/Source/Programs/Shared/EpicGames.UHT/Tables/UhtPropertyTypeTable.cs
--- a/Engine/Source/Programs/Shared/EpicGames.UHT/Tables/UhtPropertyTypeTable.cs
+++ b/Engine/Source/Programs/Shared/EpicGames.UHT/Tables/UhtPropertyTypeTable.cs
@@ -124,12 +124,12 @@ namespace EpicGames.UHT.Tables
 		/// <summary>
 		/// The expected keyword.  Must be set unless this is the default processor
 		/// </summary>
-		public string? Keyword = null;
+		public string? Keyword { get; set; } = null;
 
 		/// <summary>
 		/// Options
 		/// </summary>
-		public UhtPropertyTypeOptions Options = UhtPropertyTypeOptions.None;
+		public UhtPropertyTypeOptions Options { get; set; } = UhtPropertyTypeOptions.None;
 	}
 
 	/// <summary>
@@ -141,12 +141,12 @@ namespace EpicGames.UHT.Tables
 		/// <summary>
 		/// Delegate to invoke
 		/// </summary>
-		public UhtResolvePropertyDelegate Delegate;
+		public UhtResolvePropertyDelegate Delegate { get; set; }
 
 		/// <summary>
 		/// Options
 		/// </summary>
-		public UhtPropertyTypeOptions Options;
+		public UhtPropertyTypeOptions Options { get; set; }
 	}
 
 	/// <summary>
diff --git a/Engine/Source/Programs/Shared/EpicGames.UHT/Tables/UhtSpecifierTable.cs b/Engine/Source/Programs/Shared/EpicGames.UHT/Tables/UhtSpecifierTable.cs
--- a/Engine/Source/Programs/Shared/EpicGames.UHT/Tables/UhtSpecifierTable.cs
+++ b/Engine/Source/Programs/Shared/EpicGames.UHT/Tables/UhtSpecifierTable.cs
@@ -93,22 +93,22 @@ namespace EpicGames.UHT.Tables
 		/// <summary>
 		/// Current parsing scope (i.e. global, class, ...)
 		/// </summary>
-		public UhtParsingScope Scope;
+		public UhtParsingScope Scope { get; set; }
 
 		/// <summary>
 		/// Message site for messages
 		/// </summary>
-		public IUhtMessageSite MessageSite;
+		public IUhtMessageSite MessageSite { get; set; }
 
 		/// <summary>
 		/// Meta data currently being parsed.
 		/// </summary>
-		public UhtMetaData MetaData;
+		public UhtMetaData MetaData { get; set; }
 
 		/// <summary>
 		/// Make data key index utilized by enumeration values
 		/// </summary>
-		public int MetaNameIndex;
+		public int MetaNameIndex { get; set; }
 
 		/// <summary>
 		/// Construct a new specifier context
@@ -152,17 +152,17 @@ namespace EpicGames.UHT.Tables
 		/// <summary>
 		/// Name of the specifier
 		/// </summary>
-		public string Name = String.Empty;
+		public string Name { get; set; } = String.Empty;
 
 		/// <summary>
 		/// Expected value type
 		/// </summary>
-		public UhtSpecifierValueType ValueType;
+		public UhtSpecifierValueType ValueType { get; set; }
 
 		/// <summary>
 		/// When is the specifier executed
 		/// </summary>
-		public UhtSpecifierWhen When = UhtSpecifierWhen.Deferred;
+		public UhtSpecifierWhen When { get; set; } = UhtSpecifierWhen.Deferred;
 
 		/// <summary>
 		/// Dispatch an instance of the specifier
@@ -489,22 +489,22 @@ namespace EpicGames.UHT.Tables
 		/// <summary>
 		/// Name of the specifier.   If not supplied, the method name must end in "Specifier" and the name will be the method name with "Specifier" stripped.
 		/// </summary>
-		public string? Name;
+		public string? Name { get; set; }
 
 		/// <summary>
 		/// Name of the table/scope this specifier applies
 		/// </summary>
-		public string? Extends;
+		public string? Extends { get; set; }
 
 		/// <summary>
 		/// Value type of the specifier
 		/// </summary>
-		public UhtSpecifierValueType ValueType = UhtSpecifierValueType.NotSet;
+		public UhtSpecifierValueType ValueType { get; set; } = UhtSpecifierValueType.NotSet;
 
 		/// <summary>
 		/// When the specifier is dispatched
 		/// </summary>
-		public UhtSpecifierWhen When = UhtSpecifierWhen.Deferred;
+		public UhtSpecifierWhen When { get; set; } = UhtSpecifierWhen.Deferred;
 	}
 
 	/// <summary>
diff --git a/Engine/Source/Programs/Shared/EpicGames.UHT/Tables/UhtSpecifierValidatorTable.cs b/Engine/Source/Programs/Shared/EpicGames.UHT/Tables/UhtSpecifierValidatorTable.cs
--- a/Engine/Source/Programs/Shared/EpicGames.UHT/Tables/UhtSpecifierValidatorTable.cs
+++ b/Engine/Source/Programs/Shared/EpicGames.UHT/Tables/UhtSpecifierValidatorTable.cs
@@ -27,12 +27,12 @@ namespace EpicGames.UHT.Tables
 		/// <summary>
 		/// Name of the validator
 		/// </summary>
-		public string Name;
+		public string Name { get; set; }
 
 		/// <summary>
 		/// Delegate for the validator
 		/// </summary>
-		public UhtSpecifierValidatorDelegate Delegate;
+		public UhtSpecifierValidatorDelegate Delegate { get; set; }
 
 		/// <summary>
 		/// Construct a new instance
@@ -56,12 +56,12 @@ namespace EpicGames.UHT.Tables
 		/// <summary>
 		/// Name of the validator. If not supplied &quot;SpecifierValidator&quot; will be removed from the end of the method name
 		/// </summary>
-		public string? Name;
+		public string? Name { get; set; }
 
 		/// <summary>
 		/// Name of the table/scope for the validator
 		/// </summary>
-		public string? Extends;
+		public string? Extends { get; set; }
 	}
 
 	/// <summary>
diff --git a/Engine/Source/Programs/Shared/EpicGames.UHT/Tables/UhtStructDefaultValueTable.cs b/Engine/Source/Programs/Shared/EpicGames.UHT/Tables/UhtStructDefaultValueTable.cs
--- a/Engine/Source/Programs/Shared/EpicGames.UHT/Tables/UhtStructDefaultValueTable.cs
+++ b/Engine/Source/Programs/Shared/EpicGames.UHT/Tables/UhtStructDefaultValueTable.cs
@@ -91,12 +91,12 @@ namespace EpicGames.UHT.Tables
 		/// <summary>
 		/// Name of the structure.  Not required for default processor.  
 		/// </summary>
-		public string? Name;
+		public string? Name { get; set; }
 
 		/// <summary>
 		/// Options
 		/// </summary>
-		public UhtStructDefaultValueOptions Options = UhtStructDefaultValueOptions.None;
+		public UhtStructDefaultValueOptions Options { get; set; } = UhtStructDefaultValueOptions.None;
 	}
 
 	/// <summary>
@@ -108,7 +108,7 @@ namespace EpicGames.UHT.Tables
 		/// <summary>
 		/// The delegate to invoke
 		/// </summary>
-		public UhtStructDefaultValueDelegate Delegate;
+		public UhtStructDefaultValueDelegate Delegate { get; set; }
 	}
 
 	/// <summary>
diff --git a/Engine/Source/Programs/Shared/EpicGames.UHT/Tokenizer/UhtTokenReader.cs b/Engine/Source/Programs/Shared/EpicGames.UHT/Tokenizer/UhtTokenReader.cs
--- a/Engine/Source/Programs/Shared/EpicGames.UHT/Tokenizer/UhtTokenReader.cs
+++ b/Engine/Source/Programs/Shared/EpicGames.UHT/Tokenizer/UhtTokenReader.cs
@@ -255,12 +255,12 @@ namespace EpicGames.UHT.Tokenizer
 		/// <summary>
 		/// The token
 		/// </summary>
-		public UhtToken Token;
+		public UhtToken Token { get; set; }
 
 		/// <summary>
 		/// The next token in the list
 		/// </summary>
-		public UhtTokenList? Next;
+		public UhtTokenList? Next { get; set; }
 
 		/// <summary>
 		/// Join the tokens in the list
diff --git a/Engine/Source/Programs/Shared/EpicGames.UHT/Types/Properties/UhtBoolProperty.cs b/Engine/Source/Programs/Shared/EpicGames.UHT/Types/Properties/UhtBoolProperty.cs
--- a/Engine/Source/Programs/Shared/EpicGames.UHT/Types/Properties/UhtBoolProperty.cs
+++ b/Engine/Source/Programs/Shared/EpicGames.UHT/Types/Properties/UhtBoolProperty.cs
@@ -106,7 +106,7 @@ namespace EpicGames.UHT.Types
 		/// <summary>
 		/// Type of the boolean
 		/// </summary>
-		public readonly UhtBoolType BoolType;
+		public UhtBoolType BoolType { get; }
 
 		/// <summary>
 		/// Construct a new boolean property
diff --git a/Engine/Source/Programs/Shared/EpicGames.UHT/Types/Properties/UhtNumericProperty.cs b/Engine/Source/Programs/Shared/EpicGames.UHT/Types/Properties/UhtNumericProperty.cs
--- a/Engine/Source/Programs/Shared/EpicGames.UHT/Types/Properties/UhtNumericProperty.cs
+++ b/Engine/Source/Programs/Shared/EpicGames.UHT/Types/Properties/UhtNumericProperty.cs
@@ -19,7 +19,7 @@ namespace EpicGames.UHT.Types
 		/// <summary>
 		/// Describes the integer as either being sized on unsized
 		/// </summary>
-		public readonly UhtPropertyIntType IntType;
+		public UhtPropertyIntType IntType { get; }
 
 		/// <summary>
 		/// Construct new property
diff --git a/Engine/Source/Programs/Shared/EpicGames.UHT/Types/UhtClass.cs b/Engine/Source/Programs/Shared/EpicGames.UHT/Types/UhtClass.cs
--- a/Engine/Source/Programs/Shared/EpicGames.UHT/Types/UhtClass.cs
+++ b/Engine/Source/Programs/Shared/EpicGames.UHT/Types/UhtClass.cs
@@ -227,17 +227,18 @@ namespace EpicGames.UHT.Types
 		/// <summary>
 		/// Compiler directives when declaration was parsed
 		/// </summary>
-		public UhtCompilerDirective CompilerDirectives;
+		public UhtCompilerDirective CompilerDirectives { get; set; }
 
 		/// <summary>
 		/// Collection of tokens parsed in the declaration
 		/// </summary>
+		[System.Diagnostics.CodeAnalysis.SuppressMessage("Design", "CA1051:Do not declare visible instance fields", Justification = "<Pending>")]
 		public UhtToken[] Tokens;
 
 		/// <summary>
 		/// If this declaration is part of a UFUNCTION, this will be set
 		/// </summary>
-		public UhtFunction? Function;
+		public UhtFunction? Function { get; set; }
 	}
 
 	/// <summary>
@@ -249,22 +250,23 @@ namespace EpicGames.UHT.Types
 		/// <summary>
 		/// Compiler directives when declaration was parsed
 		/// </summary>
-		public UhtCompilerDirective CompilerDirectives;
+		public UhtCompilerDirective CompilerDirectives { get; set; }
 
 		/// <summary>
 		/// Collection of tokens parsed in the declaration
 		/// </summary>
+		[System.Diagnostics.CodeAnalysis.SuppressMessage("Design", "CA1051:Do not declare visible instance fields", Justification = "<Pending>")]
 		public UhtToken[] Tokens;
 
 		/// <summary>
 		/// Token index for the matching name
 		/// </summary>
-		public int NameTokenIndex;
+		public int NameTokenIndex { get; set; }
 
 		/// <summary>
 		/// True if "virtual" was found prior to the matching name
 		/// </summary>
-		public bool bIsVirtual;
+		public bool bIsVirtual { get; set; }
 	}
 
 	/// <summary>
@@ -312,7 +314,7 @@ namespace EpicGames.UHT.Types
 		/// <summary>
 		/// Type of the class
 		/// </summary>
-		public UhtClassType ClassType = UhtClassType.Class;
+		public UhtClassType ClassType { get; set; } = UhtClassType.Class;
 
 		/// <summary>
 		/// Type of archivers present
@@ -334,7 +336,7 @@ namespace EpicGames.UHT.Types
 		/// <summary>
 		/// If this, this class is a UINTERFACE and NativeInterface is the associated native interface
 		/// </summary>
-		public UhtClass? NativeInterface = null;
+		public UhtClass? NativeInterface { get; set; } = null;
 
 		/// <summary>
 		/// Line number o the generated body statement
@@ -350,7 +352,7 @@ namespace EpicGames.UHT.Types
 		/// <summary>
 		/// True if GENERATED_BODY was used.  If false, GENERATED_UCLASS_BODY was used.
 		/// </summary>
-		public bool bHasGeneratedBody = false;
+		public bool bHasGeneratedBody { get; set; } = false;
 
 		/// <inheritdoc/>
 		[JsonIgnore]
diff --git a/Engine/Source/Programs/Shared/EpicGames.UHT/Types/UhtFunction.cs b/Engine/Source/Programs/Shared/EpicGames.UHT/Types/UhtFunction.cs
--- a/Engine/Source/Programs/Shared/EpicGames.UHT/Types/UhtFunction.cs
+++ b/Engine/Source/Programs/Shared/EpicGames.UHT/Types/UhtFunction.cs
@@ -218,7 +218,7 @@ namespace EpicGames.UHT.Types
 		/// <summary>
 		/// The type of function
 		/// </summary>
-		public UhtFunctionType FunctionType = UhtFunctionType.Function;
+		public UhtFunctionType FunctionType { get; set; } = UhtFunctionType.Function;
 
 		/// <summary>
 		/// The line number for the macro.
@@ -291,37 +291,37 @@ namespace EpicGames.UHT.Types
 		/// <summary>
 		/// Identifier for an RPC call to a platform service
 		/// </summary>
-		public UInt16 RPCId = 0;
+		public UInt16 RPCId { get; set; } = 0;
 
 		/// <summary>
 		/// Identifier for an RPC call expecting a response
 		/// </summary>
-		public UInt16 RPCResponseId = 0;
+		public UInt16 RPCResponseId { get; set; } = 0;
 
 		/// <summary>
 		/// Name of the actual implementation
 		/// </summary>
-		public string CppImplName = string.Empty;
+		public string CppImplName { get; set; } = string.Empty;
 
 		/// <summary>
 		/// Name of the actual validation implementation
 		/// </summary>
-		public string CppValidationImplName = string.Empty;
+		public string CppValidationImplName { get; set; } = string.Empty;
 
 		/// <summary>
 		/// Name of the wrapper function that marshals the arguments and does the indirect call
 		/// </summary>
-		public string MarshalAndCallName = string.Empty;
+		public string MarshalAndCallName { get; set; } = string.Empty;
 
 		/// <summary>
 		/// Name for callback-style names
 		/// </summary>
-		public string UnMarshalAndCallName = string.Empty;
+		public string UnMarshalAndCallName { get; set; } = string.Empty;
 
 		/// <summary>
 		/// Endpoint name
 		/// </summary>
-		public string EndpointName = string.Empty;
+		public string EndpointName { get; set; } = string.Empty;
 
 		/// <summary>
 		/// True if the function has a return value.
diff --git a/Engine/Source/Programs/Shared/EpicGames.UHT/Types/UhtHeaderFile.cs b/Engine/Source/Programs/Shared/EpicGames.UHT/Types/UhtHeaderFile.cs
--- a/Engine/Source/Programs/Shared/EpicGames.UHT/Types/UhtHeaderFile.cs
+++ b/Engine/Source/Programs/Shared/EpicGames.UHT/Types/UhtHeaderFile.cs
@@ -122,37 +122,37 @@ namespace EpicGames.UHT.Types
 		/// <summary>
 		/// File name without the extension
 		/// </summary>
-		public readonly string FileNameWithoutExtension;
+		public string FileNameWithoutExtension { get; }
 
 		/// <summary>
 		/// Required name for the generated.h file name.  Used to validate parsed code
 		/// </summary>
-		public readonly string GeneratedHeaderFileName;
+		public string GeneratedHeaderFileName { get; }
 
 		/// <summary>
 		/// True if this header is NoExportTypes.h
 		/// </summary>
-		public readonly bool bIsNoExportTypes;
+		public bool bIsNoExportTypes { get; }
 
 		/// <summary>
 		/// The file path of the header relative to the module location
 		/// </summary>
-		public string ModuleRelativeFilePath = string.Empty;
+		public string ModuleRelativeFilePath { get; set; } = string.Empty;
 
 		/// <summary>
 		/// Include file path added as meta data to the types
 		/// </summary>
-		public string IncludeFilePath = string.Empty;
+		public string IncludeFilePath { get; set; } = string.Empty;
 
 		/// <summary>
 		/// Location where the header file was found
 		/// </summary>
-		public UhtHeaderFileType HeaderFileType = UhtHeaderFileType.Private;
+		public UhtHeaderFileType HeaderFileType { get; set; } = UhtHeaderFileType.Private;
 
 		/// <summary>
 		/// Unique index of the header file
 		/// </summary>
-		public readonly int HeaderFileTypeIndex;
+		public int HeaderFileTypeIndex { get; }
 
 		/// <summary>
 		/// UHT flags for the header
@@ -169,7 +169,8 @@ namespace EpicGames.UHT.Types
 		/// <summary>
 		/// Resource collector for the header file
 		/// </summary>
-		public UhtReferenceCollector References = new UhtReferenceCollector();
+		[JsonIgnore]
+		public UhtReferenceCollector References { get; } = new UhtReferenceCollector();
 
 		/// <inheritdoc/>
 		[JsonIgnore]
diff --git a/Engine/Source/Programs/Shared/EpicGames.UHT/Types/UhtManifestFile.cs b/Engine/Source/Programs/Shared/EpicGames.UHT/Types/UhtManifestFile.cs
--- a/Engine/Source/Programs/Shared/EpicGames.UHT/Types/UhtManifestFile.cs
+++ b/Engine/Source/Programs/Shared/EpicGames.UHT/Types/UhtManifestFile.cs
@@ -16,7 +16,7 @@ namespace EpicGames.UHT.Types
 		/// <summary>
 		/// Loaded manifest from the json manifest file
 		/// </summary>
-		public UHTManifest? Manifest = null;
+		public UHTManifest? Manifest { get; set; } = null;
 
 		private readonly UhtSourceFile SourceFile;
 
diff --git a/Engine/Source/Programs/Shared/EpicGames.UHT/Types/UhtMetaData.cs b/Engine/Source/Programs/Shared/EpicGames.UHT/Types/UhtMetaData.cs
--- a/Engine/Source/Programs/Shared/EpicGames.UHT/Types/UhtMetaData.cs
+++ b/Engine/Source/Programs/Shared/EpicGames.UHT/Types/UhtMetaData.cs
@@ -5,6 +5,7 @@ using EpicGames.UHT.Utils;
 using System;
 using System.Collections.Generic;
 using System.Diagnostics.CodeAnalysis;
+using System.Text.Json.Serialization;
 
 namespace EpicGames.UHT.Types
 {
@@ -33,12 +34,12 @@ namespace EpicGames.UHT.Types
 		/// <summary>
 		/// The name of the meta data
 		/// </summary>
-		public string Name;
+		public string Name { get; set; }
 
 		/// <summary>
 		/// The index of the meta data name (i.e. enum value index) or -1 for the owning object meta data
 		/// </summary>
-		public int Index;
+		public int Index { get; set; }
 
 		/// <summary>
 		/// Construct a new meta data key
@@ -134,33 +135,38 @@ namespace EpicGames.UHT.Types
 		/// <summary>
 		/// The meta data of the outer object for the type that owns this meta data
 		/// </summary>
-		public UhtMetaData? Parent = null;
+		[JsonIgnore]
+		public UhtMetaData? Parent { get; set; } = null;
 
 		/// <summary>
 		/// Message site associated with the meta data.  That in combination with the line number is used to generate errors
 		/// </summary>
-		public IUhtMessageSite? MessageSite = null;
+		[JsonIgnore]
+		public IUhtMessageSite? MessageSite { get; set; } = null;
 
 		/// <summary>
 		/// Configuration object used to redirect meta data
 		/// </summary>
-		public IUhtConfig? Config = null;
-		
+		[JsonIgnore]
+		public IUhtConfig? Config { get; set; } = null;
+
 		/// <summary>
 		/// Source code line number where the meta data is declared
 		/// </summary>
-		public int LineNumber { get; internal set; } = 1;
+		[JsonIgnore]
+		public int LineNumber { get; set; } = 1;
 
 		/// <summary>
 		/// Contains the meta data entries.  Due to the small size of these dictionaries, a SortedList performs
 		/// better than a Dictionary.
 		/// </summary>
-		public SortedList<UhtMetaDataKey, string>? Dictionary = null;
+		[SuppressMessage("Usage", "CA2227:Collection properties should be read only", Justification = "")]
+		public SortedList<UhtMetaDataKey, string>? Dictionary { get; set; } = null;
 
 		/// <summary>
 		/// Enumerations implement this interface so that the index part of the key can be converted to a string
 		/// </summary>
-		internal IUhtMetaDataKeyConversion? KeyConversion = null;
+		internal IUhtMetaDataKeyConversion? KeyConversion { get; set; } = null;
 
 		/// <summary>
 		/// Index for a meta data key associated with the owning object
diff --git a/Engine/Source/Programs/Shared/EpicGames.UHT/Types/UhtObject.cs b/Engine/Source/Programs/Shared/EpicGames.UHT/Types/UhtObject.cs
--- a/Engine/Source/Programs/Shared/EpicGames.UHT/Types/UhtObject.cs
+++ b/Engine/Source/Programs/Shared/EpicGames.UHT/Types/UhtObject.cs
@@ -16,18 +16,18 @@ namespace EpicGames.UHT.Types
 		/// <summary>
 		/// Internal object flags.
 		/// </summary>
-		public EInternalObjectFlags InternalObjectFlags = EInternalObjectFlags.None;
+		public EInternalObjectFlags InternalObjectFlags { get; set; } = EInternalObjectFlags.None;
 
 		/// <summary>
 		/// Unique index of the object
 		/// </summary>
-		public readonly int ObjectTypeIndex;
+		public int ObjectTypeIndex { get; }
 
 		/// <summary>
 		/// The alternate object is used by the interface system where the native interface will
 		/// update this setting to point to the UInterface derived companion object.
 		/// </summary>
-		public UhtObject? AlternateObject = null;
+		public UhtObject? AlternateObject { get; set; } = null;
 
 		/// <inheritdoc/>
 		public override string EngineClassName { get => "Object"; }
diff --git a/Engine/Source/Programs/Shared/EpicGames.UHT/Types/UhtPackage.cs b/Engine/Source/Programs/Shared/EpicGames.UHT/Types/UhtPackage.cs
--- a/Engine/Source/Programs/Shared/EpicGames.UHT/Types/UhtPackage.cs
+++ b/Engine/Source/Programs/Shared/EpicGames.UHT/Types/UhtPackage.cs
@@ -20,7 +20,7 @@ namespace EpicGames.UHT.Types
 		/// <summary>
 		/// Unique index of the package
 		/// </summary>
-		public readonly int PackageTypeIndex;
+		public int PackageTypeIndex { get; }
 
 		/// <summary>
 		/// Engine package flags
@@ -86,7 +86,7 @@ namespace EpicGames.UHT.Types
 		/// <summary>
 		/// Short name of the package (without the /Script/)
 		/// </summary>
-		public string ShortName;
+		public string ShortName { get; set; }
 
 		/// <summary>
 		/// Construct a new instance of a package
diff --git a/Engine/Source/Programs/Shared/EpicGames.UHT/Types/UhtProperty.cs b/Engine/Source/Programs/Shared/EpicGames.UHT/Types/UhtProperty.cs
--- a/Engine/Source/Programs/Shared/EpicGames.UHT/Types/UhtProperty.cs
+++ b/Engine/Source/Programs/Shared/EpicGames.UHT/Types/UhtProperty.cs
@@ -630,92 +630,93 @@ namespace EpicGames.UHT.Types
 		/// <summary>
 		/// Source name of the property
 		/// </summary>
-		public string SourceName = String.Empty;
+		public string SourceName { get; set; } = String.Empty;
 
 		/// <summary>
 		/// Engine name of the property
 		/// </summary>
-		public string EngineName = String.Empty;
+		public string EngineName { get; set; } = String.Empty;
 
 		/// <summary>
 		/// Property's meta data
 		/// </summary>
-		public UhtMetaData MetaData = UhtMetaData.Empty;
+		public UhtMetaData MetaData { get; set; } = UhtMetaData.Empty;
 
 		/// <summary>
 		/// Property outer object
 		/// </summary>
-		public UhtType Outer;
+		public UhtType Outer { get; set; }
 
 		/// <summary>
 		/// Line number of the property declaration
 		/// </summary>
-		public int LineNumber;
+		public int LineNumber { get; set; }
 
 		/// <summary>
 		/// Property category
 		/// </summary>
-		public UhtPropertyCategory PropertyCategory;
+		public UhtPropertyCategory PropertyCategory { get; set; }
 
 		/// <summary>
 		/// Engine property flags
 		/// </summary>
-		public EPropertyFlags PropertyFlags;
+		public EPropertyFlags PropertyFlags { get; set; }
 
 		/// <summary>
 		/// Property flags not allowed by the context of the property parsing
 		/// </summary>
-		public EPropertyFlags DisallowPropertyFlags;
+		public EPropertyFlags DisallowPropertyFlags { get; set; }
 
 		/// <summary>
 		/// UHT specified property flags
 		/// </summary>
-		public UhtPropertyExportFlags PropertyExportFlags;
+		public UhtPropertyExportFlags PropertyExportFlags { get; set; }
 
 		/// <summary>
 		/// Allocator used for containers
 		/// </summary>
-		public UhtPropertyAllocator Allocator;
+		public UhtPropertyAllocator Allocator { get; set; }
 
 		/// <summary>
 		/// Options for property parsing
 		/// </summary>
-		public UhtPropertyOptions Options;
+		public UhtPropertyOptions Options { get; set; }
 
 		/// <summary>
 		/// Property pointer type
 		/// </summary>
-		public UhtPointerType PointerType;
+		public UhtPointerType PointerType { get; set; }
 
 		/// <summary>
 		/// Replication notify name
 		/// </summary>
-		public string? RepNotifyName;
+		public string? RepNotifyName { get; set; }
 
 		/// <summary>
 		/// If set, the array size of the property
 		/// </summary>
-		public string? ArrayDimensions;
+		public string? ArrayDimensions { get; set; }
 
 		/// <summary>
 		/// Getter method
 		/// </summary>
-		public string? Setter;
+		public string? Setter { get; set; }
 
 		/// <summary>
 		/// Setter method
 		/// </summary>
-		public string? Getter;
+		public string? Getter { get; set; }
 
 		/// <summary>
 		/// Default value of the property
 		/// </summary>
-		public List<UhtToken>? DefaultValueTokens;
+		[SuppressMessage("Usage", "CA2227:Collection properties should be read only", Justification = "")]
+		public List<UhtToken>? DefaultValueTokens { get; set; }
 
 		/// <summary>
 		/// If true, the property is a bit field
 		/// </summary>
-		public bool bIsBitfield;
+		public bool bIsBitfield { get; set; }
 
 		/// <summary>
 		/// Construct a new, uninitialized version of the property settings
@@ -882,7 +883,7 @@ namespace EpicGames.UHT.Types
 		/// <summary>
 		/// Property category
 		/// </summary>
-		public UhtPropertyCategory PropertyCategory = UhtPropertyCategory.Member;
+		public UhtPropertyCategory PropertyCategory { get; set; } = UhtPropertyCategory.Member;
 
 		/// <summary>
 		/// Engine property flags
@@ -900,57 +901,59 @@ namespace EpicGames.UHT.Types
 		/// <summary>
 		/// Engine flags that are disallowed on this property
 		/// </summary>
-		public EPropertyFlags DisallowPropertyFlags = EPropertyFlags.None;
+		public EPropertyFlags DisallowPropertyFlags { get; set; } = EPropertyFlags.None;
 
 		/// <summary>
 		/// UHT specified property flags
 		/// </summary>
-		public UhtPropertyExportFlags PropertyExportFlags = UhtPropertyExportFlags.Public;
+		public UhtPropertyExportFlags PropertyExportFlags { get; set; } = UhtPropertyExportFlags.Public;
 
 		/// <summary>
 		/// Reference type of the property
 		/// </summary>
-		public UhtPropertyRefQualifier RefQualifier = UhtPropertyRefQualifier.None;
+		public UhtPropertyRefQualifier RefQualifier { get; set; } = UhtPropertyRefQualifier.None;
 
 		/// <summary>
 		/// Pointer type of the property
 		/// </summary>
-		public UhtPointerType PointerType = UhtPointerType.None;
+		public UhtPointerType PointerType { get; set; } = UhtPointerType.None;
 
 		/// <summary>
 		/// Allocator to be used with containers
 		/// </summary>
-		public UhtPropertyAllocator Allocator = UhtPropertyAllocator.Default;
+		public UhtPropertyAllocator Allocator { get; set; } = UhtPropertyAllocator.Default;
 
 		/// <summary>
 		/// Replication notify name
 		/// </summary>
-		public string? RepNotifyName = null;
+		public string? RepNotifyName { get; set; } = null;
 
 		/// <summary>
 		/// Fixed array size
 		/// </summary>
-		public string? ArrayDimensions = null;
+		public string? ArrayDimensions { get; set; } = null;
 
 		/// <summary>
 		/// Property setter
 		/// </summary>
-		public string? Setter = null;
+		public string? Setter { get; set; } = null;
 
 		/// <summary>
 		/// Property getter
 		/// </summary>
-		public string? Getter = null;
+		public string? Getter { get; set; } = null;
 
 		/// <summary>
 		/// Default value of property
 		/// </summary>
-		public List<UhtToken>? DefaultValueTokens = null;
+		[JsonIgnore]
+		[SuppressMessage("Usage", "CA2227:Collection properties should be read only", Justification = "<Pending>")]
+		public List<UhtToken>? DefaultValueTokens { get; set; } = null;
 
 		/// <summary>
 		/// If true, this property is a bit field
 		/// </summary>
-		public bool bIsBitfield = false;
+		public bool bIsBitfield { get; set; } = false;
 
 		///<inheritdoc/>
 		[JsonIgnore]
diff --git a/Engine/Source/Programs/Shared/EpicGames.UHT/Types/UhtScriptStruct.cs b/Engine/Source/Programs/Shared/EpicGames.UHT/Types/UhtScriptStruct.cs
--- a/Engine/Source/Programs/Shared/EpicGames.UHT/Types/UhtScriptStruct.cs
+++ b/Engine/Source/Programs/Shared/EpicGames.UHT/Types/UhtScriptStruct.cs
@@ -109,7 +109,7 @@ namespace EpicGames.UHT.Types
 		/// <summary>
 		/// Property associated with the RigVM parameter
 		/// </summary>
-		public readonly UhtProperty? Property;
+		public UhtProperty? Property { get; }
 
 		/// <summary>
 		/// Name of the property
@@ -124,17 +124,17 @@ namespace EpicGames.UHT.Types
 		/// <summary>
 		/// Cast name
 		/// </summary>
-		public readonly string? CastName = null;
+		public string? CastName { get; }
 
 		/// <summary>
 		/// Cast type
 		/// </summary>
-		public readonly string? CastType = null;
+		public string? CastType { get; }
 
 		/// <summary>
 		/// Flags associated with the parameter
 		/// </summary>
-		public UhtRigVMParameterFlags ParameterFlags = UhtRigVMParameterFlags.None;
+		public UhtRigVMParameterFlags ParameterFlags { get; set; } = UhtRigVMParameterFlags.None;
 
 		/// <summary>
 		/// True if the parameter is marked as "Constant"
@@ -442,12 +442,12 @@ namespace EpicGames.UHT.Types
 		/// <summary>
 		/// True if the GetUpgradeInfoMethod was found. 
 		/// </summary>
-		public bool bHasGetUpgradeInfoMethod = false;
+		public bool bHasGetUpgradeInfoMethod { get; set; } = false;
 		
 		/// <summary>
 		/// True if the GetNextAggregateNameMethod was found. 
 		/// </summary>
-		public bool bHasGetNextAggregateNameMethod = false;
+		public bool bHasGetNextAggregateNameMethod { get; set; } = false;
 
 		/// <summary>
 		/// Engine name of the owning script struct
@@ -457,7 +457,7 @@ namespace EpicGames.UHT.Types
 		/// <summary>
 		/// List of the members
 		/// </summary>
-		public List<UhtRigVMParameter> Members = new List<UhtRigVMParameter>();
+		public List<UhtRigVMParameter> Members { get; } = new List<UhtRigVMParameter>();
 
 		/// <summary>
 		/// List of the methods
diff --git a/Engine/Source/Programs/Shared/EpicGames.UHT/Types/UhtSourceFile.cs b/Engine/Source/Programs/Shared/EpicGames.UHT/Types/UhtSourceFile.cs
--- a/Engine/Source/Programs/Shared/EpicGames.UHT/Types/UhtSourceFile.cs
+++ b/Engine/Source/Programs/Shared/EpicGames.UHT/Types/UhtSourceFile.cs
@@ -18,17 +18,17 @@ namespace EpicGames.UHT.Types
 		/// <summary>
 		/// The session associated with the source
 		/// </summary>
-		public readonly UhtSession Session;
+		public UhtSession Session { get; }
 
 		/// <summary>
 		/// The full file path
 		/// </summary>
-		public readonly string FilePath;
+		public string FilePath { get; }
 
 		/// <summary>
 		/// The file name
 		/// </summary>
-		public readonly string FileName;
+		public string FileName { get; }
 
 		/// <summary>
 		/// The contents of the source file
diff --git a/Engine/Source/Programs/Shared/EpicGames.UHT/Types/UhtStruct.cs b/Engine/Source/Programs/Shared/EpicGames.UHT/Types/UhtStruct.cs
--- a/Engine/Source/Programs/Shared/EpicGames.UHT/Types/UhtStruct.cs
+++ b/Engine/Source/Programs/Shared/EpicGames.UHT/Types/UhtStruct.cs
@@ -21,7 +21,7 @@ namespace EpicGames.UHT.Types
 		/// <summary>
 		/// Generated code version of the type.  Set via specifiers
 		/// </summary>
-		public EGeneratedCodeVersion GeneratedCodeVersion = EGeneratedCodeVersion.None;
+		public EGeneratedCodeVersion GeneratedCodeVersion { get; set; } = EGeneratedCodeVersion.None;
 
 		/// <summary>
 		/// Return a collection of children that are properties
@@ -79,7 +79,7 @@ namespace EpicGames.UHT.Types
 		/// <summary>
 		/// Super type
 		/// </summary>
-		public UhtStruct? Super = null;
+		public UhtStruct? Super { get; set; } = null;
 
 		/// <summary>
 		/// Base types
@@ -180,7 +180,7 @@ namespace EpicGames.UHT.Types
 		/// <param name="SuperIdentifier">Token that represent the super</param>
 		/// <param name="FindOptions">Find options to restrict types</param>
 		/// <exception cref="UhtException">Thrown if super can not be found</exception>
-		public void BindAndResolveSuper(ref UhtToken SuperIdentifier, UhtFindOptions FindOptions)
+		public void BindAndResolveSuper(UhtToken SuperIdentifier, UhtFindOptions FindOptions)
 		{
 			if (SuperIdentifier)
 			{
diff --git a/Engine/Source/Programs/Shared/EpicGames.UHT/Types/UhtType.cs b/Engine/Source/Programs/Shared/EpicGames.UHT/Types/UhtType.cs
--- a/Engine/Source/Programs/Shared/EpicGames.UHT/Types/UhtType.cs
+++ b/Engine/Source/Programs/Shared/EpicGames.UHT/Types/UhtType.cs
@@ -610,32 +610,32 @@ namespace EpicGames.UHT.Types
 		/// <summary>
 		/// If true, this is a valid documentation policy that should be respected
 		/// </summary>
-		public bool bPolicySet;
+		public bool bPolicySet { get; set; }
 
 		/// <summary>
 		/// A tool tip must be defined
 		/// </summary>
-		public bool bClassOrStructCommentRequired;
+		public bool bClassOrStructCommentRequired { get; set; }
 
 		/// <summary>
 		/// Functions must have a tool tip
 		/// </summary>
-		public bool bFunctionToolTipsRequired;
+		public bool bFunctionToolTipsRequired { get; set; }
 
 		/// <summary>
 		/// Property members must have a tool tip
 		/// </summary>
-		public bool bMemberToolTipsRequired;
+		public bool bMemberToolTipsRequired { get; set; }
 
 		/// <summary>
 		/// Function parameters must have a tool tip
 		/// </summary>
-		public bool bParameterToolTipsRequired;
+		public bool bParameterToolTipsRequired { get; set; }
 
 		/// <summary>
 		/// Float properties must have a range specified
 		/// </summary>
-		public bool bFloatRangesRequired;
+		public bool bFloatRangesRequired { get; set; }
 
 		/// <summary>
 		/// Return the documentation policy given the name
@@ -670,17 +670,20 @@ namespace EpicGames.UHT.Types
 		/// <summary>
 		/// All UHT runs are associated with a given session.  The session holds all the global information for a run.
 		/// </summary>
-		public UhtSession Session;
+		[JsonIgnore]
+		public UhtSession Session { get; set; }
 
 		/// <summary>
 		/// Every type in a session is assigned a unique, non-zero id.
 		/// </summary>
-		public readonly int TypeIndex;
+		[JsonIgnore]
+		public int TypeIndex { get; }
 
 		/// <summary>
 		/// The outer object containing this type.  For example, the header file would be the outer for a class defined in the global scope.
 		/// </summary>
-		public UhtType? Outer = null;
+		[JsonIgnore]
+		public UhtType? Outer { get; set; } = null;
 
 		/// <summary>
 		/// The owning package of the type
@@ -740,7 +743,7 @@ namespace EpicGames.UHT.Types
 		/// <summary>
 		/// If true, then this type is visible to the symbol table.
 		/// </summary>
-		public bool bVisibleType = true;
+		public bool bVisibleType { get; set; } = true;
 
 		/// <summary>
 		/// Return true if the type has been deprecated
@@ -795,7 +798,8 @@ namespace EpicGames.UHT.Types
 		/// <summary>
 		/// Meta data associated with the type
 		/// </summary>
-		public UhtMetaData MetaData;
+		[JsonIgnore] //ETSTODO - disable for now
+		public UhtMetaData MetaData { get; set; }
 
 		/// <summary>
 		/// Children types of this type
diff --git a/Engine/Source/Programs/Shared/EpicGames.UHT/Utils/UhtBuffer.cs b/Engine/Source/Programs/Shared/EpicGames.UHT/Utils/UhtBuffer.cs
--- a/Engine/Source/Programs/Shared/EpicGames.UHT/Utils/UhtBuffer.cs
+++ b/Engine/Source/Programs/Shared/EpicGames.UHT/Utils/UhtBuffer.cs
@@ -38,23 +38,24 @@ namespace EpicGames.UHT.Utils
 		/// <summary>
 		/// The bucket index associated with the buffer
 		/// </summary>
-		private readonly int Bucket;
+		private int Bucket { get; }
 
 		/// <summary>
 		/// Single list link to the next cached buffer
 		/// </summary>
-		private UhtBuffer? NextBuffer = null;
+		private UhtBuffer? NextBuffer { get; set; } = null;
 
 		/// <summary>
 		/// The backing character block.  The size of the array will normally be larger than the 
 		/// requested size.
 		/// </summary>
+		[System.Diagnostics.CodeAnalysis.SuppressMessage("Design", "CA1051:Do not declare visible instance fields", Justification = "")]
 		public char[] Block;
 
 		/// <summary>
 		/// Memory region sized to the requested size
 		/// </summary>
-		public Memory<char> Memory;
+		public Memory<char> Memory { get; set; }
 
 		/// <summary>
 		/// Construct a new buffer
@@ -178,7 +179,7 @@ namespace EpicGames.UHT.Utils
 		/// <summary>
 		/// The borrowed buffer
 		/// </summary>
-		public UhtBuffer Buffer;
+		public UhtBuffer Buffer { get; set; }
 
 		/// <summary>
 		/// Borrow a buffer with the given size
diff --git a/Engine/Source/Programs/Shared/EpicGames.UHT/Utils/UhtException.cs b/Engine/Source/Programs/Shared/EpicGames.UHT/Utils/UhtException.cs
--- a/Engine/Source/Programs/Shared/EpicGames.UHT/Utils/UhtException.cs
+++ b/Engine/Source/Programs/Shared/EpicGames.UHT/Utils/UhtException.cs
@@ -14,7 +14,7 @@ namespace EpicGames.UHT.Utils
 		/// <summary>
 		/// The generated message
 		/// </summary>
-		public UhtMessage UhtMessage;
+		public UhtMessage UhtMessage { get; set; }
 
 		/// <summary>
 		/// Internal do nothing constructor
diff --git a/Engine/Source/Programs/Shared/EpicGames.UHT/Utils/UhtMessage.cs b/Engine/Source/Programs/Shared/EpicGames.UHT/Utils/UhtMessage.cs
--- a/Engine/Source/Programs/Shared/EpicGames.UHT/Utils/UhtMessage.cs
+++ b/Engine/Source/Programs/Shared/EpicGames.UHT/Utils/UhtMessage.cs
@@ -146,27 +146,27 @@ namespace EpicGames.UHT.Utils
 		/// <summary>
 		/// The type of message
 		/// </summary>
-		public UhtMessageType MessageType;
+		public UhtMessageType MessageType { get; set; }
 
 		/// <summary>
 		/// Optional message source for the message.  Either the MessageSource or FilePath must be set.
 		/// </summary>
-		public IUhtMessageSource? MessageSource;
+		public IUhtMessageSource? MessageSource { get; set; }
 
 		/// <summary>
 		/// Optional file path for the message.  Either the MessageSource or FilePath must be set.
 		/// </summary>
-		public string? FilePath;
+		public string? FilePath { get; set; }
 
 		/// <summary>
 		/// Line number where error occurred.
 		/// </summary>
-		public int LineNumber;
+		public int LineNumber { get; set; }
 
 		/// <summary>
 		/// Text of the message
 		/// </summary>
-		public string Message;
+		public string Message { get; set; }
 
 		/// <summary>
 		/// Make a new message with the given settings
diff --git a/Engine/Source/Programs/Shared/EpicGames.UHT/Utils/UhtReferenceCollector.cs b/Engine/Source/Programs/Shared/EpicGames.UHT/Utils/UhtReferenceCollector.cs
--- a/Engine/Source/Programs/Shared/EpicGames.UHT/Utils/UhtReferenceCollector.cs
+++ b/Engine/Source/Programs/Shared/EpicGames.UHT/Utils/UhtReferenceCollector.cs
@@ -64,12 +64,12 @@ namespace EpicGames.UHT.Utils
 		/// <summary>
 		/// Collection use to quickly detect if a reference is already in the collection
 		/// </summary>
-		private readonly HashSet<int> Uniques = new HashSet<int>();
+		private HashSet<int> Uniques { get; } = new HashSet<int>();
 
 		/// <summary>
 		/// List of all unique reference keys.  Use UngetKey to get the object index and the flag.
 		/// </summary>
-		public readonly List<int> References = new List<int>();
+		public List<int> References { get; } = new List<int>();
 
 		/// <summary>
 		/// Return an encoded key that represents the object and registered flag. 
@@ -165,32 +165,32 @@ namespace EpicGames.UHT.Utils
 		/// <summary>
 		/// Collection of unique cross module references
 		/// </summary>
-		public UhtUniqueReferenceCollection CrossModule = new UhtUniqueReferenceCollection();
+		public UhtUniqueReferenceCollection CrossModule { get; set; } = new UhtUniqueReferenceCollection();
 
 		/// <summary>
 		/// Collection of unique declarations
 		/// </summary>
-		public UhtUniqueReferenceCollection Declaration = new UhtUniqueReferenceCollection();
+		public UhtUniqueReferenceCollection Declaration { get; set; } = new UhtUniqueReferenceCollection();
 
 		/// <summary>
 		/// Collection of singletons
 		/// </summary>
-		public List<UhtField> Singletons = new List<UhtField>();
+		public List<UhtField> Singletons { get; } = new List<UhtField>();
 
 		/// <summary>
 		/// Collection of types to export
 		/// </summary>
-		public List<UhtField> ExportTypes = new List<UhtField>();
+		public List<UhtField> ExportTypes { get; } = new List<UhtField>();
 
 		/// <summary>
 		/// Collection of forward declarations
 		/// </summary>
-		public HashSet<string> ForwardDeclarations = new HashSet<string>();
+		public HashSet<string> ForwardDeclarations { get; } = new HashSet<string>();
 
 		/// <summary>
 		/// Collection of referenced headers
 		/// </summary>
-		public HashSet<UhtHeaderFile> ReferencedHeaders = new HashSet<UhtHeaderFile>();
+		public HashSet<UhtHeaderFile> ReferencedHeaders { get; } = new HashSet<UhtHeaderFile>();
 
 		/// <summary>
 		/// Add a cross module reference
diff --git a/Engine/Source/Programs/Shared/EpicGames.UHT/Utils/UhtSession.cs b/Engine/Source/Programs/Shared/EpicGames.UHT/Utils/UhtSession.cs
--- a/Engine/Source/Programs/Shared/EpicGames.UHT/Utils/UhtSession.cs
+++ b/Engine/Source/Programs/Shared/EpicGames.UHT/Utils/UhtSession.cs
@@ -29,22 +29,22 @@ namespace EpicGames.UHT.Utils
 		/// <summary>
 		/// When not null, this source comes from another source file
 		/// </summary>
-		public UhtSourceFile? SourceFile;
+		public UhtSourceFile? SourceFile { get; set; }
 
 		/// <summary>
 		/// The file path of the source
 		/// </summary>
-		public string FilePath;
+		public string FilePath { get; set; }
 
 		/// <summary>
 		/// The line number of the fragment in the containing source file.
 		/// </summary>
-		public int LineNumber;
+		public int LineNumber { get; set; }
 
 		/// <summary>
 		/// Data of the source file
 		/// </summary>
-		public StringView Data;
+		public StringView Data { get; set; }
 	}
 
 	/// <summary>
@@ -482,16 +482,13 @@ namespace EpicGames.UHT.Utils
 		/// <param name="FileName">File name to test</param>
 		/// <param name="Filters">List of wildcard filters</param>
 		/// <returns>True if there is a match</returns>
-		private static bool IsFilterMatch(string FileName, string[]? Filters)
+		private static bool IsFilterMatch(string FileName, IEnumerable<string> Filters)
 		{
-			if (Filters != null)
+			foreach (string Filter in Filters)
 			{
-				foreach (string Filter in Filters)
+				if (FileSystemName.MatchesSimpleExpression(Filter, FileName, true))
 				{
-					if (FileSystemName.MatchesSimpleExpression(Filter, FileName, true))
-					{
-						return true;
-					}
+					return true;
 				}
 			}
 			return false;
@@ -571,98 +568,98 @@ namespace EpicGames.UHT.Utils
 		/// <summary>
 		/// Interface used to read/write files
 		/// </summary>
-		public IUhtFileManager? FileManager;
+		public IUhtFileManager? FileManager { get; set; }
 
 		/// <summary>
 		/// Location of the engine code
 		/// </summary>
-		public string? EngineDirectory;
+		public string? EngineDirectory { get; set; }
 
 		/// <summary>
 		/// If set, the name of the project file.
 		/// </summary>
-		public string? ProjectFile;
+		public string? ProjectFile { get; set; }
 
 		/// <summary>
 		/// Optional location of the project
 		/// </summary>
-		public string? ProjectDirectory;
+		public string? ProjectDirectory { get; set; }
 
 		/// <summary>
 		/// Root directory for the engine.  This is usually just EngineDirectory without the Engine directory.
 		/// </summary>
-		public string? RootDirectory;
+		public string? RootDirectory { get; set; }
 
 		/// <summary>
 		/// Directory to store the reference output
 		/// </summary>
-		public string ReferenceDirectory = string.Empty;
+		public string ReferenceDirectory { get; set; } = string.Empty;
 
 		/// <summary>
 		/// Directory to store the verification output
 		/// </summary>
-		public string VerifyDirectory = string.Empty;
+		public string VerifyDirectory { get; set; } = string.Empty;
 
 		/// <summary>
 		/// Mode for generating and/or testing reference output
 		/// </summary>
-		public UhtReferenceMode ReferenceMode = UhtReferenceMode.None;
+		public UhtReferenceMode ReferenceMode { get; set; } = UhtReferenceMode.None;
 
 		/// <summary>
 		/// If true, warnings are considered to be errors
 		/// </summary>
-		public bool bWarningsAsErrors = false;
+		public bool bWarningsAsErrors { get; set; } = false;
 
 		/// <summary>
 		/// If true, include relative file paths in the log file
 		/// </summary>
-		public bool bRelativePathInLog = false;
+		public bool bRelativePathInLog { get; set; } = false;
 
 		/// <summary>
 		/// If true, use concurrent tasks to run UHT
 		/// </summary>
-		public bool bGoWide = true;
+		public bool bGoWide { get; set; } = true;
 
 		/// <summary>
 		/// If any output file mismatches existing outputs, an error will be generated
 		/// </summary>
-		public bool bFailIfGeneratedCodeChanges = false;
+		public bool bFailIfGeneratedCodeChanges { get; set; } = false;
 
 		/// <summary>
 		/// If true, no output files will be saved
 		/// </summary>
-		public bool bNoOutput = false;
+		public bool bNoOutput { get; set; } = false;
 
 		/// <summary>
 		/// If true, cull the output for any extra files
 		/// </summary>
-		public bool bCullOutput = true;
+		public bool bCullOutput { get; set; } = true;
 
 		/// <summary>
 		/// If true, include extra output in code generation
 		/// </summary>
-		public bool bIncludeDebugOutput = false;
+		public bool bIncludeDebugOutput { get; set; } = false;
 
 		/// <summary>
 		/// If true, disable all exporters which would normally be run by default
 		/// </summary>
-		public bool bNoDefaultExporters = false;
+		public bool bNoDefaultExporters { get; set; } = false;
 
 		/// <summary>
 		/// If true, cache any error messages until the end of processing.  This is used by the testing
 		/// harness to generate more stable console output.
 		/// </summary>
-		public bool bCacheMessages = false;
+		public bool bCacheMessages { get; set; } = false;
 
 		/// <summary>
 		/// Collection of UHT tables
 		/// </summary>
-		public UhtTables? Tables;
+		public UhtTables? Tables { get; set; }
 
 		/// <summary>
 		/// Configuration for the session
 		/// </summary>
-		public IUhtConfig? Config;
+		public IUhtConfig? Config { get; set; }
 		#endregion
 
 		/// <summary>
@@ -761,13 +758,13 @@ namespace EpicGames.UHT.Utils
 		/// After headers are parsed, returns the AActor class.  Unlike such properties as "UObject", there
 		/// is no requirement for AActor to be defined.  May be null.
 		/// </summary>
-		public UhtClass? AActor = null;
+		public UhtClass? AActor { get; set; } = null;
 
 		/// <summary>
 		/// After headers are parsed, return the INotifyFieldValueChanged interface.  There is no requirement 
 		/// that this interface be defined.
 		/// </summary>
-		public UhtClass? INotifyFieldValueChanged = null;
+		public UhtClass? INotifyFieldValueChanged { get; set; } = null;
 
 		private readonly List<UhtPackage> PackagesInternal = new List<UhtPackage>();
 		private readonly List<UhtHeaderFile> HeaderFilesInternal = new List<UhtHeaderFile>();
diff --git a/Engine/Source/Programs/Shared/EpicGames.UHT/Utils/UhtStringBuilderExtensions.cs b/Engine/Source/Programs/Shared/EpicGames.UHT/Utils/UhtStringBuilderExtensions.cs
--- a/Engine/Source/Programs/Shared/EpicGames.UHT/Utils/UhtStringBuilderExtensions.cs
+++ b/Engine/Source/Programs/Shared/EpicGames.UHT/Utils/UhtStringBuilderExtensions.cs
@@ -572,12 +572,12 @@ namespace EpicGames.UHT.Utils
 		/// <summary>
 		/// Owning cache
 		/// </summary>
-		private readonly StringBuilderCache Cache;
+		private StringBuilderCache Cache { get; }
 
 		/// <summary>
 		/// Borrowed string builder
 		/// </summary>
-		public readonly StringBuilder StringBuilder;
+		public StringBuilder StringBuilder { get; }
 
 		/// <summary>
 		/// Borrow a string builder from the given cache
diff --git a/Engine/Source/Programs/Shared/EpicGames.UHT/Utils/UhtTables.cs b/Engine/Source/Programs/Shared/EpicGames.UHT/Utils/UhtTables.cs
--- a/Engine/Source/Programs/Shared/EpicGames.UHT/Utils/UhtTables.cs
+++ b/Engine/Source/Programs/Shared/EpicGames.UHT/Utils/UhtTables.cs
@@ -96,12 +96,12 @@ namespace EpicGames.UHT.Utils
 		/// <summary>
 		/// This table inherits entries for the given table
 		/// </summary>
-		public UhtLookupTableBase? ParentTable = null;
+		public UhtLookupTableBase? ParentTable { get; set; } = null;
 
 		/// <summary>
 		/// Name of the table
 		/// </summary>
-		public string TableName = string.Empty;
+		public string TableName { get; set; } = string.Empty;
 
 		/// <summary>
 		/// User facing name of the table
@@ -115,17 +115,17 @@ namespace EpicGames.UHT.Utils
 		/// <summary>
 		/// Check to see if the table is internal
 		/// </summary>
-		public bool Internal = false;
+		public bool Internal { get; set; } = false;
 
 		/// <summary>
 		/// If true, this table has been implemented and not just created on demand by another table
 		/// </summary>
-		public bool Implemented = false;
+		public bool Implemented { get; set; } = false;
 
 		/// <summary>
 		/// Internal version of the user name.  If it hasn't been set, then the table name will be used
 		/// </summary>
-		private string UserNameInternal = string.Empty;
+		private string UserNameInternal { get; set; } = string.Empty;
 
 		/// <summary>
 		/// Merge the lookup table.  Duplicates will be ignored.
@@ -226,12 +226,12 @@ namespace EpicGames.UHT.Utils
 		/// <summary>
 		/// Collection of named tables
 		/// </summary>
-		public Dictionary<string, TTable> Tables = new Dictionary<string, TTable>();
+		public Dictionary<string, TTable> Tables { get; } = new Dictionary<string, TTable>();
 
 		/// <summary>
 		/// The name of the group of tables
 		/// </summary>
-		public string Name;
+		public string Name { get; set; }
 
 		/// <summary>
 		/// Create a new group of tables
@@ -418,7 +418,7 @@ namespace EpicGames.UHT.Utils
 		/// If specified, this method will be invoked once during the scan for attributes.
 		/// It can be used to perform some one time initialization.
 		/// </summary>
-		public string InitMethod = string.Empty;
+		public string InitMethod { get; set; } = string.Empty;
 	}
 
 	/// <summary>
@@ -452,42 +452,42 @@ namespace EpicGames.UHT.Utils
 		/// <summary>
 		/// Collection of specifier tables
 		/// </summary>
-		public readonly UhtSpecifierTables SpecifierTables = new UhtSpecifierTables();
+		public UhtSpecifierTables SpecifierTables { get; } = new UhtSpecifierTables();
 
 		/// <summary>
 		/// Collection of specifier validator tables
 		/// </summary>
-		public readonly UhtSpecifierValidatorTables SpecifierValidatorTables = new UhtSpecifierValidatorTables();
+		public UhtSpecifierValidatorTables SpecifierValidatorTables { get; } = new UhtSpecifierValidatorTables();
 
 		/// <summary>
 		/// Collection of keyword tables
 		/// </summary>
-		public readonly UhtKeywordTables KeywordTables = new UhtKeywordTables();
+		public UhtKeywordTables KeywordTables { get; } = new UhtKeywordTables();
 
 		/// <summary>
 		/// Collection of property types
 		/// </summary>
-		public readonly UhtPropertyTypeTable PropertyTypeTable = new UhtPropertyTypeTable();
+		public UhtPropertyTypeTable PropertyTypeTable { get; } = new UhtPropertyTypeTable();
 
 		/// <summary>
 		/// Collection of structure default values
 		/// </summary>
-		public readonly UhtStructDefaultValueTable StructDefaultValueTable = new UhtStructDefaultValueTable();
+		public UhtStructDefaultValueTable StructDefaultValueTable { get; } = new UhtStructDefaultValueTable();
 
 		/// <summary>
 		/// Collection of engine class types
 		/// </summary>
-		public readonly UhtEngineClassTable EngineClassTable = new UhtEngineClassTable();
+		public UhtEngineClassTable EngineClassTable { get; } = new UhtEngineClassTable();
 
 		/// <summary>
 		/// Collection of exporters
 		/// </summary>
-		public readonly UhtExporterTable ExporterTable = new UhtExporterTable();
+		public UhtExporterTable ExporterTable { get; } = new UhtExporterTable();
 
 		/// <summary>
 		/// Collection loc text default values
 		/// </summary>
-		public readonly UhtLocTextDefaultValueTable LocTextDefaultValueTable = new UhtLocTextDefaultValueTable();
+		public UhtLocTextDefaultValueTable LocTextDefaultValueTable { get; } = new UhtLocTextDefaultValueTable();
 
 		/// <summary>
 		/// Construct a new table collection
diff --git a/Engine/Source/Programs/Shared/EpicGames.UHT/Utils/UhtUtilities.cs b/Engine/Source/Programs/Shared/EpicGames.UHT/Utils/UhtUtilities.cs
--- a/Engine/Source/Programs/Shared/EpicGames.UHT/Utils/UhtUtilities.cs
+++ b/Engine/Source/Programs/Shared/EpicGames.UHT/Utils/UhtUtilities.cs
@@ -95,17 +95,17 @@ namespace EpicGames.UHT.Utils
 		/// <summary>
 		/// Any prefix removed from the source name to create the engine name
 		/// </summary>
-		public StringView Prefix;
+		public StringView Prefix { get; set; }
 
 		/// <summary>
 		/// The name to be used by the Unreal Engine.
 		/// </summary>
-		public StringView EngineName;
+		public StringView EngineName { get; set; }
 
 		/// <summary>
 		/// The name contained the "DEPRECATED" text which has been removed from the engine name
 		/// </summary>
-		public bool bIsDeprecated;
+		public bool bIsDeprecated { get; set; }
 	}
 
 	/// <summary>
