EpicGames.Perforce: Read environment settings from ~/.p4enviro on MacOS and Linux.

diff --git a/Engine/Source/Programs/Shared/EpicGames.Perforce/PerforceEnvironment.cs b/Engine/Source/Programs/Shared/EpicGames.Perforce/PerforceEnvironment.cs
--- a/Engine/Source/Programs/Shared/EpicGames.Perforce/PerforceEnvironment.cs
+++ b/Engine/Source/Programs/Shared/EpicGames.Perforce/PerforceEnvironment.cs
@@ -1,9 +1,10 @@
-// Copyright Epic Games, Inc. All Rights Reserved.
+﻿// Copyright Epic Games, Inc. All Rights Reserved.
 
 using System;
 using System.Collections;
 using System.Collections.Generic;
 using System.Diagnostics;
+using System.IO;
 using System.Linq;
 using System.Runtime.InteropServices;
 using EpicGames.Core;
@@ -165,10 +166,22 @@ namespace EpicGames.Perforce
 			}
 
 			string? enviroValue = environment.GetValue("P4ENVIRO");
+			if (enviroValue == null && !RuntimeInformation.IsOSPlatform(OSPlatform.Windows)) // Linux AND MacOS, despite what P4 docs say: https://www.perforce.com/manuals/v20.1/cmdref/Content/CmdRef/P4ENVIRO.html
+			{
+				string? homeDir = Environment.GetFolderPath(Environment.SpecialFolder.UserProfile);
+				if (homeDir != null)
+				{
+					enviroValue = Path.Combine(homeDir, ".p4enviro");
+				}
+			}
+
 			if (enviroValue != null)
 			{
 				FileReference location = new FileReference(enviroValue);
-				environment = new PerforceEnvironmentFile(environment, location);
+				if (FileReference.Exists(location))
+				{
+					environment = new PerforceEnvironmentFile(environment, location);
+				}
 			}
 
 			return environment;
