EpicGames.Perforce: Fall back to running external process when running on Apple silicon. We don't have a native P4 library for arm64 yet.

diff --git a/Engine/Source/Programs/Shared/EpicGames.Perforce/NativePerforceConnection.cs b/Engine/Source/Programs/Shared/EpicGames.Perforce/NativePerforceConnection.cs
--- a/Engine/Source/Programs/Shared/EpicGames.Perforce/NativePerforceConnection.cs
+++ b/Engine/Source/Programs/Shared/EpicGames.Perforce/NativePerforceConnection.cs
@@ -1,4 +1,4 @@
-// Copyright Epic Games, Inc. All Rights Reserved.
+﻿// Copyright Epic Games, Inc. All Rights Reserved.
 
 using System;
 using System.Collections.Concurrent;
@@ -324,9 +324,10 @@ namespace EpicGames.Perforce
 			NativePerforceConnection connection = new NativePerforceConnection(settings, logger);
 			await connection.ConnectAsync();
 			return connection;
-
 		}
 
+		public static bool IsSupported() => !(RuntimeInformation.IsOSPlatform(OSPlatform.OSX) && RuntimeInformation.OSArchitecture == Architecture.Arm64);
+
 		void GetNextWriteBuffer(NativeWriteBuffer nativeWriteBuffer, int minSize)
 		{
 			PinnedBuffer buffer = _writeBuffers.Take();
diff --git a/Engine/Source/Programs/Shared/EpicGames.Perforce/PerforceConnection.cs b/Engine/Source/Programs/Shared/EpicGames.Perforce/PerforceConnection.cs
--- a/Engine/Source/Programs/Shared/EpicGames.Perforce/PerforceConnection.cs
+++ b/Engine/Source/Programs/Shared/EpicGames.Perforce/PerforceConnection.cs
@@ -1,4 +1,4 @@
-// Copyright Epic Games, Inc. All Rights Reserved.
+﻿// Copyright Epic Games, Inc. All Rights Reserved.
 
 using System;
 using System.Collections.Generic;
@@ -76,7 +76,7 @@ namespace EpicGames.Perforce
 		/// <returns></returns>
 		public static async Task<IPerforceConnection> CreateAsync(IPerforceSettings settings, ILogger logger)
 		{
-			if (settings.PreferNativeClient)
+			if (settings.PreferNativeClient && NativePerforceConnection.IsSupported())
 			{
 				return await NativePerforceConnection.CreateAsync(settings, logger);
 			}
