More coding standard fixes
#rnx

diff --git a/Engine/Plugins/ScriptPlugin/Source/ScriptGeneratorUbtPlugin/ScriptCodeGeneratorBase.cs b/Engine/Plugins/ScriptPlugin/Source/ScriptGeneratorUbtPlugin/ScriptCodeGeneratorBase.cs
--- a/Engine/Plugins/ScriptPlugin/Source/ScriptGeneratorUbtPlugin/ScriptCodeGeneratorBase.cs
+++ b/Engine/Plugins/ScriptPlugin/Source/ScriptGeneratorUbtPlugin/ScriptCodeGeneratorBase.cs
@@ -30,12 +30,11 @@ namespace ScriptGeneratorUbtPlugin
 		{
 			DirectoryReference ConfigDirectory = DirectoryReference.Combine(Unreal.EngineDirectory, "Programs/UnrealBuildTool");
 			ConfigHierarchy Ini = ConfigCache.ReadHierarchy(ConfigHierarchyType.Engine, ConfigDirectory, BuildHostPlatform.Current.Platform);
-			List<string>? SupportedScriptModules = null;
-			Ini.GetArray("Plugins", "ScriptSupportedModules", out SupportedScriptModules);
+			Ini.GetArray("Plugins", "ScriptSupportedModules", out List<string>? SupportedScriptModules);
 
 			// Loop through the packages making sure they should be exported.  Queue the export of the classes
-			List<UhtClass> Classes = new List<UhtClass>();
-			List<Task?> Tasks = new List<Task?>();
+			List<UhtClass> Classes = new();
+			List<Task?> Tasks = new();
 			foreach (UhtPackage Package in this.Session.Packages)
 			{
 				if (Package.Module.ModuleType != UHTModuleType.EngineRuntime && Package.Module.ModuleType != UHTModuleType.GameRuntime)
@@ -76,7 +75,7 @@ namespace ScriptGeneratorUbtPlugin
 				if (CanExportClass(Class))
 				{
 					Classes.Add(Class);
-					Tasks.Add(Factory.CreateTask((Factory) => { ExportClass(Package, Class); }));
+					Tasks.Add(Factory.CreateTask((Factory) => { ExportClass(Class); }));
 				}
 			}
 			foreach (UhtType Child in Type.Children)
@@ -120,7 +119,7 @@ namespace ScriptGeneratorUbtPlugin
 				{
 					return false;
 				}
-				if (Child is UhtProperty Property && Property.bIsStaticArray)
+				if (Child is UhtProperty Property && Property.IsStaticArray)
 				{
 					return false;
 				}
@@ -150,7 +149,7 @@ namespace ScriptGeneratorUbtPlugin
 			}
 
 			// Reject if any of the parameter types is unsupported yet
-			if (Property.bIsStaticArray ||
+			if (Property.IsStaticArray ||
 				Property is UhtArrayProperty ||
 				Property is UhtDelegateProperty ||
 				Property is UhtMulticastDelegateProperty ||
@@ -167,26 +166,21 @@ namespace ScriptGeneratorUbtPlugin
 		/// Export the given class
 		/// </summary>
 		/// <param name="Factory">Factory associated with the export</param>
-		/// <param name="Package">Owning package</param>
 		/// <param name="Class">Class to export</param>
-		private void ExportClass(UhtPackage Package, UhtClass Class)
+		private void ExportClass(UhtClass Class)
 		{
-			using (BorrowStringBuilder Borrower = new BorrowStringBuilder(StringBuilderCache.Big))
-			{
-				ExportClass(Borrower.StringBuilder, Class);
-				string FileName = this.Factory.MakePath(Class.EngineName, ".script.h"); //ETSTODO - WRONG
-				this.Factory.CommitOutput(FileName, Borrower.StringBuilder);
-			}
+			using BorrowStringBuilder Borrower = new(StringBuilderCache.Big);
+			ExportClass(Borrower.StringBuilder, Class);
+			string FileName = this.Factory.MakePath(Class.EngineName, ".script.h");
+			this.Factory.CommitOutput(FileName, Borrower.StringBuilder);
 		}
 
 		private void Finish(List<UhtClass> Classes)
 		{
-			using (BorrowStringBuilder Borrower = new BorrowStringBuilder(StringBuilderCache.Big))
-			{
-				Finish(Borrower.StringBuilder, Classes);
-				string FileName = this.Factory.MakePath("GeneratedScriptLibraries", ".inl"); //ETSTODO - WRONG
-				this.Factory.CommitOutput(FileName, Borrower.StringBuilder);
-			}
+			using BorrowStringBuilder Borrower = new(StringBuilderCache.Big);
+			Finish(Borrower.StringBuilder, Classes);
+			string FileName = this.Factory.MakePath("GeneratedScriptLibraries", ".inl");
+			this.Factory.CommitOutput(FileName, Borrower.StringBuilder);
 		}
 
 		abstract protected void ExportClass(StringBuilder Builder, UhtClass Class);
diff --git a/Engine/Source/Programs/Shared/EpicGames.UHT/Exporters/CodeGen/UhtCodeGenerator.cs b/Engine/Source/Programs/Shared/EpicGames.UHT/Exporters/CodeGen/UhtCodeGenerator.cs
--- a/Engine/Source/Programs/Shared/EpicGames.UHT/Exporters/CodeGen/UhtCodeGenerator.cs
+++ b/Engine/Source/Programs/Shared/EpicGames.UHT/Exporters/CodeGen/UhtCodeGenerator.cs
@@ -1,14 +1,14 @@
 // Copyright Epic Games, Inc. All Rights Reserved.
 
-using EpicGames.Core;
-using EpicGames.UHT.Tables;
-using EpicGames.UHT.Types;
-using EpicGames.UHT.Utils;
 using System;
 using System.Collections.Generic;
 using System.IO;
 using System.Text;
 using System.Threading.Tasks;
+using EpicGames.Core;
+using EpicGames.UHT.Tables;
+using EpicGames.UHT.Types;
+using EpicGames.UHT.Utils;
 
 namespace EpicGames.UHT.Exporters.CodeGen
 {
@@ -18,121 +18,121 @@ namespace EpicGames.UHT.Exporters.CodeGen
 		[UhtExporter(Name = "CodeGen", Description = "Standard UnrealEngine code generation", Options = UhtExporterOptions.Default,
 			CppFilters = new string[] { "*.generated.cpp", "*.generated.*.cpp", "*.gen.cpp", "*.gen.*.cpp" },
 			HeaderFilters = new string[] { "*.generated.h" })]
-		public static void CodeGenerator(IUhtExportFactory Factory)
+		public static void CodeGenerator(IUhtExportFactory factory)
 		{
-			UhtCodeGenerator Generator = new UhtCodeGenerator(Factory);
-			Generator.Generate();
+			UhtCodeGenerator generator = new UhtCodeGenerator(factory);
+			generator.Generate();
 		}
 
 		public struct PackageInfo
 		{
-			public string StrippedName;
-			public string Api;
+			public string _strippedName;
+			public string _api;
 		}
-		public PackageInfo[] PackageInfos;
+		public PackageInfo[] _packageInfos;
 
 		public struct HeaderInfo
 		{
-			public Task? Task;
-			public string IncludePath;
-			public string FileId;
-			public uint BodyHash;
-			public bool bNeedsPushModelHeaders;
+			public Task? _task;
+			public string _includePath;
+			public string _fileId;
+			public uint _bodyHash;
+			public bool _needsPushModelHeaders;
 		}
-		public HeaderInfo[] HeaderInfos;
+		public HeaderInfo[] _headerInfos;
 
 		public struct ObjectInfo
 		{
-			public string RegisteredSingletonName;
-			public string UnregisteredSingletonName;
-			public string RegsiteredCrossReference;
-			public string UnregsiteredCrossReference;
-			public string RegsiteredExternalDecl;
-			public string UnregisteredExternalDecl;
-			public UhtClass? NativeInterface;
-			public uint Hash;
+			public string _registeredSingletonName;
+			public string _unregisteredSingletonName;
+			public string _regsiteredCrossReference;
+			public string _unregsiteredCrossReference;
+			public string _regsiteredExternalDecl;
+			public string _unregisteredExternalDecl;
+			public UhtClass? _nativeInterface;
+			public uint _hash;
 		}
-		public ObjectInfo[] ObjectInfos;
+		public ObjectInfo[] _objectInfos;
 
 		public readonly IUhtExportFactory Factory;
 		public UhtSession Session => Factory.Session;
 
-		private UhtCodeGenerator(IUhtExportFactory Factory)
+		private UhtCodeGenerator(IUhtExportFactory factory)
 		{
-			this.Factory = Factory;
-			this.HeaderInfos = new HeaderInfo[this.Factory.Session.HeaderFileTypeCount];
-			this.ObjectInfos = new ObjectInfo[this.Factory.Session.ObjectTypeCount];
-			this.PackageInfos = new PackageInfo[this.Factory.Session.PackageTypeCount];
+			this.Factory = factory;
+			this._headerInfos = new HeaderInfo[this.Factory.Session.HeaderFileTypeCount];
+			this._objectInfos = new ObjectInfo[this.Factory.Session.ObjectTypeCount];
+			this._packageInfos = new PackageInfo[this.Factory.Session.PackageTypeCount];
 		}
 
 		private void Generate()
 		{
-			List<Task?> Prereqs = new List<Task?>();
+			List<Task?> prereqs = new List<Task?>();
 
 			// Perform some startup initialization to compute things we need over and over again
-			if (this.Session.bGoWide)
+			if (this.Session.GoWide)
 			{
-				Parallel.ForEach(this.Factory.Session.Packages, Package =>
+				Parallel.ForEach(this.Factory.Session.Packages, package =>
 				{
-					InitPackageInfo(Package);
+					InitPackageInfo(package);
 				});
 			}
 			else
 			{
-				foreach (UhtPackage Package in this.Factory.Session.Packages)
+				foreach (UhtPackage package in this.Factory.Session.Packages)
 				{
-					InitPackageInfo(Package);
+					InitPackageInfo(package);
 				}
 			}
 
 			// Generate the files for the header files
-			foreach (UhtHeaderFile HeaderFile in this.Session.SortedHeaderFiles)
+			foreach (UhtHeaderFile headerFile in this.Session.SortedHeaderFiles)
 			{
-				if (HeaderFile.bShouldExport)
+				if (headerFile.ShouldExport)
 				{
-					UhtPackage Package = HeaderFile.Package;
-					UHTManifest.Module Module = Package.Module;
+					UhtPackage package = headerFile.Package;
+					UHTManifest.Module module = package.Module;
 
-					Prereqs.Clear();
-					foreach (UhtHeaderFile Referenced in HeaderFile.ReferencedHeadersNoLock)
+					prereqs.Clear();
+					foreach (UhtHeaderFile referenced in headerFile.ReferencedHeadersNoLock)
 					{
-						if (HeaderFile != Referenced)
+						if (headerFile != referenced)
 						{
-							Prereqs.Add(this.HeaderInfos[Referenced.HeaderFileTypeIndex].Task);
+							prereqs.Add(this._headerInfos[referenced.HeaderFileTypeIndex]._task);
 						}
 					}
 
-					this.HeaderInfos[HeaderFile.HeaderFileTypeIndex].Task = Factory.CreateTask(Prereqs,  
-						(IUhtExportFactory Factory) =>
+					this._headerInfos[headerFile.HeaderFileTypeIndex]._task = Factory.CreateTask(prereqs,
+						(IUhtExportFactory factory) =>
 						{
-							new UhtHeaderCodeGeneratorHFile(this, Package, HeaderFile).Generate(Factory);
-							new UhtHeaderCodeGeneratorCppFile(this, Package, HeaderFile).Generate(Factory);
+							new UhtHeaderCodeGeneratorHFile(this, package, headerFile).Generate(factory);
+							new UhtHeaderCodeGeneratorCppFile(this, package, headerFile).Generate(factory);
 						});
 				}
 			}
 
 			// Generate the files for the packages
-			List<Task?> GeneratedPackages = new List<Task?>(this.Session.PackageTypeCount);
-			foreach (UhtPackage Package in this.Session.Packages)
+			List<Task?> generatedPackages = new List<Task?>(this.Session.PackageTypeCount);
+			foreach (UhtPackage package in this.Session.Packages)
 			{
-				UHTManifest.Module Module = Package.Module;
+				UHTManifest.Module module = package.Module;
 
-				bool bWriteHeader = false;
-				Prereqs.Clear();
-				foreach (UhtHeaderFile HeaderFile in Package.Children)
+				bool writeHeader = false;
+				prereqs.Clear();
+				foreach (UhtHeaderFile headerFile in package.Children)
 				{
-					Prereqs.Add(this.HeaderInfos[HeaderFile.HeaderFileTypeIndex].Task);
-					if (!bWriteHeader)
+					prereqs.Add(this._headerInfos[headerFile.HeaderFileTypeIndex]._task);
+					if (!writeHeader)
 					{
-						foreach (UhtType Type in HeaderFile.Children)
+						foreach (UhtType type in headerFile.Children)
 						{
-							if (Type is UhtClass Class)
+							if (type is UhtClass classObj)
 							{
-								if (Class.ClassType != UhtClassType.NativeInterface && 
-									Class.ClassFlags.HasExactFlags(EClassFlags.Native | EClassFlags.Intrinsic, EClassFlags.Native) &&
-									!Class.ClassExportFlags.HasAllFlags(UhtClassExportFlags.NoExport))
+								if (classObj.ClassType != UhtClassType.NativeInterface &&
+									classObj.ClassFlags.HasExactFlags(EClassFlags.Native | EClassFlags.Intrinsic, EClassFlags.Native) &&
+									!classObj.ClassExportFlags.HasAllFlags(UhtClassExportFlags.NoExport))
 								{
-									bWriteHeader = true;
+									writeHeader = true;
 									break;
 								}
 							}
@@ -140,254 +140,251 @@ namespace EpicGames.UHT.Exporters.CodeGen
 					}
 				}
 
-				GeneratedPackages.Add(Factory.CreateTask(Prereqs, 
-					(IUhtExportFactory Factory) =>
+				generatedPackages.Add(Factory.CreateTask(prereqs,
+					(IUhtExportFactory factory) =>
 					{
-						List<UhtHeaderFile> PackageSortedHeaders = GetSortedHeaderFiles(Package);
-						if (bWriteHeader)
+						List<UhtHeaderFile> packageSortedHeaders = GetSortedHeaderFiles(package);
+						if (writeHeader)
 						{
-							new UhtPackageCodeGeneratorHFile(this, Package).Generate(Factory, PackageSortedHeaders);
+							new UhtPackageCodeGeneratorHFile(this, package).Generate(factory, packageSortedHeaders);
 						}
-						new UhtPackageCodeGeneratorCppFile(this, Package).Generate(Factory, PackageSortedHeaders);
+						new UhtPackageCodeGeneratorCppFile(this, package).Generate(factory, packageSortedHeaders);
 					}));
 			}
 
 			// Wait for all the packages to complete
-			List<Task> PackageTasks = new List<Task>(this.Session.PackageTypeCount);
-			foreach (Task? Output in GeneratedPackages)
+			List<Task> packageTasks = new List<Task>(this.Session.PackageTypeCount);
+			foreach (Task? output in generatedPackages)
 			{
-				if (Output != null)
+				if (output != null)
 				{
-					PackageTasks.Add(Output);
+					packageTasks.Add(output);
 				}
 			}
-			Task.WaitAll(PackageTasks.ToArray());
+			Task.WaitAll(packageTasks.ToArray());
 		}
 
-#region Utility functions
+		#region Utility functions
 		/// <summary>
 		/// Return the singleton name for an object
 		/// </summary>
-		/// <param name="Object">The object in question.</param>
-		/// <param name="bRegistered">If true, return the registered singleton name.  Otherwise return the unregistered.</param>
+		/// <param name="obj">The object in question.</param>
+		/// <param name="registered">If true, return the registered singleton name.  Otherwise return the unregistered.</param>
 		/// <returns>Singleton name or "nullptr" if Object is null</returns>
-		public string GetSingletonName(UhtObject? Object, bool bRegistered)
+		public string GetSingletonName(UhtObject? obj, bool registered)
 		{
-			if (Object == null)
+			if (obj == null)
 			{
 				return "nullptr";
 			}
-			return bRegistered ? this.ObjectInfos[Object.ObjectTypeIndex].RegisteredSingletonName : this.ObjectInfos[Object.ObjectTypeIndex].UnregisteredSingletonName;
+			return registered ? this._objectInfos[obj.ObjectTypeIndex]._registeredSingletonName : this._objectInfos[obj.ObjectTypeIndex]._unregisteredSingletonName;
 		}
 
 		/// <summary>
 		/// Return the external declaration for an object
 		/// </summary>
-		/// <param name="Object">The object in question.</param>
-		/// <param name="bRegistered">If true, return the registered external declaration.  Otherwise return the unregistered.</param>
+		/// <param name="obj">The object in question.</param>
+		/// <param name="registered">If true, return the registered external declaration.  Otherwise return the unregistered.</param>
 		/// <returns>External declaration</returns>
-		public string GetExternalDecl(UhtObject Object, bool bRegistered)
+		public string GetExternalDecl(UhtObject obj, bool registered)
 		{
-			return GetExternalDecl(Object.ObjectTypeIndex, bRegistered);
+			return GetExternalDecl(obj.ObjectTypeIndex, registered);
 		}
 
 		/// <summary>
 		/// Return the external declaration for an object
 		/// </summary>
-		/// <param name="ObjectIndex">The object in question.</param>
-		/// <param name="bRegistered">If true, return the registered external declaration.  Otherwise return the unregistered.</param>
+		/// <param name="objectIndex">The object in question.</param>
+		/// <param name="registered">If true, return the registered external declaration.  Otherwise return the unregistered.</param>
 		/// <returns>External declaration</returns>
-		public string GetExternalDecl(int ObjectIndex, bool bRegistered)
+		public string GetExternalDecl(int objectIndex, bool registered)
 		{
-			return bRegistered ? this.ObjectInfos[ObjectIndex].RegsiteredExternalDecl : this.ObjectInfos[ObjectIndex].UnregisteredExternalDecl;
+			return registered ? this._objectInfos[objectIndex]._regsiteredExternalDecl : this._objectInfos[objectIndex]._unregisteredExternalDecl;
 		}
 
 		/// <summary>
 		/// Return the cross reference for an object
 		/// </summary>
-		/// <param name="Object">The object in question.</param>
-		/// <param name="bRegistered">If true, return the registered cross reference.  Otherwise return the unregistered.</param>
+		/// <param name="obj">The object in question.</param>
+		/// <param name="registered">If true, return the registered cross reference.  Otherwise return the unregistered.</param>
 		/// <returns>Cross reference</returns>
-		public string GetCrossReference(UhtObject Object, bool bRegistered)
+		public string GetCrossReference(UhtObject obj, bool registered)
 		{
-			return GetCrossReference(Object.ObjectTypeIndex, bRegistered);
+			return GetCrossReference(obj.ObjectTypeIndex, registered);
 		}
 
 		/// <summary>
 		/// Return the cross reference for an object
 		/// </summary>
-		/// <param name="ObjectIndex">The object in question.</param>
-		/// <param name="bRegistered">If true, return the registered cross reference.  Otherwise return the unregistered.</param>
+		/// <param name="objectIndex">The object in question.</param>
+		/// <param name="registered">If true, return the registered cross reference.  Otherwise return the unregistered.</param>
 		/// <returns>Cross reference</returns>
-		public string GetCrossReference(int ObjectIndex, bool bRegistered)
+		public string GetCrossReference(int objectIndex, bool registered)
 		{
-			return bRegistered ? this.ObjectInfos[ObjectIndex].RegsiteredCrossReference : this.ObjectInfos[ObjectIndex].UnregsiteredCrossReference;
+			return registered ? this._objectInfos[objectIndex]._regsiteredCrossReference : this._objectInfos[objectIndex]._unregsiteredCrossReference;
 		}
-#endregion
+		#endregion
 
-#region Information initialization
-		private void InitPackageInfo(UhtPackage Package)
+		#region Information initialization
+		private void InitPackageInfo(UhtPackage package)
 		{
-			StringBuilder Builder = new StringBuilder();
+			StringBuilder builder = new StringBuilder();
 
-			ref PackageInfo PackageInfo = ref this.PackageInfos[Package.PackageTypeIndex];
-			PackageInfo.StrippedName = Package.SourceName.Replace('/', '_');
-			PackageInfo.Api = $"{Package.ShortName.ToString().ToUpper()}_API ";
+			ref PackageInfo packageInfo = ref this._packageInfos[package.PackageTypeIndex];
+			packageInfo._strippedName = package.SourceName.Replace('/', '_');
+			packageInfo._api = $"{package.ShortName.ToString().ToUpper()}_API ";
 
 			// Construct the names used commonly during export
-			ref ObjectInfo ObjectInfo = ref this.ObjectInfos[Package.ObjectTypeIndex];
-			Builder.Append("Z_Construct_UPackage_");
-			Builder.Append(PackageInfo.StrippedName);
-			ObjectInfo.UnregisteredSingletonName = ObjectInfo.RegisteredSingletonName = Builder.ToString();
-			ObjectInfo.UnregisteredExternalDecl = ObjectInfo.RegsiteredExternalDecl = $"\t{PackageInfo.Api}_API UPackage* {ObjectInfo.RegisteredSingletonName}();\r\n"; //COMPATIBILITY-TODO remove the extra _API
-			ObjectInfo.UnregsiteredCrossReference = ObjectInfo.RegsiteredCrossReference = $"\tUPackage* {ObjectInfo.RegisteredSingletonName}();\r\n";
-
-			foreach (UhtHeaderFile HeaderFile in Package.Children)
+			ref ObjectInfo objectInfo = ref this._objectInfos[package.ObjectTypeIndex];
+			builder.Append("Z_Construct_UPackage_");
+			builder.Append(packageInfo._strippedName);
+			objectInfo._unregisteredSingletonName = objectInfo._registeredSingletonName = builder.ToString();
+			objectInfo._unregisteredExternalDecl = objectInfo._regsiteredExternalDecl = $"\t{packageInfo._api}_API UPackage* {objectInfo._registeredSingletonName}();\r\n"; //COMPATIBILITY-TODO remove the extra _API
+			objectInfo._unregsiteredCrossReference = objectInfo._regsiteredCrossReference = $"\tUPackage* {objectInfo._registeredSingletonName}();\r\n";
+
+			foreach (UhtHeaderFile headerFile in package.Children)
 			{
-				InitHeaderInfo(Builder, Package, ref PackageInfo, HeaderFile);
+				InitHeaderInfo(builder, package, ref packageInfo, headerFile);
 			}
 		}
 
-		private void InitHeaderInfo(StringBuilder Builder, UhtPackage Package, ref PackageInfo PackageInfo, UhtHeaderFile HeaderFile)
+		private void InitHeaderInfo(StringBuilder builder, UhtPackage package, ref PackageInfo packageInfo, UhtHeaderFile headerFile)
 		{
-			ref HeaderInfo HeaderInfo = ref this.HeaderInfos[HeaderFile.HeaderFileTypeIndex];
+			ref HeaderInfo headerInfo = ref this._headerInfos[headerFile.HeaderFileTypeIndex];
 
-			HeaderInfo.IncludePath = Path.GetRelativePath(Package.Module.IncludeBase, HeaderFile.FilePath).Replace('\\', '/');
+			headerInfo._includePath = Path.GetRelativePath(package.Module.IncludeBase, headerFile.FilePath).Replace('\\', '/');
 
 			// Convert the file path to a C identifier
-			string FilePath = HeaderFile.FilePath;
-			bool bIsRelative = !Path.IsPathRooted(FilePath);
-			if (!bIsRelative && Session.EngineDirectory != null)
+			string filePath = headerFile.FilePath;
+			bool isRelative = !Path.IsPathRooted(filePath);
+			if (!isRelative && Session.EngineDirectory != null)
 			{
-				string? Directory = Path.GetDirectoryName(Session.EngineDirectory);
-				if (!string.IsNullOrEmpty(Directory))
+				string? directory = Path.GetDirectoryName(Session.EngineDirectory);
+				if (!String.IsNullOrEmpty(directory))
 				{
-					FilePath = Path.GetRelativePath(Directory, FilePath);
-					bIsRelative = !Path.IsPathRooted(FilePath);
+					filePath = Path.GetRelativePath(directory, filePath);
+					isRelative = !Path.IsPathRooted(filePath);
 				}
 			}
-			if (!bIsRelative && Session.ProjectDirectory != null)
+			if (!isRelative && Session.ProjectDirectory != null)
 			{
-				string? Directory = Path.GetDirectoryName(Session.ProjectDirectory);
-				if (!string.IsNullOrEmpty(Directory))
+				string? directory = Path.GetDirectoryName(Session.ProjectDirectory);
+				if (!String.IsNullOrEmpty(directory))
 				{
-					FilePath = Path.GetRelativePath(Directory, FilePath);
-					bIsRelative = !Path.IsPathRooted(FilePath);
+					filePath = Path.GetRelativePath(directory, filePath);
+					isRelative = !Path.IsPathRooted(filePath);
 				}
 			}
-			FilePath = FilePath.Replace('\\', '/');
-			if (bIsRelative)
+			filePath = filePath.Replace('\\', '/');
+			if (isRelative)
 			{
-				while (FilePath.StartsWith("../", StringComparison.Ordinal))
+				while (filePath.StartsWith("../", StringComparison.Ordinal))
 				{
-					FilePath = FilePath.Substring(3);
+					filePath = filePath.Substring(3);
 				}
 			}
 
-			char[] OutFilePath = new char[FilePath.Length + 4];
-			OutFilePath[0] = 'F';
-			OutFilePath[1] = 'I';
-			OutFilePath[2] = 'D';
-			OutFilePath[3] = '_';
-			for (int Index = 0; Index < FilePath.Length; ++Index)
+			char[] outFilePath = new char[filePath.Length + 4];
+			outFilePath[0] = 'F';
+			outFilePath[1] = 'I';
+			outFilePath[2] = 'D';
+			outFilePath[3] = '_';
+			for (int index = 0; index < filePath.Length; ++index)
 			{
-				OutFilePath[Index+4] = UhtFCString.IsAlnum(FilePath[Index]) ? FilePath[Index] : '_';
+				outFilePath[index + 4] = UhtFCString.IsAlnum(filePath[index]) ? filePath[index] : '_';
 			}
-			HeaderInfo.FileId = new string(OutFilePath);
+			headerInfo._fileId = new string(outFilePath);
 
-			foreach (UhtObject Object in HeaderFile.Children)
+			foreach (UhtObject obj in headerFile.Children)
 			{
-				InitObjectInfo(Builder, Package, ref PackageInfo, ref HeaderInfo, Object);
-				if (Object is UhtClass Class)
-				{
-				}
+				InitObjectInfo(builder, package, ref packageInfo, ref headerInfo, obj);
 			}
 		}
 
-		private void InitObjectInfo(StringBuilder Builder, UhtPackage Package, ref PackageInfo PackageInfo, ref HeaderInfo HeaderInfo, UhtObject Object)
+		private void InitObjectInfo(StringBuilder builder, UhtPackage package, ref PackageInfo packageInfo, ref HeaderInfo headerInfo, UhtObject obj)
 		{
-			ref ObjectInfo ObjectInfo = ref this.ObjectInfos[Object.ObjectTypeIndex];
+			ref ObjectInfo objectInfo = ref this._objectInfos[obj.ObjectTypeIndex];
 
-			Builder.Clear();
+			builder.Clear();
 
 			// Construct the names used commonly during export
-			bool bIsNonIntrinsicClass = false;
-			Builder.Append("Z_Construct_U").Append(Object.EngineClassName).AppendOuterNames(Object);
+			bool isNonIntrinsicClass = false;
+			builder.Append("Z_Construct_U").Append(obj.EngineClassName).AppendOuterNames(obj);
 
-			string EngineClassName = Object.EngineClassName;
-			if (Object is UhtClass Class)
+			string engineClassName = obj.EngineClassName;
+			if (obj is UhtClass classObj)
 			{
-				if (!Class.ClassFlags.HasAnyFlags(EClassFlags.Intrinsic))
+				if (!classObj.ClassFlags.HasAnyFlags(EClassFlags.Intrinsic))
 				{
-					bIsNonIntrinsicClass = true;
+					isNonIntrinsicClass = true;
 				}
-				if (Class.ClassExportFlags.HasExactFlags(UhtClassExportFlags.HasReplciatedProperties, UhtClassExportFlags.SelfHasReplicatedProperties))
+				if (classObj.ClassExportFlags.HasExactFlags(UhtClassExportFlags.HasReplciatedProperties, UhtClassExportFlags.SelfHasReplicatedProperties))
 				{
-					HeaderInfo.bNeedsPushModelHeaders = true;
+					headerInfo._needsPushModelHeaders = true;
 				}
-				if (Class.ClassType == UhtClassType.NativeInterface)
+				if (classObj.ClassType == UhtClassType.NativeInterface)
 				{
-					if (Class.AlternateObject != null)
+					if (classObj.AlternateObject != null)
 					{
-						this.ObjectInfos[Class.AlternateObject.ObjectTypeIndex].NativeInterface = Class;
+						this._objectInfos[classObj.AlternateObject.ObjectTypeIndex]._nativeInterface = classObj;
 					}
 				}
 			}
-			else if (Object is UhtFunction)
+			else if (obj is UhtFunction)
 			{
 				// The method for EngineClassName returns type specific where in this case we need just the simple return type
-				EngineClassName = "Function";
+				engineClassName = "Function";
 			}
 
-			if (bIsNonIntrinsicClass)
+			if (isNonIntrinsicClass)
 			{
-				ObjectInfo.RegisteredSingletonName = Builder.ToString();
-				Builder.Append("_NoRegister");
-				ObjectInfo.UnregisteredSingletonName = Builder.ToString();
+				objectInfo._registeredSingletonName = builder.ToString();
+				builder.Append("_NoRegister");
+				objectInfo._unregisteredSingletonName = builder.ToString();
 
-				ObjectInfo.UnregisteredExternalDecl = $"\t{PackageInfo.Api}U{EngineClassName}* {ObjectInfo.UnregisteredSingletonName}();\r\n";
-				ObjectInfo.RegsiteredExternalDecl = $"\t{PackageInfo.Api}U{EngineClassName}* {ObjectInfo.RegisteredSingletonName}();\r\n";
+				objectInfo._unregisteredExternalDecl = $"\t{packageInfo._api}U{engineClassName}* {objectInfo._unregisteredSingletonName}();\r\n";
+				objectInfo._regsiteredExternalDecl = $"\t{packageInfo._api}U{engineClassName}* {objectInfo._registeredSingletonName}();\r\n";
 			}
 			else
 			{
-				ObjectInfo.UnregisteredSingletonName = ObjectInfo.RegisteredSingletonName = Builder.ToString();
-				ObjectInfo.UnregisteredExternalDecl = ObjectInfo.RegsiteredExternalDecl = $"\t{PackageInfo.Api}U{EngineClassName}* {ObjectInfo.RegisteredSingletonName}();\r\n";
+				objectInfo._unregisteredSingletonName = objectInfo._registeredSingletonName = builder.ToString();
+				objectInfo._unregisteredExternalDecl = objectInfo._regsiteredExternalDecl = $"\t{packageInfo._api}U{engineClassName}* {objectInfo._registeredSingletonName}();\r\n";
 			}
 
 			//COMPATIBILITY-TODO - The cross reference string should match the extern decl string always.  But currently, it is different for packages.
-			ObjectInfo.UnregsiteredCrossReference = ObjectInfo.UnregisteredExternalDecl;
-			ObjectInfo.RegsiteredCrossReference = ObjectInfo.RegsiteredExternalDecl;
+			objectInfo._unregsiteredCrossReference = objectInfo._unregisteredExternalDecl;
+			objectInfo._regsiteredCrossReference = objectInfo._regsiteredExternalDecl;
 
 			// Init the children
-			foreach (UhtType Child in Object.Children)
+			foreach (UhtType child in obj.Children)
 			{
-				if (Child is UhtObject ChildObject)
+				if (child is UhtObject childObject)
 				{
-					InitObjectInfo(Builder, Package, ref PackageInfo, ref HeaderInfo, ChildObject);
+					InitObjectInfo(builder, package, ref packageInfo, ref headerInfo, childObject);
 				}
 			}
 		}
-#endregion
+		#endregion
 
-#region Utility functions
+		#region Utility functions
 		/// <summary>
 		/// Return a package's sorted header file list of all header files that or referenced or have declarations.
 		/// </summary>
-		/// <param name="Package">The package in question</param>
+		/// <param name="package">The package in question</param>
 		/// <returns>Sorted list of the header files</returns>
-		private static List<UhtHeaderFile> GetSortedHeaderFiles(UhtPackage Package)
+		private static List<UhtHeaderFile> GetSortedHeaderFiles(UhtPackage package)
 		{
-			List<UhtHeaderFile> Out = new List<UhtHeaderFile>(Package.Children.Count);
-			foreach (UhtHeaderFile HeaderFile in Package.Children)
+			List<UhtHeaderFile> sortedHeaders = new List<UhtHeaderFile>(package.Children.Count);
+			foreach (UhtHeaderFile headerFile in package.Children)
 			{
-				if (HeaderFile.bShouldExport)
+				if (headerFile.ShouldExport)
 				{
-					Out.Add(HeaderFile);
+					sortedHeaders.Add(headerFile);
 				}
 			}
-			Out.Sort((Lhs, Rhs) => { return StringComparerUE.OrdinalIgnoreCase.Compare(Lhs.FilePath, Rhs.FilePath); });
-			return Out;
+			sortedHeaders.Sort((lhs, rhs) => { return StringComparerUE.OrdinalIgnoreCase.Compare(lhs.FilePath, rhs.FilePath); });
+			return sortedHeaders;
 		}
-#endregion
+		#endregion
 	}
 }
diff --git a/Engine/Source/Programs/Shared/EpicGames.UHT/Exporters/CodeGen/UhtHeaderCodeGenerator.cs b/Engine/Source/Programs/Shared/EpicGames.UHT/Exporters/CodeGen/UhtHeaderCodeGenerator.cs
--- a/Engine/Source/Programs/Shared/EpicGames.UHT/Exporters/CodeGen/UhtHeaderCodeGenerator.cs
+++ b/Engine/Source/Programs/Shared/EpicGames.UHT/Exporters/CodeGen/UhtHeaderCodeGenerator.cs
@@ -1,14 +1,14 @@
 // Copyright Epic Games, Inc. All Rights Reserved.
 
+using System;
+using System.Text;
 using EpicGames.Core;
 using EpicGames.UHT.Types;
 using EpicGames.UHT.Utils;
-using System;
-using System.Text;
 
 namespace EpicGames.UHT.Exporters.CodeGen
 {
-	internal class UhtHeaderCodeGenerator : UhtPackageCodeGenerator
+	class UhtHeaderCodeGenerator : UhtPackageCodeGenerator
 	{
 		#region Define block macro names
 		public const string EventParamsMacroSuffix = "EVENT_PARMS";
@@ -35,396 +35,396 @@ namespace EpicGames.UHT.Exporters.CodeGen
 		#endregion
 
 		public readonly UhtHeaderFile HeaderFile;
-		public string FileId => this.HeaderInfos[this.HeaderFile.HeaderFileTypeIndex].FileId;
+		public string FileId => this.HeaderInfos[this.HeaderFile.HeaderFileTypeIndex]._fileId;
 
-		public UhtHeaderCodeGenerator(UhtCodeGenerator CodeGenerator, UhtPackage Package, UhtHeaderFile HeaderFile)
-			: base(CodeGenerator, Package)
+		public UhtHeaderCodeGenerator(UhtCodeGenerator codeGenerator, UhtPackage package, UhtHeaderFile headerFile)
+			: base(codeGenerator, package)
 		{
-			this.HeaderFile = HeaderFile;
+			this.HeaderFile = headerFile;
 		}
 
 		#region Event parameter
-		protected static StringBuilder AppendEventParameter(StringBuilder Builder, UhtFunction Function, string StrippedFunctionName, UhtPropertyTextType TextType, bool bOutputConstructor, int Tabs, string Endl)
+		protected static StringBuilder AppendEventParameter(StringBuilder builder, UhtFunction function, string strippedFunctionName, UhtPropertyTextType textType, bool outputConstructor, int tabs, string endl)
 		{
-			if (!WillExportEventParameters(Function))
+			if (!WillExportEventParameters(function))
 			{
-				return Builder;
+				return builder;
 			}
 
-			string EventParameterStructName = GetEventStructParametersName(Function.Outer, StrippedFunctionName);
+			string eventParameterStructName = GetEventStructParametersName(function.Outer, strippedFunctionName);
 
-			Builder.AppendTabs(Tabs).Append("struct ").Append(EventParameterStructName).Append(Endl);
-			Builder.AppendTabs(Tabs).Append('{').Append(Endl);
+			builder.AppendTabs(tabs).Append("struct ").Append(eventParameterStructName).Append(endl);
+			builder.AppendTabs(tabs).Append('{').Append(endl);
 
-			++Tabs;
-			foreach (UhtProperty Property in Function.Children)
+			++tabs;
+			foreach (UhtProperty property in function.Children)
 			{
-				bool bEmitConst = Property.PropertyFlags.HasAnyFlags(EPropertyFlags.ConstParm) && Property is UhtObjectProperty;
+				bool emitConst = property.PropertyFlags.HasAnyFlags(EPropertyFlags.ConstParm) && property is UhtObjectProperty;
 
 				//@TODO: UCREMOVAL: This is awful code duplication to avoid a double-const
 				{
 					//@TODO: bEmitConst will only be true if we have an object, so checking interface here doesn't do anything.
 					// export 'const' for parameters
-					bool bIsConstParam = Property is UhtInterfaceProperty && !Property.PropertyFlags.HasAnyFlags(EPropertyFlags.OutParm); //@TODO: This should be const once that flag exists
-					bool bIsOnConstClass = false;
-					if (Property is UhtObjectProperty ObjectProperty)
+					bool isConstParam = property is UhtInterfaceProperty && !property.PropertyFlags.HasAnyFlags(EPropertyFlags.OutParm); //@TODO: This should be const once that flag exists
+					bool isOnConstClass = false;
+					if (property is UhtObjectProperty objectProperty)
 					{
-						bIsOnConstClass = ObjectProperty.Class.ClassFlags.HasAnyFlags(EClassFlags.Const);
+						isOnConstClass = objectProperty.Class.ClassFlags.HasAnyFlags(EClassFlags.Const);
 					}
 
-					if (bIsConstParam || bIsOnConstClass)
+					if (isConstParam || isOnConstClass)
 					{
-						bEmitConst = false; // ExportCppDeclaration will do it for us
+						emitConst = false; // ExportCppDeclaration will do it for us
 					}
 				}
 
-				Builder.AppendTabs(Tabs);
-				if (bEmitConst)
+				builder.AppendTabs(tabs);
+				if (emitConst)
 				{
-					Builder.Append("const ");
+					builder.Append("const ");
 				}
 
-				Builder.AppendFullDecl(Property, TextType, false);
-				Builder.Append(';').Append(Endl);
+				builder.AppendFullDecl(property, textType, false);
+				builder.Append(';').Append(endl);
 			}
 
-			if (bOutputConstructor)
+			if (outputConstructor)
 			{
-				UhtProperty? ReturnProperty = Function.ReturnProperty;
-				if (ReturnProperty != null && ReturnProperty.PropertyCaps.HasAnyFlags(UhtPropertyCaps.RequiresNullConstructorArg))
+				UhtProperty? returnProperty = function.ReturnProperty;
+				if (returnProperty != null && returnProperty.PropertyCaps.HasAnyFlags(UhtPropertyCaps.RequiresNullConstructorArg))
 				{
-					Builder.Append(Endl);
-					Builder.AppendTabs(Tabs).Append("/** Constructor, initializes return property only **/").Append(Endl);
-					Builder.AppendTabs(Tabs).Append(EventParameterStructName).Append("()").Append(Endl);
-					Builder.AppendTabs(Tabs + 1).Append(": ").Append(ReturnProperty.SourceName).Append('(').AppendNullConstructorArg(ReturnProperty, true).Append(')').Append(Endl);
-					Builder.AppendTabs(Tabs).Append('{').Append(Endl);
-					Builder.AppendTabs(Tabs).Append('}').Append(Endl);
+					builder.Append(endl);
+					builder.AppendTabs(tabs).Append("/** Constructor, initializes return property only **/").Append(endl);
+					builder.AppendTabs(tabs).Append(eventParameterStructName).Append("()").Append(endl);
+					builder.AppendTabs(tabs + 1).Append(": ").Append(returnProperty.SourceName).Append('(').AppendNullConstructorArg(returnProperty, true).Append(')').Append(endl);
+					builder.AppendTabs(tabs).Append('{').Append(endl);
+					builder.AppendTabs(tabs).Append('}').Append(endl);
 				}
 			}
-			--Tabs;
+			--tabs;
 
-			Builder.AppendTabs(Tabs).Append("};").Append(Endl);
+			builder.AppendTabs(tabs).Append("};").Append(endl);
 
-			return Builder;
+			return builder;
 		}
 
-		protected static bool WillExportEventParameters(UhtFunction Function)
+		protected static bool WillExportEventParameters(UhtFunction function)
 		{
-			return Function.Children.Count > 0;
+			return function.Children.Count > 0;
 		}
 
-		protected static string GetEventStructParametersName(UhtType? Outer, string FunctionName)
+		protected static string GetEventStructParametersName(UhtType? outer, string functionName)
 		{
-			string OuterName = String.Empty;
-			if (Outer == null)
+			string outerName = String.Empty;
+			if (outer == null)
 			{
 				throw new UhtIceException("Outer type not set on delegate function");
 			}
-			else if (Outer is UhtClass Class)
+			else if (outer is UhtClass classObj)
 			{
-				OuterName = Class.EngineName;
+				outerName = classObj.EngineName;
 			}
-			else if (Outer is UhtHeaderFile)
+			else if (outer is UhtHeaderFile)
 			{
-				string PackageName = Outer.Package.EngineName;
-				OuterName = PackageName.Replace('/', '_');
+				string packageName = outer.Package.EngineName;
+				outerName = packageName.Replace('/', '_');
 			}
 
-			string Result = $"{OuterName}_event{FunctionName}_Parms";
-			if (UhtFCString.IsDigit(Result[0]))
+			string result = $"{outerName}_event{functionName}_Parms";
+			if (UhtFCString.IsDigit(result[0]))
 			{
-				Result = "_" + Result;
+				result = "_" + result;
 			}
-			return Result;
+			return result;
 		}
 		#endregion
 
 		#region Function helper methods
-		protected StringBuilder AppendNativeFunctionHeader(StringBuilder Builder, UhtFunction Function, UhtPropertyTextType TextType, bool bIsDeclaration,
-			string? AlternateFunctionName, string? ExtraParam, UhtFunctionExportFlags ExtraExportFlags, string Endl)
+		protected StringBuilder AppendNativeFunctionHeader(StringBuilder builder, UhtFunction function, UhtPropertyTextType textType, bool isDeclaration,
+			string? alternateFunctionName, string? extraParam, UhtFunctionExportFlags extraExportFlags, string endl)
 		{
-			UhtClass? OuterClass = Function.Outer as UhtClass;
-			UhtFunctionExportFlags ExportFlags = Function.FunctionExportFlags | ExtraExportFlags;
-			bool bIsDelegate = Function.FunctionFlags.HasAnyFlags(EFunctionFlags.Delegate);
-			bool bIsInterface = !bIsDelegate && (OuterClass != null && OuterClass.ClassFlags.HasAnyFlags(EClassFlags.Interface));
-			bool bIsK2Override = Function.FunctionFlags.HasAnyFlags(EFunctionFlags.BlueprintEvent);
+			UhtClass? outerClass = function.Outer as UhtClass;
+			UhtFunctionExportFlags exportFlags = function.FunctionExportFlags | extraExportFlags;
+			bool isDelegate = function.FunctionFlags.HasAnyFlags(EFunctionFlags.Delegate);
+			bool isInterface = !isDelegate && (outerClass != null && outerClass.ClassFlags.HasAnyFlags(EClassFlags.Interface));
+			bool isK2Override = function.FunctionFlags.HasAnyFlags(EFunctionFlags.BlueprintEvent);
 
-			if (!bIsDelegate)
+			if (!isDelegate)
 			{
-				Builder.Append('\t');
+				builder.Append('\t');
 			}
 
-			if (bIsDeclaration)
+			if (isDeclaration)
 			{
 				// If the function was marked as 'RequiredAPI', then add the *_API macro prefix.  Note that if the class itself
 				// was marked 'RequiredAPI', this is not needed as C++ will exports all methods automatically.
-				if (TextType != UhtPropertyTextType.EventFunctionArgOrRetVal &&
-					!(OuterClass != null && OuterClass.ClassFlags.HasAnyFlags(EClassFlags.RequiredAPI)) &&
-					ExportFlags.HasAnyFlags(UhtFunctionExportFlags.RequiredAPI))
+				if (textType != UhtPropertyTextType.EventFunctionArgOrRetVal &&
+					!(outerClass != null && outerClass.ClassFlags.HasAnyFlags(EClassFlags.RequiredAPI)) &&
+					exportFlags.HasAnyFlags(UhtFunctionExportFlags.RequiredAPI))
 				{
-					Builder.Append(this.PackageApi);
+					builder.Append(this.PackageApi);
 				}
 
-				if (TextType == UhtPropertyTextType.InterfaceFunctionArgOrRetVal)
+				if (textType == UhtPropertyTextType.InterfaceFunctionArgOrRetVal)
 				{
-					Builder.Append("static ");
+					builder.Append("static ");
 				}
-				else if (bIsK2Override)
+				else if (isK2Override)
 				{
-					Builder.Append("virtual ");
+					builder.Append("virtual ");
 				}
 				// if the owning class is an interface class
-				else if (bIsInterface)
+				else if (isInterface)
 				{
-					Builder.Append("virtual ");
+					builder.Append("virtual ");
 				}
 				// this is not an event, the function is not a static function and the function is not marked final
-				else if (TextType != UhtPropertyTextType.EventFunctionArgOrRetVal && !Function.FunctionFlags.HasAnyFlags(EFunctionFlags.Static) && !ExportFlags.HasAnyFlags(UhtFunctionExportFlags.Final))
+				else if (textType != UhtPropertyTextType.EventFunctionArgOrRetVal && !function.FunctionFlags.HasAnyFlags(EFunctionFlags.Static) && !exportFlags.HasAnyFlags(UhtFunctionExportFlags.Final))
 				{
-					Builder.Append("virtual ");
+					builder.Append("virtual ");
 				}
-				else if (ExportFlags.HasAnyFlags(UhtFunctionExportFlags.Inline))
+				else if (exportFlags.HasAnyFlags(UhtFunctionExportFlags.Inline))
 				{
-					Builder.Append("inline ");
+					builder.Append("inline ");
 				}
 			}
 
-			UhtProperty? ReturnProperty = Function.ReturnProperty;
-			if (ReturnProperty != null)
+			UhtProperty? returnProperty = function.ReturnProperty;
+			if (returnProperty != null)
 			{
-				if (ReturnProperty.PropertyFlags.HasAnyFlags(EPropertyFlags.ConstParm))
+				if (returnProperty.PropertyFlags.HasAnyFlags(EPropertyFlags.ConstParm))
 				{
-					Builder.Append("const ");
+					builder.Append("const ");
 				}
-				Builder.AppendPropertyText(ReturnProperty, TextType);
+				builder.AppendPropertyText(returnProperty, textType);
 			}
 			else
 			{
-				Builder.Append("void");
+				builder.Append("void");
 			}
 
-			Builder.Append(' ');
-			if (!bIsDeclaration && OuterClass != null)
+			builder.Append(' ');
+			if (!isDeclaration && outerClass != null)
 			{
-				Builder.AppendClassSourceNameOrInterfaceName(OuterClass).Append("::");
+				builder.AppendClassSourceNameOrInterfaceName(outerClass).Append("::");
 			}
 
-			if (AlternateFunctionName != null)
+			if (alternateFunctionName != null)
 			{
-				Builder.Append(AlternateFunctionName);
+				builder.Append(alternateFunctionName);
 			}
 			else
 			{
-				switch (TextType)
+				switch (textType)
 				{
 					case UhtPropertyTextType.InterfaceFunctionArgOrRetVal:
-						Builder.Append("Execute_").Append(Function.SourceName);
+						builder.Append("Execute_").Append(function.SourceName);
 						break;
 					case UhtPropertyTextType.EventFunctionArgOrRetVal:
-						Builder.Append(Function.MarshalAndCallName);
+						builder.Append(function.MarshalAndCallName);
 						break;
 					case UhtPropertyTextType.ClassFunctionArgOrRetVal:
-						Builder.Append(Function.CppImplName);
+						builder.Append(function.CppImplName);
 						break;
 					default:
 						throw new UhtIceException("Unexpected type text");
 				}
 			}
 
-			AppendParameters(Builder, Function, TextType, ExtraParam, false);
+			AppendParameters(builder, function, textType, extraParam, false);
 
-			if (TextType != UhtPropertyTextType.InterfaceFunctionArgOrRetVal)
+			if (textType != UhtPropertyTextType.InterfaceFunctionArgOrRetVal)
 			{
-				if (!bIsDelegate && Function.FunctionFlags.HasAnyFlags(EFunctionFlags.Const))
+				if (!isDelegate && function.FunctionFlags.HasAnyFlags(EFunctionFlags.Const))
 				{
-					Builder.Append(" const");
+					builder.Append(" const");
 				}
 
-				if (bIsInterface && bIsDeclaration)
+				if (isInterface && isDeclaration)
 				{
 					// all methods in interface classes are pure virtuals
-					if (bIsK2Override)
+					if (isK2Override)
 					{
 						// For BlueprintNativeEvent methods we emit a stub implementation. This allows Blueprints that implement the interface class to be nativized.
-						Builder.Append(" {");
-						if (ReturnProperty != null)
+						builder.Append(" {");
+						if (returnProperty != null)
 						{
-							if (ReturnProperty is UhtByteProperty ByteProperty && ByteProperty.Enum != null)
+							if (returnProperty is UhtByteProperty byteProperty && byteProperty.Enum != null)
 							{
-								Builder.Append(" return TEnumAsByte<").Append(ByteProperty.Enum.CppType).Append(">(").AppendNullConstructorArg(ReturnProperty, false).Append("); ");
+								builder.Append(" return TEnumAsByte<").Append(byteProperty.Enum.CppType).Append(">(").AppendNullConstructorArg(returnProperty, false).Append("); ");
 							}
 							else
 							{
-								Builder.Append(" return ").AppendNullConstructorArg(ReturnProperty, false).Append("; ");
+								builder.Append(" return ").AppendNullConstructorArg(returnProperty, false).Append("; ");
 							}
 						}
-						Builder.Append('}');
+						builder.Append('}');
 					}
 					else
 					{
-						Builder.Append("=0");
+						builder.Append("=0");
 					}
 				}
 			}
-			Builder.Append(Endl);
+			builder.Append(endl);
 
-			return Builder;
+			return builder;
 		}
 
-		protected static StringBuilder AppendEventFunctionPrologue(StringBuilder Builder, UhtFunction Function, string FunctionName, int Tabs, string Endl)
+		protected static StringBuilder AppendEventFunctionPrologue(StringBuilder builder, UhtFunction function, string functionName, int tabs, string endl)
 		{
-			Builder.AppendTabs(Tabs).Append('{').Append(Endl);
-			if (Function.Children.Count == 0)
+			builder.AppendTabs(tabs).Append('{').Append(endl);
+			if (function.Children.Count == 0)
 			{
-				return Builder;
+				return builder;
 			}
 
-			++Tabs;
+			++tabs;
 
-			string EventParameterStructName = GetEventStructParametersName(Function.Outer, FunctionName);
+			string eventParameterStructName = GetEventStructParametersName(function.Outer, functionName);
 
-			Builder.AppendTabs(Tabs).Append(EventParameterStructName).Append(" Parms;").Append(Endl);
+			builder.AppendTabs(tabs).Append(eventParameterStructName).Append(" Parms;").Append(endl);
 
 			// Declare a parameter struct for this event/delegate and assign the struct members using the values passed into the event/delegate call.
-			foreach (UhtProperty Property in Function.ParameterProperties.Span)
+			foreach (UhtProperty property in function.ParameterProperties.Span)
 			{
-				if (Property.bIsStaticArray)
+				if (property.IsStaticArray)
 				{
-					Builder.AppendTabs(Tabs).Append("FMemory::Memcpy(Parm.")
-						.Append(Property.SourceName).Append(',')
-						.Append(Property.SourceName).Append(",sizeof(Parms.")
-						.Append(Property.SourceName).Append(");").Append(Endl);
+					builder.AppendTabs(tabs).Append("FMemory::Memcpy(Parm.")
+						.Append(property.SourceName).Append(',')
+						.Append(property.SourceName).Append(",sizeof(Parms.")
+						.Append(property.SourceName).Append(");").Append(endl);
 				}
 				else
 				{
-					Builder.AppendTabs(Tabs).Append("Parms.").Append(Property.SourceName).Append('=').Append(Property.SourceName);
-					if (Property is UhtBoolProperty)
+					builder.AppendTabs(tabs).Append("Parms.").Append(property.SourceName).Append('=').Append(property.SourceName);
+					if (property is UhtBoolProperty)
 					{
-						Builder.Append(" ? true : false");
+						builder.Append(" ? true : false");
 					}
-					Builder.Append(';').Append(Endl);
+					builder.Append(';').Append(endl);
 				}
 			}
-			return Builder;
+			return builder;
 		}
 
-		protected static StringBuilder AppendEventFunctionEpilogue(StringBuilder Builder, UhtFunction Function, int Tabs, string Endl)
+		protected static StringBuilder AppendEventFunctionEpilogue(StringBuilder builder, UhtFunction function, int tabs, string endl)
 		{
-			++Tabs;
+			++tabs;
 
 			// Out parm copying.
-			foreach (UhtProperty Property in Function.ParameterProperties.Span)
+			foreach (UhtProperty property in function.ParameterProperties.Span)
 			{
-				if (Property.PropertyFlags.HasExactFlags(EPropertyFlags.ConstParm | EPropertyFlags.OutParm, EPropertyFlags.OutParm))
+				if (property.PropertyFlags.HasExactFlags(EPropertyFlags.ConstParm | EPropertyFlags.OutParm, EPropertyFlags.OutParm))
 				{
-					if (Property.bIsStaticArray)
+					if (property.IsStaticArray)
 					{
-						Builder
-							.AppendTabs(Tabs)
+						builder
+							.AppendTabs(tabs)
 							.Append("FMemory::Memcpy(&")
-							.Append(Property.SourceName)
+							.Append(property.SourceName)
 							.Append(",&Parms.")
-							.Append(Property.SourceName)
+							.Append(property.SourceName)
 							.Append(",sizeof(")
-							.Append(Property.SourceName)
-							.Append("));").Append(Endl);
+							.Append(property.SourceName)
+							.Append("));").Append(endl);
 					}
 					else
 					{
-						Builder
-							.AppendTabs(Tabs)
-							.Append(Property.SourceName)
+						builder
+							.AppendTabs(tabs)
+							.Append(property.SourceName)
 							.Append("=Parms.")
-							.Append(Property.SourceName)
-							.Append(';').Append(Endl);
+							.Append(property.SourceName)
+							.Append(';').Append(endl);
 					}
 				}
 			}
 
 			// Return value.
-			UhtProperty? ReturnProperty = Function.ReturnProperty;
-			if (ReturnProperty != null)
+			UhtProperty? returnProperty = function.ReturnProperty;
+			if (returnProperty != null)
 			{
 				// Make sure uint32 -> bool is supported
-				if (ReturnProperty is UhtBoolProperty)
+				if (returnProperty is UhtBoolProperty)
 				{
-					Builder
-						.AppendTabs(Tabs)
+					builder
+						.AppendTabs(tabs)
 						.Append("return !!Parms.")
-						.Append(ReturnProperty.SourceName)
-						.Append(';').Append(Endl);
+						.Append(returnProperty.SourceName)
+						.Append(';').Append(endl);
 				}
 				else
 				{
-					Builder
-						.AppendTabs(Tabs)
+					builder
+						.AppendTabs(tabs)
 						.Append("return Parms.")
-						.Append(ReturnProperty.SourceName)
-						.Append(';').Append(Endl);
+						.Append(returnProperty.SourceName)
+						.Append(';').Append(endl);
 				}
 			}
 
-			--Tabs;
-			Builder.AppendTabs(Tabs).Append('}').Append(Endl);
-			return Builder;
+			--tabs;
+			builder.AppendTabs(tabs).Append('}').Append(endl);
+			return builder;
 		}
 
-		protected static StringBuilder AppendParameters(StringBuilder Builder, UhtFunction Function, UhtPropertyTextType TextType, string? ExtraParameter, bool bSkipParameterName)
+		protected static StringBuilder AppendParameters(StringBuilder builder, UhtFunction function, UhtPropertyTextType textType, string? extraParameter, bool skipParameterName)
 		{
-			bool bNeedsSeperator = false;
+			bool needsSeperator = false;
 
-			Builder.Append('(');
+			builder.Append('(');
 
-			if (ExtraParameter != null)
+			if (extraParameter != null)
 			{
-				Builder.Append(ExtraParameter);
-				bNeedsSeperator = true;
+				builder.Append(extraParameter);
+				needsSeperator = true;
 			}
 
-			foreach (UhtProperty Parameter in Function.ParameterProperties.Span)
+			foreach (UhtProperty parameter in function.ParameterProperties.Span)
 			{
-				if (bNeedsSeperator)
+				if (needsSeperator)
 				{
-					Builder.Append(", ");
+					builder.Append(", ");
 				}
 				else
 				{
-					bNeedsSeperator = true;
+					needsSeperator = true;
 				}
-				Builder.AppendFullDecl(Parameter, TextType, bSkipParameterName);
+				builder.AppendFullDecl(parameter, textType, skipParameterName);
 			}
 
-			Builder.Append(')');
-			return Builder;
+			builder.Append(')');
+			return builder;
 		}
 
-		protected static bool IsRpcFunction(UhtFunction Function)
+		protected static bool IsRpcFunction(UhtFunction function)
 		{
-			return Function.FunctionFlags.HasAnyFlags(EFunctionFlags.Native) &&
-				!Function.FunctionExportFlags.HasAnyFlags(UhtFunctionExportFlags.CustomThunk);
+			return function.FunctionFlags.HasAnyFlags(EFunctionFlags.Native) &&
+				!function.FunctionExportFlags.HasAnyFlags(UhtFunctionExportFlags.CustomThunk);
 		}
 
-		protected static bool IsRpcFunction(UhtFunction Function, bool bEditorOnly)
+		protected static bool IsRpcFunction(UhtFunction function, bool editorOnly)
 		{
-			return IsRpcFunction(Function) && Function.FunctionFlags.HasAnyFlags(EFunctionFlags.EditorOnly) == bEditorOnly;
+			return IsRpcFunction(function) && function.FunctionFlags.HasAnyFlags(EFunctionFlags.EditorOnly) == editorOnly;
 		}
 
 		/// <summary>
 		/// Determines whether the glue version of the specified native function should be exported.
 		/// </summary>
-		/// <param name="Function">The function to check</param>
+		/// <param name="function">The function to check</param>
 		/// <returns>True if the glue version of the function should be exported.</returns>
-		protected static bool ShouldExportFunction(UhtFunction Function)
+		protected static bool ShouldExportFunction(UhtFunction function)
 		{
 			// export any script stubs for native functions declared in interface classes
-			bool bIsBlueprintNativeEvent = Function.FunctionFlags.HasAllFlags(EFunctionFlags.BlueprintEvent | EFunctionFlags.Native);
-			if (!bIsBlueprintNativeEvent)
+			bool isBlueprintNativeEvent = function.FunctionFlags.HasAllFlags(EFunctionFlags.BlueprintEvent | EFunctionFlags.Native);
+			if (!isBlueprintNativeEvent)
 			{
-				if (Function.Outer != null)
+				if (function.Outer != null)
 				{
-					if (Function.Outer is UhtClass Class)
+					if (function.Outer is UhtClass classObj)
 					{
-						if (Class.ClassFlags.HasAnyFlags(EClassFlags.Interface))
+						if (classObj.ClassFlags.HasAnyFlags(EClassFlags.Interface))
 						{
 							return true;
 						}
@@ -433,16 +433,16 @@ namespace EpicGames.UHT.Exporters.CodeGen
 			}
 
 			// always export if the function is static
-			if (Function.FunctionFlags.HasAnyFlags(EFunctionFlags.Static))
+			if (function.FunctionFlags.HasAnyFlags(EFunctionFlags.Static))
 			{
 				return true;
 			}
 
 			// don't export the function if this is not the original declaration and there is
 			// at least one parent version of the function that is declared native
-			for (UhtFunction? ParentFunction = Function.SuperFunction; ParentFunction != null; ParentFunction = ParentFunction.SuperFunction)
+			for (UhtFunction? parentFunction = function.SuperFunction; parentFunction != null; parentFunction = parentFunction.SuperFunction)
 			{
-				if (ParentFunction.FunctionFlags.HasAnyFlags(EFunctionFlags.Native))
+				if (parentFunction.FunctionFlags.HasAnyFlags(EFunctionFlags.Native))
 				{
 					return false;
 				}
@@ -450,171 +450,171 @@ namespace EpicGames.UHT.Exporters.CodeGen
 			return true;
 		}
 
-		protected static string GetDelegateFunctionExportName(UhtFunction Function)
+		protected static string GetDelegateFunctionExportName(UhtFunction function)
 		{
 			const string DelegatePrefix = "delegate";
 
-			if (!Function.FunctionFlags.HasAnyFlags(EFunctionFlags.Delegate))
+			if (!function.FunctionFlags.HasAnyFlags(EFunctionFlags.Delegate))
 			{
 				throw new UhtIceException("Attempt to export a function that isn't a delegate as a delegate");
 			}
-			if (!Function.MarshalAndCallName.StartsWith(DelegatePrefix, StringComparison.Ordinal))
+			if (!function.MarshalAndCallName.StartsWith(DelegatePrefix, StringComparison.Ordinal))
 			{
 				throw new UhtIceException("Marshal and call name must begin with 'delegate'");
 			}
-			return $"F{Function.MarshalAndCallName.Substring(DelegatePrefix.Length)}_DelegateWrapper";
+			return $"F{function.MarshalAndCallName.Substring(DelegatePrefix.Length)}_DelegateWrapper";
 		}
 
-		protected static string GetDelegateFunctionExtraParameter(UhtFunction Function)
+		protected static string GetDelegateFunctionExtraParameter(UhtFunction function)
 		{
-			string ReturnType = Function.FunctionFlags.HasAnyFlags(EFunctionFlags.MulticastDelegate) ? "FMulticastScriptDelegate" : "FScriptDelegate";
-			return $"const {ReturnType}& {Function.SourceName.Substring(1)}";
+			string returnType = function.FunctionFlags.HasAnyFlags(EFunctionFlags.MulticastDelegate) ? "FMulticastScriptDelegate" : "FScriptDelegate";
+			return $"const {returnType}& {function.SourceName.Substring(1)}";
 		}
 		#endregion
 
 		#region Field notify support
-		protected static bool NeedFieldNotifyCodeGen(UhtClass Class)
+		protected static bool NeedFieldNotifyCodeGen(UhtClass classObj)
 		{
-			return 
-				!Class.ClassExportFlags.HasAnyFlags(UhtClassExportFlags.HasCustomFieldNotify) &&
-				Class.ClassExportFlags.HasAnyFlags(UhtClassExportFlags.HasFieldNotify);
+			return
+				!classObj.ClassExportFlags.HasAnyFlags(UhtClassExportFlags.HasCustomFieldNotify) &&
+				classObj.ClassExportFlags.HasAnyFlags(UhtClassExportFlags.HasFieldNotify);
 		}
 
-		protected static void GetFieldNotifyStats(UhtClass Class, out bool bHasProperties, out bool bHasFunctions, out bool bHasEditorFields, out bool bAllEditorFields)
+		protected static void GetFieldNotifyStats(UhtClass classObj, out bool hasProperties, out bool hasFunctions, out bool hasEditorFields, out bool allEditorFields)
 		{
 			// Scan the children to see what we have
-			bHasProperties = false;
-			bHasFunctions = false;
-			bHasEditorFields = false;
-			bAllEditorFields = true;
-			foreach (UhtType Type in Class.Children)
+			hasProperties = false;
+			hasFunctions = false;
+			hasEditorFields = false;
+			allEditorFields = true;
+			foreach (UhtType type in classObj.Children)
 			{
-				if (Type is UhtProperty Property)
+				if (type is UhtProperty property)
 				{
-					if (Property.PropertyExportFlags.HasAnyFlags(UhtPropertyExportFlags.FieldNotify))
+					if (property.PropertyExportFlags.HasAnyFlags(UhtPropertyExportFlags.FieldNotify))
 					{
-						bHasProperties = true;
-						bHasEditorFields |= Property.bIsEditorOnlyProperty;
-						bAllEditorFields &= Property.bIsEditorOnlyProperty;
+						hasProperties = true;
+						hasEditorFields |= property.IsEditorOnlyProperty;
+						allEditorFields &= property.IsEditorOnlyProperty;
 					}
 				}
-				else if (Type is UhtFunction Function)
+				else if (type is UhtFunction function)
 				{
-					if (Function.FunctionExportFlags.HasAnyFlags(UhtFunctionExportFlags.FieldNotify))
+					if (function.FunctionExportFlags.HasAnyFlags(UhtFunctionExportFlags.FieldNotify))
 					{
-						bHasFunctions = true;
-						bHasEditorFields |= Function.FunctionFlags.HasAnyFlags(EFunctionFlags.EditorOnly);
-						bAllEditorFields &= Function.FunctionFlags.HasAnyFlags(EFunctionFlags.EditorOnly);
+						hasFunctions = true;
+						hasEditorFields |= function.FunctionFlags.HasAnyFlags(EFunctionFlags.EditorOnly);
+						allEditorFields &= function.FunctionFlags.HasAnyFlags(EFunctionFlags.EditorOnly);
 					}
 				}
 			}
 
 			// If we have no editor fields, then by definition, all fields can't be editor fields
-			bAllEditorFields &= bHasEditorFields;
+			allEditorFields &= hasEditorFields;
 		}
 
-		protected static StringBuilder AppendFieldNotify(StringBuilder Builder, UhtClass Class,
-			bool bHasProperties, bool bHasFunctions, bool bHasEditorFields, bool bAllEditorFields,
-			bool bIncludeEditorOnlyFields, bool bAppendDefine, Action<StringBuilder, UhtClass, string> AppendAction)
+		protected static StringBuilder AppendFieldNotify(StringBuilder builder, UhtClass classObj,
+			bool hasProperties, bool hasFunctions, bool hasEditorFields, bool allEditorFields,
+			bool includeEditorOnlyFields, bool appendDefine, Action<StringBuilder, UhtClass, string> appendAction)
 		{
-			if (bHasProperties && !bAllEditorFields)
+			if (hasProperties && !allEditorFields)
 			{
-				foreach (UhtType Child in Class.Children)
+				foreach (UhtType child in classObj.Children)
 				{
-					if (Child is UhtProperty Property)
+					if (child is UhtProperty property)
 					{
-						if (Property.PropertyExportFlags.HasAnyFlags(UhtPropertyExportFlags.FieldNotify) && !Property.bIsEditorOnlyProperty)
+						if (property.PropertyExportFlags.HasAnyFlags(UhtPropertyExportFlags.FieldNotify) && !property.IsEditorOnlyProperty)
 						{
-							AppendAction(Builder, Class, Property.SourceName);
+							appendAction(builder, classObj, property.SourceName);
 						}
 					}
 				}
 			}
 
-			if (bHasFunctions && !bAllEditorFields)
+			if (hasFunctions && !allEditorFields)
 			{
-				foreach (UhtType Child in Class.Children)
+				foreach (UhtType child in classObj.Children)
 				{
-					if (Child is UhtFunction Function)
+					if (child is UhtFunction function)
 					{
-						if (Function.FunctionExportFlags.HasAnyFlags(UhtFunctionExportFlags.FieldNotify) && !Function.FunctionFlags.HasAnyFlags(EFunctionFlags.EditorOnly))
+						if (function.FunctionExportFlags.HasAnyFlags(UhtFunctionExportFlags.FieldNotify) && !function.FunctionFlags.HasAnyFlags(EFunctionFlags.EditorOnly))
 						{
-							AppendAction(Builder, Class, Function.CppImplName);
+							appendAction(builder, classObj, function.CppImplName);
 						}
 					}
 				}
 			}
 
-			if (bHasEditorFields && bIncludeEditorOnlyFields)
+			if (hasEditorFields && includeEditorOnlyFields)
 			{
-				if (!bAllEditorFields && bAppendDefine)
+				if (!allEditorFields && appendDefine)
 				{
-					Builder.Append("#if WITH_EDITORONLY_DATA\r\n");
+					builder.Append("#if WITH_EDITORONLY_DATA\r\n");
 				}
 
-				if (bHasProperties)
+				if (hasProperties)
 				{
-					foreach (UhtType Child in Class.Children)
+					foreach (UhtType child in classObj.Children)
 					{
-						if (Child is UhtProperty Property)
+						if (child is UhtProperty property)
 						{
-							if (Property.PropertyExportFlags.HasAnyFlags(UhtPropertyExportFlags.FieldNotify) && Property.bIsEditorOnlyProperty)
+							if (property.PropertyExportFlags.HasAnyFlags(UhtPropertyExportFlags.FieldNotify) && property.IsEditorOnlyProperty)
 							{
-								AppendAction(Builder, Class, Property.SourceName);
+								appendAction(builder, classObj, property.SourceName);
 							}
 						}
 					}
 				}
 
-				if (bHasFunctions)
+				if (hasFunctions)
 				{
-					foreach (UhtType Child in Class.Children)
+					foreach (UhtType child in classObj.Children)
 					{
-						if (Child is UhtFunction Function)
+						if (child is UhtFunction function)
 						{
-							if (Function.FunctionExportFlags.HasAnyFlags(UhtFunctionExportFlags.FieldNotify) && Function.FunctionFlags.HasAnyFlags(EFunctionFlags.EditorOnly))
+							if (function.FunctionExportFlags.HasAnyFlags(UhtFunctionExportFlags.FieldNotify) && function.FunctionFlags.HasAnyFlags(EFunctionFlags.EditorOnly))
 							{
-								AppendAction(Builder, Class, Function.CppImplName);
+								appendAction(builder, classObj, function.CppImplName);
 							}
 						}
 					}
 				}
 
-				if (!bAllEditorFields && bAppendDefine)
+				if (!allEditorFields && appendDefine)
 				{
-					Builder.Append("#endif // WITH_EDITORONLY_DATA\r\n");
+					builder.Append("#endif // WITH_EDITORONLY_DATA\r\n");
 				}
 			}
-			return Builder;
+			return builder;
 		}
 		#endregion
 	}
 
 	internal static class UhtHaederCodeGeneratorStringBuilderExtensions
 	{
-		public static StringBuilder AppendMacroName(this StringBuilder Builder, string FileId, int LineNumber, string MacroSuffix)
+		public static StringBuilder AppendMacroName(this StringBuilder builder, string fileId, int lineNumber, string macroSuffix)
 		{
-			return Builder.Append(FileId).Append('_').Append(LineNumber).Append('_').Append(MacroSuffix);
+			return builder.Append(fileId).Append('_').Append(lineNumber).Append('_').Append(macroSuffix);
 		}
 
-		public static StringBuilder AppendMacroName(this StringBuilder Builder, UhtHeaderCodeGenerator Generator, int LineNumber, string MacroSuffix)
+		public static StringBuilder AppendMacroName(this StringBuilder builder, UhtHeaderCodeGenerator generator, int lineNumber, string macroSuffix)
 		{
-			return Builder.AppendMacroName(Generator.FileId, LineNumber, MacroSuffix);
+			return builder.AppendMacroName(generator.FileId, lineNumber, macroSuffix);
 		}
 
-		public static StringBuilder AppendMacroName(this StringBuilder Builder, UhtHeaderCodeGenerator Generator, UhtClass Class, string MacroSuffix)
+		public static StringBuilder AppendMacroName(this StringBuilder builder, UhtHeaderCodeGenerator generator, UhtClass classObj, string macroSuffix)
 		{
-			return Builder.AppendMacroName(Generator, Class.GeneratedBodyLineNumber, MacroSuffix);
+			return builder.AppendMacroName(generator, classObj.GeneratedBodyLineNumber, macroSuffix);
 		}
 
-		public static StringBuilder AppendMacroName(this StringBuilder Builder, UhtHeaderCodeGenerator Generator, UhtScriptStruct ScriptStruct, string MacroSuffix)
+		public static StringBuilder AppendMacroName(this StringBuilder builder, UhtHeaderCodeGenerator generator, UhtScriptStruct scriptStruct, string macroSuffix)
 		{
-			return Builder.AppendMacroName(Generator, ScriptStruct.MacroDeclaredLineNumber, MacroSuffix);
+			return builder.AppendMacroName(generator, scriptStruct.MacroDeclaredLineNumber, macroSuffix);
 		}
 
-		public static StringBuilder AppendMacroName(this StringBuilder Builder, UhtHeaderCodeGenerator Generator, UhtFunction Function, string MacroSuffix)
+		public static StringBuilder AppendMacroName(this StringBuilder builder, UhtHeaderCodeGenerator generator, UhtFunction function, string macroSuffix)
 		{
-			return Builder.AppendMacroName(Generator, Function.MacroLineNumber, MacroSuffix);
+			return builder.AppendMacroName(generator, function.MacroLineNumber, macroSuffix);
 		}
 	}
 }
