More coding standard fixes

diff --git a/Engine/Source/Programs/Shared/EpicGames.UHT/Exporters/CodeGen/UhtHeaderCodeGeneratorCppFile.cs b/Engine/Source/Programs/Shared/EpicGames.UHT/Exporters/CodeGen/UhtHeaderCodeGeneratorCppFile.cs
--- a/Engine/Source/Programs/Shared/EpicGames.UHT/Exporters/CodeGen/UhtHeaderCodeGeneratorCppFile.cs
+++ b/Engine/Source/Programs/Shared/EpicGames.UHT/Exporters/CodeGen/UhtHeaderCodeGeneratorCppFile.cs
@@ -1,12 +1,12 @@
 // Copyright Epic Games, Inc. All Rights Reserved.
 
-using EpicGames.Core;
-using EpicGames.UHT.Types;
-using EpicGames.UHT.Utils;
 using System;
 using System.Collections.Generic;
 using System.Linq;
 using System.Text;
+using EpicGames.Core;
+using EpicGames.UHT.Types;
+using EpicGames.UHT.Utils;
 
 namespace EpicGames.UHT.Exporters.CodeGen
 {
@@ -16,1722 +16,1722 @@ namespace EpicGames.UHT.Exporters.CodeGen
 		/// <summary>
 		/// Construct an instance of this generator object
 		/// </summary>
-		/// <param name="CodeGenerator">The base code generator</param>
-		/// <param name="Package">Package being generated</param>
-		/// <param name="HeaderFile">Header file being generated</param>
-		public UhtHeaderCodeGeneratorCppFile(UhtCodeGenerator CodeGenerator, UhtPackage Package, UhtHeaderFile HeaderFile)
-			: base(CodeGenerator, Package, HeaderFile)
+		/// <param name="codeGenerator">The base code generator</param>
+		/// <param name="package">Package being generated</param>
+		/// <param name="headerFile">Header file being generated</param>
+		public UhtHeaderCodeGeneratorCppFile(UhtCodeGenerator codeGenerator, UhtPackage package, UhtHeaderFile headerFile)
+			: base(codeGenerator, package, headerFile)
 		{
 		}
 
 		/// <summary>
 		/// For a given UE header file, generated the generated H file
 		/// </summary>
-		/// <param name="Factory">Requesting factory</param>
-		public void Generate(IUhtExportFactory Factory)
+		/// <param name="factory">Requesting factory</param>
+		public void Generate(IUhtExportFactory factory)
 		{
-			ref UhtCodeGenerator.HeaderInfo HeaderInfo = ref this.HeaderInfos[this.HeaderFile.HeaderFileTypeIndex];
-			using (BorrowStringBuilder Borrower = new BorrowStringBuilder(StringBuilderCache.Big))
+			ref UhtCodeGenerator.HeaderInfo headerInfo = ref this.HeaderInfos[this.HeaderFile.HeaderFileTypeIndex];
+			using (BorrowStringBuilder borrower = new BorrowStringBuilder(StringBuilderCache.Big))
 			{
-				StringBuilder Builder = Borrower.StringBuilder;
+				StringBuilder builder = borrower.StringBuilder;
 
-				Builder.Append(HeaderCopyright);
-				Builder.Append(RequiredCPPIncludes);
-				Builder.Append("#include \"").Append(HeaderInfo.IncludePath).Append("\"\r\n");
+				builder.Append(HeaderCopyright);
+				builder.Append(RequiredCPPIncludes);
+				builder.Append("#include \"").Append(headerInfo._includePath).Append("\"\r\n");
 
-				bool bAddedStructuredArchiveFromArchiveHeader = false;
-				bool bAddedArchiveUObjectFromStructuredArchiveHeader = false;
-				HashSet<UhtHeaderFile> AddedIncludes = new HashSet<UhtHeaderFile>();
-				AddedIncludes.Add(this.HeaderFile);
-				foreach (UhtType Type in this.HeaderFile.Children)
+				bool addedStructuredArchiveFromArchiveHeader = false;
+				bool addedArchiveUObjectFromStructuredArchiveHeader = false;
+				HashSet<UhtHeaderFile> addedIncludes = new HashSet<UhtHeaderFile>();
+				addedIncludes.Add(this.HeaderFile);
+				foreach (UhtType type in this.HeaderFile.Children)
 				{
-					if (Type is UhtClass Class)
+					if (type is UhtClass classObj)
 					{
-						if (Class.ClassWithin != this.Session.UObject && !Class.ClassWithin.HeaderFile.bIsNoExportTypes)
+						if (classObj.ClassWithin != this.Session.UObject && !classObj.ClassWithin.HeaderFile.IsNoExportTypes)
 						{
-							if (AddedIncludes.Add(Class.ClassWithin.HeaderFile))
+							if (addedIncludes.Add(classObj.ClassWithin.HeaderFile))
 							{
-								Builder.Append("#include \"").Append(this.HeaderInfos[Class.ClassWithin.HeaderFile.HeaderFileTypeIndex].IncludePath).Append("\"\r\n");
+								builder.Append("#include \"").Append(this.HeaderInfos[classObj.ClassWithin.HeaderFile.HeaderFileTypeIndex]._includePath).Append("\"\r\n");
 							}
 						}
 
-						switch (Class.SerializerArchiveType)
+						switch (classObj.SerializerArchiveType)
 						{
 							case UhtSerializerArchiveType.None:
 								break;
 
 							case UhtSerializerArchiveType.Archive:
-								if (!bAddedArchiveUObjectFromStructuredArchiveHeader)
+								if (!addedArchiveUObjectFromStructuredArchiveHeader)
 								{
-									Builder.Append("#include \"Serialization/ArchiveUObjectFromStructuredArchive.h\"\r\n");
-									bAddedArchiveUObjectFromStructuredArchiveHeader = true;
+									builder.Append("#include \"Serialization/ArchiveUObjectFromStructuredArchive.h\"\r\n");
+									addedArchiveUObjectFromStructuredArchiveHeader = true;
 								}
 								break;
 
 							case UhtSerializerArchiveType.StructuredArchiveRecord:
-								if (!bAddedStructuredArchiveFromArchiveHeader)
+								if (!addedStructuredArchiveFromArchiveHeader)
 								{
-									Builder.Append("#include \"Serialization/StructuredArchive.h\"\r\n");
-									bAddedStructuredArchiveFromArchiveHeader = true;
+									builder.Append("#include \"Serialization/StructuredArchive.h\"\r\n");
+									addedStructuredArchiveFromArchiveHeader = true;
 								}
 								break;
 						}
 					}
 				}
 
-				Builder.Append(DisableDeprecationWarnings).Append("\r\n");
-				string CleanFileName = this.HeaderFile.FileNameWithoutExtension.Replace('.', '_');
-				Builder.Append("void EmptyLinkFunctionForGeneratedCode").Append(CleanFileName).Append("() {}\r\n");
+				builder.Append(DisableDeprecationWarnings).Append("\r\n");
+				string cleanFileName = this.HeaderFile.FileNameWithoutExtension.Replace('.', '_');
+				builder.Append("void EmptyLinkFunctionForGeneratedCode").Append(cleanFileName).Append("() {}\r\n");
 
 				if (this.HeaderFile.References.CrossModule.References.Count > 0)
 				{
-					ReadOnlyMemory<string> Sorted = this.HeaderFile.References.CrossModule.GetSortedReferences(
-						(int ObjectIndex, bool bRegistered) => GetCrossReference(ObjectIndex, bRegistered));
-					Builder.Append("// Cross Module References\r\n");
-					foreach (string CrossReference in Sorted.Span)
+					ReadOnlyMemory<string> sorted = this.HeaderFile.References.CrossModule.GetSortedReferences(
+						(int objectIndex, bool registered) => GetCrossReference(objectIndex, registered));
+					builder.Append("// Cross Module References\r\n");
+					foreach (string crossReference in sorted.Span)
 					{
-						Builder.Append(CrossReference);
+						builder.Append(crossReference);
 					}
-					Builder.Append("// End Cross Module References\r\n");
+					builder.Append("// End Cross Module References\r\n");
 				}
 
-				int GeneratedBodyStart = Builder.Length;
+				int generatedBodyStart = builder.Length;
 
-				bool bHasRegisteredEnums = false;
-				bool bHasRegisteredScriptStructs = false;
-				bool bHasRegisteredClasses = false;
-				bool bAllEnumsEditorOnly = true;
-				foreach (UhtField Field in this.HeaderFile.References.ExportTypes)
+				bool hasRegisteredEnums = false;
+				bool hasRegisteredScriptStructs = false;
+				bool hasRegisteredClasses = false;
+				bool allEnumsEditorOnly = true;
+				foreach (UhtField field in this.HeaderFile.References.ExportTypes)
 				{
-					if (Field is UhtEnum Enum)
+					if (field is UhtEnum enumObj)
 					{
-						AppendEnum(Builder, Enum);
-						bHasRegisteredEnums = true;
-						bAllEnumsEditorOnly &= Enum.bIsEditorOnly;
+						AppendEnum(builder, enumObj);
+						hasRegisteredEnums = true;
+						allEnumsEditorOnly &= enumObj.IsEditorOnly;
 					}
-					else if (Field is UhtScriptStruct ScriptStruct)
+					else if (field is UhtScriptStruct scriptStruct)
 					{
-						AppendScriptStruct(Builder, ScriptStruct);
-						if (ScriptStruct.ScriptStructFlags.HasAnyFlags(EStructFlags.Native))
+						AppendScriptStruct(builder, scriptStruct);
+						if (scriptStruct.ScriptStructFlags.HasAnyFlags(EStructFlags.Native))
 						{
-							bHasRegisteredScriptStructs = true;
+							hasRegisteredScriptStructs = true;
 						}
 					}
-					else if (Field is UhtFunction Function)
+					else if (field is UhtFunction function)
 					{
-						AppendDelegate(Builder, Function);
+						AppendDelegate(builder, function);
 					}
-					else if (Field is UhtClass Class)
+					else if (field is UhtClass classObj)
 					{
-						if (!Class.ClassFlags.HasAnyFlags(EClassFlags.Intrinsic))
+						if (!classObj.ClassFlags.HasAnyFlags(EClassFlags.Intrinsic))
 						{
-							AppendClass(Builder, Class);
-							bHasRegisteredClasses = true;
+							AppendClass(builder, classObj);
+							hasRegisteredClasses = true;
 						}
 					}
 				}
 
-				if (bHasRegisteredEnums || bHasRegisteredScriptStructs || bHasRegisteredClasses)
+				if (hasRegisteredEnums || hasRegisteredScriptStructs || hasRegisteredClasses)
 				{
-					string Name = $"Z_CompiledInDeferFile_{HeaderInfo.FileId}";
-					string StaticsName = $"{Name}_Statics";
+					string name = $"Z_CompiledInDeferFile_{headerInfo._fileId}";
+					string staticsName = $"{name}_Statics";
 
-					Builder.Append("\tstruct ").Append(StaticsName).Append("\r\n");
-					Builder.Append("\t{\r\n");
+					builder.Append("\tstruct ").Append(staticsName).Append("\r\n");
+					builder.Append("\t{\r\n");
 
-					if (bHasRegisteredEnums)
+					if (hasRegisteredEnums)
 					{
-						if (bAllEnumsEditorOnly)
+						if (allEnumsEditorOnly)
 						{
-							Builder.Append("#if WITH_EDITORONLY_DATA\r\n");
+							builder.Append("#if WITH_EDITORONLY_DATA\r\n");
 						}
-						Builder.Append("\t\tstatic const FEnumRegisterCompiledInInfo EnumInfo[];\r\n");
-						if (bAllEnumsEditorOnly)
+						builder.Append("\t\tstatic const FEnumRegisterCompiledInInfo EnumInfo[];\r\n");
+						if (allEnumsEditorOnly)
 						{
-							Builder.Append("#endif\r\n");
+							builder.Append("#endif\r\n");
 						}
 					}
-					if (bHasRegisteredScriptStructs)
+					if (hasRegisteredScriptStructs)
 					{
-						Builder.Append("\t\tstatic const FStructRegisterCompiledInInfo ScriptStructInfo[];\r\n");
+						builder.Append("\t\tstatic const FStructRegisterCompiledInInfo ScriptStructInfo[];\r\n");
 					}
-					if (bHasRegisteredClasses)
+					if (hasRegisteredClasses)
 					{
-						Builder.Append("\t\tstatic const FClassRegisterCompiledInInfo ClassInfo[];\r\n");
+						builder.Append("\t\tstatic const FClassRegisterCompiledInInfo ClassInfo[];\r\n");
 					}
 
-					Builder.Append("\t};\r\n");
+					builder.Append("\t};\r\n");
 
-					uint CombinedHash = uint.MaxValue;
+					uint combinedHash = UInt32.MaxValue;
 
-					if (bHasRegisteredEnums)
+					if (hasRegisteredEnums)
 					{
-						if (bAllEnumsEditorOnly)
+						if (allEnumsEditorOnly)
 						{
-							Builder.Append("#if WITH_EDITORONLY_DATA\r\n");
+							builder.Append("#if WITH_EDITORONLY_DATA\r\n");
 						}
-						Builder.Append("\tconst FEnumRegisterCompiledInInfo ").Append(StaticsName).Append("::EnumInfo[] = {\r\n");
-						foreach (UhtObject Object in this.HeaderFile.References.ExportTypes)
+						builder.Append("\tconst FEnumRegisterCompiledInInfo ").Append(staticsName).Append("::EnumInfo[] = {\r\n");
+						foreach (UhtObject obj in this.HeaderFile.References.ExportTypes)
 						{
-							if (Object is UhtEnum Enum)
+							if (obj is UhtEnum enumObj)
 							{
-								if (!bAllEnumsEditorOnly && Enum.bIsEditorOnly)
+								if (!allEnumsEditorOnly && enumObj.IsEditorOnly)
 								{
-									Builder.Append("#if WITH_EDITORONLY_DATA\r\n");
+									builder.Append("#if WITH_EDITORONLY_DATA\r\n");
 								}
-								uint Hash = this.ObjectInfos[Enum.ObjectTypeIndex].Hash;
-								Builder
+								uint hash = this.ObjectInfos[enumObj.ObjectTypeIndex]._hash;
+								builder
 									.Append("\t\t{ ")
-									.Append(Enum.SourceName)
+									.Append(enumObj.SourceName)
 									.Append("_StaticEnum, TEXT(\"")
-									.Append(Enum.EngineName)
+									.Append(enumObj.EngineName)
 									.Append("\"), &Z_Registration_Info_UEnum_")
-									.Append(Enum.EngineName)
-									.Append($", CONSTRUCT_RELOAD_VERSION_INFO(FEnumReloadVersionInfo, {Hash}U) }},\r\n");
-								if (!bAllEnumsEditorOnly && Enum.bIsEditorOnly)
+									.Append(enumObj.EngineName)
+									.Append($", CONSTRUCT_RELOAD_VERSION_INFO(FEnumReloadVersionInfo, {hash}U) }},\r\n");
+								if (!allEnumsEditorOnly && enumObj.IsEditorOnly)
 								{
-									Builder.Append("#endif\r\n");
+									builder.Append("#endif\r\n");
 								}
-								CombinedHash = HashCombine(CombinedHash, Hash);
+								combinedHash = HashCombine(combinedHash, hash);
 							}
 						}
-						Builder.Append("\t};\r\n");
-						if (bAllEnumsEditorOnly)
+						builder.Append("\t};\r\n");
+						if (allEnumsEditorOnly)
 						{
-							Builder.Append("#endif\r\n");
+							builder.Append("#endif\r\n");
 						}
 					}
 
-					if (bHasRegisteredScriptStructs)
+					if (hasRegisteredScriptStructs)
 					{
-						Builder.Append("\tconst FStructRegisterCompiledInInfo ").Append(StaticsName).Append("::ScriptStructInfo[] = {\r\n");
-						foreach (UhtObject Object in this.HeaderFile.References.ExportTypes)
+						builder.Append("\tconst FStructRegisterCompiledInInfo ").Append(staticsName).Append("::ScriptStructInfo[] = {\r\n");
+						foreach (UhtObject obj in this.HeaderFile.References.ExportTypes)
 						{
-							if (Object is UhtScriptStruct ScriptStruct)
+							if (obj is UhtScriptStruct scriptStruct)
 							{
-								if (ScriptStruct.ScriptStructFlags.HasAnyFlags(EStructFlags.Native))
+								if (scriptStruct.ScriptStructFlags.HasAnyFlags(EStructFlags.Native))
 								{
-									uint Hash = this.ObjectInfos[ScriptStruct.ObjectTypeIndex].Hash;
-									Builder
+									uint hash = this.ObjectInfos[scriptStruct.ObjectTypeIndex]._hash;
+									builder
 										.Append("\t\t{ ")
-										.Append(ScriptStruct.SourceName)
+										.Append(scriptStruct.SourceName)
 										.Append("::StaticStruct, Z_Construct_UScriptStruct_")
-										.Append(ScriptStruct.SourceName)
+										.Append(scriptStruct.SourceName)
 										.Append("_Statics::NewStructOps, TEXT(\"")
-										.Append(ScriptStruct.EngineName)
+										.Append(scriptStruct.EngineName)
 										.Append("\"), &Z_Registration_Info_UScriptStruct_")
-										.Append(ScriptStruct.EngineName)
+										.Append(scriptStruct.EngineName)
 										.Append(", CONSTRUCT_RELOAD_VERSION_INFO(FStructReloadVersionInfo, sizeof(")
-										.Append(ScriptStruct.SourceName)
-										.Append($"), {Hash}U) }},\r\n");
-									CombinedHash = HashCombine(CombinedHash, Hash);
+										.Append(scriptStruct.SourceName)
+										.Append($"), {hash}U) }},\r\n");
+									combinedHash = HashCombine(combinedHash, hash);
 								}
 							}
 						}
-						Builder.Append("\t};\r\n");
+						builder.Append("\t};\r\n");
 					}
 
-					if (bHasRegisteredClasses)
+					if (hasRegisteredClasses)
 					{
-						Builder.Append("\tconst FClassRegisterCompiledInInfo ").Append(StaticsName).Append("::ClassInfo[] = {\r\n");
-						foreach (UhtObject Object in this.HeaderFile.References.ExportTypes)
+						builder.Append("\tconst FClassRegisterCompiledInInfo ").Append(staticsName).Append("::ClassInfo[] = {\r\n");
+						foreach (UhtObject obj in this.HeaderFile.References.ExportTypes)
 						{
-							if (Object is UhtClass Class)
+							if (obj is UhtClass classObj)
 							{
-								if (!Class.ClassFlags.HasAnyFlags(EClassFlags.Intrinsic))
+								if (!classObj.ClassFlags.HasAnyFlags(EClassFlags.Intrinsic))
 								{
-									uint Hash = this.ObjectInfos[Class.ObjectTypeIndex].Hash;
-									Builder
+									uint hash = this.ObjectInfos[classObj.ObjectTypeIndex]._hash;
+									builder
 										.Append("\t\t{ Z_Construct_UClass_")
-										.Append(Class.SourceName)
+										.Append(classObj.SourceName)
 										.Append(", ")
-										.Append(Class.SourceName)
+										.Append(classObj.SourceName)
 										.Append("::StaticClass, TEXT(\"")
-										.Append(Class.SourceName)
+										.Append(classObj.SourceName)
 										.Append("\"), &Z_Registration_Info_UClass_")
-										.Append(Class.SourceName)
+										.Append(classObj.SourceName)
 										.Append(", CONSTRUCT_RELOAD_VERSION_INFO(FClassReloadVersionInfo, sizeof(")
-										.Append(Class.SourceName)
-										.Append($"), {Hash}U) }},\r\n");
-									CombinedHash = HashCombine(CombinedHash, Hash);
+										.Append(classObj.SourceName)
+										.Append($"), {hash}U) }},\r\n");
+									combinedHash = HashCombine(combinedHash, hash);
 								}
 							}
 						}
-						Builder.Append("\t};\r\n");
+						builder.Append("\t};\r\n");
 					}
 
-					Builder
+					builder
 						.Append("\tstatic FRegisterCompiledInInfo ")
-						.Append(Name)
-						.Append($"_{CombinedHash}(TEXT(\"")
+						.Append(name)
+						.Append($"_{combinedHash}(TEXT(\"")
 						.Append(Package.EngineName)
 						.Append("\"),\r\n");
 
-					Builder.Append("\t\t").AppendArray(!bHasRegisteredClasses, false, StaticsName, "ClassInfo").Append(",\r\n");
-					Builder.Append("\t\t").AppendArray(!bHasRegisteredScriptStructs, false, StaticsName, "ScriptStructInfo").Append(",\r\n");
-					Builder.Append("\t\t").AppendArray(!bHasRegisteredEnums, bAllEnumsEditorOnly, StaticsName, "EnumInfo").Append(");\r\n");
+					builder.Append("\t\t").AppendArray(!hasRegisteredClasses, false, staticsName, "ClassInfo").Append(",\r\n");
+					builder.Append("\t\t").AppendArray(!hasRegisteredScriptStructs, false, staticsName, "ScriptStructInfo").Append(",\r\n");
+					builder.Append("\t\t").AppendArray(!hasRegisteredEnums, allEnumsEditorOnly, staticsName, "EnumInfo").Append(");\r\n");
 				}
 
-				if (this.Session.bIncludeDebugOutput)
+				if (this.Session.IncludeDebugOutput)
 				{
-					Builder.Append("#if 0\r\n");
-					ReadOnlyMemory<string> Sorted = this.HeaderFile.References.Declaration.GetSortedReferences(
-						(int ObjectIndex, bool bRegistered) => GetExternalDecl(ObjectIndex, bRegistered));
-					foreach (string Declaration in Sorted.Span)
+					builder.Append("#if 0\r\n");
+					ReadOnlyMemory<string> sorted = this.HeaderFile.References.Declaration.GetSortedReferences(
+						(int objectIndex, bool registered) => GetExternalDecl(objectIndex, registered));
+					foreach (string declaration in sorted.Span)
 					{
-						Builder.Append(Declaration);
+						builder.Append(declaration);
 					}
-					Builder.Append("#endif\r\n");
+					builder.Append("#endif\r\n");
 				}
 
-				int GeneratedBodyEnd = Builder.Length;
+				int generatedBodyEnd = builder.Length;
 
-				Builder.Append(EnableDeprecationWarnings).Append("\r\n");
+				builder.Append(EnableDeprecationWarnings).Append("\r\n");
 
-				using (UhtBorrowBuffer BorrowBuffer = new UhtBorrowBuffer(Builder))
+				using (UhtBorrowBuffer borrowBuffer = new UhtBorrowBuffer(builder))
 				{
-					string CppFilePath = Factory.MakePath(HeaderFile, ".gen.cpp");
-					StringView GeneratedBody = new StringView(BorrowBuffer.Buffer.Memory);
+					string cppFilePath = factory.MakePath(HeaderFile, ".gen.cpp");
+					StringView generatedBody = new StringView(borrowBuffer.Buffer.Memory);
 					if (this.SaveExportedHeaders)
 					{
-						Factory.CommitOutput(CppFilePath, GeneratedBody);
+						factory.CommitOutput(cppFilePath, generatedBody);
 					}
 
 					// Save the hash of the generated body 
-					this.HeaderInfos[this.HeaderFile.HeaderFileTypeIndex].BodyHash = UhtHash.GenenerateTextHash(GeneratedBody.Span.Slice(GeneratedBodyStart, GeneratedBodyEnd - GeneratedBodyStart));
+					this.HeaderInfos[this.HeaderFile.HeaderFileTypeIndex]._bodyHash = UhtHash.GenenerateTextHash(generatedBody.Span.Slice(generatedBodyStart, generatedBodyEnd - generatedBodyStart));
 				}
 			}
 		}
 
-		private StringBuilder AppendEnum(StringBuilder Builder, UhtEnum Enum)
+		private StringBuilder AppendEnum(StringBuilder builder, UhtEnum enumObj)
 		{
 			const string MetaDataParamsName = "Enum_MetaDataParams";
 			const string ObjectFlags = "RF_Public|RF_Transient|RF_MarkAsNative";
-			string SingletonName = GetSingletonName(Enum, true);
-			string StaticsName = SingletonName + "_Statics";
-			string RegistrationName = $"Z_Registration_Info_UEnum_{Enum.SourceName}";
+			string singletonName = GetSingletonName(enumObj, true);
+			string staticsName = singletonName + "_Statics";
+			string registrationName = $"Z_Registration_Info_UEnum_{enumObj.SourceName}";
 
-			string EnumDisplayNameFn = Enum.MetaData.GetValueOrDefault(UhtNames.EnumDisplayNameFn);
-			if (EnumDisplayNameFn.Length == 0)
+			string enumDisplayNameFn = enumObj.MetaData.GetValueOrDefault(UhtNames.EnumDisplayNameFn);
+			if (enumDisplayNameFn.Length == 0)
 			{
-				EnumDisplayNameFn = "nullptr";
+				enumDisplayNameFn = "nullptr";
 			}
 
-			using (UhtMacroBlockEmitter MacroBlockEmitter = new UhtMacroBlockEmitter(Builder, "WITH_EDITORONLY_DATA", Enum.bIsEditorOnly))
+			using (UhtMacroBlockEmitter macroBlockEmitter = new UhtMacroBlockEmitter(builder, "WITH_EDITORONLY_DATA", enumObj.IsEditorOnly))
 			{
 
 				// If we don't have a zero 0 then we emit a static assert to verify we have one
-				if (!Enum.IsValidEnumValue(0) && Enum.MetaData.ContainsKey(UhtNames.BlueprintType))
+				if (!enumObj.IsValidEnumValue(0) && enumObj.MetaData.ContainsKey(UhtNames.BlueprintType))
 				{
-					bool bHasUnparsedValue = Enum.EnumValues.Exists(x => x.Value == -1);
-					if (bHasUnparsedValue)
+					bool hasUnparsedValue = enumObj.EnumValues.Exists(x => x.Value == -1);
+					if (hasUnparsedValue)
 					{
-						Builder.Append("\tstatic_assert(");
-						bool bDoneFirst = false;
-						foreach (UhtEnumValue Value in Enum.EnumValues)
+						builder.Append("\tstatic_assert(");
+						bool doneFirst = false;
+						foreach (UhtEnumValue value in enumObj.EnumValues)
 						{
-							if (Value.Value == -1)
+							if (value.Value == -1)
 							{
-								if (bDoneFirst)
+								if (doneFirst)
 								{
-									Builder.Append("||");
+									builder.Append("||");
 								}
-								bDoneFirst = true;
-								Builder.Append("!int64(").Append(Value.Name).Append(')');
+								doneFirst = true;
+								builder.Append("!int64(").Append(value.Name).Append(')');
 							}
 						}
-						Builder.Append(", \"'").Append(Enum.SourceName).Append("' does not have a 0 entry!(This is a problem when the enum is initalized by default)\");\r\n");
+						builder.Append(", \"'").Append(enumObj.SourceName).Append("' does not have a 0 entry!(This is a problem when the enum is initalized by default)\");\r\n");
 					}
 				}
 
-				Builder.Append("\tstatic FEnumRegistrationInfo ").Append(RegistrationName).Append(";\r\n");
-				Builder.Append("\tstatic UEnum* ").Append(Enum.SourceName).Append("_StaticEnum()\r\n");
-				Builder.Append("\t{\r\n");
+				builder.Append("\tstatic FEnumRegistrationInfo ").Append(registrationName).Append(";\r\n");
+				builder.Append("\tstatic UEnum* ").Append(enumObj.SourceName).Append("_StaticEnum()\r\n");
+				builder.Append("\t{\r\n");
 
-				Builder.Append("\t\tif (!").Append(RegistrationName).Append(".OuterSingleton)\r\n");
-				Builder.Append("\t\t{\r\n");
-				Builder.Append("\t\t\t").Append(RegistrationName).Append(".OuterSingleton = GetStaticEnum(").Append(SingletonName).Append(", ")
-					.Append(this.PackageSingletonName).Append("(), TEXT(\"").Append(Enum.SourceName).Append("\"));\r\n");
-				Builder.Append("\t\t}\r\n");
-				Builder.Append("\t\treturn ").Append(RegistrationName).Append(".OuterSingleton;\r\n");
+				builder.Append("\t\tif (!").Append(registrationName).Append(".OuterSingleton)\r\n");
+				builder.Append("\t\t{\r\n");
+				builder.Append("\t\t\t").Append(registrationName).Append(".OuterSingleton = GetStaticEnum(").Append(singletonName).Append(", ")
+					.Append(this.PackageSingletonName).Append("(), TEXT(\"").Append(enumObj.SourceName).Append("\"));\r\n");
+				builder.Append("\t\t}\r\n");
+				builder.Append("\t\treturn ").Append(registrationName).Append(".OuterSingleton;\r\n");
 
-				Builder.Append("\t}\r\n");
+				builder.Append("\t}\r\n");
 
-				Builder.Append("\ttemplate<> ").Append(this.PackageApi).Append("UEnum* StaticEnum<").Append(Enum.CppType).Append(">()\r\n");
-				Builder.Append("\t{\r\n");
-				Builder.Append("\t\treturn ").Append(Enum.SourceName).Append("_StaticEnum();\r\n");
-				Builder.Append("\t}\r\n");
+				builder.Append("\ttemplate<> ").Append(this.PackageApi).Append("UEnum* StaticEnum<").Append(enumObj.CppType).Append(">()\r\n");
+				builder.Append("\t{\r\n");
+				builder.Append("\t\treturn ").Append(enumObj.SourceName).Append("_StaticEnum();\r\n");
+				builder.Append("\t}\r\n");
 
 				// Everything from this point on will be part of the definition hash
-				int HashCodeBlockStart = Builder.Length;
+				int hashCodeBlockStart = builder.Length;
 
 				// Statics declaration
-				Builder.Append("\tstruct ").Append(StaticsName).Append("\r\n");
-				Builder.Append("\t{\r\n");
-				Builder.Append("\t\tstatic const UECodeGen_Private::FEnumeratorParam Enumerators[];\r\n");
-				Builder.AppendMetaDataDecl(Enum, MetaDataParamsName, 2);
-				Builder.Append("\t\tstatic const UECodeGen_Private::FEnumParams EnumParams;\r\n");
-				Builder.Append("\t};\r\n");
+				builder.Append("\tstruct ").Append(staticsName).Append("\r\n");
+				builder.Append("\t{\r\n");
+				builder.Append("\t\tstatic const UECodeGen_Private::FEnumeratorParam Enumerators[];\r\n");
+				builder.AppendMetaDataDecl(enumObj, MetaDataParamsName, 2);
+				builder.Append("\t\tstatic const UECodeGen_Private::FEnumParams EnumParams;\r\n");
+				builder.Append("\t};\r\n");
 
 				// Enumerators
-				Builder.Append("\tconst UECodeGen_Private::FEnumeratorParam ").Append(StaticsName).Append("::Enumerators[] = {\r\n");
-				int EnumIndex = 0;
-				foreach (UhtEnumValue Value in Enum.EnumValues)
+				builder.Append("\tconst UECodeGen_Private::FEnumeratorParam ").Append(staticsName).Append("::Enumerators[] = {\r\n");
+				int enumIndex = 0;
+				foreach (UhtEnumValue value in enumObj.EnumValues)
 				{
-					string? KeyName;
-					if (!Enum.MetaData.TryGetValue("OverrideName", EnumIndex, out KeyName))
+					string? keyName;
+					if (!enumObj.MetaData.TryGetValue("OverrideName", enumIndex, out keyName))
 					{
-						KeyName = Value.Name.ToString();
+						keyName = value.Name.ToString();
 					}
-					Builder.Append("\t\t{ ").AppendUTF8LiteralString(KeyName).Append(", (int64)").Append(Value.Name).Append(" },\r\n");
-					++EnumIndex;
+					builder.Append("\t\t{ ").AppendUTF8LiteralString(keyName).Append(", (int64)").Append(value.Name).Append(" },\r\n");
+					++enumIndex;
 				}
-				Builder.Append("\t};\r\n");
+				builder.Append("\t};\r\n");
 
 				// Meta data
-				Builder.AppendMetaDataDef(Enum, StaticsName, MetaDataParamsName, 1);
+				builder.AppendMetaDataDef(enumObj, staticsName, MetaDataParamsName, 1);
 
 				// Singleton parameters
-				Builder.Append("\tconst UECodeGen_Private::FEnumParams ").Append(StaticsName).Append("::EnumParams = {\r\n");
-				Builder.Append("\t\t(UObject*(*)())").Append(this.PackageSingletonName).Append(",\r\n");
-				Builder.Append("\t\t").Append(EnumDisplayNameFn).Append(",\r\n");
-				Builder.Append("\t\t").AppendUTF8LiteralString(Enum.SourceName).Append(",\r\n");
-				Builder.Append("\t\t").AppendUTF8LiteralString(Enum.CppType).Append(",\r\n");
-				Builder.Append("\t\t").Append(StaticsName).Append("::Enumerators,\r\n");
-				Builder.Append("\t\tUE_ARRAY_COUNT(").Append(StaticsName).Append("::Enumerators),\r\n");
-				Builder.Append("\t\t").Append(ObjectFlags).Append(",\r\n");
-				Builder.Append("\t\t").Append(Enum.EnumFlags.HasAnyFlags(EEnumFlags.Flags) ? "EEnumFlags::Flags" : "EEnumFlags::None").Append(",\r\n");
-				Builder.Append("\t\t(uint8)UEnum::ECppForm::").Append(Enum.CppForm.ToString()).Append(",\r\n");
-				Builder.Append("\t\t").AppendMetaDataParams(Enum, StaticsName, MetaDataParamsName).Append("\r\n");
-				Builder.Append("\t};\r\n");
+				builder.Append("\tconst UECodeGen_Private::FEnumParams ").Append(staticsName).Append("::EnumParams = {\r\n");
+				builder.Append("\t\t(UObject*(*)())").Append(this.PackageSingletonName).Append(",\r\n");
+				builder.Append("\t\t").Append(enumDisplayNameFn).Append(",\r\n");
+				builder.Append("\t\t").AppendUTF8LiteralString(enumObj.SourceName).Append(",\r\n");
+				builder.Append("\t\t").AppendUTF8LiteralString(enumObj.CppType).Append(",\r\n");
+				builder.Append("\t\t").Append(staticsName).Append("::Enumerators,\r\n");
+				builder.Append("\t\tUE_ARRAY_COUNT(").Append(staticsName).Append("::Enumerators),\r\n");
+				builder.Append("\t\t").Append(ObjectFlags).Append(",\r\n");
+				builder.Append("\t\t").Append(enumObj.EnumFlags.HasAnyFlags(EEnumFlags.Flags) ? "EEnumFlags::Flags" : "EEnumFlags::None").Append(",\r\n");
+				builder.Append("\t\t(uint8)UEnum::ECppForm::").Append(enumObj.CppForm.ToString()).Append(",\r\n");
+				builder.Append("\t\t").AppendMetaDataParams(enumObj, staticsName, MetaDataParamsName).Append("\r\n");
+				builder.Append("\t};\r\n");
 
 				// Registration singleton
-				Builder.Append("\tUEnum* ").Append(SingletonName).Append("()\r\n");
-				Builder.Append("\t{\r\n");
-				Builder.Append("\t\tif (!").Append(RegistrationName).Append(".InnerSingleton)\r\n");
-				Builder.Append("\t\t{\r\n");
-				Builder.Append("\t\t\tUECodeGen_Private::ConstructUEnum(").Append(RegistrationName).Append(".InnerSingleton, ").Append(StaticsName).Append("::EnumParams);\r\n");
-				Builder.Append("\t\t}\r\n");
-				Builder.Append("\t\treturn ").Append(RegistrationName).Append(".InnerSingleton;\r\n");
-				Builder.Append("\t}\r\n");
+				builder.Append("\tUEnum* ").Append(singletonName).Append("()\r\n");
+				builder.Append("\t{\r\n");
+				builder.Append("\t\tif (!").Append(registrationName).Append(".InnerSingleton)\r\n");
+				builder.Append("\t\t{\r\n");
+				builder.Append("\t\t\tUECodeGen_Private::ConstructUEnum(").Append(registrationName).Append(".InnerSingleton, ").Append(staticsName).Append("::EnumParams);\r\n");
+				builder.Append("\t\t}\r\n");
+				builder.Append("\t\treturn ").Append(registrationName).Append(".InnerSingleton;\r\n");
+				builder.Append("\t}\r\n");
 
-				using (UhtBorrowBuffer BorrowBuffer = new UhtBorrowBuffer(Builder, HashCodeBlockStart, Builder.Length - HashCodeBlockStart))
+				using (UhtBorrowBuffer borrowBuffer = new UhtBorrowBuffer(builder, hashCodeBlockStart, builder.Length - hashCodeBlockStart))
 				{
-					this.ObjectInfos[Enum.ObjectTypeIndex].Hash = UhtHash.GenenerateTextHash(BorrowBuffer.Buffer.Memory.Span);
+					this.ObjectInfos[enumObj.ObjectTypeIndex]._hash = UhtHash.GenenerateTextHash(borrowBuffer.Buffer.Memory.Span);
 				}
 			}
-			return Builder;
+			return builder;
 		}
 
-		private StringBuilder AppendScriptStruct(StringBuilder Builder, UhtScriptStruct ScriptStruct)
+		private StringBuilder AppendScriptStruct(StringBuilder builder, UhtScriptStruct scriptStruct)
 		{
 			const string MetaDataParamsName = "Struct_MetaDataParams";
-			string SingletonName = GetSingletonName(ScriptStruct, true);
-			string StaticsName = SingletonName + "_Statics";
-			string RegistrationName = $"Z_Registration_Info_UScriptStruct_{ScriptStruct.EngineName}";
-			List<UhtScriptStruct> NoExportStructs = FindNoExportStructs(ScriptStruct);
+			string singletonName = GetSingletonName(scriptStruct, true);
+			string staticsName = singletonName + "_Statics";
+			string registrationName = $"Z_Registration_Info_UScriptStruct_{scriptStruct.EngineName}";
+			List<UhtScriptStruct> noExportStructs = FindNoExportStructs(scriptStruct);
 
-			if (ScriptStruct.ScriptStructFlags.HasAnyFlags(EStructFlags.Native))
+			if (scriptStruct.ScriptStructFlags.HasAnyFlags(EStructFlags.Native))
 			{
 				// Inject static assert to verify that we do not add vtable
-				if (ScriptStruct.SuperScriptStruct != null)
+				if (scriptStruct.SuperScriptStruct != null)
 				{
-					Builder.Append("\r\n");
-					Builder.Append("static_assert(std::is_polymorphic<")
-						.Append(ScriptStruct.SourceName)
+					builder.Append("\r\n");
+					builder.Append("static_assert(std::is_polymorphic<")
+						.Append(scriptStruct.SourceName)
 						.Append(">() == std::is_polymorphic<")
-						.Append(ScriptStruct.SuperScriptStruct.SourceName)
+						.Append(scriptStruct.SuperScriptStruct.SourceName)
 						.Append(">(), \"USTRUCT ")
-						.Append(ScriptStruct.SourceName)
+						.Append(scriptStruct.SourceName)
 						.Append(" cannot be polymorphic unless super ")
-						.Append(ScriptStruct.SuperScriptStruct.SourceName)
+						.Append(scriptStruct.SuperScriptStruct.SourceName)
 						.Append(" is polymorphic\");\r\n");
-					Builder.Append("\r\n");
+					builder.Append("\r\n");
 				}
 
 				// Outer singleton
-				Builder.Append("\tstatic FStructRegistrationInfo ").Append(RegistrationName).Append(";\r\n");
-				Builder.Append("class UScriptStruct* ").Append(ScriptStruct.SourceName).Append("::StaticStruct()\r\n");
-				Builder.Append("{\r\n");
-				Builder.Append("\tif (!").Append(RegistrationName).Append(".OuterSingleton)\r\n");
-				Builder.Append("\t{\r\n");
-				Builder.Append("\t\t")
-					.Append(RegistrationName)
+				builder.Append("\tstatic FStructRegistrationInfo ").Append(registrationName).Append(";\r\n");
+				builder.Append("class UScriptStruct* ").Append(scriptStruct.SourceName).Append("::StaticStruct()\r\n");
+				builder.Append("{\r\n");
+				builder.Append("\tif (!").Append(registrationName).Append(".OuterSingleton)\r\n");
+				builder.Append("\t{\r\n");
+				builder.Append("\t\t")
+					.Append(registrationName)
 					.Append(".OuterSingleton = GetStaticStruct(")
-					.Append(SingletonName)
+					.Append(singletonName)
 					.Append(", ")
 					.Append(this.PackageSingletonName)
 					.Append("(), TEXT(\"")
-					.Append(ScriptStruct.EngineName)
+					.Append(scriptStruct.EngineName)
 					.Append("\"));\r\n");
 
 				// if this struct has RigVM methods - we need to register the method to our central
 				// registry on construction of the static struct
-				if (ScriptStruct.RigVMStructInfo != null)
+				if (scriptStruct.RigVMStructInfo != null)
 				{
-					foreach (UhtRigVMMethodInfo MethodInfo in ScriptStruct.RigVMStructInfo.Methods)
+					foreach (UhtRigVMMethodInfo methodInfo in scriptStruct.RigVMStructInfo.Methods)
 					{
-						Builder.Append("\t\tTArray<FRigVMFunctionArgument> ").AppendArgumentsName(ScriptStruct, MethodInfo).Append(";\r\n");
-						foreach (UhtRigVMParameter Parameter in ScriptStruct.RigVMStructInfo.Members)
+						builder.Append("\t\tTArray<FRigVMFunctionArgument> ").AppendArgumentsName(scriptStruct, methodInfo).Append(";\r\n");
+						foreach (UhtRigVMParameter parameter in scriptStruct.RigVMStructInfo.Members)
 						{
-							Builder
+							builder
 								.Append("\t\t")
-								.AppendArgumentsName(ScriptStruct, MethodInfo)
+								.AppendArgumentsName(scriptStruct, methodInfo)
 								.Append(".Emplace(TEXT(\"")
-								.Append(Parameter.NameOriginal())
+								.Append(parameter.NameOriginal())
 								.Append("\"), TEXT(\"")
-								.Append(Parameter.TypeOriginal())
+								.Append(parameter.TypeOriginal())
 								.Append("\"));\r\n");
 						}
-						foreach (UhtRigVMParameter Parameter in MethodInfo.Parameters)
+						foreach (UhtRigVMParameter parameter in methodInfo.Parameters)
 						{
-							Builder
+							builder
 								.Append("\t\t")
-								.AppendArgumentsName(ScriptStruct, MethodInfo)
+								.AppendArgumentsName(scriptStruct, methodInfo)
 								.Append(".Emplace(TEXT(\"")
-								.Append(Parameter.NameOriginal())
+								.Append(parameter.NameOriginal())
 								.Append("\"), TEXT(\"")
-								.Append(Parameter.TypeOriginal())
+								.Append(parameter.TypeOriginal())
 								.Append("\"));\r\n");
 						}
-						Builder
+						builder
 							.Append("\t\tFRigVMRegistry::Get().Register(TEXT(\"")
-							.Append(ScriptStruct.SourceName)
+							.Append(scriptStruct.SourceName)
 							.Append("::")
-							.Append(MethodInfo.Name)
+							.Append(methodInfo.Name)
 							.Append("\"), &")
-							.Append(ScriptStruct.SourceName)
+							.Append(scriptStruct.SourceName)
 							.Append("::RigVM")
-							.Append(MethodInfo.Name)
+							.Append(methodInfo.Name)
 							.Append(", ")
-							.Append(RegistrationName)
+							.Append(registrationName)
 							.Append(".OuterSingleton, ")
-							.AppendArgumentsName(ScriptStruct, MethodInfo)
+							.AppendArgumentsName(scriptStruct, methodInfo)
 							.Append(");\r\n");
 					}
 				}
 
-				Builder.Append("\t}\r\n");
-				Builder.Append("\treturn ").Append(RegistrationName).Append(".OuterSingleton;\r\n");
-				Builder.Append("}\r\n");
+				builder.Append("\t}\r\n");
+				builder.Append("\treturn ").Append(registrationName).Append(".OuterSingleton;\r\n");
+				builder.Append("}\r\n");
 
 				// Generate the StaticStruct specialization
-				Builder.Append("template<> ").Append(this.PackageApi).Append("UScriptStruct* StaticStruct<").Append(ScriptStruct.SourceName).Append(">()\r\n");
-				Builder.Append("{\r\n");
-				Builder.Append("\treturn ").Append(ScriptStruct.SourceName).Append("::StaticStruct();\r\n");
-				Builder.Append("}\r\n");
+				builder.Append("template<> ").Append(this.PackageApi).Append("UScriptStruct* StaticStruct<").Append(scriptStruct.SourceName).Append(">()\r\n");
+				builder.Append("{\r\n");
+				builder.Append("\treturn ").Append(scriptStruct.SourceName).Append("::StaticStruct();\r\n");
+				builder.Append("}\r\n");
 			}
 
 			// Everything from this point on will be part of the definition hash
-			int HashCodeBlockStart = Builder.Length;
+			int hashCodeBlockStart = builder.Length;
 
 			// Declare the statics structure
 			{
-				Builder.Append("\tstruct ").Append(StaticsName).Append("\r\n");
-				Builder.Append("\t{\r\n");
+				builder.Append("\tstruct ").Append(staticsName).Append("\r\n");
+				builder.Append("\t{\r\n");
 
-				foreach (UhtScriptStruct NoExportStruct in NoExportStructs)
+				foreach (UhtScriptStruct noExportStruct in noExportStructs)
 				{
-					AppendMirrorsForNoexportStruct(Builder, NoExportStruct, 2);
+					AppendMirrorsForNoexportStruct(builder, noExportStruct, 2);
 				}
 
 				// Meta data
-				Builder.AppendMetaDataDecl(ScriptStruct, MetaDataParamsName, 2);
+				builder.AppendMetaDataDecl(scriptStruct, MetaDataParamsName, 2);
 
 				// New struct ops
-				if (ScriptStruct.ScriptStructFlags.HasAnyFlags(EStructFlags.Native))
+				if (scriptStruct.ScriptStructFlags.HasAnyFlags(EStructFlags.Native))
 				{
-					Builder.Append("\t\tstatic void* NewStructOps();\r\n");
+					builder.Append("\t\tstatic void* NewStructOps();\r\n");
 				}
 
-				AppendPropertiesDecl(Builder, ScriptStruct, ScriptStruct.SourceName, StaticsName, 2);
+				AppendPropertiesDecl(builder, scriptStruct, scriptStruct.SourceName, staticsName, 2);
 
-				Builder.Append("\t\tstatic const UECodeGen_Private::FStructParams ReturnStructParams;\r\n");
-				Builder.Append("\t};\r\n");
+				builder.Append("\t\tstatic const UECodeGen_Private::FStructParams ReturnStructParams;\r\n");
+				builder.Append("\t};\r\n");
 			}
 
 			// Populate the elements of the static structure
 			{
-				Builder.AppendMetaDataDef(ScriptStruct, StaticsName, MetaDataParamsName, 1);
+				builder.AppendMetaDataDef(scriptStruct, staticsName, MetaDataParamsName, 1);
 
-				if (ScriptStruct.ScriptStructFlags.HasAnyFlags(EStructFlags.Native))
+				if (scriptStruct.ScriptStructFlags.HasAnyFlags(EStructFlags.Native))
 				{
-					Builder.Append("\tvoid* ").Append(StaticsName).Append("::NewStructOps()\r\n");
-					Builder.Append("\t{\r\n");
-					Builder.Append("\t\treturn (UScriptStruct::ICppStructOps*)new UScriptStruct::TCppStructOps<").Append(ScriptStruct.SourceName).Append(">();\r\n");
-					Builder.Append("\t}\r\n");
+					builder.Append("\tvoid* ").Append(staticsName).Append("::NewStructOps()\r\n");
+					builder.Append("\t{\r\n");
+					builder.Append("\t\treturn (UScriptStruct::ICppStructOps*)new UScriptStruct::TCppStructOps<").Append(scriptStruct.SourceName).Append(">();\r\n");
+					builder.Append("\t}\r\n");
 				}
 
-				AppendPropertiesDefs(Builder, ScriptStruct, ScriptStruct.SourceName, StaticsName, 1);
+				AppendPropertiesDefs(builder, scriptStruct, scriptStruct.SourceName, staticsName, 1);
 
-				Builder.Append("\tconst UECodeGen_Private::FStructParams ").Append(StaticsName).Append("::ReturnStructParams = {\r\n");
-				Builder.Append("\t\t(UObject* (*)())").Append(this.PackageSingletonName).Append(",\r\n");
-				Builder.Append("\t\t").Append(GetSingletonName(ScriptStruct.SuperScriptStruct, true)).Append(",\r\n");
-				if (ScriptStruct.ScriptStructFlags.HasAnyFlags(EStructFlags.Native))
+				builder.Append("\tconst UECodeGen_Private::FStructParams ").Append(staticsName).Append("::ReturnStructParams = {\r\n");
+				builder.Append("\t\t(UObject* (*)())").Append(this.PackageSingletonName).Append(",\r\n");
+				builder.Append("\t\t").Append(GetSingletonName(scriptStruct.SuperScriptStruct, true)).Append(",\r\n");
+				if (scriptStruct.ScriptStructFlags.HasAnyFlags(EStructFlags.Native))
 				{
-					Builder.Append("\t\t&NewStructOps,\r\n");
+					builder.Append("\t\t&NewStructOps,\r\n");
 				}
 				else
 				{
-					Builder.Append("\t\tnullptr,\r\n");
+					builder.Append("\t\tnullptr,\r\n");
 				}
-				Builder.Append("\t\t").AppendUTF8LiteralString(ScriptStruct.EngineName).Append(",\r\n");
-				Builder.Append("\t\tsizeof(").Append(ScriptStruct.SourceName).Append("),\r\n");
-				Builder.Append("\t\talignof(").Append(ScriptStruct.SourceName).Append("),\r\n");
-				Builder.AppendPropertiesParams(ScriptStruct, StaticsName, 2, "\r\n");
-				Builder.Append("\t\tRF_Public|RF_Transient|RF_MarkAsNative,\r\n");
-				Builder.Append($"\t\tEStructFlags(0x{(uint)(ScriptStruct.ScriptStructFlags & ~EStructFlags.ComputedFlags):X8}),\r\n");
-				Builder.Append("\t\t").AppendMetaDataParams(ScriptStruct, StaticsName, MetaDataParamsName).Append("\r\n");
-				Builder.Append("\t};\r\n");
+				builder.Append("\t\t").AppendUTF8LiteralString(scriptStruct.EngineName).Append(",\r\n");
+				builder.Append("\t\tsizeof(").Append(scriptStruct.SourceName).Append("),\r\n");
+				builder.Append("\t\talignof(").Append(scriptStruct.SourceName).Append("),\r\n");
+				builder.AppendPropertiesParams(scriptStruct, staticsName, 2, "\r\n");
+				builder.Append("\t\tRF_Public|RF_Transient|RF_MarkAsNative,\r\n");
+				builder.Append($"\t\tEStructFlags(0x{(uint)(scriptStruct.ScriptStructFlags & ~EStructFlags.ComputedFlags):X8}),\r\n");
+				builder.Append("\t\t").AppendMetaDataParams(scriptStruct, staticsName, MetaDataParamsName).Append("\r\n");
+				builder.Append("\t};\r\n");
 			}
 
 			// Generate the registration function
 			{
-				Builder.Append("\tUScriptStruct* ").Append(SingletonName).Append("()\r\n");
-				Builder.Append("\t{\r\n");
-				string InnerSingletonName;
-				if (ScriptStruct.ScriptStructFlags.HasAnyFlags(EStructFlags.Native))
+				builder.Append("\tUScriptStruct* ").Append(singletonName).Append("()\r\n");
+				builder.Append("\t{\r\n");
+				string innerSingletonName;
+				if (scriptStruct.ScriptStructFlags.HasAnyFlags(EStructFlags.Native))
 				{
-					InnerSingletonName = $"{RegistrationName}.InnerSingleton";
+					innerSingletonName = $"{registrationName}.InnerSingleton";
 				}
 				else
 				{
-					Builder.Append("\t\tstatic UScriptStruct* ReturnStruct = nullptr;\r\n");
-					InnerSingletonName = "ReturnStruct";
+					builder.Append("\t\tstatic UScriptStruct* ReturnStruct = nullptr;\r\n");
+					innerSingletonName = "ReturnStruct";
 				}
-				Builder.Append("\t\tif (!").Append(InnerSingletonName).Append(")\r\n");
-				Builder.Append("\t\t{\r\n");
-				Builder.Append("\t\t\tUECodeGen_Private::ConstructUScriptStruct(").Append(InnerSingletonName).Append(", ").Append(StaticsName).Append("::ReturnStructParams);\r\n");
-				Builder.Append("\t\t}\r\n");
-				Builder.Append("\t\treturn ").Append(InnerSingletonName).Append(";\r\n");
-				Builder.Append("\t}\r\n");
+				builder.Append("\t\tif (!").Append(innerSingletonName).Append(")\r\n");
+				builder.Append("\t\t{\r\n");
+				builder.Append("\t\t\tUECodeGen_Private::ConstructUScriptStruct(").Append(innerSingletonName).Append(", ").Append(staticsName).Append("::ReturnStructParams);\r\n");
+				builder.Append("\t\t}\r\n");
+				builder.Append("\t\treturn ").Append(innerSingletonName).Append(";\r\n");
+				builder.Append("\t}\r\n");
 			}
 
-			using (UhtBorrowBuffer BorrowBuffer = new UhtBorrowBuffer(Builder, HashCodeBlockStart, Builder.Length - HashCodeBlockStart))
+			using (UhtBorrowBuffer borrowBuffer = new UhtBorrowBuffer(builder, hashCodeBlockStart, builder.Length - hashCodeBlockStart))
 			{
-				this.ObjectInfos[ScriptStruct.ObjectTypeIndex].Hash = UhtHash.GenenerateTextHash(BorrowBuffer.Buffer.Memory.Span);
+				this.ObjectInfos[scriptStruct.ObjectTypeIndex]._hash = UhtHash.GenenerateTextHash(borrowBuffer.Buffer.Memory.Span);
 			}
 
 			// if this struct has RigVM methods we need to implement both the 
 			// virtual function as well as the stub method here.
 			// The static method is implemented by the user using a macro.
-			if (ScriptStruct.RigVMStructInfo != null)
+			if (scriptStruct.RigVMStructInfo != null)
 			{
-				foreach (UhtRigVMMethodInfo MethodInfo in ScriptStruct.RigVMStructInfo.Methods)
+				foreach (UhtRigVMMethodInfo methodInfo in scriptStruct.RigVMStructInfo.Methods)
 				{
-					Builder.Append("\r\n");
+					builder.Append("\r\n");
 
-					Builder
-						.Append(MethodInfo.ReturnType)
+					builder
+						.Append(methodInfo.ReturnType)
 						.Append(' ')
-						.Append(ScriptStruct.SourceName)
+						.Append(scriptStruct.SourceName)
 						.Append("::")
-						.Append(MethodInfo.Name)
+						.Append(methodInfo.Name)
 						.Append('(')
-						.AppendParameterDecls(MethodInfo.Parameters, false, ",\r\n\t\t")
+						.AppendParameterDecls(methodInfo.Parameters, false, ",\r\n\t\t")
 						.Append(")\r\n");
-					Builder.Append("{\r\n");
-					Builder.Append("\tFRigVMExecuteContext RigVMExecuteContext;\r\n");
+					builder.Append("{\r\n");
+					builder.Append("\tFRigVMExecuteContext RigVMExecuteContext;\r\n");
 
-					bool bWroteLine = false;
-					foreach (UhtRigVMParameter Parameter in ScriptStruct.RigVMStructInfo.Members)
+					bool wroteLine = false;
+					foreach (UhtRigVMParameter parameter in scriptStruct.RigVMStructInfo.Members)
 					{
-						if (!Parameter.RequiresCast())
+						if (!parameter.RequiresCast())
 						{
 							continue;
 						}
-						bWroteLine = true;
-						Builder.Append('\t').Append(Parameter.CastType).Append(' ').Append(Parameter.CastName).Append('(').Append(Parameter.Name).Append(");\r\n");
+						wroteLine = true;
+						builder.Append('\t').Append(parameter.CastType).Append(' ').Append(parameter.CastName).Append('(').Append(parameter.Name).Append(");\r\n");
 					}
-					if (bWroteLine)
+					if (wroteLine)
 					{
 						//COMPATIBILITY-TODO - Remove the tab
-						Builder.Append("\t\r\n");
+						builder.Append("\t\r\n");
 					}
 
 					//COMPATIBILITY-TODO - Replace spaces with \t
-					Builder.Append('\t').Append(MethodInfo.ReturnPrefix()).Append("Static").Append(MethodInfo.Name).Append("(\r\n");
-					Builder.Append("\t\tRigVMExecuteContext");
-					Builder.AppendParameterNames(ScriptStruct.RigVMStructInfo.Members, true, ",\r\n\t\t", true);
-					Builder.AppendParameterNames(MethodInfo.Parameters, true, ",\r\n\t\t");
-					Builder.Append("\r\n");
-					Builder.Append("\t);\r\n");
-					Builder.Append("}\r\n");
+					builder.Append('\t').Append(methodInfo.ReturnPrefix()).Append("Static").Append(methodInfo.Name).Append("(\r\n");
+					builder.Append("\t\tRigVMExecuteContext");
+					builder.AppendParameterNames(scriptStruct.RigVMStructInfo.Members, true, ",\r\n\t\t", true);
+					builder.AppendParameterNames(methodInfo.Parameters, true, ",\r\n\t\t");
+					builder.Append("\r\n");
+					builder.Append("\t);\r\n");
+					builder.Append("}\r\n");
 				}
-				Builder.Append("\r\n");
+				builder.Append("\r\n");
 			}
-			return Builder;
+			return builder;
 		}
 
-		private StringBuilder AppendDelegate(StringBuilder Builder, UhtFunction Function)
+		private StringBuilder AppendDelegate(StringBuilder builder, UhtFunction function)
 		{
-			AppendFunction(Builder, Function, false);
-			return Builder;
+			AppendFunction(builder, function, false);
+			return builder;
 		}
 
-		private StringBuilder AppendFunction(StringBuilder Builder, UhtFunction Function, bool bIsNoExport)
+		private StringBuilder AppendFunction(StringBuilder builder, UhtFunction function, bool isNoExport)
 		{
 			const string MetaDataParamsName = "Function_MetaDataParams";
-			string SingletonName = GetSingletonName(Function, true);
-			string StaticsName = SingletonName + "_Statics";
-			bool bParamsInStatic = bIsNoExport || !Function.FunctionFlags.HasAnyFlags(EFunctionFlags.Event);
-			bool bIsNet = Function.FunctionFlags.HasAnyFlags(EFunctionFlags.NetRequest | EFunctionFlags.NetResponse);
+			string singletonName = GetSingletonName(function, true);
+			string staticsName = singletonName + "_Statics";
+			bool paramsInStatic = isNoExport || !function.FunctionFlags.HasAnyFlags(EFunctionFlags.Event);
+			bool isNet = function.FunctionFlags.HasAnyFlags(EFunctionFlags.NetRequest | EFunctionFlags.NetResponse);
 
-			string StrippedFunctionName = Function.StrippedFunctionName;
-			string EventParameters = GetEventStructParametersName(Function.Outer, StrippedFunctionName);
+			string strippedFunctionName = function.StrippedFunctionName;
+			string eventParameters = GetEventStructParametersName(function.Outer, strippedFunctionName);
 
 			// Everything from this point on will be part of the definition hash
-			int HashCodeBlockStart = Builder.Length;
+			int hashCodeBlockStart = builder.Length;
 
-			Builder.AppendBeginEditorOnlyGuard(Function.FunctionFlags.HasAnyFlags(EFunctionFlags.EditorOnly));
+			builder.AppendBeginEditorOnlyGuard(function.FunctionFlags.HasAnyFlags(EFunctionFlags.EditorOnly));
 
 			// Statics declaration
 			{
-				Builder.Append("\tstruct ").Append(StaticsName).Append("\r\n");
-				Builder.Append("\t{\r\n");
+				builder.Append("\tstruct ").Append(staticsName).Append("\r\n");
+				builder.Append("\t{\r\n");
 
-				if (bParamsInStatic)
+				if (paramsInStatic)
 				{
-					List<UhtScriptStruct> NoExportStructs = FindNoExportStructs(Function);
-					foreach (UhtScriptStruct NoExportStruct in NoExportStructs)
+					List<UhtScriptStruct> noExportStructs = FindNoExportStructs(function);
+					foreach (UhtScriptStruct noExportStruct in noExportStructs)
 					{
-						AppendMirrorsForNoexportStruct(Builder, NoExportStruct, 2);
+						AppendMirrorsForNoexportStruct(builder, noExportStruct, 2);
 					}
-					AppendEventParameter(Builder, Function, StrippedFunctionName, UhtPropertyTextType.EventParameterFunctionMember, false, 2, "\r\n");
+					AppendEventParameter(builder, function, strippedFunctionName, UhtPropertyTextType.EventParameterFunctionMember, false, 2, "\r\n");
 				}
 
-				AppendPropertiesDecl(Builder, Function, EventParameters, StaticsName, 2);
+				AppendPropertiesDecl(builder, function, eventParameters, staticsName, 2);
 
-				Builder.AppendMetaDataDecl(Function, MetaDataParamsName, 2);
+				builder.AppendMetaDataDecl(function, MetaDataParamsName, 2);
 
-				Builder.Append("\t\tstatic const UECodeGen_Private::FFunctionParams FuncParams;\r\n");
+				builder.Append("\t\tstatic const UECodeGen_Private::FFunctionParams FuncParams;\r\n");
 
-				Builder.Append("\t};\r\n");
+				builder.Append("\t};\r\n");
 			}
 
 			// Statics definition
 			{
-				AppendPropertiesDefs(Builder, Function, EventParameters, StaticsName, 1);
+				AppendPropertiesDefs(builder, function, eventParameters, staticsName, 1);
 
-				Builder.AppendMetaDataDef(Function, StaticsName, MetaDataParamsName, 1);
+				builder.AppendMetaDataDef(function, staticsName, MetaDataParamsName, 1);
 
-				Builder
+				builder
 					.Append("\tconst UECodeGen_Private::FFunctionParams ")
-					.Append(StaticsName).Append("::FuncParams = { ")
-					.Append("(UObject*(*)())").Append(GetFunctionOuterFunc(Function)).Append(", ")
-					.Append(GetSingletonName(Function.SuperFunction, true)).Append(", ")
-					.AppendUTF8LiteralString(Function.EngineName).Append(", ")
-					.AppendUTF8LiteralString(Function.FunctionType == UhtFunctionType.SparseDelegate, Function.SparseOwningClassName).Append(", ")
-					.AppendUTF8LiteralString(Function.FunctionType == UhtFunctionType.SparseDelegate, Function.SparseDelegateName).Append(", ");
+					.Append(staticsName).Append("::FuncParams = { ")
+					.Append("(UObject*(*)())").Append(GetFunctionOuterFunc(function)).Append(", ")
+					.Append(GetSingletonName(function.SuperFunction, true)).Append(", ")
+					.AppendUTF8LiteralString(function.EngineName).Append(", ")
+					.AppendUTF8LiteralString(function.FunctionType == UhtFunctionType.SparseDelegate, function.SparseOwningClassName).Append(", ")
+					.AppendUTF8LiteralString(function.FunctionType == UhtFunctionType.SparseDelegate, function.SparseDelegateName).Append(", ");
 
-				if (Function.Children.Count > 0)
+				if (function.Children.Count > 0)
 				{
-					UhtFunction TempFunction = Function;
-					while (TempFunction.SuperFunction != null)
+					UhtFunction tempFunction = function;
+					while (tempFunction.SuperFunction != null)
 					{
-						TempFunction = TempFunction.SuperFunction;
+						tempFunction = tempFunction.SuperFunction;
 					}
 
-					if (bParamsInStatic)
+					if (paramsInStatic)
 					{
-						Builder.Append("sizeof(").Append(StaticsName).Append("::").Append(GetEventStructParametersName(TempFunction.Outer, TempFunction.StrippedFunctionName)).Append(')');
+						builder.Append("sizeof(").Append(staticsName).Append("::").Append(GetEventStructParametersName(tempFunction.Outer, tempFunction.StrippedFunctionName)).Append(')');
 					}
 					else
 					{
-						Builder.Append("sizeof(").Append(GetEventStructParametersName(TempFunction.Outer, TempFunction.StrippedFunctionName)).Append(')');
+						builder.Append("sizeof(").Append(GetEventStructParametersName(tempFunction.Outer, tempFunction.StrippedFunctionName)).Append(')');
 					}
 				}
 				else
 				{
-					Builder.Append('0');
+					builder.Append('0');
 				}
-				Builder.Append(", ");
+				builder.Append(", ");
 
-				Builder
-					.AppendPropertiesParams(Function, StaticsName, 0, " ")
+				builder
+					.AppendPropertiesParams(function, staticsName, 0, " ")
 					.Append("RF_Public|RF_Transient|RF_MarkAsNative, ")
-					.Append($"(EFunctionFlags)0x{(uint)Function.FunctionFlags:X8}, ")
-					.Append(bIsNet ? Function.RPCId : 0).Append(", ")
-					.Append(bIsNet ? Function.RPCResponseId : 0).Append(", ")
-					.AppendMetaDataParams(Function, StaticsName, MetaDataParamsName)
+					.Append($"(EFunctionFlags)0x{(uint)function.FunctionFlags:X8}, ")
+					.Append(isNet ? function.RPCId : 0).Append(", ")
+					.Append(isNet ? function.RPCResponseId : 0).Append(", ")
+					.AppendMetaDataParams(function, staticsName, MetaDataParamsName)
 					.Append(" };\r\n");
 			}
 
 			// Registration function
 			{
-				Builder.Append("\tUFunction* ").Append(SingletonName).Append("()\r\n");
-				Builder.Append("\t{\r\n");
-				Builder.Append("\t\tstatic UFunction* ReturnFunction = nullptr;\r\n");
-				Builder.Append("\t\tif (!ReturnFunction)\r\n");
-				Builder.Append("\t\t{\r\n");
-				Builder.Append("\t\t\tUECodeGen_Private::ConstructUFunction(&ReturnFunction, ").Append(StaticsName).Append("::FuncParams);\r\n");
-				Builder.Append("\t\t}\r\n");
-				Builder.Append("\t\treturn ReturnFunction;\r\n");
-				Builder.Append("\t}\r\n");
+				builder.Append("\tUFunction* ").Append(singletonName).Append("()\r\n");
+				builder.Append("\t{\r\n");
+				builder.Append("\t\tstatic UFunction* ReturnFunction = nullptr;\r\n");
+				builder.Append("\t\tif (!ReturnFunction)\r\n");
+				builder.Append("\t\t{\r\n");
+				builder.Append("\t\t\tUECodeGen_Private::ConstructUFunction(&ReturnFunction, ").Append(staticsName).Append("::FuncParams);\r\n");
+				builder.Append("\t\t}\r\n");
+				builder.Append("\t\treturn ReturnFunction;\r\n");
+				builder.Append("\t}\r\n");
 			}
 
-			Builder.AppendEndEditorOnlyGuard(Function.FunctionFlags.HasAnyFlags(EFunctionFlags.EditorOnly));
+			builder.AppendEndEditorOnlyGuard(function.FunctionFlags.HasAnyFlags(EFunctionFlags.EditorOnly));
 
-			using (UhtBorrowBuffer BorrowBuffer = new UhtBorrowBuffer(Builder, HashCodeBlockStart, Builder.Length - HashCodeBlockStart))
+			using (UhtBorrowBuffer borrowBuffer = new UhtBorrowBuffer(builder, hashCodeBlockStart, builder.Length - hashCodeBlockStart))
 			{
-				this.ObjectInfos[Function.ObjectTypeIndex].Hash = UhtHash.GenenerateTextHash(BorrowBuffer.Buffer.Memory.Span);
+				this.ObjectInfos[function.ObjectTypeIndex]._hash = UhtHash.GenenerateTextHash(borrowBuffer.Buffer.Memory.Span);
 			}
-			return Builder;
+			return builder;
 		}
 
-		private string GetFunctionOuterFunc(UhtFunction Function)
+		private string GetFunctionOuterFunc(UhtFunction function)
 		{
-			if (Function.Outer == null)
+			if (function.Outer == null)
 			{
 				return "nullptr";
 			}
-			else if (Function.Outer is UhtHeaderFile HeaderFile)
+			else if (function.Outer is UhtHeaderFile headerFile)
 			{
-				return GetSingletonName(HeaderFile.Package, true);
+				return GetSingletonName(headerFile.Package, true);
 			}
 			else
 			{
-				return GetSingletonName((UhtObject)Function.Outer, true);
+				return GetSingletonName((UhtObject)function.Outer, true);
 			}
 		}
 
-		private static StringBuilder AppendMirrorsForNoexportStruct(StringBuilder Builder, UhtScriptStruct NoExportStruct, int Tabs)
+		private static StringBuilder AppendMirrorsForNoexportStruct(StringBuilder builder, UhtScriptStruct noExportStruct, int tabs)
 		{
-			Builder.AppendTabs(Tabs).Append("struct ").Append(NoExportStruct.SourceName);
-			if (NoExportStruct.SuperScriptStruct != null)
+			builder.AppendTabs(tabs).Append("struct ").Append(noExportStruct.SourceName);
+			if (noExportStruct.SuperScriptStruct != null)
 			{
-				Builder.Append(" : public ").Append(NoExportStruct.SuperScriptStruct.SourceName);
+				builder.Append(" : public ").Append(noExportStruct.SuperScriptStruct.SourceName);
 			}
-			Builder.Append("\r\n");
-			Builder.AppendTabs(Tabs).Append("{\r\n");
+			builder.Append("\r\n");
+			builder.AppendTabs(tabs).Append("{\r\n");
 
 			// Export the struct's CPP properties
-			AppendExportProperties(Builder, NoExportStruct, Tabs + 1);
+			AppendExportProperties(builder, noExportStruct, tabs + 1);
 
-			Builder.AppendTabs(Tabs).Append("};\r\n");
-			Builder.Append("\r\n");
-			return Builder;
+			builder.AppendTabs(tabs).Append("};\r\n");
+			builder.Append("\r\n");
+			return builder;
 		}
 
-		private static StringBuilder AppendExportProperties(StringBuilder Builder, UhtScriptStruct ScriptStruct, int Tabs)
+		private static StringBuilder AppendExportProperties(StringBuilder builder, UhtScriptStruct scriptStruct, int tabs)
 		{
-			using (UhtMacroBlockEmitter Emitter = new UhtMacroBlockEmitter(Builder, "WITH_EDITORONLY_DATA"))
+			using (UhtMacroBlockEmitter emitter = new UhtMacroBlockEmitter(builder, "WITH_EDITORONLY_DATA"))
 			{
-				foreach (UhtProperty Property in ScriptStruct.Properties)
+				foreach (UhtProperty property in scriptStruct.Properties)
 				{
-					Emitter.Set(Property.bIsEditorOnlyProperty);
-					Builder.AppendTabs(Tabs).AppendFullDecl(Property, UhtPropertyTextType.ExportMember, false).Append(";\r\n");
+					emitter.Set(property.IsEditorOnlyProperty);
+					builder.AppendTabs(tabs).AppendFullDecl(property, UhtPropertyTextType.ExportMember, false).Append(";\r\n");
 				}
 			}
-			return Builder;
+			return builder;
 		}
 
-		private StringBuilder AppendPropertiesDecl(StringBuilder Builder, UhtStruct Struct, string StructSourceName, string StaticsName, int Tabs)
+		private StringBuilder AppendPropertiesDecl(StringBuilder builder, UhtStruct structObj, string structSourceName, string staticsName, int tabs)
 		{
-			if (!Struct.Children.Any(x => x is UhtProperty))
+			if (!structObj.Children.Any(x => x is UhtProperty))
 			{
-				return Builder;
+				return builder;
 			}
 
-			PropertyMemberContextImpl Context = new PropertyMemberContextImpl(this.CodeGenerator, Struct, StructSourceName, StaticsName);
+			PropertyMemberContextImpl context = new PropertyMemberContextImpl(this.CodeGenerator, structObj, structSourceName, staticsName);
 
-			using (UhtMacroBlockEmitter Emitter = new UhtMacroBlockEmitter(Builder, "WITH_EDITORONLY_DATA"))
+			using (UhtMacroBlockEmitter emitter = new UhtMacroBlockEmitter(builder, "WITH_EDITORONLY_DATA"))
 			{
-				bool bHasAllEditorOnlyDataProperties = true;
-				foreach (UhtType Type in Struct.Children)
+				bool hasAllEditorOnlyDataProperties = true;
+				foreach (UhtType type in structObj.Children)
 				{
-					if (Type is UhtProperty Property)
+					if (type is UhtProperty property)
 					{
-						Emitter.Set(Property.bIsEditorOnlyProperty);
-						bHasAllEditorOnlyDataProperties &= Property.bIsEditorOnlyProperty;
-						Builder.AppendMemberDecl(Property, Context, Property.EngineName, "", Tabs);
+						emitter.Set(property.IsEditorOnlyProperty);
+						hasAllEditorOnlyDataProperties &= property.IsEditorOnlyProperty;
+						builder.AppendMemberDecl(property, context, property.EngineName, "", tabs);
 					}
 				}
 
 				// This will force it off if the last one was editor only but we has some that weren't
-				Emitter.Set(bHasAllEditorOnlyDataProperties);
+				emitter.Set(hasAllEditorOnlyDataProperties);
 
-				Builder.AppendTabs(Tabs).Append("static const UECodeGen_Private::FPropertyParamsBase* const PropPointers[];\r\n");
+				builder.AppendTabs(tabs).Append("static const UECodeGen_Private::FPropertyParamsBase* const PropPointers[];\r\n");
 
-				foreach (UhtType Type in Struct.Children)
+				foreach (UhtType type in structObj.Children)
 				{
-					if (Type is UhtProperty Property)
+					if (type is UhtProperty property)
 					{
-						Emitter.Set(Property.bIsEditorOnlyProperty);
+						emitter.Set(property.IsEditorOnlyProperty);
 					}
 				}
-				Emitter.Set(bHasAllEditorOnlyDataProperties);
+				emitter.Set(hasAllEditorOnlyDataProperties);
 			}
-			return Builder;
+			return builder;
 		}
 
-		private StringBuilder AppendPropertiesDefs(StringBuilder Builder, UhtStruct Struct, string StructSourceName, string StaticsName, int Tabs)
+		private StringBuilder AppendPropertiesDefs(StringBuilder builder, UhtStruct structObj, string structSourceName, string staticsName, int tabs)
 		{
-			if (!Struct.Children.Any(x => x is UhtProperty))
+			if (!structObj.Children.Any(x => x is UhtProperty))
 			{
-				return Builder;
+				return builder;
 			}
 
-			PropertyMemberContextImpl Context = new PropertyMemberContextImpl(this.CodeGenerator, Struct, StructSourceName, StaticsName);
+			PropertyMemberContextImpl context = new PropertyMemberContextImpl(this.CodeGenerator, structObj, structSourceName, staticsName);
 
-			using (UhtMacroBlockEmitter Emitter = new UhtMacroBlockEmitter(Builder, "WITH_EDITORONLY_DATA"))
+			using (UhtMacroBlockEmitter emitter = new UhtMacroBlockEmitter(builder, "WITH_EDITORONLY_DATA"))
 			{
-				bool bHasAllEditorOnlyDataProperties = true;
-				foreach (UhtType Type in Struct.Children)
+				bool hasAllEditorOnlyDataProperties = true;
+				foreach (UhtType type in structObj.Children)
 				{
-					if (Type is UhtProperty Property)
+					if (type is UhtProperty property)
 					{
-						Emitter.Set(Property.bIsEditorOnlyProperty);
-						bHasAllEditorOnlyDataProperties &= Property.bIsEditorOnlyProperty;
-						Builder.AppendMemberDef(Property, Context, Property.EngineName, "", null, Tabs);
+						emitter.Set(property.IsEditorOnlyProperty);
+						hasAllEditorOnlyDataProperties &= property.IsEditorOnlyProperty;
+						builder.AppendMemberDef(property, context, property.EngineName, "", null, tabs);
 					}
 				}
 
 				// This will force it off if the last one was editor only but we has some that weren't
-				Emitter.Set(bHasAllEditorOnlyDataProperties);
+				emitter.Set(hasAllEditorOnlyDataProperties);
 
-				Builder.AppendTabs(Tabs).Append("const UECodeGen_Private::FPropertyParamsBase* const ").Append(StaticsName).Append("::PropPointers[] = {\r\n");
-				foreach (UhtType Type in Struct.Children)
+				builder.AppendTabs(tabs).Append("const UECodeGen_Private::FPropertyParamsBase* const ").Append(staticsName).Append("::PropPointers[] = {\r\n");
+				foreach (UhtType type in structObj.Children)
 				{
-					if (Type is UhtProperty Property)
+					if (type is UhtProperty property)
 					{
-						Emitter.Set(Property.bIsEditorOnlyProperty);
-						Builder.AppendMemberPtr(Property, Context, Property.EngineName, "", Tabs + 1);
+						emitter.Set(property.IsEditorOnlyProperty);
+						builder.AppendMemberPtr(property, context, property.EngineName, "", tabs + 1);
 					}
 				}
 
 				// This will force it off if the last one was editor only but we has some that weren't
-				Emitter.Set(bHasAllEditorOnlyDataProperties);
+				emitter.Set(hasAllEditorOnlyDataProperties);
 
-				Builder.AppendTabs(Tabs).Append("};\r\n");
+				builder.AppendTabs(tabs).Append("};\r\n");
 			}
-			return Builder;
+			return builder;
 		}
 
-		private StringBuilder AppendClass(StringBuilder Builder, UhtClass Class)
+		private StringBuilder AppendClass(StringBuilder builder, UhtClass classObj)
 		{
 			// Collect the functions in reversed order
-			List<UhtFunction> ReversedFunctions = new List<UhtFunction>(Class.Functions.Where(x => IsRpcFunction(x) && ShouldExportFunction(x)));
-			ReversedFunctions.Reverse();
+			List<UhtFunction> reversedFunctions = new List<UhtFunction>(classObj.Functions.Where(x => IsRpcFunction(x) && ShouldExportFunction(x)));
+			reversedFunctions.Reverse();
 
 			// Check to see if we have any RPC functions for the editor
-			bool bHasEditorRpc = ReversedFunctions.Any(x => x.FunctionFlags.HasAnyFlags(EFunctionFlags.EditorOnly));
+			bool hasEditorRpc = reversedFunctions.Any(x => x.FunctionFlags.HasAnyFlags(EFunctionFlags.EditorOnly));
 
 			// Output the RPC methods
-			AppendRpcFunctions(Builder, Class, ReversedFunctions, false);
-			if (bHasEditorRpc)
+			AppendRpcFunctions(builder, classObj, reversedFunctions, false);
+			if (hasEditorRpc)
 			{
-				Builder.AppendBeginEditorOnlyGuard();
-				AppendRpcFunctions(Builder, Class, ReversedFunctions, true);
-				Builder.AppendEndEditorOnlyGuard();
+				builder.AppendBeginEditorOnlyGuard();
+				AppendRpcFunctions(builder, classObj, reversedFunctions, true);
+				builder.AppendEndEditorOnlyGuard();
 			}
 
 			// Add the accessors
-			AppendPropertyAccessors(Builder, Class);
+			AppendPropertyAccessors(builder, classObj);
 
 			// Collect the callback function and sort by name to make the order stable
-			List<UhtFunction> CallbackFunctions = new List<UhtFunction>(Class.Functions.Where(x => x.FunctionFlags.HasAnyFlags(EFunctionFlags.Event) && x.SuperFunction == null));
-			CallbackFunctions.Sort((x, y) => StringComparerUE.OrdinalIgnoreCase.Compare(x.EngineName, y.EngineName));
+			List<UhtFunction> callbackFunctions = new List<UhtFunction>(classObj.Functions.Where(x => x.FunctionFlags.HasAnyFlags(EFunctionFlags.Event) && x.SuperFunction == null));
+			callbackFunctions.Sort((x, y) => StringComparerUE.OrdinalIgnoreCase.Compare(x.EngineName, y.EngineName));
 
 			// VM -> C++ proxies (events and delegates).
-			if (!Class.ClassExportFlags.HasAnyFlags(UhtClassExportFlags.NoExport))
+			if (!classObj.ClassExportFlags.HasAnyFlags(UhtClassExportFlags.NoExport))
 			{
-				AppendCallbackFunctions(Builder, Class, CallbackFunctions);
+				AppendCallbackFunctions(builder, classObj, callbackFunctions);
 			}
 
-			if (!Class.ClassExportFlags.HasAnyFlags(UhtClassExportFlags.NoExport))
+			if (!classObj.ClassExportFlags.HasAnyFlags(UhtClassExportFlags.NoExport))
 			{
-				AppendNatives(Builder, Class);
+				AppendNatives(builder, classObj);
 			}
 
-			AppendNativeGeneratedInitCode(Builder, Class);
+			AppendNativeGeneratedInitCode(builder, classObj);
 
-			if (!Class.ClassExportFlags.HasAnyFlags(UhtClassExportFlags.HasCustomVTableHelperConstructor))
+			if (!classObj.ClassExportFlags.HasAnyFlags(UhtClassExportFlags.HasCustomVTableHelperConstructor))
 			{
-				Builder.Append("\tDEFINE_VTABLE_PTR_HELPER_CTOR(").Append(Class.SourceName).Append(");\r\n");
+				builder.Append("\tDEFINE_VTABLE_PTR_HELPER_CTOR(").Append(classObj.SourceName).Append(");\r\n");
 			}
 
 			// Only write out adapters if the user has provided one or the other of the Serialize overloads
-			if (Class.SerializerArchiveType != UhtSerializerArchiveType.None && Class.SerializerArchiveType != UhtSerializerArchiveType.All)
+			if (classObj.SerializerArchiveType != UhtSerializerArchiveType.None && classObj.SerializerArchiveType != UhtSerializerArchiveType.All)
 			{
-				AppendSerializer(Builder, Class, UhtSerializerArchiveType.Archive, "IMPLEMENT_FARCHIVE_SERIALIZER");
-				AppendSerializer(Builder, Class, UhtSerializerArchiveType.StructuredArchiveRecord, "IMPLEMENT_FSTRUCTUREDARCHIVE_SERIALIZER");
+				AppendSerializer(builder, classObj, UhtSerializerArchiveType.Archive, "IMPLEMENT_FARCHIVE_SERIALIZER");
+				AppendSerializer(builder, classObj, UhtSerializerArchiveType.StructuredArchiveRecord, "IMPLEMENT_FSTRUCTUREDARCHIVE_SERIALIZER");
 			}
 
-			AppendFieldNotify(Builder, Class);
+			AppendFieldNotify(builder, classObj);
 
-			if (Class.ClassFlags.HasAnyFlags(EClassFlags.Interface))
+			if (classObj.ClassFlags.HasAnyFlags(EClassFlags.Interface))
 			{
-				AppendInterfaceCallFunctions(Builder, Class, CallbackFunctions);
+				AppendInterfaceCallFunctions(builder, classObj, callbackFunctions);
 			}
-			return Builder;
+			return builder;
 		}
 
-		private static StringBuilder AppendFieldNotify(StringBuilder Builder, UhtClass Class)
+		private static StringBuilder AppendFieldNotify(StringBuilder builder, UhtClass classObj)
 		{
-			if (!NeedFieldNotifyCodeGen(Class))
+			if (!NeedFieldNotifyCodeGen(classObj))
 			{
-				return Builder;
+				return builder;
 			}
 
 			// Scan the children to see what we have
-			bool bHasProperties;
-			bool bHasFunctions;
-			bool bHasEditorFields;
-			bool bAllEditorFields;
-			GetFieldNotifyStats(Class, out bHasProperties, out bHasFunctions, out bHasEditorFields, out bAllEditorFields);
+			bool hasProperties;
+			bool hasFunctions;
+			bool hasEditorFields;
+			bool allEditorFields;
+			GetFieldNotifyStats(classObj, out hasProperties, out hasFunctions, out hasEditorFields, out allEditorFields);
 
-			if (bAllEditorFields)
+			if (allEditorFields)
 			{
-				Builder.Append("#if WITH_EDITORONLY_DATA\r\n");
+				builder.Append("#if WITH_EDITORONLY_DATA\r\n");
 			}
 
 			//UE_FIELD_NOTIFICATION_DECLARE_FIELD
-			AppendFieldNotify(Builder, Class, bHasProperties, bHasFunctions, bHasEditorFields, bAllEditorFields, true, true,
-				(StringBuilder Builder, UhtClass Class, string Name) =>
+			AppendFieldNotify(builder, classObj, hasProperties, hasFunctions, hasEditorFields, allEditorFields, true, true,
+				(StringBuilder builder, UhtClass classObj, string name) =>
 				{
-					Builder.Append($"\tUE_FIELD_NOTIFICATION_IMPLEMENT_FIELD({Class.SourceName}, {Name})\r\n");
+					builder.Append($"\tUE_FIELD_NOTIFICATION_IMPLEMENT_FIELD({classObj.SourceName}, {name})\r\n");
 				});
 
 			//UE_FIELD_NOTIFICATION_DECLARE_ENUM_FIELD
-			Builder.Append($"\tUE_FIELD_NOTIFICATION_IMPLEMENTATION_BEGIN({Class.SourceName})\r\n"); 
-			AppendFieldNotify(Builder, Class, bHasProperties, bHasFunctions, bHasEditorFields, bAllEditorFields, true, true,
-				(StringBuilder Builder, UhtClass Class, string Name) =>
+			builder.Append($"\tUE_FIELD_NOTIFICATION_IMPLEMENTATION_BEGIN({classObj.SourceName})\r\n");
+			AppendFieldNotify(builder, classObj, hasProperties, hasFunctions, hasEditorFields, allEditorFields, true, true,
+				(StringBuilder builder, UhtClass classObj, string name) =>
 				{
-					Builder.Append($"\tUE_FIELD_NOTIFICATION_IMPLEMENT_ENUM_FIELD({Class.SourceName}, {Name})\r\n");
+					builder.Append($"\tUE_FIELD_NOTIFICATION_IMPLEMENT_ENUM_FIELD({classObj.SourceName}, {name})\r\n");
 				});
-			Builder.Append($"\tUE_FIELD_NOTIFICATION_IMPLEMENTATION_END({Class.SourceName});\r\n");
+			builder.Append($"\tUE_FIELD_NOTIFICATION_IMPLEMENTATION_END({classObj.SourceName});\r\n");
 
-			if (bAllEditorFields)
+			if (allEditorFields)
 			{
-				Builder.Append("#endif // WITH_EDITORONLY_DATA\r\n");
+				builder.Append("#endif // WITH_EDITORONLY_DATA\r\n");
 			}
-			return Builder;
+			return builder;
 		}
 
-		private static StringBuilder AppendPropertyAccessors(StringBuilder Builder, UhtClass Class)
+		private static StringBuilder AppendPropertyAccessors(StringBuilder builder, UhtClass classObj)
 		{
-			foreach (UhtType Type in Class.Children)
+			foreach (UhtType type in classObj.Children)
 			{
-				if (Type is UhtProperty Property)
+				if (type is UhtProperty property)
 				{
-					bool bEditorOnlyProperty = Property.bIsEditorOnlyProperty;
-					if (Property.PropertyExportFlags.HasAnyFlags(UhtPropertyExportFlags.GetterFound))
+					bool editorOnlyProperty = property.IsEditorOnlyProperty;
+					if (property.PropertyExportFlags.HasAnyFlags(UhtPropertyExportFlags.GetterFound))
 					{
-						Builder.Append("\tvoid ").Append(Class.SourceName).Append("::").AppendPropertyGetterWrapperName(Property).Append("(const void* Object, void* OutValue)\r\n");
-						Builder.Append("\t{\r\n");
-						if (bEditorOnlyProperty)
+						builder.Append("\tvoid ").Append(classObj.SourceName).Append("::").AppendPropertyGetterWrapperName(property).Append("(const void* Object, void* OutValue)\r\n");
+						builder.Append("\t{\r\n");
+						if (editorOnlyProperty)
 						{
-							Builder.Append("\t#if WITH_EDITORONLY_DATA\r\n");
+							builder.Append("\t#if WITH_EDITORONLY_DATA\r\n");
 						}
-						Builder.Append("\t\tconst ").Append(Class.SourceName).Append("* Obj = (const ").Append(Class.SourceName).Append("*)Object;\r\n");
-						if (Property.bIsStaticArray)
+						builder.Append("\t\tconst ").Append(classObj.SourceName).Append("* Obj = (const ").Append(classObj.SourceName).Append("*)Object;\r\n");
+						if (property.IsStaticArray)
 						{
-							Builder
+							builder
 								.Append("\t\t")
-								.AppendPropertyText(Property, UhtPropertyTextType.GetterSetterArg)
+								.AppendPropertyText(property, UhtPropertyTextType.GetterSetterArg)
 								.Append("* Source = (")
-								.AppendPropertyText(Property, UhtPropertyTextType.GetterSetterArg)
+								.AppendPropertyText(property, UhtPropertyTextType.GetterSetterArg)
 								.Append("*)Obj->")
-								.Append(Property.Getter!)
+								.Append(property.Getter!)
 								.Append("();\r\n");
-							Builder
+							builder
 								.Append("\t\t")
-								.AppendPropertyText(Property, UhtPropertyTextType.GetterSetterArg)
+								.AppendPropertyText(property, UhtPropertyTextType.GetterSetterArg)
 								.Append("* Result = (")
-								.AppendPropertyText(Property, UhtPropertyTextType.GetterSetterArg)
+								.AppendPropertyText(property, UhtPropertyTextType.GetterSetterArg)
 								.Append("*)OutValue;\r\n");
-							Builder
+							builder
 								.Append("\t\tCopyAssignItems(Result, Source, ")
-								.Append(Property.ArrayDimensions)
+								.Append(property.ArrayDimensions)
 								.Append(");\r\n");
 						}
 						else
 						{
-							Builder
+							builder
 								.Append("\t\t")
-								.AppendPropertyText(Property, UhtPropertyTextType.GetterSetterArg)
+								.AppendPropertyText(property, UhtPropertyTextType.GetterSetterArg)
 								.Append("& Result = *(")
-								.AppendPropertyText(Property, UhtPropertyTextType.GetterSetterArg)
+								.AppendPropertyText(property, UhtPropertyTextType.GetterSetterArg)
 								.Append("*)OutValue;\r\n");
-							Builder
+							builder
 								.Append("\t\tResult = (")
-								.AppendPropertyText(Property, UhtPropertyTextType.GetterSetterArg)
+								.AppendPropertyText(property, UhtPropertyTextType.GetterSetterArg)
 								.Append(")Obj->")
-								.Append(Property.Getter!)
+								.Append(property.Getter!)
 								.Append("();\r\n");
 						}
-						if (bEditorOnlyProperty)
+						if (editorOnlyProperty)
 						{
-							Builder.Append("\t#endif // WITH_EDITORONLY_DATA\r\n");
+							builder.Append("\t#endif // WITH_EDITORONLY_DATA\r\n");
 						}
-						Builder.Append("\t}\r\n");
+						builder.Append("\t}\r\n");
 					}
-					if (Property.PropertyExportFlags.HasAnyFlags(UhtPropertyExportFlags.SetterFound))
+					if (property.PropertyExportFlags.HasAnyFlags(UhtPropertyExportFlags.SetterFound))
 					{
-						Builder.Append("\tvoid ").Append(Class.SourceName).Append("::").AppendPropertySetterWrapperName(Property).Append("(void* Object, const void* InValue)\r\n");
-						Builder.Append("\t{\r\n");
-						if (bEditorOnlyProperty)
+						builder.Append("\tvoid ").Append(classObj.SourceName).Append("::").AppendPropertySetterWrapperName(property).Append("(void* Object, const void* InValue)\r\n");
+						builder.Append("\t{\r\n");
+						if (editorOnlyProperty)
 						{
-							Builder.Append("\t#if WITH_EDITORONLY_DATA\r\n");
+							builder.Append("\t#if WITH_EDITORONLY_DATA\r\n");
 						}
-						Builder.Append("\t\t").Append(Class.SourceName).Append("* Obj = (").Append(Class.SourceName).Append("*)Object;\r\n");
-						if (Property.bIsStaticArray)
+						builder.Append("\t\t").Append(classObj.SourceName).Append("* Obj = (").Append(classObj.SourceName).Append("*)Object;\r\n");
+						if (property.IsStaticArray)
 						{
-							Builder
+							builder
 								.Append("\t\t")
-								.AppendPropertyText(Property, UhtPropertyTextType.GetterSetterArg)
+								.AppendPropertyText(property, UhtPropertyTextType.GetterSetterArg)
 								.Append("* Value = (")
-								.AppendPropertyText(Property, UhtPropertyTextType.GetterSetterArg)
+								.AppendPropertyText(property, UhtPropertyTextType.GetterSetterArg)
 								.Append("*)InValue;\r\n");
 						}
 						else
 						{
-							Builder
+							builder
 								.Append("\t\t")
-								.AppendPropertyText(Property, UhtPropertyTextType.GetterSetterArg)
+								.AppendPropertyText(property, UhtPropertyTextType.GetterSetterArg)
 								.Append("& Value = *(")
-								.AppendPropertyText(Property, UhtPropertyTextType.GetterSetterArg)
+								.AppendPropertyText(property, UhtPropertyTextType.GetterSetterArg)
 								.Append("*)InValue;\r\n");
 						}
-						Builder
+						builder
 							.Append("\t\tObj->")
-							.Append(Property.Setter!)
+							.Append(property.Setter!)
 							.Append("(Value);\r\n");
-						if (bEditorOnlyProperty)
+						if (editorOnlyProperty)
 						{
-							Builder.Append("\t#endif // WITH_EDITORONLY_DATA\r\n");
+							builder.Append("\t#endif // WITH_EDITORONLY_DATA\r\n");
 						}
-						Builder.Append("\t}\r\n");
+						builder.Append("\t}\r\n");
 					}
 				}
 			}
-			return Builder;
+			return builder;
 		}
 
-		private StringBuilder AppendInterfaceCallFunctions(StringBuilder Builder, UhtClass Class, List<UhtFunction> CallbackFunctions)
+		private StringBuilder AppendInterfaceCallFunctions(StringBuilder builder, UhtClass classObj, List<UhtFunction> callbackFunctions)
 		{
-			foreach (UhtFunction Function in CallbackFunctions)
+			foreach (UhtFunction function in callbackFunctions)
 			{
-				Builder.Append("\tstatic FName NAME_").Append(Function.Outer?.SourceName).Append('_').Append(Function.SourceName).Append(" = FName(TEXT(\"").Append(Function.EngineName).Append("\"));\r\n");
-				string ExtraParameter = Function.FunctionFlags.HasAnyFlags(EFunctionFlags.Const) ? "const UObject* O" : "UObject* O";
-				AppendNativeFunctionHeader(Builder, Function, UhtPropertyTextType.InterfaceFunctionArgOrRetVal, false, null, ExtraParameter, 0, "\r\n");
-				Builder.Append("\t{\r\n");
-				Builder.Append("\t\tcheck(O != NULL);\r\n");
-				Builder.Append("\t\tcheck(O->GetClass()->ImplementsInterface(").Append(Class.SourceName).Append("::StaticClass()));\r\n");
-				if (Function.Children.Count > 0)
+				builder.Append("\tstatic FName NAME_").Append(function.Outer?.SourceName).Append('_').Append(function.SourceName).Append(" = FName(TEXT(\"").Append(function.EngineName).Append("\"));\r\n");
+				string extraParameter = function.FunctionFlags.HasAnyFlags(EFunctionFlags.Const) ? "const UObject* O" : "UObject* O";
+				AppendNativeFunctionHeader(builder, function, UhtPropertyTextType.InterfaceFunctionArgOrRetVal, false, null, extraParameter, 0, "\r\n");
+				builder.Append("\t{\r\n");
+				builder.Append("\t\tcheck(O != NULL);\r\n");
+				builder.Append("\t\tcheck(O->GetClass()->ImplementsInterface(").Append(classObj.SourceName).Append("::StaticClass()));\r\n");
+				if (function.Children.Count > 0)
 				{
-					Builder.Append("\t\t").Append(GetEventStructParametersName(Class, Function.StrippedFunctionName)).Append(" Parms;\r\n");
+					builder.Append("\t\t").Append(GetEventStructParametersName(classObj, function.StrippedFunctionName)).Append(" Parms;\r\n");
 				}
-				Builder.Append("\t\tUFunction* const Func = O->FindFunction(NAME_").Append(Function.Outer?.SourceName).Append('_').Append(Function.SourceName).Append(");\r\n");
-				Builder.Append("\t\tif (Func)\r\n");
-				Builder.Append("\t\t{\r\n");
-				foreach (UhtProperty Property in Function.ParameterProperties.Span)
+				builder.Append("\t\tUFunction* const Func = O->FindFunction(NAME_").Append(function.Outer?.SourceName).Append('_').Append(function.SourceName).Append(");\r\n");
+				builder.Append("\t\tif (Func)\r\n");
+				builder.Append("\t\t{\r\n");
+				foreach (UhtProperty property in function.ParameterProperties.Span)
 				{
-					Builder.Append("\t\t\tParms.").Append(Property.SourceName).Append('=').Append(Property.SourceName).Append(";\r\n");
+					builder.Append("\t\t\tParms.").Append(property.SourceName).Append('=').Append(property.SourceName).Append(";\r\n");
 				}
-				Builder
+				builder
 					.Append("\t\t\t")
-					.Append(Function.FunctionFlags.HasAnyFlags(EFunctionFlags.Const) ? "const_cast<UObject*>(O)" : "O")
+					.Append(function.FunctionFlags.HasAnyFlags(EFunctionFlags.Const) ? "const_cast<UObject*>(O)" : "O")
 					.Append("->ProcessEvent(Func, ")
-					.Append(Function.Children.Count > 0 ? "&Parms" : "NULL")
+					.Append(function.Children.Count > 0 ? "&Parms" : "NULL")
 					.Append(");\r\n");
-				foreach (UhtProperty Property in Function.ParameterProperties.Span)
+				foreach (UhtProperty property in function.ParameterProperties.Span)
 				{
-					if (Property.PropertyFlags.HasExactFlags(EPropertyFlags.OutParm | EPropertyFlags.ConstParm, EPropertyFlags.OutParm))
+					if (property.PropertyFlags.HasExactFlags(EPropertyFlags.OutParm | EPropertyFlags.ConstParm, EPropertyFlags.OutParm))
 					{
-						Builder.Append("\t\t\t").Append(Property.SourceName).Append("=Parms.").Append(Property.SourceName).Append(";\r\n");
+						builder.Append("\t\t\t").Append(property.SourceName).Append("=Parms.").Append(property.SourceName).Append(";\r\n");
 					}
 				}
-				Builder.Append("\t\t}\r\n");
+				builder.Append("\t\t}\r\n");
 
 				// else clause to call back into native if it's a BlueprintNativeEvent
-				if (Function.FunctionFlags.HasAnyFlags(EFunctionFlags.Native))
+				if (function.FunctionFlags.HasAnyFlags(EFunctionFlags.Native))
 				{
-					Builder
+					builder
 						.Append("\t\telse if (auto I = (")
-						.Append(Function.FunctionFlags.HasAnyFlags(EFunctionFlags.Const) ? "const I" : "I")
-						.Append(Class.EngineName)
+						.Append(function.FunctionFlags.HasAnyFlags(EFunctionFlags.Const) ? "const I" : "I")
+						.Append(classObj.EngineName)
 						.Append("*)(O->GetNativeInterfaceAddress(U")
-						.Append(Class.EngineName)
+						.Append(classObj.EngineName)
 						.Append("::StaticClass())))\r\n");
-					Builder.Append("\t\t{\r\n");
-					Builder.Append("\t\t\t");
-					if (Function.bHasReturnProperty)
+					builder.Append("\t\t{\r\n");
+					builder.Append("\t\t\t");
+					if (function.HasReturnProperty)
 					{
-						Builder.Append("Parms.ReturnValue = ");
+						builder.Append("Parms.ReturnValue = ");
 					}
-					Builder.Append("I->").Append(Function.SourceName).Append("_Implementation(");
+					builder.Append("I->").Append(function.SourceName).Append("_Implementation(");
 
-					bool bFirst = true;
-					foreach (UhtProperty Property in Function.ParameterProperties.Span)
+					bool first = true;
+					foreach (UhtProperty property in function.ParameterProperties.Span)
 					{
-						if (!bFirst)
+						if (!first)
 						{
-							Builder.Append(',');
+							builder.Append(',');
 						}
-						bFirst = false;
-						Builder.Append(Property.SourceName);
+						first = false;
+						builder.Append(property.SourceName);
 					}
-					Builder.Append(");\r\n");
-					Builder.Append("\t\t}\r\n");
+					builder.Append(");\r\n");
+					builder.Append("\t\t}\r\n");
 				}
 
-				if (Function.bHasReturnProperty)
+				if (function.HasReturnProperty)
 				{
-					Builder.Append("\t\treturn Parms.ReturnValue;\r\n");
+					builder.Append("\t\treturn Parms.ReturnValue;\r\n");
 				}
-				Builder.Append("\t}\r\n");
+				builder.Append("\t}\r\n");
 			}
-			return Builder;
+			return builder;
 		}
 
-		private static StringBuilder AppendSerializer(StringBuilder Builder, UhtClass Class, UhtSerializerArchiveType SerializerType, string MacroText)
+		private static StringBuilder AppendSerializer(StringBuilder builder, UhtClass classObj, UhtSerializerArchiveType serializerType, string macroText)
 		{
-			if (!Class.SerializerArchiveType.HasAnyFlags(SerializerType))
+			if (!classObj.SerializerArchiveType.HasAnyFlags(serializerType))
 			{
-				if (Class.EnclosingDefine.Length > 0)
+				if (classObj.EnclosingDefine.Length > 0)
 				{
-					Builder.Append("#if ").Append(Class.EnclosingDefine).Append("\r\n");
+					builder.Append("#if ").Append(classObj.EnclosingDefine).Append("\r\n");
 				}
-				Builder.Append('\t').Append(MacroText).Append('(').Append(Class.SourceName).Append(")\r\n");
-				if (Class.EnclosingDefine.Length > 0)
+				builder.Append('\t').Append(macroText).Append('(').Append(classObj.SourceName).Append(")\r\n");
+				if (classObj.EnclosingDefine.Length > 0)
 				{
-					Builder.Append("#endif\r\n");
+					builder.Append("#endif\r\n");
 				}
 			}
-			return Builder;
+			return builder;
 		}
 
-		private StringBuilder AppendNativeGeneratedInitCode(StringBuilder Builder, UhtClass Class)
+		private StringBuilder AppendNativeGeneratedInitCode(StringBuilder builder, UhtClass classObj)
 		{
 			const string MetaDataParamsName = "Class_MetaDataParams";
-			string SingletonName = GetSingletonName(Class, true);
-			string StaticsName = SingletonName + "_Statics";
-			string RegistrationName = $"Z_Registration_Info_UClass_{Class.SourceName}";
-			string[]? SparseDataTypes = Class.MetaData.GetStringArray(UhtNames.SparseClassDataTypes);
+			string singletonName = GetSingletonName(classObj, true);
+			string staticsName = singletonName + "_Statics";
+			string registrationName = $"Z_Registration_Info_UClass_{classObj.SourceName}";
+			string[]? sparseDataTypes = classObj.MetaData.GetStringArray(UhtNames.SparseClassDataTypes);
 
-			PropertyMemberContextImpl Context = new PropertyMemberContextImpl(this.CodeGenerator, Class, Class.SourceName, StaticsName);
+			PropertyMemberContextImpl context = new PropertyMemberContextImpl(this.CodeGenerator, classObj, classObj.SourceName, staticsName);
 
-			bool bHasInterfaces = Class.Bases.Any(x => x is UhtClass BaseClass ? BaseClass.ClassFlags.HasAnyFlags(EClassFlags.Interface) : false);
+			bool hasInterfaces = classObj.Bases.Any(x => x is UhtClass baseClass && baseClass.ClassFlags.HasAnyFlags(EClassFlags.Interface));
 
 			// Collect the functions to be exported
-			bool bAllEditorOnlyFunctions = true;
-			List<UhtFunction> SortedFunctions = new List<UhtFunction>();
-			foreach (UhtFunction Function in Class.Functions)
+			bool allEditorOnlyFunctions = true;
+			List<UhtFunction> sortedFunctions = new List<UhtFunction>();
+			foreach (UhtFunction function in classObj.Functions)
 			{
-				if (!Function.FunctionFlags.HasAnyFlags(EFunctionFlags.Delegate))
+				if (!function.FunctionFlags.HasAnyFlags(EFunctionFlags.Delegate))
 				{
-					bAllEditorOnlyFunctions &= Function.FunctionFlags.HasAnyFlags(EFunctionFlags.EditorOnly);
+					allEditorOnlyFunctions &= function.FunctionFlags.HasAnyFlags(EFunctionFlags.EditorOnly);
 				}
-				SortedFunctions.Add(Function);
+				sortedFunctions.Add(function);
 			}
-			SortedFunctions.Sort((x, y) => StringComparerUE.OrdinalIgnoreCase.Compare(x.EngineName, y.EngineName));
+			sortedFunctions.Sort((x, y) => StringComparerUE.OrdinalIgnoreCase.Compare(x.EngineName, y.EngineName));
 
 			// Output any function
-			foreach (UhtFunction Function in SortedFunctions)
+			foreach (UhtFunction function in sortedFunctions)
 			{
-				if (!Function.FunctionFlags.HasAnyFlags(EFunctionFlags.Delegate))
+				if (!function.FunctionFlags.HasAnyFlags(EFunctionFlags.Delegate))
 				{
-					AppendFunction(Builder, Function, Class.ClassExportFlags.HasAnyFlags(UhtClassExportFlags.NoExport));
-				}				
+					AppendFunction(builder, function, classObj.ClassExportFlags.HasAnyFlags(UhtClassExportFlags.NoExport));
+				}
 			}
 
-			Builder.Append("\tIMPLEMENT_CLASS_NO_AUTO_REGISTRATION(").Append(Class.SourceName).Append(");\r\n");
+			builder.Append("\tIMPLEMENT_CLASS_NO_AUTO_REGISTRATION(").Append(classObj.SourceName).Append(");\r\n");
 
 			// Everything from this point on will be part of the definition hash
-			int HashCodeBlockStart = Builder.Length;
+			int hashCodeBlockStart = builder.Length;
 
 			// simple ::StaticClass wrapper to avoid header, link and DLL hell
 			{
-				Builder.Append("\tUClass* ").Append(GetSingletonName(Class, false)).Append("()\r\n");
-				Builder.Append("\t{\r\n");
-				Builder.Append("\t\treturn ").Append(Class.SourceName).Append("::StaticClass();\r\n");
-				Builder.Append("\t}\r\n");
+				builder.Append("\tUClass* ").Append(GetSingletonName(classObj, false)).Append("()\r\n");
+				builder.Append("\t{\r\n");
+				builder.Append("\t\treturn ").Append(classObj.SourceName).Append("::StaticClass();\r\n");
+				builder.Append("\t}\r\n");
 			}
 
 			// Declare the statics object
 			{
-				Builder.Append("\tstruct ").Append(StaticsName).Append("\r\n");
-				Builder.Append("\t{\r\n");
-				Builder.Append("\t\tstatic UObject* (*const DependentSingletons[])();\r\n");
+				builder.Append("\tstruct ").Append(staticsName).Append("\r\n");
+				builder.Append("\t{\r\n");
+				builder.Append("\t\tstatic UObject* (*const DependentSingletons[])();\r\n");
 
-				if (SortedFunctions.Count != 0)
+				if (sortedFunctions.Count != 0)
 				{
-					Builder.AppendBeginEditorOnlyGuard(bAllEditorOnlyFunctions);
-					Builder.Append("\t\tstatic const FClassFunctionLinkInfo FuncInfo[];\r\n");
-					Builder.AppendEndEditorOnlyGuard(bAllEditorOnlyFunctions);
+					builder.AppendBeginEditorOnlyGuard(allEditorOnlyFunctions);
+					builder.Append("\t\tstatic const FClassFunctionLinkInfo FuncInfo[];\r\n");
+					builder.AppendEndEditorOnlyGuard(allEditorOnlyFunctions);
 				}
 
-				Builder.AppendMetaDataDecl(Class, MetaDataParamsName, 2);
+				builder.AppendMetaDataDecl(classObj, MetaDataParamsName, 2);
 
-				AppendPropertiesDecl(Builder, Class, Class.SourceName, StaticsName, 2);
+				AppendPropertiesDecl(builder, classObj, classObj.SourceName, staticsName, 2);
 
-				if (bHasInterfaces)
+				if (hasInterfaces)
 				{
-					Builder.Append("\t\tstatic const UECodeGen_Private::FImplementedInterfaceParams InterfaceParams[];\r\n");
+					builder.Append("\t\tstatic const UECodeGen_Private::FImplementedInterfaceParams InterfaceParams[];\r\n");
 				}
 
-				Builder.Append("\t\tstatic const FCppClassTypeInfoStatic StaticCppClassTypeInfo;\r\n");
+				builder.Append("\t\tstatic const FCppClassTypeInfoStatic StaticCppClassTypeInfo;\r\n");
 
-				Builder.Append("\t\tstatic const UECodeGen_Private::FClassParams ClassParams;\r\n");
+				builder.Append("\t\tstatic const UECodeGen_Private::FClassParams ClassParams;\r\n");
 
-				Builder.Append("\t};\r\n");
+				builder.Append("\t};\r\n");
 			}
 
 			// Define the statics object
 			{
-				Builder.Append("\tUObject* (*const ").Append(StaticsName).Append("::DependentSingletons[])() = {\r\n");
-				if (Class.SuperClass != null && Class.SuperClass != Class)
+				builder.Append("\tUObject* (*const ").Append(staticsName).Append("::DependentSingletons[])() = {\r\n");
+				if (classObj.SuperClass != null && classObj.SuperClass != classObj)
 				{
-					Builder.Append("\t\t(UObject* (*)())").Append(GetSingletonName(Class.SuperClass, true)).Append(",\r\n");
+					builder.Append("\t\t(UObject* (*)())").Append(GetSingletonName(classObj.SuperClass, true)).Append(",\r\n");
 				}
-				Builder.Append("\t\t(UObject* (*)())").Append(GetSingletonName(this.Package, true)).Append(",\r\n");
-				Builder.Append("\t};\r\n");
+				builder.Append("\t\t(UObject* (*)())").Append(GetSingletonName(this.Package, true)).Append(",\r\n");
+				builder.Append("\t};\r\n");
 
-				if (SortedFunctions.Count > 0)
+				if (sortedFunctions.Count > 0)
 				{
-					Builder.AppendBeginEditorOnlyGuard(bAllEditorOnlyFunctions);
-					Builder.Append("\tconst FClassFunctionLinkInfo ").Append(StaticsName).Append("::FuncInfo[] = {\r\n");
+					builder.AppendBeginEditorOnlyGuard(allEditorOnlyFunctions);
+					builder.Append("\tconst FClassFunctionLinkInfo ").Append(staticsName).Append("::FuncInfo[] = {\r\n");
 
-					foreach (UhtFunction Function in SortedFunctions)
+					foreach (UhtFunction function in sortedFunctions)
 					{
-						bool bIsEditorOnlyFunction = Function.FunctionFlags.HasAnyFlags(EFunctionFlags.EditorOnly);
-						Builder.AppendBeginEditorOnlyGuard(bIsEditorOnlyFunction);
-						Builder
+						bool isEditorOnlyFunction = function.FunctionFlags.HasAnyFlags(EFunctionFlags.EditorOnly);
+						builder.AppendBeginEditorOnlyGuard(isEditorOnlyFunction);
+						builder
 							.Append("\t\t{ &")
-							.Append(GetSingletonName(Function, true))
+							.Append(GetSingletonName(function, true))
 							.Append(", ")
-							.AppendUTF8LiteralString(Function.EngineName)
+							.AppendUTF8LiteralString(function.EngineName)
 							.Append(" },")
-							.AppendObjectHash(Class, Context, Function)
+							.AppendObjectHash(classObj, context, function)
 							.Append("\r\n");
-						Builder.AppendEndEditorOnlyGuard(bIsEditorOnlyFunction);
+						builder.AppendEndEditorOnlyGuard(isEditorOnlyFunction);
 					}
 
-					Builder.Append("\t};\r\n");
-					Builder.AppendEndEditorOnlyGuard(bAllEditorOnlyFunctions);
+					builder.Append("\t};\r\n");
+					builder.AppendEndEditorOnlyGuard(allEditorOnlyFunctions);
 				}
 
-				Builder.AppendMetaDataDef(Class, StaticsName, MetaDataParamsName, 1);
+				builder.AppendMetaDataDef(classObj, staticsName, MetaDataParamsName, 1);
 
-				AppendPropertiesDefs(Builder, Class, Class.SourceName, StaticsName, 1);
+				AppendPropertiesDefs(builder, classObj, classObj.SourceName, staticsName, 1);
 
-				if (bHasInterfaces)
+				if (hasInterfaces)
 				{
-					Builder.Append("\t\tconst UECodeGen_Private::FImplementedInterfaceParams ").Append(StaticsName).Append("::InterfaceParams[] = {\r\n");
+					builder.Append("\t\tconst UECodeGen_Private::FImplementedInterfaceParams ").Append(staticsName).Append("::InterfaceParams[] = {\r\n");
 
-					foreach (UhtStruct Struct in Class.Bases)
+					foreach (UhtStruct structObj in classObj.Bases)
 					{
-						if (Struct is UhtClass BaseClass)
+						if (structObj is UhtClass baseClass)
 						{
-							if (BaseClass.ClassFlags.HasAnyFlags(EClassFlags.Interface))
+							if (baseClass.ClassFlags.HasAnyFlags(EClassFlags.Interface))
 							{
-								Builder
+								builder
 									.Append("\t\t\t{ ")
-									.Append(GetSingletonName(BaseClass.AlternateObject, false))
+									.Append(GetSingletonName(baseClass.AlternateObject, false))
 									.Append(", (int32)VTABLE_OFFSET(")
-									.Append(Class.SourceName)
+									.Append(classObj.SourceName)
 									.Append(", ")
-									.Append(BaseClass.SourceName)
+									.Append(baseClass.SourceName)
 									.Append("), false }, ")
-									.AppendObjectHash(Class, Context, BaseClass.AlternateObject)
+									.AppendObjectHash(classObj, context, baseClass.AlternateObject)
 									.Append("\r\n");
 							}
 						}
 					}
 
-					Builder.Append("\t\t};\r\n");
+					builder.Append("\t\t};\r\n");
 				}
 
-				Builder.Append("\tconst FCppClassTypeInfoStatic ").Append(StaticsName).Append("::StaticCppClassTypeInfo = {\r\n");
-				if (Class.ClassFlags.HasAnyFlags(EClassFlags.Interface))
+				builder.Append("\tconst FCppClassTypeInfoStatic ").Append(staticsName).Append("::StaticCppClassTypeInfo = {\r\n");
+				if (classObj.ClassFlags.HasAnyFlags(EClassFlags.Interface))
 				{
-					Builder.Append("\t\tTCppClassTypeTraits<I").Append(Class.EngineName).Append(">::IsAbstract,\r\n");
+					builder.Append("\t\tTCppClassTypeTraits<I").Append(classObj.EngineName).Append(">::IsAbstract,\r\n");
 				}
 				else
 				{
-					Builder.Append("\t\tTCppClassTypeTraits<").Append(Class.SourceName).Append(">::IsAbstract,\r\n");
+					builder.Append("\t\tTCppClassTypeTraits<").Append(classObj.SourceName).Append(">::IsAbstract,\r\n");
 				}
-				Builder.Append("\t};\r\n");
+				builder.Append("\t};\r\n");
 
-				EClassFlags ClassFlags = Class.ClassFlags & EClassFlags.SaveInCompiledInClasses;
+				EClassFlags classFlags = classObj.ClassFlags & EClassFlags.SaveInCompiledInClasses;
 
-				Builder.Append("\tconst UECodeGen_Private::FClassParams ").Append(StaticsName).Append("::ClassParams = {\r\n");
-				Builder.Append("\t\t&").Append(Class.SourceName).Append("::StaticClass,\r\n");
-				if (Class.Config.Length > 0)
+				builder.Append("\tconst UECodeGen_Private::FClassParams ").Append(staticsName).Append("::ClassParams = {\r\n");
+				builder.Append("\t\t&").Append(classObj.SourceName).Append("::StaticClass,\r\n");
+				if (classObj.Config.Length > 0)
 				{
-					Builder.Append("\t\t").AppendUTF8LiteralString(Class.Config).Append(",\r\n");
+					builder.Append("\t\t").AppendUTF8LiteralString(classObj.Config).Append(",\r\n");
 				}
 				else
 				{
-					Builder.Append("\t\tnullptr,\r\n");
+					builder.Append("\t\tnullptr,\r\n");
 				}
-				Builder.Append("\t\t&StaticCppClassTypeInfo,\r\n");
-				Builder.Append("\t\tDependentSingletons,\r\n");
-				if (SortedFunctions.Count == 0)
+				builder.Append("\t\t&StaticCppClassTypeInfo,\r\n");
+				builder.Append("\t\tDependentSingletons,\r\n");
+				if (sortedFunctions.Count == 0)
 				{
-					Builder.Append("\t\tnullptr,\r\n");
+					builder.Append("\t\tnullptr,\r\n");
 				}
-				else if (bAllEditorOnlyFunctions)
+				else if (allEditorOnlyFunctions)
 				{
-					Builder.Append("\t\tIF_WITH_EDITOR(FuncInfo, nullptr),\r\n");
+					builder.Append("\t\tIF_WITH_EDITOR(FuncInfo, nullptr),\r\n");
 				}
 				else
 				{
-					Builder.Append("\t\tFuncInfo,\r\n");
+					builder.Append("\t\tFuncInfo,\r\n");
 				}
-				Builder.AppendPropertiesParamsList(Class, StaticsName, 2, "\r\n");
-				Builder.Append("\t\t").Append(bHasInterfaces ? "InterfaceParams" : "nullptr").Append(",\r\n");
-				Builder.Append("\t\tUE_ARRAY_COUNT(DependentSingletons),\r\n");
-				if (SortedFunctions.Count == 0)
+				builder.AppendPropertiesParamsList(classObj, staticsName, 2, "\r\n");
+				builder.Append("\t\t").Append(hasInterfaces ? "InterfaceParams" : "nullptr").Append(",\r\n");
+				builder.Append("\t\tUE_ARRAY_COUNT(DependentSingletons),\r\n");
+				if (sortedFunctions.Count == 0)
 				{
-					Builder.Append("\t\t0,\r\n");
+					builder.Append("\t\t0,\r\n");
 				}
-				else if (bAllEditorOnlyFunctions)
+				else if (allEditorOnlyFunctions)
 				{
-					Builder.Append("\t\tIF_WITH_EDITOR(UE_ARRAY_COUNT(FuncInfo), 0),\r\n");
+					builder.Append("\t\tIF_WITH_EDITOR(UE_ARRAY_COUNT(FuncInfo), 0),\r\n");
 				}
 				else
 				{
-					Builder.Append("\t\tUE_ARRAY_COUNT(FuncInfo),\r\n");
+					builder.Append("\t\tUE_ARRAY_COUNT(FuncInfo),\r\n");
 				}
-				Builder.AppendPropertiesParamsCount(Class, StaticsName, 2, "\r\n");
-				Builder.Append("\t\t").Append(bHasInterfaces ? "UE_ARRAY_COUNT(InterfaceParams)" : "0").Append(",\r\n");
-				Builder.Append($"\t\t0x{(uint)ClassFlags:X8}u,\r\n");
-				Builder.Append("\t\t").AppendMetaDataParams(Class, StaticsName, MetaDataParamsName).Append("\r\n");
-				Builder.Append("\t};\r\n");
+				builder.AppendPropertiesParamsCount(classObj, staticsName, 2, "\r\n");
+				builder.Append("\t\t").Append(hasInterfaces ? "UE_ARRAY_COUNT(InterfaceParams)" : "0").Append(",\r\n");
+				builder.Append($"\t\t0x{(uint)classFlags:X8}u,\r\n");
+				builder.Append("\t\t").AppendMetaDataParams(classObj, staticsName, MetaDataParamsName).Append("\r\n");
+				builder.Append("\t};\r\n");
 			}
 
 			// Class registration
 			{
-				Builder.Append("\tUClass* ").Append(SingletonName).Append("()\r\n");
-				Builder.Append("\t{\r\n");
-				Builder.Append("\t\tif (!").Append(RegistrationName).Append(".OuterSingleton)\r\n");
-				Builder.Append("\t\t{\r\n");
-				Builder.Append("\t\t\tUECodeGen_Private::ConstructUClass(").Append(RegistrationName).Append(".OuterSingleton, ").Append(StaticsName).Append("::ClassParams);\r\n");
-				if (SparseDataTypes != null)
+				builder.Append("\tUClass* ").Append(singletonName).Append("()\r\n");
+				builder.Append("\t{\r\n");
+				builder.Append("\t\tif (!").Append(registrationName).Append(".OuterSingleton)\r\n");
+				builder.Append("\t\t{\r\n");
+				builder.Append("\t\t\tUECodeGen_Private::ConstructUClass(").Append(registrationName).Append(".OuterSingleton, ").Append(staticsName).Append("::ClassParams);\r\n");
+				if (sparseDataTypes != null)
 				{
-					foreach (string SparseClass in SparseDataTypes)
+					foreach (string sparseClass in sparseDataTypes)
 					{
-						Builder.Append("\t\t\t").Append(RegistrationName).Append(".OuterSingleton->SetSparseClassDataStruct(F").Append(SparseClass).Append("::StaticStruct());\r\n");
+						builder.Append("\t\t\t").Append(registrationName).Append(".OuterSingleton->SetSparseClassDataStruct(F").Append(sparseClass).Append("::StaticStruct());\r\n");
 					}
 				}
-				Builder.Append("\t\t}\r\n");
-				Builder.Append("\t\treturn ").Append(RegistrationName).Append(".OuterSingleton;\r\n");
-				Builder.Append("\t}\r\n");
+				builder.Append("\t\t}\r\n");
+				builder.Append("\t\treturn ").Append(registrationName).Append(".OuterSingleton;\r\n");
+				builder.Append("\t}\r\n");
 			}
 
 			// At this point, we can compute the hash... HOWEVER, in the old UHT extra data is appended to the hash block that isn't emitted to the actual output
-			using (BorrowStringBuilder HashBorrower = new BorrowStringBuilder(StringBuilderCache.Small))
+			using (BorrowStringBuilder hashBorrower = new BorrowStringBuilder(StringBuilderCache.Small))
 			{
-				StringBuilder HashBuilder = HashBorrower.StringBuilder;
-				HashBuilder.Append(Builder, HashCodeBlockStart, Builder.Length - HashCodeBlockStart);
+				StringBuilder hashBuilder = hashBorrower.StringBuilder;
+				hashBuilder.Append(builder, hashCodeBlockStart, builder.Length - hashCodeBlockStart);
 
-				int SaveLength = HashBuilder.Length;
+				int saveLength = hashBuilder.Length;
 
 				// Append base class' hash at the end of the generated code, this will force update derived classes
 				// when base class changes during hot-reload.
-				uint BaseClassHash = 0;
-				if (Class.SuperClass != null && !Class.SuperClass.ClassFlags.HasAnyFlags(EClassFlags.Intrinsic))
+				uint baseClassHash = 0;
+				if (classObj.SuperClass != null && !classObj.SuperClass.ClassFlags.HasAnyFlags(EClassFlags.Intrinsic))
 				{
-					BaseClassHash = this.ObjectInfos[Class.SuperClass.ObjectTypeIndex].Hash;
+					baseClassHash = this.ObjectInfos[classObj.SuperClass.ObjectTypeIndex]._hash;
 				}
-				HashBuilder.Append($"\r\n// {BaseClassHash}\r\n");
+				hashBuilder.Append($"\r\n// {baseClassHash}\r\n");
 
 				// Append info for the sparse class data struct onto the text to be hashed
-				if (SparseDataTypes != null)
+				if (sparseDataTypes != null)
 				{
-					foreach (string SparseDataType in SparseDataTypes)
+					foreach (string sparseDataType in sparseDataTypes)
 					{
-						UhtType? Type = this.Session.FindType(Class, UhtFindOptions.ScriptStruct | UhtFindOptions.EngineName, SparseDataType);
-						if (Type != null)
+						UhtType? type = this.Session.FindType(classObj, UhtFindOptions.ScriptStruct | UhtFindOptions.EngineName, sparseDataType);
+						if (type != null)
 						{
-							HashBuilder.Append(Type.EngineName).Append("\r\n");
-							for (UhtScriptStruct? SparseStruct = Type as UhtScriptStruct; SparseStruct != null; SparseStruct = SparseStruct.SuperScriptStruct)
+							hashBuilder.Append(type.EngineName).Append("\r\n");
+							for (UhtScriptStruct? sparseStruct = type as UhtScriptStruct; sparseStruct != null; sparseStruct = sparseStruct.SuperScriptStruct)
 							{
-								foreach (UhtProperty Property in SparseStruct.Properties)
+								foreach (UhtProperty property in sparseStruct.Properties)
 								{
-									HashBuilder.AppendPropertyText(Property, UhtPropertyTextType.SparseShort).Append(' ').Append(Property.SourceName).Append("\r\n");
+									hashBuilder.AppendPropertyText(property, UhtPropertyTextType.SparseShort).Append(' ').Append(property.SourceName).Append("\r\n");
 								}
 							}
 						}
 					}
 				}
 
-				if (Class.ClassExportFlags.HasAnyFlags(UhtClassExportFlags.NoExport))
+				if (classObj.ClassExportFlags.HasAnyFlags(UhtClassExportFlags.NoExport))
 				{
-					Builder.Append("\t/* friend declarations for pasting into noexport class ").Append(Class.SourceName).Append("\r\n");
-					Builder.Append("\tfriend struct ").Append(StaticsName).Append(";\r\n");
-					Builder.Append("\t*/\r\n");
+					builder.Append("\t/* friend declarations for pasting into noexport class ").Append(classObj.SourceName).Append("\r\n");
+					builder.Append("\tfriend struct ").Append(staticsName).Append(";\r\n");
+					builder.Append("\t*/\r\n");
 				}
 
-				if (this.Session.bIncludeDebugOutput)
+				if (this.Session.IncludeDebugOutput)
 				{
-					using (UhtBorrowBuffer BorrowBuffer = new UhtBorrowBuffer(HashBuilder, SaveLength, HashBuilder.Length - SaveLength))
+					using (UhtBorrowBuffer borrowBuffer = new UhtBorrowBuffer(hashBuilder, saveLength, hashBuilder.Length - saveLength))
 					{
-						Builder.Append("#if 0\r\n");
-						Builder.Append(BorrowBuffer.Buffer.Memory);
-						Builder.Append("#endif\r\n");
+						builder.Append("#if 0\r\n");
+						builder.Append(borrowBuffer.Buffer.Memory);
+						builder.Append("#endif\r\n");
 					}
 				}
 
 				// Calculate generated class initialization code hash so that we know when it changes after hot-reload
-				using (UhtBorrowBuffer BorrowBuffer = new UhtBorrowBuffer(HashBuilder))
+				using (UhtBorrowBuffer borrowBuffer = new UhtBorrowBuffer(hashBuilder))
 				{
-					this.ObjectInfos[Class.ObjectTypeIndex].Hash = UhtHash.GenenerateTextHash(BorrowBuffer.Buffer.Memory.Span);
+					this.ObjectInfos[classObj.ObjectTypeIndex]._hash = UhtHash.GenenerateTextHash(borrowBuffer.Buffer.Memory.Span);
 				}
 			}
 
-			Builder.Append("\ttemplate<> ").Append(this.PackageApi).Append("UClass* StaticClass<").Append(Class.SourceName).Append(">()\r\n");
-			Builder.Append("\t{\r\n");
-			Builder.Append("\t\treturn ").Append(Class.SourceName).Append("::StaticClass();\r\n");
-			Builder.Append("\t}\r\n");
+			builder.Append("\ttemplate<> ").Append(this.PackageApi).Append("UClass* StaticClass<").Append(classObj.SourceName).Append(">()\r\n");
+			builder.Append("\t{\r\n");
+			builder.Append("\t\treturn ").Append(classObj.SourceName).Append("::StaticClass();\r\n");
+			builder.Append("\t}\r\n");
 
-			if (Class.ClassExportFlags.HasAnyFlags(UhtClassExportFlags.SelfHasReplicatedProperties))
+			if (classObj.ClassExportFlags.HasAnyFlags(UhtClassExportFlags.SelfHasReplicatedProperties))
 			{
-				Builder.Append("\r\n");
-				Builder.Append("\tvoid ").Append(Class.SourceName).Append("::ValidateGeneratedRepEnums(const TArray<struct FRepRecord>& ClassReps) const\r\n");
-				Builder.Append("\t{\r\n");
+				builder.Append("\r\n");
+				builder.Append("\tvoid ").Append(classObj.SourceName).Append("::ValidateGeneratedRepEnums(const TArray<struct FRepRecord>& ClassReps) const\r\n");
+				builder.Append("\t{\r\n");
 
-				foreach (UhtProperty Property in Class.Properties)
+				foreach (UhtProperty property in classObj.Properties)
 				{
-					if (Property.PropertyFlags.HasAnyFlags(EPropertyFlags.Net))
+					if (property.PropertyFlags.HasAnyFlags(EPropertyFlags.Net))
 					{
-						Builder.Append("\t\tstatic const FName Name_").Append(Property.SourceName).Append("(TEXT(\"").Append(Property.SourceName).Append("\"));\r\n");
+						builder.Append("\t\tstatic const FName Name_").Append(property.SourceName).Append("(TEXT(\"").Append(property.SourceName).Append("\"));\r\n");
 					}
 				}
-				Builder.Append("\r\n");
-				Builder.Append("\t\tconst bool bIsValid = true");
-				foreach (UhtProperty Property in Class.Properties)
+				builder.Append("\r\n");
+				builder.Append("\t\tconst bool bIsValid = true");
+				foreach (UhtProperty property in classObj.Properties)
 				{
-					if (Property.PropertyFlags.HasAnyFlags(EPropertyFlags.Net))
+					if (property.PropertyFlags.HasAnyFlags(EPropertyFlags.Net))
 					{
-						if (!Property.bIsStaticArray)
+						if (!property.IsStaticArray)
 						{
-							Builder.Append("\r\n\t\t\t&& Name_").Append(Property.SourceName).Append(" == ClassReps[(int32)ENetFields_Private::").Append(Property.SourceName).Append("].Property->GetFName()");
+							builder.Append("\r\n\t\t\t&& Name_").Append(property.SourceName).Append(" == ClassReps[(int32)ENetFields_Private::").Append(property.SourceName).Append("].Property->GetFName()");
 						}
 						else
 						{
-							Builder.Append("\r\n\t\t\t&& Name_").Append(Property.SourceName).Append(" == ClassReps[(int32)ENetFields_Private::").Append(Property.SourceName).Append("_STATIC_ARRAY].Property->GetFName()");
+							builder.Append("\r\n\t\t\t&& Name_").Append(property.SourceName).Append(" == ClassReps[(int32)ENetFields_Private::").Append(property.SourceName).Append("_STATIC_ARRAY].Property->GetFName()");
 						}
 					}
 				}
-				Builder.Append(";\r\n");
-				Builder.Append("\r\n");
-				Builder.Append("\t\tcheckf(bIsValid, TEXT(\"UHT Generated Rep Indices do not match runtime populated Rep Indices for properties in ").Append(Class.SourceName).Append("\"));\r\n");
-				Builder.Append("\t}\r\n");
+				builder.Append(";\r\n");
+				builder.Append("\r\n");
+				builder.Append("\t\tcheckf(bIsValid, TEXT(\"UHT Generated Rep Indices do not match runtime populated Rep Indices for properties in ").Append(classObj.SourceName).Append("\"));\r\n");
+				builder.Append("\t}\r\n");
 			}
-			return Builder;
+			return builder;
 		}
 
-		private static StringBuilder AppendNatives(StringBuilder Builder, UhtClass Class)
+		private static StringBuilder AppendNatives(StringBuilder builder, UhtClass classObj)
 		{
-			Builder.Append("\tvoid ").Append(Class.SourceName).Append("::StaticRegisterNatives").Append(Class.SourceName).Append("()\r\n");
-			Builder.Append("\t{\r\n");
+			builder.Append("\tvoid ").Append(classObj.SourceName).Append("::StaticRegisterNatives").Append(classObj.SourceName).Append("()\r\n");
+			builder.Append("\t{\r\n");
 
-			bool bAllEditorOnly = true;
+			bool allEditorOnly = true;
 
-			List<UhtFunction> SortedFunctions = new List<UhtFunction>();
-			foreach (UhtFunction Function in Class.Functions)
+			List<UhtFunction> sortedFunctions = new List<UhtFunction>();
+			foreach (UhtFunction function in classObj.Functions)
 			{
-				if (Function.FunctionFlags.HasExactFlags(EFunctionFlags.Native | EFunctionFlags.NetRequest, EFunctionFlags.Native))
+				if (function.FunctionFlags.HasExactFlags(EFunctionFlags.Native | EFunctionFlags.NetRequest, EFunctionFlags.Native))
 				{
-					SortedFunctions.Add(Function);
-					bAllEditorOnly &= Function.FunctionFlags.HasAnyFlags(EFunctionFlags.EditorOnly);
+					sortedFunctions.Add(function);
+					allEditorOnly &= function.FunctionFlags.HasAnyFlags(EFunctionFlags.EditorOnly);
 				}
 			}
-			SortedFunctions.Sort((x, y) => StringComparerUE.OrdinalIgnoreCase.Compare(x.EngineName, y.EngineName));
+			sortedFunctions.Sort((x, y) => StringComparerUE.OrdinalIgnoreCase.Compare(x.EngineName, y.EngineName));
 
-			if (SortedFunctions.Count != 0)
+			if (sortedFunctions.Count != 0)
 			{
-				using (UhtMacroBlockEmitter BlockEmitter = new UhtMacroBlockEmitter(Builder, "WITH_EDITOR", bAllEditorOnly))
+				using (UhtMacroBlockEmitter blockEmitter = new UhtMacroBlockEmitter(builder, "WITH_EDITOR", allEditorOnly))
 				{
-					Builder.Append("\t\tUClass* Class = ").Append(Class.SourceName).Append("::StaticClass();\r\n");
-					Builder.Append("\t\tstatic const FNameNativePtrPair Funcs[] = {\r\n");
-					
-					foreach (UhtFunction Function in SortedFunctions)
+					builder.Append("\t\tUClass* Class = ").Append(classObj.SourceName).Append("::StaticClass();\r\n");
+					builder.Append("\t\tstatic const FNameNativePtrPair Funcs[] = {\r\n");
+
+					foreach (UhtFunction function in sortedFunctions)
 					{
-						BlockEmitter.Set(Function.FunctionFlags.HasAnyFlags(EFunctionFlags.EditorOnly));
-						Builder
+						blockEmitter.Set(function.FunctionFlags.HasAnyFlags(EFunctionFlags.EditorOnly));
+						builder
 							.Append("\t\t\t{ ")
-							.AppendUTF8LiteralString(Function.EngineName)
+							.AppendUTF8LiteralString(function.EngineName)
 							.Append(", &")
-							.AppendClassSourceNameOrInterfaceName(Class)
+							.AppendClassSourceNameOrInterfaceName(classObj)
 							.Append("::exec")
-							.Append(Function.EngineName)
+							.Append(function.EngineName)
 							.Append(" },\r\n");
 					}
 
 					// This will close the block if we have one that isn't editor only
-					BlockEmitter.Set(bAllEditorOnly);
+					blockEmitter.Set(allEditorOnly);
 
-					Builder.Append("\t\t};\r\n");
-					Builder.Append("\t\tFNativeFunctionRegistrar::RegisterFunctions(Class, Funcs, UE_ARRAY_COUNT(Funcs));\r\n");
+					builder.Append("\t\t};\r\n");
+					builder.Append("\t\tFNativeFunctionRegistrar::RegisterFunctions(Class, Funcs, UE_ARRAY_COUNT(Funcs));\r\n");
 				}
 			}
 
-			Builder.Append("\t}\r\n");
-			return Builder;
+			builder.Append("\t}\r\n");
+			return builder;
 		}
 
-		private StringBuilder AppendCallbackFunctions(StringBuilder Builder, UhtClass Class, List<UhtFunction> CallbackFunctions)
+		private StringBuilder AppendCallbackFunctions(StringBuilder builder, UhtClass classObj, List<UhtFunction> callbackFunctions)
 		{
-			if (CallbackFunctions.Count > 0)
+			if (callbackFunctions.Count > 0)
 			{
-				bool bIsInterfaceClass = Class.ClassFlags.HasAnyFlags(EClassFlags.Interface);
-				using (UhtMacroBlockEmitter BlockEmitter = new UhtMacroBlockEmitter(Builder, "WITH_EDITOR"))
+				bool isInterfaceClass = classObj.ClassFlags.HasAnyFlags(EClassFlags.Interface);
+				using (UhtMacroBlockEmitter blockEmitter = new UhtMacroBlockEmitter(builder, "WITH_EDITOR"))
 				{
-					foreach (UhtFunction Function in CallbackFunctions)
+					foreach (UhtFunction function in callbackFunctions)
 					{
 						// Net response functions don't go into the VM
-						if (Function.FunctionFlags.HasAnyFlags(EFunctionFlags.NetResponse))
+						if (function.FunctionFlags.HasAnyFlags(EFunctionFlags.NetResponse))
 						{
 							continue;
 						}
 
-						BlockEmitter.Set(Function.FunctionFlags.HasAnyFlags(EFunctionFlags.EditorOnly));
+						blockEmitter.Set(function.FunctionFlags.HasAnyFlags(EFunctionFlags.EditorOnly));
 
-						if (!bIsInterfaceClass)
+						if (!isInterfaceClass)
 						{
-							Builder.Append("\tstatic FName NAME_").Append(Class.SourceName).Append('_').Append(Function.EngineName).Append(" = FName(TEXT(\"").Append(Function.EngineName).Append("\"));\r\n");
+							builder.Append("\tstatic FName NAME_").Append(classObj.SourceName).Append('_').Append(function.EngineName).Append(" = FName(TEXT(\"").Append(function.EngineName).Append("\"));\r\n");
 						}
 
-						AppendNativeFunctionHeader(Builder, Function, UhtPropertyTextType.EventFunctionArgOrRetVal, false, null, null, UhtFunctionExportFlags.None, "\r\n");
+						AppendNativeFunctionHeader(builder, function, UhtPropertyTextType.EventFunctionArgOrRetVal, false, null, null, UhtFunctionExportFlags.None, "\r\n");
 
-						if (bIsInterfaceClass)
+						if (isInterfaceClass)
 						{
-							Builder.Append("\t{\r\n");
+							builder.Append("\t{\r\n");
 
 							// assert if this is ever called directly
-							Builder
+							builder
 								.Append("\t\tcheck(0 && \"Do not directly call Event functions in Interfaces. Call Execute_")
-								.Append(Function.EngineName)
+								.Append(function.EngineName)
 								.Append(" instead.\");\r\n");
 
 							// satisfy compiler if it's expecting a return value
-							if (Function.ReturnProperty != null)
+							if (function.ReturnProperty != null)
 							{
-								string EventParmStructName = GetEventStructParametersName(Class, Function.EngineName);
-								Builder.Append("\t\t").Append(EventParmStructName).Append(" Parms;\r\n");
-								Builder.Append("\t\treturn Parms.ReturnValue;\r\n");
+								string eventParmStructName = GetEventStructParametersName(classObj, function.EngineName);
+								builder.Append("\t\t").Append(eventParmStructName).Append(" Parms;\r\n");
+								builder.Append("\t\treturn Parms.ReturnValue;\r\n");
 							}
-							Builder.Append("\t}\r\n");
+							builder.Append("\t}\r\n");
 						}
 						else
 						{
-							AppendEventFunctionPrologue(Builder, Function, Function.EngineName, 1, "\r\n");
+							AppendEventFunctionPrologue(builder, function, function.EngineName, 1, "\r\n");
 
 							// Cast away const just in case, because ProcessEvent isn't const
-							Builder.Append("\t\t");
-							if (Function.FunctionFlags.HasAnyFlags(EFunctionFlags.Const))
+							builder.Append("\t\t");
+							if (function.FunctionFlags.HasAnyFlags(EFunctionFlags.Const))
 							{
-								Builder.Append("const_cast<").Append(Class.SourceName).Append("*>(this)->");
+								builder.Append("const_cast<").Append(classObj.SourceName).Append("*>(this)->");
 							}
-							Builder
+							builder
 								.Append("ProcessEvent(FindFunctionChecked(")
 								.Append("NAME_")
-								.Append(Class.SourceName)
+								.Append(classObj.SourceName)
 								.Append('_')
-								.Append(Function.EngineName)
+								.Append(function.EngineName)
 								.Append("),")
-								.Append(Function.Children.Count > 0 ? "&Parms" : "NULL")
+								.Append(function.Children.Count > 0 ? "&Parms" : "NULL")
 								.Append(");\r\n");
 
-							AppendEventFunctionEpilogue(Builder, Function, 1, "\r\n");
+							AppendEventFunctionEpilogue(builder, function, 1, "\r\n");
 						}
 					}
 				}
 			}
-			return Builder;
+			return builder;
 		}
 
-		private static StringBuilder AppendRpcFunctions(StringBuilder Builder, UhtClass Class, List<UhtFunction> ReversedFunctions, bool bEditorOnly)
+		private static StringBuilder AppendRpcFunctions(StringBuilder builder, UhtClass classObj, List<UhtFunction> reversedFunctions, bool editorOnly)
 		{
-			foreach (UhtFunction Function in ReversedFunctions)
+			foreach (UhtFunction function in reversedFunctions)
 			{
-				if (Function.FunctionFlags.HasAnyFlags(EFunctionFlags.EditorOnly) == bEditorOnly)
+				if (function.FunctionFlags.HasAnyFlags(EFunctionFlags.EditorOnly) == editorOnly)
 				{
-					Builder.Append("\tDEFINE_FUNCTION(").AppendClassSourceNameOrInterfaceName(Class).Append("::").Append(Function.UnMarshalAndCallName).Append(")\r\n");
-					Builder.Append("\t{\r\n");
-					AppendFunctionThunk(Builder, Function);
-					Builder.Append("\t}\r\n");
+					builder.Append("\tDEFINE_FUNCTION(").AppendClassSourceNameOrInterfaceName(classObj).Append("::").Append(function.UnMarshalAndCallName).Append(")\r\n");
+					builder.Append("\t{\r\n");
+					AppendFunctionThunk(builder, function);
+					builder.Append("\t}\r\n");
 				}
 			}
-			return Builder;
+			return builder;
 		}
 
-		private static StringBuilder AppendFunctionThunk(StringBuilder Builder, UhtFunction Function)
+		private static StringBuilder AppendFunctionThunk(StringBuilder builder, UhtFunction function)
 		{
 			// Export the GET macro for the parameters
-			foreach (UhtProperty Parameter in Function.ParameterProperties.Span)
+			foreach (UhtProperty parameter in function.ParameterProperties.Span)
 			{
-				Builder.Append("\t\t").AppendFunctionThunkParameterGet(Parameter).Append(";\r\n");
+				builder.Append("\t\t").AppendFunctionThunkParameterGet(parameter).Append(";\r\n");
 			}
 
-			Builder.Append("\t\tP_FINISH;\r\n");
-			Builder.Append("\t\tP_NATIVE_BEGIN;\r\n");
+			builder.Append("\t\tP_FINISH;\r\n");
+			builder.Append("\t\tP_NATIVE_BEGIN;\r\n");
 
 			// Call the validate function if there is one
-			if (!Function.FunctionExportFlags.HasAnyFlags(UhtFunctionExportFlags.CppStatic) && Function.FunctionFlags.HasAnyFlags(EFunctionFlags.NetValidate))
+			if (!function.FunctionExportFlags.HasAnyFlags(UhtFunctionExportFlags.CppStatic) && function.FunctionFlags.HasAnyFlags(EFunctionFlags.NetValidate))
 			{
-				Builder.Append("\t\tif (!P_THIS->").Append(Function.CppValidationImplName).Append('(').AppendFunctionThunkParameterNames(Function).Append("))\r\n");
-				Builder.Append("\t\t{\r\n");
-				Builder.Append("\t\t\tRPC_ValidateFailed(TEXT(\"").Append(Function.CppValidationImplName).Append("\"));\r\n");
-				Builder.Append("\t\t\treturn;\r\n");   // If we got here, the validation function check failed
-				Builder.Append("\t\t}\r\n");
+				builder.Append("\t\tif (!P_THIS->").Append(function.CppValidationImplName).Append('(').AppendFunctionThunkParameterNames(function).Append("))\r\n");
+				builder.Append("\t\t{\r\n");
+				builder.Append("\t\t\tRPC_ValidateFailed(TEXT(\"").Append(function.CppValidationImplName).Append("\"));\r\n");
+				builder.Append("\t\t\treturn;\r\n");   // If we got here, the validation function check failed
+				builder.Append("\t\t}\r\n");
 			}
 
 			// Write out the return value
-			Builder.Append("\t\t");
-			UhtProperty? ReturnProperty = Function.ReturnProperty;
-			if (ReturnProperty != null)
+			builder.Append("\t\t");
+			UhtProperty? returnProperty = function.ReturnProperty;
+			if (returnProperty != null)
 			{
-				Builder.Append("*(").AppendFunctionThunkReturn(ReturnProperty).Append("*)Z_Param__Result=");
+				builder.Append("*(").AppendFunctionThunkReturn(returnProperty).Append("*)Z_Param__Result=");
 			}
 
 			// Export the call to the C++ version
-			if (Function.FunctionExportFlags.HasAnyFlags(UhtFunctionExportFlags.CppStatic))
+			if (function.FunctionExportFlags.HasAnyFlags(UhtFunctionExportFlags.CppStatic))
 			{
-				Builder.Append(Function.Outer?.SourceName).Append("::").Append(Function.CppImplName).Append('(').AppendFunctionThunkParameterNames(Function).Append(");\r\n");
+				builder.Append(function.Outer?.SourceName).Append("::").Append(function.CppImplName).Append('(').AppendFunctionThunkParameterNames(function).Append(");\r\n");
 			}
 			else
 			{
-				Builder.Append("P_THIS->").Append(Function.CppImplName).Append('(').AppendFunctionThunkParameterNames(Function).Append(");\r\n");
+				builder.Append("P_THIS->").Append(function.CppImplName).Append('(').AppendFunctionThunkParameterNames(function).Append(");\r\n");
 			}
-			Builder.Append("\t\tP_NATIVE_END;\r\n");
-			return Builder;
+			builder.Append("\t\tP_NATIVE_END;\r\n");
+			return builder;
 		}
 
-		private static void FindNoExportStructsRecursive(List<UhtScriptStruct> Out, UhtStruct Struct)
+		private static void FindNoExportStructsRecursive(List<UhtScriptStruct> outScriptStructs, UhtStruct structObj)
 		{
-			for (UhtStruct? Current = Struct; Current != null; Current = Current.SuperStruct)
+			for (UhtStruct? current = structObj; current != null; current = current.SuperStruct)
 			{
 				// Is isn't true for noexport structs
-				if (Current is UhtScriptStruct ScriptStruct)
+				if (current is UhtScriptStruct scriptStruct)
 				{
-					if (ScriptStruct.ScriptStructFlags.HasAnyFlags(EStructFlags.Native))
+					if (scriptStruct.ScriptStructFlags.HasAnyFlags(EStructFlags.Native))
 					{
 						break;
 					}
 
 					// these are a special cases that already exists and if wrong if exported naively
-					if (!ScriptStruct.bIsAlwaysAccessible)
+					if (!scriptStruct.IsAlwaysAccessible)
 					{
-						Out.Remove(ScriptStruct);
-						Out.Add(ScriptStruct);
+						outScriptStructs.Remove(scriptStruct);
+						outScriptStructs.Add(scriptStruct);
 					}
 				}
 
-				foreach (UhtType Type in Current.Children)
+				foreach (UhtType type in current.Children)
 				{
-					if (Type is UhtProperty Property)
+					if (type is UhtProperty property)
 					{
-						foreach (UhtType ReferenceType in Property.EnumerateReferencedTypes())
+						foreach (UhtType referenceType in property.EnumerateReferencedTypes())
 						{
-							if (ReferenceType is UhtScriptStruct PropertyScriptStruct)
+							if (referenceType is UhtScriptStruct propertyScriptStruct)
 							{
-								FindNoExportStructsRecursive(Out, PropertyScriptStruct);
+								FindNoExportStructsRecursive(outScriptStructs, propertyScriptStruct);
 							}
 						}
 					}
@@ -1739,43 +1739,43 @@ namespace EpicGames.UHT.Exporters.CodeGen
 			}
 		}
 
-		private static List<UhtScriptStruct> FindNoExportStructs(UhtStruct Struct)
+		private static List<UhtScriptStruct> FindNoExportStructs(UhtStruct structObj)
 		{
-			List<UhtScriptStruct> Out = new List<UhtScriptStruct>();
-			FindNoExportStructsRecursive(Out, Struct);
-			Out.Reverse();
-			return Out;
+			List<UhtScriptStruct> outScriptStructs = new List<UhtScriptStruct>();
+			FindNoExportStructsRecursive(outScriptStructs, structObj);
+			outScriptStructs.Reverse();
+			return outScriptStructs;
 		}
 
 		private class PropertyMemberContextImpl : IUhtPropertyMemberContext
 		{
-			private readonly UhtCodeGenerator CodeGenerator;
-			private readonly UhtStruct OuterStructInternal;
-			private readonly string OuterStructSourceNameInternal;
-			private readonly string StaticsNameInternal;
+			private readonly UhtCodeGenerator _codeGenerator;
+			private readonly UhtStruct _outerStructInternal;
+			private readonly string _outerStructSourceNameInternal;
+			private readonly string _staticsNameInternal;
 
-			public PropertyMemberContextImpl(UhtCodeGenerator CodeGenerator, UhtStruct OuterStruct, string OuterStructSourceName, string StaticsName)
+			public PropertyMemberContextImpl(UhtCodeGenerator codeGenerator, UhtStruct outerStruct, string outerStructSourceName, string staticsName)
 			{
-				this.CodeGenerator = CodeGenerator;
-				this.OuterStructInternal = OuterStruct;
-				this.StaticsNameInternal = StaticsName;
-				this.OuterStructSourceNameInternal = OuterStructSourceName.Length == 0 ? OuterStruct.SourceName : OuterStructSourceName;
+				this._codeGenerator = codeGenerator;
+				this._outerStructInternal = outerStruct;
+				this._staticsNameInternal = staticsName;
+				this._outerStructSourceNameInternal = outerStructSourceName.Length == 0 ? outerStruct.SourceName : outerStructSourceName;
 			}
 
-			public UhtStruct OuterStruct { get => this.OuterStructInternal; }
-			public string OuterStructSourceName { get => this.OuterStructSourceNameInternal; }
-			public string StaticsName { get => this.StaticsNameInternal; }
-			public string NamePrefix { get => "NewProp_"; }
-			public string MetaDataSuffix { get => "_MetaData"; }
+			public UhtStruct OuterStruct => this._outerStructInternal;
+			public string OuterStructSourceName => this._outerStructSourceNameInternal;
+			public string StaticsName => this._staticsNameInternal;
+			public string NamePrefix => "NewProp_";
+			public string MetaDataSuffix => "_MetaData";
 
-			public string GetSingletonName(UhtObject? Object, bool bRegistered)
+			public string GetSingletonName(UhtObject? obj, bool registered)
 			{
-				return this.CodeGenerator.GetSingletonName(Object, bRegistered);
+				return this._codeGenerator.GetSingletonName(obj, registered);
 			}
 
-			public uint GetTypeHash(UhtObject Object)
+			public uint GetTypeHash(UhtObject obj)
 			{
-				return this.CodeGenerator.ObjectInfos[Object.ObjectTypeIndex].Hash;
+				return this._codeGenerator._objectInfos[obj.ObjectTypeIndex]._hash;
 			}
 		}
 	}
@@ -1789,123 +1789,123 @@ namespace EpicGames.UHT.Exporters.CodeGen
 		/// <summary>
 		/// Append the parameter names for a function
 		/// </summary>
-		/// <param name="Builder">Destination builder</param>
-		/// <param name="Function">Function in question</param>
+		/// <param name="builder">Destination builder</param>
+		/// <param name="function">Function in question</param>
 		/// <returns>Destination builder</returns>
-		public static StringBuilder AppendFunctionThunkParameterNames(this StringBuilder Builder, UhtFunction Function)
+		public static StringBuilder AppendFunctionThunkParameterNames(this StringBuilder builder, UhtFunction function)
 		{
-			bool bFirst = true;
-			foreach (UhtProperty Parameter in Function.ParameterProperties.Span)
+			bool first = true;
+			foreach (UhtProperty parameter in function.ParameterProperties.Span)
 			{
-				if (!bFirst)
+				if (!first)
 				{
-					Builder.Append(',');
+					builder.Append(',');
 				}
-				Builder.AppendFunctionThunkParameterArg(Parameter);
-				bFirst = false;
+				builder.AppendFunctionThunkParameterArg(parameter);
+				first = false;
 			}
-			return Builder;
+			return builder;
 		}
 
 		/// <summary>
 		/// Append the parameter list and count as arguments to a structure constructor
 		/// </summary>
-		/// <param name="Builder">Destination builder</param>
-		/// <param name="Struct">Structure in question</param>
-		/// <param name="StaticsName">Name of the statics section for this structure</param>
-		/// <param name="Tabs">Number of tabs to precede the line</param>
-		/// <param name="Endl">String used to terminate a line</param>
+		/// <param name="builder">Destination builder</param>
+		/// <param name="structObj">Structure in question</param>
+		/// <param name="staticsName">Name of the statics section for this structure</param>
+		/// <param name="tabs">Number of tabs to precede the line</param>
+		/// <param name="endl">String used to terminate a line</param>
 		/// <returns>Destination builder</returns>
-		public static StringBuilder AppendPropertiesParams(this StringBuilder Builder, UhtStruct Struct, string StaticsName, int Tabs, string Endl)
+		public static StringBuilder AppendPropertiesParams(this StringBuilder builder, UhtStruct structObj, string staticsName, int tabs, string endl)
 		{
-			return Builder.AppendPropertiesParamsList(Struct, StaticsName, Tabs, Endl).AppendPropertiesParamsCount(Struct, StaticsName, Tabs, Endl);
+			return builder.AppendPropertiesParamsList(structObj, staticsName, tabs, endl).AppendPropertiesParamsCount(structObj, staticsName, tabs, endl);
 		}
 
 		/// <summary>
 		/// Append the parameter list as an argument to a structure constructor
 		/// </summary>
-		/// <param name="Builder">Destination builder</param>
-		/// <param name="Struct">Structure in question</param>
-		/// <param name="StaticsName">Name of the statics section for this structure</param>
-		/// <param name="Tabs">Number of tabs to precede the line</param>
-		/// <param name="Endl">String used to terminate a line</param>
+		/// <param name="builder">Destination builder</param>
+		/// <param name="structObj">Structure in question</param>
+		/// <param name="staticsName">Name of the statics section for this structure</param>
+		/// <param name="tabs">Number of tabs to precede the line</param>
+		/// <param name="endl">String used to terminate a line</param>
 		/// <returns>Destination builder</returns>
-		public static StringBuilder AppendPropertiesParamsList(this StringBuilder Builder, UhtStruct Struct, string StaticsName, int Tabs, string Endl)
+		public static StringBuilder AppendPropertiesParamsList(this StringBuilder builder, UhtStruct structObj, string staticsName, int tabs, string endl)
 		{
-			if (!Struct.Children.Any(x => x is UhtProperty))
+			if (!structObj.Children.Any(x => x is UhtProperty))
 			{
-				Builder.AppendTabs(Tabs).Append("nullptr,").Append(Endl);
+				builder.AppendTabs(tabs).Append("nullptr,").Append(endl);
 			}
-			else if (Struct.Properties.Any(x => !x.bIsEditorOnlyProperty))
+			else if (structObj.Properties.Any(x => !x.IsEditorOnlyProperty))
 			{
-				Builder.AppendTabs(Tabs).Append(StaticsName).Append("::PropPointers,").Append(Endl);
+				builder.AppendTabs(tabs).Append(staticsName).Append("::PropPointers,").Append(endl);
 			}
 			else
 			{
-				Builder.AppendTabs(Tabs).Append("IF_WITH_EDITORONLY_DATA(").Append(StaticsName).Append("::PropPointers, nullptr),").Append(Endl);
+				builder.AppendTabs(tabs).Append("IF_WITH_EDITORONLY_DATA(").Append(staticsName).Append("::PropPointers, nullptr),").Append(endl);
 			}
-			return Builder;
+			return builder;
 		}
 
 		/// <summary>
 		/// Append the parameter count as an argument to a structure constructor
 		/// </summary>
-		/// <param name="Builder">Destination builder</param>
-		/// <param name="Struct">Structure in question</param>
-		/// <param name="StaticsName">Name of the statics section for this structure</param>
-		/// <param name="Tabs">Number of tabs to precede the line</param>
-		/// <param name="Endl">String used to terminate a line</param>
+		/// <param name="builder">Destination builder</param>
+		/// <param name="structObj">Structure in question</param>
+		/// <param name="staticsName">Name of the statics section for this structure</param>
+		/// <param name="tabs">Number of tabs to precede the line</param>
+		/// <param name="endl">String used to terminate a line</param>
 		/// <returns>Destination builder</returns>
-		public static StringBuilder AppendPropertiesParamsCount(this StringBuilder Builder, UhtStruct Struct, string StaticsName, int Tabs, string Endl)
+		public static StringBuilder AppendPropertiesParamsCount(this StringBuilder builder, UhtStruct structObj, string staticsName, int tabs, string endl)
 		{
-			if (!Struct.Children.Any(x => x is UhtProperty))
+			if (!structObj.Children.Any(x => x is UhtProperty))
 			{
-				Builder.AppendTabs(Tabs).Append("0,").Append(Endl);
+				builder.AppendTabs(tabs).Append("0,").Append(endl);
 			}
-			else if (Struct.Properties.Any(x => !x.bIsEditorOnlyProperty))
+			else if (structObj.Properties.Any(x => !x.IsEditorOnlyProperty))
 			{
-				Builder.AppendTabs(Tabs).Append("UE_ARRAY_COUNT(").Append(StaticsName).Append("::PropPointers),").Append(Endl);
+				builder.AppendTabs(tabs).Append("UE_ARRAY_COUNT(").Append(staticsName).Append("::PropPointers),").Append(endl);
 			}
 			else
 			{
-				Builder.AppendTabs(Tabs).Append("IF_WITH_EDITORONLY_DATA(UE_ARRAY_COUNT(").Append(StaticsName).Append("::PropPointers), 0),").Append(Endl);
+				builder.AppendTabs(tabs).Append("IF_WITH_EDITORONLY_DATA(UE_ARRAY_COUNT(").Append(staticsName).Append("::PropPointers), 0),").Append(endl);
 			}
-			return Builder;
+			return builder;
 		}
 
 		/// <summary>
 		/// Append the given array list and count as arguments to a structure constructor
 		/// </summary>
-		/// <param name="Builder">Destination builder</param>
-		/// <param name="bIsEmpty">If true, the array list is empty</param>
-		/// <param name="bAllEditorOnlyData">If true, the array is all editor only data</param>
-		/// <param name="StaticsName">Name of the statics section</param>
-		/// <param name="ArrayName">The name of the arrray</param>
+		/// <param name="builder">Destination builder</param>
+		/// <param name="isEmpty">If true, the array list is empty</param>
+		/// <param name="allEditorOnlyData">If true, the array is all editor only data</param>
+		/// <param name="staticsName">Name of the statics section</param>
+		/// <param name="arrayName">The name of the arrray</param>
 		/// <returns>Destination builder</returns>
-		public static StringBuilder AppendArray(this StringBuilder Builder, bool bIsEmpty, bool bAllEditorOnlyData, string StaticsName, string ArrayName)
+		public static StringBuilder AppendArray(this StringBuilder builder, bool isEmpty, bool allEditorOnlyData, string staticsName, string arrayName)
 		{
-			if (bIsEmpty)
+			if (isEmpty)
 			{
-				Builder.Append("nullptr, 0");
+				builder.Append("nullptr, 0");
 			}
-			else if (bAllEditorOnlyData)
+			else if (allEditorOnlyData)
 			{
-				Builder
+				builder
 					.Append("IF_WITH_EDITORONLY_DATA(")
-					.Append(StaticsName).Append("::").Append(ArrayName)
+					.Append(staticsName).Append("::").Append(arrayName)
 					.Append(", nullptr), IF_WITH_EDITORONLY_DATA(UE_ARRAY_COUNT(")
-					.Append(StaticsName).Append("::").Append(ArrayName)
+					.Append(staticsName).Append("::").Append(arrayName)
 					.Append("), 0)");
 			}
 			else
 			{
-				Builder
-					.Append(StaticsName).Append("::").Append(ArrayName)
+				builder
+					.Append(staticsName).Append("::").Append(arrayName)
 					.Append(", UE_ARRAY_COUNT(")
-					.Append(StaticsName).Append("::").Append(ArrayName)
+					.Append(staticsName).Append("::").Append(arrayName)
 					.Append(')');
 			}
-			return Builder;
+			return builder;
 		}
 	}
 }
