More coding standard fixes

diff --git a/Engine/Source/Programs/Shared/EpicGames.UHT/Exporters/CodeGen/UhtHeaderCodeGeneratorHFile.cs b/Engine/Source/Programs/Shared/EpicGames.UHT/Exporters/CodeGen/UhtHeaderCodeGeneratorHFile.cs
--- a/Engine/Source/Programs/Shared/EpicGames.UHT/Exporters/CodeGen/UhtHeaderCodeGeneratorHFile.cs
+++ b/Engine/Source/Programs/Shared/EpicGames.UHT/Exporters/CodeGen/UhtHeaderCodeGeneratorHFile.cs
@@ -1,13 +1,13 @@
 // Copyright Epic Games, Inc. All Rights Reserved.
 
-using EpicGames.Core;
-using EpicGames.UHT.Types;
-using EpicGames.UHT.Utils;
 using System;
 using System.Collections.Generic;
 using System.IO;
 using System.Linq;
 using System.Text;
+using EpicGames.Core;
+using EpicGames.UHT.Types;
+using EpicGames.UHT.Utils;
 
 namespace EpicGames.UHT.Exporters.CodeGen
 {
@@ -20,1473 +20,1473 @@ namespace EpicGames.UHT.Exporters.CodeGen
 		/// <summary>
 		/// Construct an instance of this generator object
 		/// </summary>
-		/// <param name="CodeGenerator">The base code generator</param>
-		/// <param name="Package">Package being generated</param>
-		/// <param name="HeaderFile">Header file being generated</param>
-		public UhtHeaderCodeGeneratorHFile(UhtCodeGenerator CodeGenerator, UhtPackage Package, UhtHeaderFile HeaderFile)
-			: base(CodeGenerator, Package, HeaderFile)
+		/// <param name="codeGenerator">The base code generator</param>
+		/// <param name="package">Package being generated</param>
+		/// <param name="headerFile">Header file being generated</param>
+		public UhtHeaderCodeGeneratorHFile(UhtCodeGenerator codeGenerator, UhtPackage package, UhtHeaderFile headerFile)
+			: base(codeGenerator, package, headerFile)
 		{
 		}
 
 		/// <summary>
 		/// For a given UE header file, generated the generated H file
 		/// </summary>
-		/// <param name="Factory">Requesting factory</param>
-		public void Generate(IUhtExportFactory Factory)
+		/// <param name="factory">Requesting factory</param>
+		public void Generate(IUhtExportFactory factory)
 		{
-			ref UhtCodeGenerator.HeaderInfo HeaderInfo = ref this.HeaderInfos[this.HeaderFile.HeaderFileTypeIndex];
-			using (BorrowStringBuilder Borrower = new BorrowStringBuilder(StringBuilderCache.Big))
+			ref UhtCodeGenerator.HeaderInfo headerInfo = ref this.HeaderInfos[this.HeaderFile.HeaderFileTypeIndex];
+			using (BorrowStringBuilder borrower = new BorrowStringBuilder(StringBuilderCache.Big))
 			{
-				StringBuilder Builder = Borrower.StringBuilder;
+				StringBuilder builder = borrower.StringBuilder;
 
-				Builder.Append(HeaderCopyright);
-				Builder.Append("#include \"UObject/ObjectMacros.h\"\r\n");
-				Builder.Append("#include \"UObject/ScriptMacros.h\"\r\n");
-				if (HeaderInfo.bNeedsPushModelHeaders)
+				builder.Append(HeaderCopyright);
+				builder.Append("#include \"UObject/ObjectMacros.h\"\r\n");
+				builder.Append("#include \"UObject/ScriptMacros.h\"\r\n");
+				if (headerInfo._needsPushModelHeaders)
 				{
-					Builder.Append("#include \"Net/Core/PushModel/PushModelMacros.h\"\r\n");
+					builder.Append("#include \"Net/Core/PushModel/PushModelMacros.h\"\r\n");
 				}
-				Builder.Append("\r\n");
-				Builder.Append(DisableDeprecationWarnings).Append("\r\n");
+				builder.Append("\r\n");
+				builder.Append(DisableDeprecationWarnings).Append("\r\n");
 
-				string StrippedName = Path.GetFileNameWithoutExtension(this.HeaderFile.FilePath);
-				string DefineName = $"{this.Package.ShortName.ToString().ToUpper()}_{StrippedName}_generated_h";
+				string strippedName = Path.GetFileNameWithoutExtension(this.HeaderFile.FilePath);
+				string defineName = $"{this.Package.ShortName.ToString().ToUpper()}_{strippedName}_generated_h";
 
 				if (this.HeaderFile.References.ForwardDeclarations.Count > 0)
 				{
-					string[] Sorted = new string[this.HeaderFile.References.ForwardDeclarations.Count];
-					int Index = 0;
-					foreach (string ForwardDeclaration in this.HeaderFile.References.ForwardDeclarations)
+					string[] sorted = new string[this.HeaderFile.References.ForwardDeclarations.Count];
+					int index = 0;
+					foreach (string forwardDeclaration in this.HeaderFile.References.ForwardDeclarations)
 					{
-						Sorted[Index++] = ForwardDeclaration;
+						sorted[index++] = forwardDeclaration;
 					}
-					Array.Sort(Sorted, StringComparerUE.OrdinalIgnoreCase);
-					foreach (string ForwardDeclaration in Sorted)
+					Array.Sort(sorted, StringComparerUE.OrdinalIgnoreCase);
+					foreach (string forwardDeclaration in sorted)
 					{
-						Builder.Append(ForwardDeclaration).Append("\r\n");
+						builder.Append(forwardDeclaration).Append("\r\n");
 					}
 				}
 
-				Builder.Append("#ifdef ").Append(DefineName).Append("\r\n");
-				Builder.Append("#error \"").Append(StrippedName).Append(".generated.h already included, missing '#pragma once' in ").Append(StrippedName).Append(".h\"\r\n");
-				Builder.Append("#endif\r\n");
-				Builder.Append("#define ").Append(DefineName).Append("\r\n");
-				Builder.Append("\r\n");
+				builder.Append("#ifdef ").Append(defineName).Append("\r\n");
+				builder.Append("#error \"").Append(strippedName).Append(".generated.h already included, missing '#pragma once' in ").Append(strippedName).Append(".h\"\r\n");
+				builder.Append("#endif\r\n");
+				builder.Append("#define ").Append(defineName).Append("\r\n");
+				builder.Append("\r\n");
 
-				foreach (UhtField Field in this.HeaderFile.References.ExportTypes)
+				foreach (UhtField field in this.HeaderFile.References.ExportTypes)
 				{
-					if (Field is UhtEnum Enum)
+					if (field is UhtEnum enumObj)
 					{
 					}
-					else if (Field is UhtScriptStruct ScriptStruct)
+					else if (field is UhtScriptStruct scriptStruct)
 					{
-						if (ScriptStruct.ScriptStructFlags.HasAnyFlags(EStructFlags.Native))
+						if (scriptStruct.ScriptStructFlags.HasAnyFlags(EStructFlags.Native))
 						{
-							AppendScriptStruct(Builder, ScriptStruct);
+							AppendScriptStruct(builder, scriptStruct);
 						}
 					}
-					else if (Field is UhtFunction Function)
+					else if (field is UhtFunction function)
 					{
-						AppendDelegate(Builder, Function);
+						AppendDelegate(builder, function);
 					}
-					else if (Field is UhtClass Class)
+					else if (field is UhtClass classObj)
 					{
-						if (!Class.ClassFlags.HasAnyFlags(EClassFlags.Intrinsic))
+						if (!classObj.ClassFlags.HasAnyFlags(EClassFlags.Intrinsic))
 						{
-							AppendClass(Builder, Class);
+							AppendClass(builder, classObj);
 						}
 					}
 				}
 
-				Builder.Append("#undef CURRENT_FILE_ID\r\n");
-				Builder.Append("#define CURRENT_FILE_ID ").Append(HeaderInfo.FileId).Append("\r\n\r\n\r\n");
+				builder.Append("#undef CURRENT_FILE_ID\r\n");
+				builder.Append("#define CURRENT_FILE_ID ").Append(headerInfo._fileId).Append("\r\n\r\n\r\n");
 
-				foreach (UhtField Field in this.HeaderFile.References.ExportTypes)
+				foreach (UhtField field in this.HeaderFile.References.ExportTypes)
 				{
-					if (Field is UhtEnum Enum)
+					if (field is UhtEnum enumObject)
 					{
-						AppendEnum(Builder, Enum);
+						AppendEnum(builder, enumObject);
 					}
 				}
 
-				Builder.Append(EnableDeprecationWarnings).Append("\r\n");
+				builder.Append(EnableDeprecationWarnings).Append("\r\n");
 
 				if (this.SaveExportedHeaders)
 				{
-					string HeaderFilePath = Factory.MakePath(this.HeaderFile, ".generated.h");
-					Factory.CommitOutput(HeaderFilePath, Builder);
+					string headerFilePath = factory.MakePath(this.HeaderFile, ".generated.h");
+					factory.CommitOutput(headerFilePath, builder);
 				}
 			}
 		}
 
-		private StringBuilder AppendScriptStruct(StringBuilder Builder, UhtScriptStruct ScriptStruct)
+		private StringBuilder AppendScriptStruct(StringBuilder builder, UhtScriptStruct scriptStruct)
 		{
-			if (ScriptStruct.RigVMStructInfo != null)
+			if (scriptStruct.RigVMStructInfo != null)
 			{
-				Builder.Append("\r\n");
-				foreach (UhtRigVMMethodInfo MethodInfo in ScriptStruct.RigVMStructInfo.Methods)
+				builder.Append("\r\n");
+				foreach (UhtRigVMMethodInfo methodInfo in scriptStruct.RigVMStructInfo.Methods)
 				{
-					Builder.Append("#define ").Append(ScriptStruct.SourceName).Append('_').Append(MethodInfo.Name).Append("() \\\r\n");
-					Builder.Append('\t').Append(MethodInfo.ReturnType).Append(' ').Append(ScriptStruct.SourceName).Append("::Static").Append(MethodInfo.Name).Append("( \\\r\n");
-					Builder.Append("\t\t").Append(RigVMExecuteContextPublicDeclaration);
-					Builder.AppendParameterDecls(ScriptStruct.RigVMStructInfo.Members, true, ", \\\r\n\t\t", true, false);
-					Builder.AppendParameterDecls(MethodInfo.Parameters, true, ", \\\r\n\t\t", false, false);
-					Builder.Append(" \\\r\n");
-					Builder.Append("\t)\r\n");
+					builder.Append("#define ").Append(scriptStruct.SourceName).Append('_').Append(methodInfo.Name).Append("() \\\r\n");
+					builder.Append('\t').Append(methodInfo.ReturnType).Append(' ').Append(scriptStruct.SourceName).Append("::Static").Append(methodInfo.Name).Append("( \\\r\n");
+					builder.Append("\t\t").Append(RigVMExecuteContextPublicDeclaration);
+					builder.AppendParameterDecls(scriptStruct.RigVMStructInfo.Members, true, ", \\\r\n\t\t", true, false);
+					builder.AppendParameterDecls(methodInfo.Parameters, true, ", \\\r\n\t\t", false, false);
+					builder.Append(" \\\r\n");
+					builder.Append("\t)\r\n");
 				}
-				Builder.Append("\r\n");
+				builder.Append("\r\n");
 			}
 
-			if (ScriptStruct.ScriptStructFlags.HasAnyFlags(EStructFlags.Native))
+			if (scriptStruct.ScriptStructFlags.HasAnyFlags(EStructFlags.Native))
 			{
-				using (UhtMacroCreator Macro = new UhtMacroCreator(Builder, this, ScriptStruct, GeneratedBodyMacroSuffix))
+				using (UhtMacroCreator macro = new UhtMacroCreator(builder, this, scriptStruct, GeneratedBodyMacroSuffix))
 				{
-					Builder.Append("\tfriend struct Z_Construct_UScriptStruct_").Append(ScriptStruct.SourceName).Append("_Statics; \\\r\n");
-					Builder.Append('\t');
-					if (!ScriptStruct.ScriptStructFlags.HasAnyFlags(EStructFlags.RequiredAPI))
+					builder.Append("\tfriend struct Z_Construct_UScriptStruct_").Append(scriptStruct.SourceName).Append("_Statics; \\\r\n");
+					builder.Append('\t');
+					if (!scriptStruct.ScriptStructFlags.HasAnyFlags(EStructFlags.RequiredAPI))
 					{
-						Builder.Append(this.PackageApi);
+						builder.Append(this.PackageApi);
 					}
-					Builder.Append("static class UScriptStruct* StaticStruct(); \\\r\n");
+					builder.Append("static class UScriptStruct* StaticStruct(); \\\r\n");
 
 					// if we have RigVM methods on this struct we need to 
 					// declare the static method as well as the stub method
-					if (ScriptStruct.RigVMStructInfo != null)
+					if (scriptStruct.RigVMStructInfo != null)
 					{
-						foreach (UhtRigVMMethodInfo MethodInfo in ScriptStruct.RigVMStructInfo.Methods)
+						foreach (UhtRigVMMethodInfo methodInfo in scriptStruct.RigVMStructInfo.Methods)
 						{
-							Builder.Append("\tstatic ").Append(MethodInfo.ReturnType).Append(" Static").Append(MethodInfo.Name).Append("( \\\r\n");
-							Builder.Append("\t\t").Append(RigVMExecuteContextPublicDeclaration);
-							Builder.AppendParameterDecls(ScriptStruct.RigVMStructInfo.Members, true, ", \\\r\n\t\t", true, false);
-							Builder.AppendParameterDecls(MethodInfo.Parameters, true, ", \\\r\n\t\t", false, false);
-							Builder.Append(" \\\r\n");
-							Builder.Append("\t); \\\r\n");
-
-							Builder.Append("\tFORCEINLINE_DEBUGGABLE static ").Append(MethodInfo.ReturnType).Append(" RigVM").Append(MethodInfo.Name).Append("( \\\r\n");
-							Builder.Append("\t\t").Append(RigVMExecuteContextDeclaration).Append(", \\\r\n");
-							Builder.Append("\t\tFRigVMMemoryHandleArray RigVMMemoryHandles \\\r\n");
-							Builder.Append("\t) \\\r\n");
-							Builder.Append("\t{ \\\r\n");
+							builder.Append("\tstatic ").Append(methodInfo.ReturnType).Append(" Static").Append(methodInfo.Name).Append("( \\\r\n");
+							builder.Append("\t\t").Append(RigVMExecuteContextPublicDeclaration);
+							builder.AppendParameterDecls(scriptStruct.RigVMStructInfo.Members, true, ", \\\r\n\t\t", true, false);
+							builder.AppendParameterDecls(methodInfo.Parameters, true, ", \\\r\n\t\t", false, false);
+							builder.Append(" \\\r\n");
+							builder.Append("\t); \\\r\n");
+
+							builder.Append("\tFORCEINLINE_DEBUGGABLE static ").Append(methodInfo.ReturnType).Append(" RigVM").Append(methodInfo.Name).Append("( \\\r\n");
+							builder.Append("\t\t").Append(RigVMExecuteContextDeclaration).Append(", \\\r\n");
+							builder.Append("\t\tFRigVMMemoryHandleArray RigVMMemoryHandles \\\r\n");
+							builder.Append("\t) \\\r\n");
+							builder.Append("\t{ \\\r\n");
 
 							// implement inline stub method body
-							if (MethodInfo.Parameters.Count > 0)
+							if (methodInfo.Parameters.Count > 0)
 							{
-								for (int ParameterIndex = 0; ParameterIndex < MethodInfo.Parameters.Count; ParameterIndex++)
+								for (int parameterIndex = 0; parameterIndex < methodInfo.Parameters.Count; parameterIndex++)
 								{
-									UhtRigVMParameter Parameter = MethodInfo.Parameters[ParameterIndex];
-									Builder.Append("\t\t").Append(Parameter.Declaration()).Append(" = *(").Append(Parameter.TypeNoRef())
-										.Append("*)RigVMExecuteContext.OpaqueArguments[").Append(ParameterIndex).Append("]; \\\r\n");
+									UhtRigVMParameter parameter = methodInfo.Parameters[parameterIndex];
+									builder.Append("\t\t").Append(parameter.Declaration()).Append(" = *(").Append(parameter.TypeNoRef())
+										.Append("*)RigVMExecuteContext.OpaqueArguments[").Append(parameterIndex).Append("]; \\\r\n");
 								}
-								Builder.Append("\t\t \\\r\n");
+								builder.Append("\t\t \\\r\n");
 							}
 
-							if (ScriptStruct.RigVMStructInfo.Members.Count > 0)
+							if (scriptStruct.RigVMStructInfo.Members.Count > 0)
 							{
-								int OperandIndex = 0;
-								foreach (UhtRigVMParameter Parameter in ScriptStruct.RigVMStructInfo.Members)
+								int operandIndex = 0;
+								foreach (UhtRigVMParameter parameter in scriptStruct.RigVMStructInfo.Members)
 								{
-									string ParamTypeOriginal = Parameter.TypeOriginal(true);
-									string ParamNameOriginal = Parameter.NameOriginal(false);
-									string AdditionalParameters = string.Empty;
-									if (!Parameter.bInput && !Parameter.bOutput && !Parameter.bSingleton)
+									string paramTypeOriginal = parameter.TypeOriginal(true);
+									string paramNameOriginal = parameter.NameOriginal(false);
+									string additionalParameters = String.Empty;
+									if (!parameter.Input && !parameter.Output && !parameter.Singleton)
 									{
-										AdditionalParameters = ", RigVMExecuteContext.GetSlice().GetIndex()";
+										additionalParameters = ", RigVMExecuteContext.GetSlice().GetIndex()";
 									}
 
-									if (Parameter.bIsArray)
+									if (parameter.IsArray)
 									{
-										string ExtendedType = Parameter.ExtendedType();
-										Builder
+										string extendedType = parameter.ExtendedType();
+										builder
 											.Append("\t\tTArray")
-											.Append(ExtendedType)
+											.Append(extendedType)
 											.Append("& ")
-											.Append(ParamNameOriginal)
+											.Append(paramNameOriginal)
 											.Append(" = *(TArray")
-											.Append(ExtendedType)
+											.Append(extendedType)
 											.Append("*)RigVMMemoryHandles[")
-											.Append(OperandIndex)
+											.Append(operandIndex)
 											.Append("].GetData(false")
-											.Append(AdditionalParameters)
+											.Append(additionalParameters)
 											.Append("); \\\r\n");
-										OperandIndex++;
+										operandIndex++;
 									}
 									else
 									{
-										string VariableType = Parameter.TypeVariableRef(true);
-										string ExtractedType = Parameter.TypeOriginal();
+										string variableType = parameter.TypeVariableRef(true);
+										string extractedType = parameter.TypeOriginal();
 
-										Builder
+										builder
 											.Append("\t\t")
-											.Append(VariableType)
+											.Append(variableType)
 											.Append(' ')
-											.Append(ParamNameOriginal)
+											.Append(paramNameOriginal)
 											.Append(" = *(")
-											.Append(ExtractedType)
+											.Append(extractedType)
 											.Append("*)RigVMMemoryHandles[")
-											.Append(OperandIndex)
+											.Append(operandIndex)
 											.Append("].GetData(false")
-											.Append(AdditionalParameters)
+											.Append(additionalParameters)
 											.Append("); \\\r\n");
-										OperandIndex++;
+										operandIndex++;
 									}
 								}
-								Builder.Append("\t\t \\\r\n");
+								builder.Append("\t\t \\\r\n");
 							}
 
-							Builder.Append("\t\t").Append(MethodInfo.ReturnPrefix()).Append("Static").Append(MethodInfo.Name).Append("( \\\r\n");
-							Builder.Append("\t\t\tRigVMExecuteContext.PublicData");
-							Builder.AppendParameterNames(ScriptStruct.RigVMStructInfo.Members, true, ", \\\r\n\t\t\t", false);
-							Builder.AppendParameterNames(MethodInfo.Parameters, true, ", \\\r\n\t\t\t");
-							Builder.Append(" \\\r\n");
-							Builder.Append("\t\t); \\\r\n");
-							Builder.Append("\t} \\\r\n");
+							builder.Append("\t\t").Append(methodInfo.ReturnPrefix()).Append("Static").Append(methodInfo.Name).Append("( \\\r\n");
+							builder.Append("\t\t\tRigVMExecuteContext.PublicData");
+							builder.AppendParameterNames(scriptStruct.RigVMStructInfo.Members, true, ", \\\r\n\t\t\t", false);
+							builder.AppendParameterNames(methodInfo.Parameters, true, ", \\\r\n\t\t\t");
+							builder.Append(" \\\r\n");
+							builder.Append("\t\t); \\\r\n");
+							builder.Append("\t} \\\r\n");
 						}
 					}
 
-					if (ScriptStruct.SuperScriptStruct != null)
+					if (scriptStruct.SuperScriptStruct != null)
 					{
-						Builder.Append("\ttypedef ").Append(ScriptStruct.SuperScriptStruct.SourceName).Append(" Super;\\\\r\n");
+						builder.Append("\ttypedef ").Append(scriptStruct.SuperScriptStruct.SourceName).Append(" Super;\\\\r\n");
 					}
 				}
 
 				// Forward declare the StaticStruct specialization in the header
-				Builder.Append("template<> ").Append(this.PackageApi).Append("UScriptStruct* StaticStruct<struct ").Append(ScriptStruct.SourceName).Append(">();\r\n");
-				Builder.Append("\r\n");
+				builder.Append("template<> ").Append(this.PackageApi).Append("UScriptStruct* StaticStruct<struct ").Append(scriptStruct.SourceName).Append(">();\r\n");
+				builder.Append("\r\n");
 			}
-			return Builder;
+			return builder;
 		}
 
-		private StringBuilder AppendEnum(StringBuilder Builder, UhtEnum Enum)
+		private StringBuilder AppendEnum(StringBuilder builder, UhtEnum enumObj)
 		{
 			// Export FOREACH macro
-			Builder.Append("#define FOREACH_ENUM_").Append(Enum.EngineName.ToUpper()).Append("(op) ");
-			bool bHasExistingMax = EnumHasExistingMax(Enum);
-			long MaxEnumValue = bHasExistingMax ? GetMaxEnumValue(Enum) : 0;
-			foreach (UhtEnumValue EnumValue in Enum.EnumValues)
+			builder.Append("#define FOREACH_ENUM_").Append(enumObj.EngineName.ToUpper()).Append("(op) ");
+			bool hasExistingMax = EnumHasExistingMax(enumObj);
+			long maxEnumValue = hasExistingMax ? GetMaxEnumValue(enumObj) : 0;
+			foreach (UhtEnumValue enumValue in enumObj.EnumValues)
 			{
-				if (!bHasExistingMax || EnumValue.Value != MaxEnumValue)
+				if (!hasExistingMax || enumValue.Value != maxEnumValue)
 				{
-					Builder.Append("\\\r\n\top(").Append(EnumValue.Name).Append(") ");
+					builder.Append("\\\r\n\top(").Append(enumValue.Name).Append(") ");
 				}
 			}
-			Builder.Append("\r\n");
+			builder.Append("\r\n");
 
 			// Forward declare the StaticEnum<> specialization for enum classes
-			if (Enum.CppForm == UhtEnumCppForm.EnumClass)
+			if (enumObj.CppForm == UhtEnumCppForm.EnumClass)
 			{
-				Builder.Append("\r\n");
-				Builder.Append("enum class ").Append(Enum.CppType);
-				if (Enum.UnderlyingType != UhtEnumUnderlyingType.Unspecified)
+				builder.Append("\r\n");
+				builder.Append("enum class ").Append(enumObj.CppType);
+				if (enumObj.UnderlyingType != UhtEnumUnderlyingType.Unspecified)
 				{
-					Builder.Append(" : ").Append(Enum.UnderlyingType.ToString());
+					builder.Append(" : ").Append(enumObj.UnderlyingType.ToString());
 				}
-				Builder.Append(";\r\n");
-				Builder.Append("template<> ").Append(this.PackageApi).Append("UEnum* StaticEnum<").Append(Enum.CppType).Append(">();\r\n");
-				Builder.Append("\r\n");
+				builder.Append(";\r\n");
+				builder.Append("template<> ").Append(this.PackageApi).Append("UEnum* StaticEnum<").Append(enumObj.CppType).Append(">();\r\n");
+				builder.Append("\r\n");
 			}
 
-			return Builder;
+			return builder;
 		}
 
-		private StringBuilder AppendDelegate(StringBuilder Builder, UhtFunction Function)
+		private StringBuilder AppendDelegate(StringBuilder builder, UhtFunction function)
 		{
-			using (UhtMacroCreator Macro = new UhtMacroCreator(Builder, this, Function, DelegateMacroSuffix))
-			{
-				int Tabs = 0;
-				string StrippedFunctionName = Function.StrippedFunctionName;
-				string ExportFunctionName = GetDelegateFunctionExportName(Function);
-				string ExtraParameter = GetDelegateFunctionExtraParameter(Function);
-
-				AppendEventParameter(Builder, Function, StrippedFunctionName, UhtPropertyTextType.EventParameterMember, true, Tabs, " \\\r\n");
-				Builder.Append("static ");
-				AppendNativeFunctionHeader(Builder, Function, UhtPropertyTextType.EventFunctionArgOrRetVal, true, ExportFunctionName, ExtraParameter, UhtFunctionExportFlags.Inline, " \\\r\n");
-				AppendEventFunctionPrologue(Builder, Function, StrippedFunctionName, Tabs, " \\\r\n");
-				Builder
+			using (UhtMacroCreator macro = new UhtMacroCreator(builder, this, function, DelegateMacroSuffix))
+			{
+				int tabs = 0;
+				string strippedFunctionName = function.StrippedFunctionName;
+				string exportFunctionName = GetDelegateFunctionExportName(function);
+				string extraParameter = GetDelegateFunctionExtraParameter(function);
+
+				AppendEventParameter(builder, function, strippedFunctionName, UhtPropertyTextType.EventParameterMember, true, tabs, " \\\r\n");
+				builder.Append("static ");
+				AppendNativeFunctionHeader(builder, function, UhtPropertyTextType.EventFunctionArgOrRetVal, true, exportFunctionName, extraParameter, UhtFunctionExportFlags.Inline, " \\\r\n");
+				AppendEventFunctionPrologue(builder, function, strippedFunctionName, tabs, " \\\r\n");
+				builder
 					.Append('\t')
-					.Append(StrippedFunctionName)
+					.Append(strippedFunctionName)
 					.Append('.')
-					.Append(Function.FunctionFlags.HasAnyFlags(EFunctionFlags.MulticastDelegate) ? "ProcessMulticastDelegate" : "ProcessDelegate")
+					.Append(function.FunctionFlags.HasAnyFlags(EFunctionFlags.MulticastDelegate) ? "ProcessMulticastDelegate" : "ProcessDelegate")
 					.Append("<UObject>(")
-					.Append(Function.Children.Count > 0 ? "&Parms" : "NULL")
+					.Append(function.Children.Count > 0 ? "&Parms" : "NULL")
 					.Append("); \\\r\n");
-				AppendEventFunctionEpilogue(Builder, Function, Tabs, " \\\r\n");
+				AppendEventFunctionEpilogue(builder, function, tabs, " \\\r\n");
 			}
-			return Builder;
+			return builder;
 		}
 
-		private StringBuilder AppendClass(StringBuilder Builder, UhtClass Class)
+		private StringBuilder AppendClass(StringBuilder builder, UhtClass classObj)
 		{
-			string Api = Class.ClassFlags.HasAnyFlags(EClassFlags.MinimalAPI) ? this.PackageApi : "NO_API ";
-			UhtClass? NativeInterface = this.ObjectInfos[Class.ObjectTypeIndex].NativeInterface;
+			string api = classObj.ClassFlags.HasAnyFlags(EClassFlags.MinimalAPI) ? this.PackageApi : "NO_API ";
+			UhtClass? nativeInterface = this.ObjectInfos[classObj.ObjectTypeIndex]._nativeInterface;
 
 			// Write the spare declarations
-			AppendSparseDeclarations(Builder, Class);
+			AppendSparseDeclarations(builder, classObj);
 
 			// Collect the functions in reversed order
-			List<UhtFunction> ReversedFunctions = new List<UhtFunction>(Class.Functions);
-			ReversedFunctions.Reverse();
+			List<UhtFunction> reversedFunctions = new List<UhtFunction>(classObj.Functions);
+			reversedFunctions.Reverse();
 
 			// Check to see if we have any RPC functions for the editor
-			bool bHasEditorRpc = ReversedFunctions.Any(x => IsRpcFunction(x, true));
+			bool hasEditorRpc = reversedFunctions.Any(x => IsRpcFunction(x, true));
 
 			// Output the RPC methods
-			AppendRpcFunctions(Builder, Class, ReversedFunctions, false);
-			if (bHasEditorRpc)
+			AppendRpcFunctions(builder, classObj, reversedFunctions, false);
+			if (hasEditorRpc)
 			{
-				AppendRpcFunctions(Builder, Class, ReversedFunctions, true);
+				AppendRpcFunctions(builder, classObj, reversedFunctions, true);
 			}
 
 			// Output property accessors
-			AppendPropertyAccessors(Builder, Class);
+			AppendPropertyAccessors(builder, classObj);
 
 			// Collect the callback function and sort by name to make the order stable
-			List<UhtFunction> CallbackFunctions = new List<UhtFunction>(Class.Functions.Where(x => x.FunctionFlags.HasAnyFlags(EFunctionFlags.Event) && x.SuperFunction == null));
-			CallbackFunctions.Sort((x, y) => StringComparerUE.OrdinalIgnoreCase.Compare(x.EngineName, y.EngineName));
+			List<UhtFunction> callbackFunctions = new List<UhtFunction>(classObj.Functions.Where(x => x.FunctionFlags.HasAnyFlags(EFunctionFlags.Event) && x.SuperFunction == null));
+			callbackFunctions.Sort((x, y) => StringComparerUE.OrdinalIgnoreCase.Compare(x.EngineName, y.EngineName));
 
 			// Generate the callback parameter structures
-			AppendCallbackParametersDecls(Builder, Class, CallbackFunctions);
+			AppendCallbackParametersDecls(builder, classObj, callbackFunctions);
 
 			// Generate the RPC wrappers for the callbacks
-			AppendCallbackRpcWrapperDecls(Builder, Class, CallbackFunctions);
+			AppendCallbackRpcWrapperDecls(builder, classObj, callbackFunctions);
 
 			// Only write out adapters if the user has provided one or the other of the Serialize overloads
-			if (Class.SerializerArchiveType != UhtSerializerArchiveType.None && Class.SerializerArchiveType != UhtSerializerArchiveType.All)
+			if (classObj.SerializerArchiveType != UhtSerializerArchiveType.None && classObj.SerializerArchiveType != UhtSerializerArchiveType.All)
 			{
-				AppendSerializer(Builder, Class, Api, UhtSerializerArchiveType.Archive, "DECLARE_FARCHIVE_SERIALIZER");
-				AppendSerializer(Builder, Class, Api, UhtSerializerArchiveType.StructuredArchiveRecord, "DECLARE_FSTRUCTUREDARCHIVE_SERIALIZER");
+				AppendSerializer(builder, classObj, api, UhtSerializerArchiveType.Archive, "DECLARE_FARCHIVE_SERIALIZER");
+				AppendSerializer(builder, classObj, api, UhtSerializerArchiveType.StructuredArchiveRecord, "DECLARE_FSTRUCTUREDARCHIVE_SERIALIZER");
 			}
 
-			if (Class.ClassFlags.HasAnyFlags(EClassFlags.Interface))
+			if (classObj.ClassFlags.HasAnyFlags(EClassFlags.Interface))
 			{
-				AppendStandardConstructors(Builder, Class, Api);
-				AppendEnhancedConstructors(Builder, Class, Api);
+				AppendStandardConstructors(builder, classObj, api);
+				AppendEnhancedConstructors(builder, classObj, api);
 
-				using (UhtMacroCreator Macro = new UhtMacroCreator(Builder, this, Class, GeneratedUInterfaceBodyMacroSuffix))
+				using (UhtMacroCreator macro = new UhtMacroCreator(builder, this, classObj, GeneratedUInterfaceBodyMacroSuffix))
 				{
-					AppendCommonGeneratedBody(Builder, Class, Api);
+					AppendCommonGeneratedBody(builder, classObj, api);
 				}
 
-				using (UhtMacroCreator Macro = new UhtMacroCreator(Builder, this, Class, GeneratedBodyLegacyMacroSuffix))
+				using (UhtMacroCreator macro = new UhtMacroCreator(builder, this, classObj, GeneratedBodyLegacyMacroSuffix))
 				{
-					Builder.Append('\t');
-					AppendGeneratedMacroDeprecationWarning(Builder, "GENERATED_UINTERFACE_BODY");
-					Builder.Append('\t').Append(DisableDeprecationWarnings).Append(" \\\r\n");
-					Builder.Append('\t').AppendMacroName(this, Class, GeneratedUInterfaceBodyMacroSuffix).Append(" \\\r\n");
-					Builder.Append('\t').AppendMacroName(this, Class, StandardConstructorsMacroSuffix).Append(" \\\r\n");
-					Builder.Append('\t').Append(EnableDeprecationWarnings).Append(" \\\r\n");
+					builder.Append('\t');
+					AppendGeneratedMacroDeprecationWarning(builder, "GENERATED_UINTERFACE_BODY");
+					builder.Append('\t').Append(DisableDeprecationWarnings).Append(" \\\r\n");
+					builder.Append('\t').AppendMacroName(this, classObj, GeneratedUInterfaceBodyMacroSuffix).Append(" \\\r\n");
+					builder.Append('\t').AppendMacroName(this, classObj, StandardConstructorsMacroSuffix).Append(" \\\r\n");
+					builder.Append('\t').Append(EnableDeprecationWarnings).Append(" \\\r\n");
 				}
 
-				using (UhtMacroCreator Macro = new UhtMacroCreator(Builder, this, Class, GeneratedBodyMacroSuffix))
+				using (UhtMacroCreator macro = new UhtMacroCreator(builder, this, classObj, GeneratedBodyMacroSuffix))
 				{
-					Builder.Append('\t').Append(DisableDeprecationWarnings).Append(" \\\r\n");
-					Builder.Append('\t').AppendMacroName(this, Class, GeneratedUInterfaceBodyMacroSuffix).Append(" \\\r\n");
-					Builder.Append('\t').AppendMacroName(this, Class, EnchancedConstructorsMacroSuffix).Append(" \\\r\n");
-					AppendAccessSpecifier(Builder, Class);
-					Builder.Append(" \\\r\n");
-					Builder.Append('\t').Append(EnableDeprecationWarnings).Append(" \\\r\n");
+					builder.Append('\t').Append(DisableDeprecationWarnings).Append(" \\\r\n");
+					builder.Append('\t').AppendMacroName(this, classObj, GeneratedUInterfaceBodyMacroSuffix).Append(" \\\r\n");
+					builder.Append('\t').AppendMacroName(this, classObj, EnchancedConstructorsMacroSuffix).Append(" \\\r\n");
+					AppendAccessSpecifier(builder, classObj);
+					builder.Append(" \\\r\n");
+					builder.Append('\t').Append(EnableDeprecationWarnings).Append(" \\\r\n");
 				}
 
-				using (UhtMacroCreator Macro = new UhtMacroCreator(Builder, this, Class, InClassIInterfaceNoPureDeclsMacroSuffix))
+				using (UhtMacroCreator macro = new UhtMacroCreator(builder, this, classObj, InClassIInterfaceNoPureDeclsMacroSuffix))
 				{
-					AppendInClassIInterface(Builder, Class, CallbackFunctions, Api);
+					AppendInClassIInterface(builder, classObj, callbackFunctions, api);
 				}
 
-				using (UhtMacroCreator Macro = new UhtMacroCreator(Builder, this, Class, InClassIInterfaceMacroSuffix))
+				using (UhtMacroCreator macro = new UhtMacroCreator(builder, this, classObj, InClassIInterfaceMacroSuffix))
 				{
-					AppendInClassIInterface(Builder, Class, CallbackFunctions, Api);
+					AppendInClassIInterface(builder, classObj, callbackFunctions, api);
 				}
 			}
 			else
 			{
-				using (UhtMacroCreator Macro = new UhtMacroCreator(Builder, this, Class, InClassNoPureDeclsMacroSuffix))
+				using (UhtMacroCreator macro = new UhtMacroCreator(builder, this, classObj, InClassNoPureDeclsMacroSuffix))
 				{
-					AppendClassGeneratedBody(Builder, Class, Api);
+					AppendClassGeneratedBody(builder, classObj, api);
 				}
 
-				using (UhtMacroCreator Macro = new UhtMacroCreator(Builder, this, Class, InClassMacroSuffix))
+				using (UhtMacroCreator macro = new UhtMacroCreator(builder, this, classObj, InClassMacroSuffix))
 				{
-					AppendClassGeneratedBody(Builder, Class, Api);
+					AppendClassGeneratedBody(builder, classObj, api);
 				}
 
-				AppendStandardConstructors(Builder, Class, Api);
-				AppendEnhancedConstructors(Builder, Class, Api);
-				AppendFieldNotify(Builder, Class);
+				AppendStandardConstructors(builder, classObj, api);
+				AppendEnhancedConstructors(builder, classObj, api);
+				AppendFieldNotify(builder, classObj);
 			}
 
-			using (UhtMacroCreator Macro = new UhtMacroCreator(Builder, this, Class.PrologLineNumber, PrologMacroSuffix))
+			using (UhtMacroCreator macro = new UhtMacroCreator(builder, this, classObj.PrologLineNumber, PrologMacroSuffix))
 			{
-				if (CallbackFunctions.Count > 0)
+				if (callbackFunctions.Count > 0)
 				{
-					Builder.Append('\t').AppendMacroName(this, Class, EventParamsMacroSuffix).Append(" \\\r\n");
+					builder.Append('\t').AppendMacroName(this, classObj, EventParamsMacroSuffix).Append(" \\\r\n");
 				}
 			}
 
-			bool bHasCallbacks = CallbackFunctions.Count > 0;
-			if (Class.ClassFlags.HasAnyFlags(EClassFlags.Interface))
+			bool hasCallbacks = callbackFunctions.Count > 0;
+			if (classObj.ClassFlags.HasAnyFlags(EClassFlags.Interface))
 			{
-				if (NativeInterface != null)
+				if (nativeInterface != null)
 				{
-					AppendGeneratedBodyMacroBlock(Builder, Class, NativeInterface, true, bHasEditorRpc, bHasCallbacks, null);
-					AppendGeneratedBodyMacroBlock(Builder, Class, NativeInterface, false, bHasEditorRpc, bHasCallbacks, null);
+					AppendGeneratedBodyMacroBlock(builder, classObj, nativeInterface, true, hasEditorRpc, hasCallbacks, null);
+					AppendGeneratedBodyMacroBlock(builder, classObj, nativeInterface, false, hasEditorRpc, hasCallbacks, null);
 				}
 			}
 			else
 			{
-				AppendGeneratedBodyMacroBlock(Builder, Class, Class, true, bHasEditorRpc, bHasCallbacks, "GENERATED_UCLASS_BODY");
-				AppendGeneratedBodyMacroBlock(Builder, Class, Class, false, bHasEditorRpc, bHasCallbacks, null);
+				AppendGeneratedBodyMacroBlock(builder, classObj, classObj, true, hasEditorRpc, hasCallbacks, "GENERATED_UCLASS_BODY");
+				AppendGeneratedBodyMacroBlock(builder, classObj, classObj, false, hasEditorRpc, hasCallbacks, null);
 			}
 
 			// Forward declare the StaticClass specialization in the header
-			Builder.Append("template<> ").Append(this.PackageApi).Append("UClass* StaticClass<class ").Append(Class.SourceName).Append(">();\r\n");
-			Builder.Append("\r\n");
-			return Builder;
+			builder.Append("template<> ").Append(this.PackageApi).Append("UClass* StaticClass<class ").Append(classObj.SourceName).Append(">();\r\n");
+			builder.Append("\r\n");
+			return builder;
 		}
 
-		private StringBuilder AppendFieldNotify(StringBuilder Builder, UhtClass Class)
+		private StringBuilder AppendFieldNotify(StringBuilder builder, UhtClass classObj)
 		{
-			if (!NeedFieldNotifyCodeGen(Class))
+			if (!NeedFieldNotifyCodeGen(classObj))
 			{
-				return Builder;
+				return builder;
 			}
 
 			// Scan the children to see what we have
-			bool bHasProperties;
-			bool bHasFunctions;
-			bool bHasEditorFields;
-			bool bAllEditorFields;
-			GetFieldNotifyStats(Class, out bHasProperties, out bHasFunctions, out bHasEditorFields, out bAllEditorFields);
+			bool hasProperties;
+			bool hasFunctions;
+			bool hasEditorFields;
+			bool allEditorFields;
+			GetFieldNotifyStats(classObj, out hasProperties, out hasFunctions, out hasEditorFields, out allEditorFields);
 
 			// If we only have editor fields or no editor fields, then we only emit one block
-			if (bHasEditorFields)
+			if (hasEditorFields)
 			{
-				Builder.Append("#if WITH_EDITORONLY_DATA\r\n");
-				using (UhtMacroCreator Macro = new UhtMacroCreator(Builder, this, Class, FieldNotifyMacroSuffix))
+				builder.Append("#if WITH_EDITORONLY_DATA\r\n");
+				using (UhtMacroCreator macro = new UhtMacroCreator(builder, this, classObj, FieldNotifyMacroSuffix))
 				{
-					AppendFieldNotify(Builder, Class, bHasProperties, bHasFunctions, bHasEditorFields, bAllEditorFields, true);
+					AppendFieldNotify(builder, classObj, hasProperties, hasFunctions, hasEditorFields, allEditorFields, true);
 				}
-				Builder.Append("#else //WITH_EDITORONLY_DATA\r\n");
-				using (UhtMacroCreator Macro = new UhtMacroCreator(Builder, this, Class, FieldNotifyMacroSuffix))
+				builder.Append("#else //WITH_EDITORONLY_DATA\r\n");
+				using (UhtMacroCreator macro = new UhtMacroCreator(builder, this, classObj, FieldNotifyMacroSuffix))
 				{
-					AppendFieldNotify(Builder, Class, bHasProperties, bHasFunctions, bHasEditorFields, bAllEditorFields, false);
+					AppendFieldNotify(builder, classObj, hasProperties, hasFunctions, hasEditorFields, allEditorFields, false);
 				}
-				Builder.Append("#endif // WITH_EDITORONLY_DATA\r\n");
+				builder.Append("#endif // WITH_EDITORONLY_DATA\r\n");
 			}
 			else
 			{
-				using (UhtMacroCreator Macro = new UhtMacroCreator(Builder, this, Class, FieldNotifyMacroSuffix))
+				using (UhtMacroCreator macro = new UhtMacroCreator(builder, this, classObj, FieldNotifyMacroSuffix))
 				{
-					AppendFieldNotify(Builder, Class, bHasProperties, bHasFunctions, bHasEditorFields, bAllEditorFields, false);
+					AppendFieldNotify(builder, classObj, hasProperties, hasFunctions, hasEditorFields, allEditorFields, false);
 				}
 			}
-			return Builder;
+			return builder;
 		}
 
-		private static StringBuilder AppendFieldNotify(StringBuilder Builder, UhtClass Class, 
-			bool bHasProperties, bool bHasFunctions, bool bHasEditorFields, bool bAllEditorFields,
-			bool bIncludeEditorOnlyFields)
+		private static StringBuilder AppendFieldNotify(StringBuilder builder, UhtClass classObj,
+			bool hasProperties, bool hasFunctions, bool hasEditorFields, bool allEditorFields,
+			bool includeEditorOnlyFields)
 		{
-			Builder.Append("\tUE_FIELD_NOTIFICATION_DECLARE_CLASS_DESCRIPTOR_BEGIN() \\\r\n");
+			builder.Append("\tUE_FIELD_NOTIFICATION_DECLARE_CLASS_DESCRIPTOR_BEGIN() \\\r\n");
 
 			//UE_FIELD_NOTIFICATION_DECLARE_FIELD
-			AppendFieldNotify(Builder, Class, bHasProperties, bHasFunctions, bHasEditorFields, bAllEditorFields, 
-				bIncludeEditorOnlyFields, false, (StringBuilder Builder, UhtClass Class, string Name) =>
+			AppendFieldNotify(builder, classObj, hasProperties, hasFunctions, hasEditorFields, allEditorFields,
+				includeEditorOnlyFields, false, (StringBuilder builder, UhtClass classObj, string name) =>
 			{
-				Builder.Append($"\tUE_FIELD_NOTIFICATION_DECLARE_FIELD({Name}) \\\r\n");
+				builder.Append($"\tUE_FIELD_NOTIFICATION_DECLARE_FIELD({name}) \\\r\n");
 			});
 
 			//UE_FIELD_NOTIFICATION_DECLARE_ENUM_FIELD
-			bool bIsFirst = true;
-			AppendFieldNotify(Builder, Class, bHasProperties, bHasFunctions, bHasEditorFields, bAllEditorFields, 
-				bIncludeEditorOnlyFields, false, (StringBuilder Builder, UhtClass Class, string Name) =>
+			bool isFirst = true;
+			AppendFieldNotify(builder, classObj, hasProperties, hasFunctions, hasEditorFields, allEditorFields,
+				includeEditorOnlyFields, false, (StringBuilder builder, UhtClass classObj, string name) =>
 			{
-				if (bIsFirst)
+				if (isFirst)
 				{
-					bIsFirst = false;
-					Builder.Append($"\tUE_FIELD_NOTIFICATION_DECLARE_ENUM_FIELD_BEGIN({Name}) \\\r\n");
+					isFirst = false;
+					builder.Append($"\tUE_FIELD_NOTIFICATION_DECLARE_ENUM_FIELD_BEGIN({name}) \\\r\n");
 				}
 				else
 				{
-					Builder.Append($"\tUE_FIELD_NOTIFICATION_DECLARE_ENUM_FIELD({Name}) \\\r\n");
+					builder.Append($"\tUE_FIELD_NOTIFICATION_DECLARE_ENUM_FIELD({name}) \\\r\n");
 				}
 			});
 
-			Builder.Append("\tUE_FIELD_NOTIFICATION_DECLARE_ENUM_FIELD_END() \\\r\n");
-			Builder.Append($"\tUE_FIELD_NOTIFICATION_DECLARE_CLASS_DESCRIPTOR_END(); \\\r\n");
-			return Builder;
+			builder.Append("\tUE_FIELD_NOTIFICATION_DECLARE_ENUM_FIELD_END() \\\r\n");
+			builder.Append($"\tUE_FIELD_NOTIFICATION_DECLARE_CLASS_DESCRIPTOR_END(); \\\r\n");
+			return builder;
 		}
 
-		private StringBuilder AppendSparseDeclarations(StringBuilder Builder, UhtClass Class)
+		private StringBuilder AppendSparseDeclarations(StringBuilder builder, UhtClass classObj)
 		{
 			// Format the sparse data
-			using (UhtMacroCreator Macro = new UhtMacroCreator(Builder, this, Class, SparseDataMacroSuffix))
+			using (UhtMacroCreator macro = new UhtMacroCreator(builder, this, classObj, SparseDataMacroSuffix))
 			{
-				string[]? SparseDataTypes = Class.MetaData.GetStringArray(UhtNames.SparseClassDataTypes);
-				if (SparseDataTypes != null)
+				string[]? sparseDataTypes = classObj.MetaData.GetStringArray(UhtNames.SparseClassDataTypes);
+				if (sparseDataTypes != null)
 				{
-					foreach (string SparseDataType in SparseDataTypes)
+					foreach (string sparseDataType in sparseDataTypes)
 					{
-						Builder.Append('F').Append(SparseDataType).Append("* Get").Append(SparseDataType).Append("() \\\r\n");
-						Builder.Append("{ \\\r\n");
-						Builder.Append("\treturn (F").Append(SparseDataType).Append("*)(GetClass()->GetOrCreateSparseClassData()); \\\r\n");
-						Builder.Append("} \\\r\n");
-
-						Builder.Append('F').Append(SparseDataType).Append("* Get").Append(SparseDataType).Append("() const \\\r\n");
-						Builder.Append("{ \\\r\n");
-						Builder.Append("\treturn (F").Append(SparseDataType).Append("*)(GetClass()->GetOrCreateSparseClassData()); \\\r\n");
-						Builder.Append("} \\\r\n");
-
-						Builder.Append("const F").Append(SparseDataType).Append("* Get").Append(SparseDataType).Append("(EGetSparseClassDataMethod GetMethod) const \\\r\n");
-						Builder.Append("{ \\\r\n");
-						Builder.Append("\treturn (const F").Append(SparseDataType).Append("*)(GetClass()->GetSparseClassData(GetMethod)); \\\r\n");
-						Builder.Append("} \\\r\n");
+						builder.Append('F').Append(sparseDataType).Append("* Get").Append(sparseDataType).Append("() \\\r\n");
+						builder.Append("{ \\\r\n");
+						builder.Append("\treturn (F").Append(sparseDataType).Append("*)(GetClass()->GetOrCreateSparseClassData()); \\\r\n");
+						builder.Append("} \\\r\n");
+
+						builder.Append('F').Append(sparseDataType).Append("* Get").Append(sparseDataType).Append("() const \\\r\n");
+						builder.Append("{ \\\r\n");
+						builder.Append("\treturn (F").Append(sparseDataType).Append("*)(GetClass()->GetOrCreateSparseClassData()); \\\r\n");
+						builder.Append("} \\\r\n");
+
+						builder.Append("const F").Append(sparseDataType).Append("* Get").Append(sparseDataType).Append("(EGetSparseClassDataMethod GetMethod) const \\\r\n");
+						builder.Append("{ \\\r\n");
+						builder.Append("\treturn (const F").Append(sparseDataType).Append("*)(GetClass()->GetSparseClassData(GetMethod)); \\\r\n");
+						builder.Append("} \\\r\n");
 					}
 
-					foreach (string SparseDataType in SparseDataTypes)
+					foreach (string sparseDataType in sparseDataTypes)
 					{
-						UhtScriptStruct? SparseScriptStruct = Class.FindType(UhtFindOptions.EngineName | UhtFindOptions.ScriptStruct | UhtFindOptions.NoSelf, SparseDataType) as UhtScriptStruct;
-						while (SparseScriptStruct != null)
+						UhtScriptStruct? sparseScriptStruct = classObj.FindType(UhtFindOptions.EngineName | UhtFindOptions.ScriptStruct | UhtFindOptions.NoSelf, sparseDataType) as UhtScriptStruct;
+						while (sparseScriptStruct != null)
 						{
-							foreach (UhtProperty SparseProperty in SparseScriptStruct.Properties)
+							foreach (UhtProperty sparseProperty in sparseScriptStruct.Properties)
 							{
-								if (!SparseProperty.MetaData.ContainsKey(UhtNames.NoGetter))
+								if (!sparseProperty.MetaData.ContainsKey(UhtNames.NoGetter))
 								{
-									string PropertyName = SparseProperty.SourceName;
-									string CleanPropertyName = PropertyName;
-									if (SparseProperty is UhtBoolProperty && PropertyName.StartsWith("b", StringComparison.Ordinal))
+									string propertyName = sparseProperty.SourceName;
+									string cleanPropertyName = propertyName;
+									if (sparseProperty is UhtBoolProperty && propertyName.StartsWith("b", StringComparison.Ordinal))
 									{
-										CleanPropertyName = PropertyName.Substring(1);
+										cleanPropertyName = propertyName.Substring(1);
 									}
 
-									bool bGetByRef = SparseProperty.MetaData.ContainsKey(UhtNames.GetByRef);
+									bool getByRef = sparseProperty.MetaData.ContainsKey(UhtNames.GetByRef);
 
-									if (bGetByRef)
+									if (getByRef)
 									{
-										Builder.Append("const ").AppendSparse(SparseProperty).Append("& Get").Append(CleanPropertyName).Append("() \\\r\n");
+										builder.Append("const ").AppendSparse(sparseProperty).Append("& Get").Append(cleanPropertyName).Append("() \\\r\n");
 									}
 									else
 									{
-										Builder.AppendSparse(SparseProperty).Append(" Get").Append(CleanPropertyName).Append("() \\\r\n");
+										builder.AppendSparse(sparseProperty).Append(" Get").Append(cleanPropertyName).Append("() \\\r\n");
 									}
-									Builder.Append("{ \\\r\n");
-									Builder.Append("\treturn Get").Append(SparseDataType).Append("()->").Append(PropertyName).Append("; \\\r\n");
-									Builder.Append("} \\\r\n");
+									builder.Append("{ \\\r\n");
+									builder.Append("\treturn Get").Append(sparseDataType).Append("()->").Append(propertyName).Append("; \\\r\n");
+									builder.Append("} \\\r\n");
 
-									if (bGetByRef)
+									if (getByRef)
 									{
-										Builder.Append("const ").AppendSparse(SparseProperty).Append("& Get").Append(CleanPropertyName).Append("() const \\\r\n");
+										builder.Append("const ").AppendSparse(sparseProperty).Append("& Get").Append(cleanPropertyName).Append("() const \\\r\n");
 									}
 									else
 									{
-										Builder.AppendSparse(SparseProperty).Append(" Get").Append(CleanPropertyName).Append("() const \\\r\n");
+										builder.AppendSparse(sparseProperty).Append(" Get").Append(cleanPropertyName).Append("() const \\\r\n");
 									}
-									Builder.Append("{ \\\r\n");
-									Builder.Append("\treturn Get").Append(SparseDataType).Append("()->").Append(PropertyName).Append("; \\\r\n");
-									Builder.Append("} \\\r\n");
+									builder.Append("{ \\\r\n");
+									builder.Append("\treturn Get").Append(sparseDataType).Append("()->").Append(propertyName).Append("; \\\r\n");
+									builder.Append("} \\\r\n");
 								}
 							}
 
-							SparseScriptStruct = SparseScriptStruct.SuperScriptStruct;
+							sparseScriptStruct = sparseScriptStruct.SuperScriptStruct;
 						}
 					}
 				}
 			}
-			return Builder;
+			return builder;
 		}
 
-		private StringBuilder AppendRpcFunctions(StringBuilder Builder, UhtClass Class, List<UhtFunction> ReversedFunctions, bool bEditorOnly)
+		private StringBuilder AppendRpcFunctions(StringBuilder builder, UhtClass classObj, List<UhtFunction> reversedFunctions, bool editorOnly)
 		{
-			Builder.AppendBeginEditorOnlyGuard(bEditorOnly);
+			builder.AppendBeginEditorOnlyGuard(editorOnly);
 
-			using (UhtMacroCreator Macro = new UhtMacroCreator(Builder, this, Class, bEditorOnly ? EditorOnlyRpcWrappersMacroSuffix : RpcWrappersMacroSuffix))
+			using (UhtMacroCreator macro = new UhtMacroCreator(builder, this, classObj, editorOnly ? EditorOnlyRpcWrappersMacroSuffix : RpcWrappersMacroSuffix))
 			{
-				AppendAutogeneratedBlueprintFunctionDeclarations(Builder, ReversedFunctions, bEditorOnly);
-				AppendRpcWrappers(Builder, ReversedFunctions, bEditorOnly);
+				AppendAutogeneratedBlueprintFunctionDeclarations(builder, reversedFunctions, editorOnly);
+				AppendRpcWrappers(builder, reversedFunctions, editorOnly);
 			}
 
-			using (UhtMacroCreator Macro = new UhtMacroCreator(Builder, this, Class, bEditorOnly ? EditorOnlyRpcWrappersNoPureDeclsMacroSuffix : RpcWrappersNoPureDeclsMacroSuffix))
+			using (UhtMacroCreator macro = new UhtMacroCreator(builder, this, classObj, editorOnly ? EditorOnlyRpcWrappersNoPureDeclsMacroSuffix : RpcWrappersNoPureDeclsMacroSuffix))
 			{
-				if (Class.GeneratedCodeVersion <= EGeneratedCodeVersion.V1)
+				if (classObj.GeneratedCodeVersion <= EGeneratedCodeVersion.V1)
 				{
-					AppendAutogeneratedBlueprintFunctionDeclarationsOnlyNotDeclared(Builder, ReversedFunctions, bEditorOnly);
+					AppendAutogeneratedBlueprintFunctionDeclarationsOnlyNotDeclared(builder, reversedFunctions, editorOnly);
 				}
-				AppendRpcWrappers(Builder, ReversedFunctions, bEditorOnly);
+				AppendRpcWrappers(builder, reversedFunctions, editorOnly);
 			}
 
-			if (bEditorOnly)
+			if (editorOnly)
 			{
-				Builder.Append("#else\r\n");
-				using (UhtMacroCreator Macro = new UhtMacroCreator(Builder, this, Class, bEditorOnly ? EditorOnlyRpcWrappersMacroSuffix : RpcWrappersMacroSuffix))
+				builder.Append("#else\r\n");
+				using (UhtMacroCreator macro = new UhtMacroCreator(builder, this, classObj, editorOnly ? EditorOnlyRpcWrappersMacroSuffix : RpcWrappersMacroSuffix))
 				{
 				}
 
-				using (UhtMacroCreator Macro = new UhtMacroCreator(Builder, this, Class, bEditorOnly ? EditorOnlyRpcWrappersNoPureDeclsMacroSuffix : RpcWrappersNoPureDeclsMacroSuffix))
+				using (UhtMacroCreator macro = new UhtMacroCreator(builder, this, classObj, editorOnly ? EditorOnlyRpcWrappersNoPureDeclsMacroSuffix : RpcWrappersNoPureDeclsMacroSuffix))
 				{
 				}
-				Builder.AppendEndEditorOnlyGuard();
+				builder.AppendEndEditorOnlyGuard();
 			}
-			return Builder;
+			return builder;
 		}
 
-		private StringBuilder AppendPropertyAccessors(StringBuilder Builder, UhtClass Class)
+		private StringBuilder AppendPropertyAccessors(StringBuilder builder, UhtClass classObj)
 		{
-			using (UhtMacroCreator Macro = new UhtMacroCreator(Builder, this, Class, AccessorsMacroSuffix))
+			using (UhtMacroCreator macro = new UhtMacroCreator(builder, this, classObj, AccessorsMacroSuffix))
 			{
-				foreach (UhtType Type in Class.Children)
+				foreach (UhtType type in classObj.Children)
 				{
-					if (Type is UhtProperty Property)
+					if (type is UhtProperty property)
 					{
-						if (Property.PropertyExportFlags.HasAnyFlags(UhtPropertyExportFlags.GetterFound))
+						if (property.PropertyExportFlags.HasAnyFlags(UhtPropertyExportFlags.GetterFound))
 						{
-							Builder.Append("static void ").AppendPropertyGetterWrapperName(Property).Append("(const void* Object, void* OutValue); \\\r\n");
+							builder.Append("static void ").AppendPropertyGetterWrapperName(property).Append("(const void* Object, void* OutValue); \\\r\n");
 						}
-						if (Property.PropertyExportFlags.HasAnyFlags(UhtPropertyExportFlags.SetterFound))
+						if (property.PropertyExportFlags.HasAnyFlags(UhtPropertyExportFlags.SetterFound))
 						{
-							Builder.Append("static void ").AppendPropertySetterWrapperName(Property).Append("(void* Object, const void* InValue); \\\r\n");
+							builder.Append("static void ").AppendPropertySetterWrapperName(property).Append("(void* Object, const void* InValue); \\\r\n");
 						}
 					}
 				}
 			}
-			return Builder;
+			return builder;
 		}
 
-		private StringBuilder AppendAutogeneratedBlueprintFunctionDeclarations(StringBuilder Builder, List<UhtFunction> ReversedFunctions, bool bEditorOnly)
+		private StringBuilder AppendAutogeneratedBlueprintFunctionDeclarations(StringBuilder builder, List<UhtFunction> reversedFunctions, bool editorOnly)
 		{
-			foreach (UhtFunction Function in ReversedFunctions)
+			foreach (UhtFunction function in reversedFunctions)
 			{
-				if (!IsRpcFunction(Function, bEditorOnly) || Function.CppImplName == Function.SourceName)
+				if (!IsRpcFunction(function, editorOnly) || function.CppImplName == function.SourceName)
 				{
 					continue;
 				}
-				AppendNetValidateDeclaration(Builder, Function);
-				AppendNativeFunctionHeader(Builder, Function, UhtPropertyTextType.ClassFunctionArgOrRetVal, true, null, null, UhtFunctionExportFlags.None, "; \\\r\n");
+				AppendNetValidateDeclaration(builder, function);
+				AppendNativeFunctionHeader(builder, function, UhtPropertyTextType.ClassFunctionArgOrRetVal, true, null, null, UhtFunctionExportFlags.None, "; \\\r\n");
 			}
-			return Builder;
+			return builder;
 		}
 
-		private StringBuilder AppendAutogeneratedBlueprintFunctionDeclarationsOnlyNotDeclared(StringBuilder Builder, List<UhtFunction> ReversedFunctions, bool bEditorOnly)
+		private StringBuilder AppendAutogeneratedBlueprintFunctionDeclarationsOnlyNotDeclared(StringBuilder builder, List<UhtFunction> reversedFunctions, bool editorOnly)
 		{
-			foreach (UhtFunction Function in ReversedFunctions)
+			foreach (UhtFunction function in reversedFunctions)
 			{
-				if (!IsRpcFunction(Function, bEditorOnly) || Function.CppImplName == Function.SourceName)
+				if (!IsRpcFunction(function, editorOnly) || function.CppImplName == function.SourceName)
 				{
 					continue;
 				}
-				if (!Function.FunctionExportFlags.HasAnyFlags(UhtFunctionExportFlags.ImplFound))
+				if (!function.FunctionExportFlags.HasAnyFlags(UhtFunctionExportFlags.ImplFound))
 				{
-					AppendNetValidateDeclaration(Builder, Function);
-					AppendNativeFunctionHeader(Builder, Function, UhtPropertyTextType.ClassFunctionArgOrRetVal, true, null, null, UhtFunctionExportFlags.None, "; \\\r\n");
+					AppendNetValidateDeclaration(builder, function);
+					AppendNativeFunctionHeader(builder, function, UhtPropertyTextType.ClassFunctionArgOrRetVal, true, null, null, UhtFunctionExportFlags.None, "; \\\r\n");
 				}
 			}
-			return Builder;
+			return builder;
 		}
 
-		private static StringBuilder AppendNetValidateDeclaration(StringBuilder Builder, UhtFunction Function)
+		private static StringBuilder AppendNetValidateDeclaration(StringBuilder builder, UhtFunction function)
 		{
-			if (Function.FunctionFlags.HasAnyFlags(EFunctionFlags.NetValidate))
+			if (function.FunctionFlags.HasAnyFlags(EFunctionFlags.NetValidate))
 			{
-				Builder.Append('\t');
-				if (!Function.FunctionFlags.HasAnyFlags(EFunctionFlags.Static) && !Function.FunctionExportFlags.HasAnyFlags(UhtFunctionExportFlags.Final))
+				builder.Append('\t');
+				if (!function.FunctionFlags.HasAnyFlags(EFunctionFlags.Static) && !function.FunctionExportFlags.HasAnyFlags(UhtFunctionExportFlags.Final))
 				{
-					Builder.Append("virtual");
+					builder.Append("virtual");
 				}
-				Builder.Append(" bool ").Append(Function.CppValidationImplName);
-				AppendParameters(Builder, Function, UhtPropertyTextType.ClassFunctionArgOrRetVal, null, true);
-				Builder.Append("; \\\r\n");
+				builder.Append(" bool ").Append(function.CppValidationImplName);
+				AppendParameters(builder, function, UhtPropertyTextType.ClassFunctionArgOrRetVal, null, true);
+				builder.Append("; \\\r\n");
 			}
-			return Builder;
+			return builder;
 		}
 
-		private static StringBuilder AppendRpcWrappers(StringBuilder Builder, List<UhtFunction> ReversedFunctions, bool bEditorOnly)
+		private static StringBuilder AppendRpcWrappers(StringBuilder builder, List<UhtFunction> reversedFunctions, bool editorOnly)
 		{
-			bool bFirst = true;
-			foreach (UhtFunction Function in ReversedFunctions)
+			bool first = true;
+			foreach (UhtFunction function in reversedFunctions)
 			{
-				if (!IsRpcFunction(Function, bEditorOnly))
+				if (!IsRpcFunction(function, editorOnly))
 				{
 					continue;
 				}
-				if (!ShouldExportFunction(Function))
+				if (!ShouldExportFunction(function))
 				{
 					continue;
 				}
 				//COMPATIBILITY-TODO - Remove once we transition to C# version
-				if (bFirst)
+				if (first)
 				{
-					Builder.Append(" \\\r\n");
-					bFirst = false;
+					builder.Append(" \\\r\n");
+					first = false;
 				}
-				Builder.Append("\tDECLARE_FUNCTION(").Append(Function.UnMarshalAndCallName).Append("); \\\r\n");
+				builder.Append("\tDECLARE_FUNCTION(").Append(function.UnMarshalAndCallName).Append("); \\\r\n");
 			}
-			return Builder;
+			return builder;
 		}
 
-		private StringBuilder AppendCallbackParametersDecls(StringBuilder Builder, UhtClass Class, List<UhtFunction> CallbackFunctions)
+		private StringBuilder AppendCallbackParametersDecls(StringBuilder builder, UhtClass classObj, List<UhtFunction> callbackFunctions)
 		{
-			if (CallbackFunctions.Count > 0)
+			if (callbackFunctions.Count > 0)
 			{
-				using (UhtMacroCreator Macro = new UhtMacroCreator(Builder, this, Class, EventParamsMacroSuffix))
+				using (UhtMacroCreator macro = new UhtMacroCreator(builder, this, classObj, EventParamsMacroSuffix))
 				{
-					foreach (UhtFunction Function in CallbackFunctions)
+					foreach (UhtFunction function in callbackFunctions)
 					{
-						AppendEventParameter(Builder, Function, Function.StrippedFunctionName, UhtPropertyTextType.EventParameterMember, true, 1, " \\\r\n");
+						AppendEventParameter(builder, function, function.StrippedFunctionName, UhtPropertyTextType.EventParameterMember, true, 1, " \\\r\n");
 					}
 				}
 			}
-			return Builder;
+			return builder;
 		}
 
-		private StringBuilder AppendCallbackRpcWrapperDecls(StringBuilder Builder, UhtClass Class, List<UhtFunction> CallbackFunctions)
+		private StringBuilder AppendCallbackRpcWrapperDecls(StringBuilder builder, UhtClass classObj, List<UhtFunction> callbackFunctions)
 		{
-			if (CallbackFunctions.Count > 0)
+			if (callbackFunctions.Count > 0)
 			{
-				using (UhtMacroCreator Macro = new UhtMacroCreator(Builder, this, Class, CallbackWrappersMacroSuffix))
+				using (UhtMacroCreator macro = new UhtMacroCreator(builder, this, classObj, CallbackWrappersMacroSuffix))
 				{
-					foreach (UhtFunction Function in CallbackFunctions)
+					foreach (UhtFunction function in callbackFunctions)
 					{
-						if (!Function.FunctionFlags.HasAnyFlags(EFunctionFlags.NetResponse) &&
-							Function.EngineName != Function.MarshalAndCallName)
+						if (!function.FunctionFlags.HasAnyFlags(EFunctionFlags.NetResponse) &&
+							function.EngineName != function.MarshalAndCallName)
 						{
-							AppendNativeFunctionHeader(Builder, Function, UhtPropertyTextType.EventFunctionArgOrRetVal, true, null, null, UhtFunctionExportFlags.None, " \\\r\n");
-							Builder.Append(" \\\r\n");
+							AppendNativeFunctionHeader(builder, function, UhtPropertyTextType.EventFunctionArgOrRetVal, true, null, null, UhtFunctionExportFlags.None, " \\\r\n");
+							builder.Append(" \\\r\n");
 						}
 					}
 				}
 			}
-			return Builder;
+			return builder;
 		}
 
-		private StringBuilder AppendSerializer(StringBuilder Builder, UhtClass Class, string Api, UhtSerializerArchiveType Type, string Declare)
+		private StringBuilder AppendSerializer(StringBuilder builder, UhtClass classObj, string api, UhtSerializerArchiveType type, string declare)
 		{
-			if (!Class.SerializerArchiveType.HasAnyFlags(Type))
+			if (!classObj.SerializerArchiveType.HasAnyFlags(type))
 			{
-				if (Class.EnclosingDefine.Length != 0)
+				if (classObj.EnclosingDefine.Length != 0)
 				{
-					Builder.Append("#if ").Append(Class.EnclosingDefine).Append("\r\n");
+					builder.Append("#if ").Append(classObj.EnclosingDefine).Append("\r\n");
 				}
-				AppendSerializerFunction(Builder, Class, Api, Declare);
-				if (Class.EnclosingDefine.Length != 0)
+				AppendSerializerFunction(builder, classObj, api, declare);
+				if (classObj.EnclosingDefine.Length != 0)
 				{
-					Builder.Append("#else\r\n");
-					using (UhtMacroCreator Macro = new UhtMacroCreator(Builder, this, Class, ArchiveSerializerMacroSuffix))
+					builder.Append("#else\r\n");
+					using (UhtMacroCreator macro = new UhtMacroCreator(builder, this, classObj, ArchiveSerializerMacroSuffix))
 					{
 					}
-					Builder.Append("#endif\r\n");
+					builder.Append("#endif\r\n");
 				}
 			}
-			return Builder;
+			return builder;
 		}
 
-		private StringBuilder AppendSerializerFunction(StringBuilder Builder, UhtClass Class, string Api, string Declare)
+		private StringBuilder AppendSerializerFunction(StringBuilder builder, UhtClass classObj, string api, string declare)
 		{
-			using (UhtMacroCreator Macro = new UhtMacroCreator(Builder, this, Class, ArchiveSerializerMacroSuffix))
+			using (UhtMacroCreator macro = new UhtMacroCreator(builder, this, classObj, ArchiveSerializerMacroSuffix))
 			{
-				Builder.Append('\t').Append(Declare).Append('(').Append(Class.SourceName).Append(", ").Append(Api.Substring(0, Api.Length - 1)).Append(") \\\r\n");
+				builder.Append('\t').Append(declare).Append('(').Append(classObj.SourceName).Append(", ").Append(api.Substring(0, api.Length - 1)).Append(") \\\r\n");
 			}
-			return Builder;
+			return builder;
 		}
 
 		/// <summary>
 		/// Generates standard constructor declarations
 		/// </summary>
-		/// <param name="Builder">Output builder</param>
-		/// <param name="Class">Class being exported</param>
-		/// <param name="Api">API text to be used</param>
+		/// <param name="builder">Output builder</param>
+		/// <param name="classObj">Class being exported</param>
+		/// <param name="api">API text to be used</param>
 		/// <returns>Output builder</returns>
-		private StringBuilder AppendStandardConstructors(StringBuilder Builder, UhtClass Class, string Api)
+		private StringBuilder AppendStandardConstructors(StringBuilder builder, UhtClass classObj, string api)
 		{
-			using (UhtMacroCreator Macro = new UhtMacroCreator(Builder, this, Class, StandardConstructorsMacroSuffix))
+			using (UhtMacroCreator macro = new UhtMacroCreator(builder, this, classObj, StandardConstructorsMacroSuffix))
 			{
-				if (!Class.ClassExportFlags.HasAnyFlags(UhtClassExportFlags.HasCustomConstructor))
+				if (!classObj.ClassExportFlags.HasAnyFlags(UhtClassExportFlags.HasCustomConstructor))
 				{
-					Builder.Append("\t/** Standard constructor, called after all reflected properties have been initialized */ \\\r\n");
-					Builder.Append('\t').Append(Api).Append(Class.SourceName).Append("(const FObjectInitializer& ObjectInitializer");
-					if (!Class.ClassExportFlags.HasAnyFlags(UhtClassExportFlags.HasDefaultConstructor))
+					builder.Append("\t/** Standard constructor, called after all reflected properties have been initialized */ \\\r\n");
+					builder.Append('\t').Append(api).Append(classObj.SourceName).Append("(const FObjectInitializer& ObjectInitializer");
+					if (!classObj.ClassExportFlags.HasAnyFlags(UhtClassExportFlags.HasDefaultConstructor))
 					{
-						Builder.Append(" = FObjectInitializer::Get()");
+						builder.Append(" = FObjectInitializer::Get()");
 					}
-					Builder.Append("); \\\r\n");
+					builder.Append("); \\\r\n");
 				}
-				if (Class.ClassFlags.HasAnyFlags(EClassFlags.Abstract))
+				if (classObj.ClassFlags.HasAnyFlags(EClassFlags.Abstract))
 				{
-					Builder.Append("\tDEFINE_ABSTRACT_DEFAULT_OBJECT_INITIALIZER_CONSTRUCTOR_CALL(").Append(Class.SourceName).Append(") \\\r\n");
+					builder.Append("\tDEFINE_ABSTRACT_DEFAULT_OBJECT_INITIALIZER_CONSTRUCTOR_CALL(").Append(classObj.SourceName).Append(") \\\r\n");
 				}
 				else
 				{
-					Builder.Append("\tDEFINE_DEFAULT_OBJECT_INITIALIZER_CONSTRUCTOR_CALL(").Append(Class.SourceName).Append(") \\\r\n");
+					builder.Append("\tDEFINE_DEFAULT_OBJECT_INITIALIZER_CONSTRUCTOR_CALL(").Append(classObj.SourceName).Append(") \\\r\n");
 				}
 
-				AppendVTableHelperCtorAndCaller(Builder, Class, Api);
-				AppendCopyConstructorDefinition(Builder, Class, Api);
+				AppendVTableHelperCtorAndCaller(builder, classObj, api);
+				AppendCopyConstructorDefinition(builder, classObj, api);
 			}
-			return Builder;
+			return builder;
 		}
 
 		/// <summary>
 		/// Generates enhanced constructor declaration.
 		/// </summary>
-		/// <param name="Builder">Output builder</param>
-		/// <param name="Class">Class being exported</param>
-		/// <param name="Api">API text to be used</param>
+		/// <param name="builder">Output builder</param>
+		/// <param name="classObj">Class being exported</param>
+		/// <param name="api">API text to be used</param>
 		/// <returns>Output builder</returns>
-		private StringBuilder AppendEnhancedConstructors(StringBuilder Builder, UhtClass Class, string Api)
+		private StringBuilder AppendEnhancedConstructors(StringBuilder builder, UhtClass classObj, string api)
 		{
-			using (UhtMacroCreator Macro = new UhtMacroCreator(Builder, this, Class, EnchancedConstructorsMacroSuffix))
+			using (UhtMacroCreator macro = new UhtMacroCreator(builder, this, classObj, EnchancedConstructorsMacroSuffix))
 			{
-				AppendConstructorDefinition(Builder, Class, Api);
-				AppendVTableHelperCtorAndCaller(Builder, Class, Api);
-				AppendDefaultConstructorCallDefinition(Builder, Class);
+				AppendConstructorDefinition(builder, classObj, api);
+				AppendVTableHelperCtorAndCaller(builder, classObj, api);
+				AppendDefaultConstructorCallDefinition(builder, classObj);
 			}
-			return Builder;
+			return builder;
 		}
 
 		/// <summary>
 		/// Generates vtable helper caller and eventual constructor body.
 		/// </summary>
-		/// <param name="Builder">Output builder</param>
-		/// <param name="Class">Class being exported</param>
-		/// <param name="Api">API text to be used</param>
+		/// <param name="builder">Output builder</param>
+		/// <param name="classObj">Class being exported</param>
+		/// <param name="api">API text to be used</param>
 		/// <returns>Output builder</returns>
-		private static StringBuilder AppendVTableHelperCtorAndCaller(StringBuilder Builder, UhtClass Class, string Api)
+		private static StringBuilder AppendVTableHelperCtorAndCaller(StringBuilder builder, UhtClass classObj, string api)
 		{
-			if (!Class.ClassExportFlags.HasAnyFlags(UhtClassExportFlags.HasCustomVTableHelperConstructor))
+			if (!classObj.ClassExportFlags.HasAnyFlags(UhtClassExportFlags.HasCustomVTableHelperConstructor))
 			{
-				Builder.Append("\tDECLARE_VTABLE_PTR_HELPER_CTOR(").Append(Api.Substring(0, Api.Length - 1)).Append(", ").Append(Class.SourceName).Append("); \\\r\n");
+				builder.Append("\tDECLARE_VTABLE_PTR_HELPER_CTOR(").Append(api.Substring(0, api.Length - 1)).Append(", ").Append(classObj.SourceName).Append("); \\\r\n");
 			}
-			Builder.Append("\tDEFINE_VTABLE_PTR_HELPER_CTOR_CALLER(").Append(Class.SourceName).Append("); \\\r\n");
-			return Builder;
+			builder.Append("\tDEFINE_VTABLE_PTR_HELPER_CTOR_CALLER(").Append(classObj.SourceName).Append("); \\\r\n");
+			return builder;
 		}
 
 		/// <summary>
 		/// Generates private copy-constructor declaration.
 		/// </summary>
-		/// <param name="Builder">Output builder</param>
-		/// <param name="Class">Class being exported</param>
-		/// <param name="Api">API text to be used</param>
+		/// <param name="builder">Output builder</param>
+		/// <param name="classObj">Class being exported</param>
+		/// <param name="api">API text to be used</param>
 		/// <returns>Output builder</returns>
-		private static StringBuilder AppendCopyConstructorDefinition(StringBuilder Builder, UhtClass Class, string Api)
+		private static StringBuilder AppendCopyConstructorDefinition(StringBuilder builder, UhtClass classObj, string api)
 		{
-			Builder.Append("private: \\\r\n");
-			Builder.Append("\t/** Private move- and copy-constructors, should never be used */ \\\r\n");
-			Builder.Append('\t').Append(Api).Append(Class.SourceName).Append('(').Append(Class.SourceName).Append("&&); \\\r\n");
-			Builder.Append('\t').Append(Api).Append(Class.SourceName).Append("(const ").Append(Class.SourceName).Append("&); \\\r\n");
-			Builder.Append("public: \\\r\n");
-			return Builder;
+			builder.Append("private: \\\r\n");
+			builder.Append("\t/** Private move- and copy-constructors, should never be used */ \\\r\n");
+			builder.Append('\t').Append(api).Append(classObj.SourceName).Append('(').Append(classObj.SourceName).Append("&&); \\\r\n");
+			builder.Append('\t').Append(api).Append(classObj.SourceName).Append("(const ").Append(classObj.SourceName).Append("&); \\\r\n");
+			builder.Append("public: \\\r\n");
+			return builder;
 		}
 
 		/// <summary>
 		/// Generates private copy-constructor declaration.
 		/// </summary>
-		/// <param name="Builder">Output builder</param>
-		/// <param name="Class">Class being exported</param>
-		/// <param name="Api">API text to be used</param>
+		/// <param name="builder">Output builder</param>
+		/// <param name="classObj">Class being exported</param>
+		/// <param name="api">API text to be used</param>
 		/// <returns>Output builder</returns>
-		private static StringBuilder AppendConstructorDefinition(StringBuilder Builder, UhtClass Class, string Api)
+		private static StringBuilder AppendConstructorDefinition(StringBuilder builder, UhtClass classObj, string api)
 		{
-			if (!Class.ClassExportFlags.HasAnyFlags(UhtClassExportFlags.HasConstructor))
+			if (!classObj.ClassExportFlags.HasAnyFlags(UhtClassExportFlags.HasConstructor))
 			{
-				Builder.Append("\t/** Standard constructor, called after all reflected properties have been initialized */ \\\r\n");
+				builder.Append("\t/** Standard constructor, called after all reflected properties have been initialized */ \\\r\n");
 
 				// Assume super class has OI constructor, this may not always be true but we should always be able to check this.
 				// In any case, it will default to old behavior before we even checked this.
-				bool bSuperClassObjectInitializerConstructorDeclared = true;
-				UhtClass? SuperClass = Class.SuperClass;
-				if (SuperClass != null)
+				bool superClassObjectInitializerConstructorDeclared = true;
+				UhtClass? superClass = classObj.SuperClass;
+				if (superClass != null)
 				{
-					if (!SuperClass.HeaderFile.bIsNoExportTypes)
+					if (!superClass.HeaderFile.IsNoExportTypes)
 					{
-						bSuperClassObjectInitializerConstructorDeclared = SuperClass.ClassExportFlags.HasAnyFlags(UhtClassExportFlags.HasObjectInitializerConstructor);
+						superClassObjectInitializerConstructorDeclared = superClass.ClassExportFlags.HasAnyFlags(UhtClassExportFlags.HasObjectInitializerConstructor);
 					}
 				}
 
-				if (bSuperClassObjectInitializerConstructorDeclared)
+				if (superClassObjectInitializerConstructorDeclared)
 				{
-					Builder.Append('\t').Append(Api).Append(Class.SourceName).Append("(const FObjectInitializer& ObjectInitializer = FObjectInitializer::Get()) : Super(ObjectInitializer) { }; \\\r\n");
+					builder.Append('\t').Append(api).Append(classObj.SourceName).Append("(const FObjectInitializer& ObjectInitializer = FObjectInitializer::Get()) : Super(ObjectInitializer) { }; \\\r\n");
 					//ETSTODO - Modification of class
-					Class.ClassExportFlags |= UhtClassExportFlags.HasObjectInitializerConstructor;
+					classObj.ClassExportFlags |= UhtClassExportFlags.HasObjectInitializerConstructor;
 				}
 				else
 				{
-					Builder.Append('\t').Append(Api).Append(Class.SourceName).Append("() { }; \\\r\n");
+					builder.Append('\t').Append(api).Append(classObj.SourceName).Append("() { }; \\\r\n");
 					//ETSTODO - Modification of class
-					Class.ClassExportFlags |= UhtClassExportFlags.HasDefaultConstructor;
+					classObj.ClassExportFlags |= UhtClassExportFlags.HasDefaultConstructor;
 				}
 
 				// The original code would mark this true at this point.  However,
 				// we are no longer allowed to modify the data.
 				//ETSTODO - Modification of class
-				Class.ClassExportFlags |= UhtClassExportFlags.HasConstructor;
+				classObj.ClassExportFlags |= UhtClassExportFlags.HasConstructor;
 			}
-			AppendCopyConstructorDefinition(Builder, Class, Api);
-			return Builder;
+			AppendCopyConstructorDefinition(builder, classObj, api);
+			return builder;
 		}
 
 		/// <summary>
 		/// Generates constructor call definition
 		/// </summary>
-		/// <param name="Builder">Output builder</param>
-		/// <param name="Class">Class being exported</param>
+		/// <param name="builder">Output builder</param>
+		/// <param name="classObj">Class being exported</param>
 		/// <returns>Output builder</returns>
-		private static StringBuilder AppendDefaultConstructorCallDefinition(StringBuilder Builder, UhtClass Class)
+		private static StringBuilder AppendDefaultConstructorCallDefinition(StringBuilder builder, UhtClass classObj)
 		{
-			if (Class.ClassExportFlags.HasAnyFlags(UhtClassExportFlags.HasObjectInitializerConstructor))
+			if (classObj.ClassExportFlags.HasAnyFlags(UhtClassExportFlags.HasObjectInitializerConstructor))
 			{
-				if (Class.ClassFlags.HasAnyFlags(EClassFlags.Abstract))
+				if (classObj.ClassFlags.HasAnyFlags(EClassFlags.Abstract))
 				{
-					Builder.Append("\tDEFINE_ABSTRACT_DEFAULT_OBJECT_INITIALIZER_CONSTRUCTOR_CALL(").Append(Class.SourceName).Append(") \\\r\n");
+					builder.Append("\tDEFINE_ABSTRACT_DEFAULT_OBJECT_INITIALIZER_CONSTRUCTOR_CALL(").Append(classObj.SourceName).Append(") \\\r\n");
 				}
 				else
 				{
-					Builder.Append("\tDEFINE_DEFAULT_OBJECT_INITIALIZER_CONSTRUCTOR_CALL(").Append(Class.SourceName).Append(") \\\r\n");
+					builder.Append("\tDEFINE_DEFAULT_OBJECT_INITIALIZER_CONSTRUCTOR_CALL(").Append(classObj.SourceName).Append(") \\\r\n");
 				}
 			}
-			else if (Class.ClassExportFlags.HasAnyFlags(UhtClassExportFlags.HasDefaultConstructor))
+			else if (classObj.ClassExportFlags.HasAnyFlags(UhtClassExportFlags.HasDefaultConstructor))
 			{
-				if (Class.ClassFlags.HasAnyFlags(EClassFlags.Abstract))
+				if (classObj.ClassFlags.HasAnyFlags(EClassFlags.Abstract))
 				{
-					Builder.Append("\tDEFINE_ABSTRACT_DEFAULT_CONSTRUCTOR_CALL(").Append(Class.SourceName).Append(") \\\r\n");
+					builder.Append("\tDEFINE_ABSTRACT_DEFAULT_CONSTRUCTOR_CALL(").Append(classObj.SourceName).Append(") \\\r\n");
 				}
 				else
 				{
-					Builder.Append("\tDEFINE_DEFAULT_CONSTRUCTOR_CALL(").Append(Class.SourceName).Append(") \\\r\n");
+					builder.Append("\tDEFINE_DEFAULT_CONSTRUCTOR_CALL(").Append(classObj.SourceName).Append(") \\\r\n");
 				}
 			}
 			else
 			{
-				Builder.Append("\tDEFINE_FORBIDDEN_DEFAULT_CONSTRUCTOR_CALL(").Append(Class.SourceName).Append(") \\\r\n");
+				builder.Append("\tDEFINE_FORBIDDEN_DEFAULT_CONSTRUCTOR_CALL(").Append(classObj.SourceName).Append(") \\\r\n");
 			}
-			return Builder;
+			return builder;
 		}
 
 		/// <summary>
 		/// Generates generated body code for classes
 		/// </summary>
-		/// <param name="Builder">Output builder</param>
-		/// <param name="Class">Class being exported</param>
-		/// <param name="Api">API text to be used</param>
+		/// <param name="builder">Output builder</param>
+		/// <param name="classObj">Class being exported</param>
+		/// <param name="api">API text to be used</param>
 		/// <returns>Output builder</returns>
-		private StringBuilder AppendClassGeneratedBody(StringBuilder Builder, UhtClass Class, string Api)
+		private StringBuilder AppendClassGeneratedBody(StringBuilder builder, UhtClass classObj, string api)
 		{
-			AppendCommonGeneratedBody(Builder, Class, Api);
+			AppendCommonGeneratedBody(builder, classObj, api);
 
 			// export the class's config name
-			UhtClass? SuperClass = Class.SuperClass;
-			if (SuperClass != null && Class.Config.Length > 0 && Class.Config != SuperClass.Config)
+			UhtClass? superClass = classObj.SuperClass;
+			if (superClass != null && classObj.Config.Length > 0 && classObj.Config != superClass.Config)
 			{
-				Builder.Append("\tstatic const TCHAR* StaticConfigName() {return TEXT(\"").Append(Class.Config).Append("\");} \\\r\n \\\r\n");
+				builder.Append("\tstatic const TCHAR* StaticConfigName() {return TEXT(\"").Append(classObj.Config).Append("\");} \\\r\n \\\r\n");
 			}
 
 			// export implementation of _getUObject for classes that implement interfaces
-			foreach (UhtStruct BaseStruct in Class.Bases)
+			foreach (UhtStruct baseStruct in classObj.Bases)
 			{
-				if (BaseStruct is UhtClass BaseClass)
+				if (baseStruct is UhtClass baseClass)
 				{
-					if (BaseClass.ClassFlags.HasAnyFlags(EClassFlags.Interface))
+					if (baseClass.ClassFlags.HasAnyFlags(EClassFlags.Interface))
 					{
-						Builder.Append("\tvirtual UObject* _getUObject() const override { return const_cast<").Append(Class.SourceName).Append("*>(this); } \\\r\n");
+						builder.Append("\tvirtual UObject* _getUObject() const override { return const_cast<").Append(classObj.SourceName).Append("*>(this); } \\\r\n");
 						break;
 					}
 				}
 			}
 
-			if (Class.ClassExportFlags.HasAnyFlags(UhtClassExportFlags.SelfHasReplicatedProperties))
+			if (classObj.ClassExportFlags.HasAnyFlags(UhtClassExportFlags.SelfHasReplicatedProperties))
 			{
-				AppendReplicatedMacroData(Builder, Class, Api);
+				AppendReplicatedMacroData(builder, classObj, api);
 			}
 
-			return Builder;
+			return builder;
 		}
 
 		/// <summary>
 		/// Generates standard generated body code for interfaces and non-interfaces
 		/// </summary>
-		/// <param name="Builder">Output builder</param>
-		/// <param name="Class">Class being exported</param>
-		/// <param name="Api">API text to be used</param>
+		/// <param name="builder">Output builder</param>
+		/// <param name="classObj">Class being exported</param>
+		/// <param name="api">API text to be used</param>
 		/// <returns>Output builder</returns>
-		private StringBuilder AppendCommonGeneratedBody(StringBuilder Builder, UhtClass Class, string Api)
+		private StringBuilder AppendCommonGeneratedBody(StringBuilder builder, UhtClass classObj, string api)
 		{
 			// Export the class's native function registration.
-			Builder.Append("private: \\\r\n");
-			Builder.Append("\tstatic void StaticRegisterNatives").Append(Class.SourceName).Append("(); \\\r\n");
-			if (!Class.ClassExportFlags.HasAnyFlags(UhtClassExportFlags.NoExport))
+			builder.Append("private: \\\r\n");
+			builder.Append("\tstatic void StaticRegisterNatives").Append(classObj.SourceName).Append("(); \\\r\n");
+			if (!classObj.ClassExportFlags.HasAnyFlags(UhtClassExportFlags.NoExport))
 			{
-				Builder.Append("\tfriend struct ").Append(this.ObjectInfos[Class.ObjectTypeIndex].RegisteredSingletonName).Append("_Statics; \\\r\n");
+				builder.Append("\tfriend struct ").Append(this.ObjectInfos[classObj.ObjectTypeIndex]._registeredSingletonName).Append("_Statics; \\\r\n");
 			}
-			Builder.Append("public: \\\r\n");
+			builder.Append("public: \\\r\n");
 
-			UhtClass? SuperClass = Class.SuperClass;
-			bool bCastedClass = Class.ClassCastFlags.HasAnyFlags(EClassCastFlags.AllFlags) && SuperClass != null && Class.ClassCastFlags != SuperClass.ClassCastFlags;
+			UhtClass? superClass = classObj.SuperClass;
+			bool castedClass = classObj.ClassCastFlags.HasAnyFlags(EClassCastFlags.AllFlags) && superClass != null && classObj.ClassCastFlags != superClass.ClassCastFlags;
 
-			Builder
+			builder
 				.Append("\tDECLARE_CLASS(")
-				.Append(Class)
+				.Append(classObj)
 				.Append(", ")
-				.Append(SuperClass != null ? SuperClass.SourceName : "None")
+				.Append(superClass != null ? superClass.SourceName : "None")
 				.Append(", COMPILED_IN_FLAGS(")
-				.Append(Class.ClassFlags.HasAnyFlags(EClassFlags.Abstract) ? "CLASS_Abstract" : "0");
+				.Append(classObj.ClassFlags.HasAnyFlags(EClassFlags.Abstract) ? "CLASS_Abstract" : "0");
 
-			AppendClassFlags(Builder, Class);
-			Builder.Append("), ");
-			if (bCastedClass)
+			AppendClassFlags(builder, classObj);
+			builder.Append("), ");
+			if (castedClass)
 			{
-				Builder
+				builder
 					.Append("CASTCLASS_")
-					.Append(Class.SourceName);
+					.Append(classObj.SourceName);
 			}
 			else
 			{
-				Builder.Append("CASTCLASS_None");
+				builder.Append("CASTCLASS_None");
 			}
-			Builder.Append(", TEXT(\"").Append(this.Package.SourceName).Append("\"), ").Append(Api.Substring(0, Api.Length - 1)).Append(") \\\r\n");
+			builder.Append(", TEXT(\"").Append(this.Package.SourceName).Append("\"), ").Append(api.Substring(0, api.Length - 1)).Append(") \\\r\n");
 
-			Builder.Append("\tDECLARE_SERIALIZER(").Append(Class.SourceName).Append(") \\\r\n");
+			builder.Append("\tDECLARE_SERIALIZER(").Append(classObj.SourceName).Append(") \\\r\n");
 
 			// Add the serialization function declaration if we generated one
-			if (Class.SerializerArchiveType != UhtSerializerArchiveType.None && Class.SerializerArchiveType != UhtSerializerArchiveType.All)
+			if (classObj.SerializerArchiveType != UhtSerializerArchiveType.None && classObj.SerializerArchiveType != UhtSerializerArchiveType.All)
 			{
-				Builder.Append('\t').AppendMacroName(this, Class, ArchiveSerializerMacroSuffix).Append(" \\\r\n");
+				builder.Append('\t').AppendMacroName(this, classObj, ArchiveSerializerMacroSuffix).Append(" \\\r\n");
 			}
 
-			if (SuperClass != null && Class.ClassWithin != SuperClass.ClassWithin)
+			if (superClass != null && classObj.ClassWithin != superClass.ClassWithin)
 			{
-				Builder.Append("\tDECLARE_WITHIN(").Append(Class.ClassWithin.SourceName).Append(") \\\r\n");
+				builder.Append("\tDECLARE_WITHIN(").Append(classObj.ClassWithin.SourceName).Append(") \\\r\n");
 			}
-			return Builder;
+			return builder;
 		}
 
 		/// <summary>
 		/// Appends the class flags in the form of CLASS_Something|CLASS_Something which represents all class flags that are set 
 		/// for the specified class which need to be exported as part of the DECLARE_CLASS macro
 		/// </summary>
-		/// <param name="Builder">Output builder</param>
-		/// <param name="Class">Class in question</param>
+		/// <param name="builder">Output builder</param>
+		/// <param name="classObj">Class in question</param>
 		/// <returns>Output builder</returns>
-		private static StringBuilder AppendClassFlags(StringBuilder Builder, UhtClass Class)
+		private static StringBuilder AppendClassFlags(StringBuilder builder, UhtClass classObj)
 		{
-			if (Class.ClassFlags.HasAnyFlags(EClassFlags.Transient))
+			if (classObj.ClassFlags.HasAnyFlags(EClassFlags.Transient))
 			{
-				Builder.Append(" | CLASS_Transient");
+				builder.Append(" | CLASS_Transient");
 			}
-			if (Class.ClassFlags.HasAnyFlags(EClassFlags.Optional))
+			if (classObj.ClassFlags.HasAnyFlags(EClassFlags.Optional))
 			{
-				Builder.Append(" | CLASS_Optional");
+				builder.Append(" | CLASS_Optional");
 			}
-			if (Class.ClassFlags.HasAnyFlags(EClassFlags.DefaultConfig))
+			if (classObj.ClassFlags.HasAnyFlags(EClassFlags.DefaultConfig))
 			{
-				Builder.Append(" | CLASS_DefaultConfig");
+				builder.Append(" | CLASS_DefaultConfig");
 			}
-			if (Class.ClassFlags.HasAnyFlags(EClassFlags.GlobalUserConfig))
+			if (classObj.ClassFlags.HasAnyFlags(EClassFlags.GlobalUserConfig))
 			{
-				Builder.Append(" | CLASS_GlobalUserConfig");
+				builder.Append(" | CLASS_GlobalUserConfig");
 			}
-			if (Class.ClassFlags.HasAnyFlags(EClassFlags.ProjectUserConfig))
+			if (classObj.ClassFlags.HasAnyFlags(EClassFlags.ProjectUserConfig))
 			{
-				Builder.Append(" | CLASS_ProjectUserConfig");
+				builder.Append(" | CLASS_ProjectUserConfig");
 			}
-			if (Class.ClassFlags.HasAnyFlags(EClassFlags.Config))
+			if (classObj.ClassFlags.HasAnyFlags(EClassFlags.Config))
 			{
-				Builder.Append(" | CLASS_Config");
+				builder.Append(" | CLASS_Config");
 			}
-			if (Class.ClassFlags.HasAnyFlags(EClassFlags.Interface))
+			if (classObj.ClassFlags.HasAnyFlags(EClassFlags.Interface))
 			{
-				Builder.Append(" | CLASS_Interface");
+				builder.Append(" | CLASS_Interface");
 			}
-			if (Class.ClassFlags.HasAnyFlags(EClassFlags.Deprecated))
+			if (classObj.ClassFlags.HasAnyFlags(EClassFlags.Deprecated))
 			{
-				Builder.Append(" | CLASS_Deprecated");
+				builder.Append(" | CLASS_Deprecated");
 			}
-			return Builder;
+			return builder;
 		}
 
 		/// <summary>
 		/// Appends preprocessor string to emit GENERATED_U*_BODY() macro is deprecated.
 		/// </summary>
-		/// <param name="Builder">Output builder</param>
-		/// <param name="MacroName">Name of the macro to deprecate</param>
+		/// <param name="builder">Output builder</param>
+		/// <param name="macroName">Name of the macro to deprecate</param>
 		/// <returns>Output builder</returns>
 		[System.Diagnostics.CodeAnalysis.SuppressMessage("Style", "IDE0060:Remove unused parameter", Justification = "Referencing code commented out")]
-		private static StringBuilder AppendGeneratedMacroDeprecationWarning(StringBuilder Builder, string MacroName)
+		private static StringBuilder AppendGeneratedMacroDeprecationWarning(StringBuilder builder, string macroName)
 		{
 			// Deprecation warning is disabled right now. After people get familiar with the new macro it should be re-enabled.
 			//Builder.Append("EMIT_DEPRECATED_WARNING_MESSAGE(\"").Append(MacroName).Append("() macro is deprecated. Please use GENERATED_BODY() macro instead.\") \\\r\n");
-			return Builder;
+			return builder;
 		}
 
-		private static StringBuilder AppendAccessSpecifier(StringBuilder Builder, UhtClass Class)
+		private static StringBuilder AppendAccessSpecifier(StringBuilder builder, UhtClass classObj)
 		{
-			switch (Class.GeneratedBodyAccessSpecifier)
+			switch (classObj.GeneratedBodyAccessSpecifier)
 			{
 				case UhtAccessSpecifier.Public:
-					Builder.Append("public:");
+					builder.Append("public:");
 					break;
 				case UhtAccessSpecifier.Private:
-					Builder.Append("private:");
+					builder.Append("private:");
 					break;
 				case UhtAccessSpecifier.Protected:
-					Builder.Append("protected:");
+					builder.Append("protected:");
 					break;
 				default:
-					Builder.Append("static_assert(false, \"Unknown access specifier for GENERATED_BODY() macro in class ").Append(Class.EngineName).Append(".\");");
+					builder.Append("static_assert(false, \"Unknown access specifier for GENERATED_BODY() macro in class ").Append(classObj.EngineName).Append(".\");");
 					break;
 			}
-			return Builder;
+			return builder;
 		}
 
-		private StringBuilder AppendInClassIInterface(StringBuilder Builder, UhtClass Class, List<UhtFunction> CallbackFunctions, string Api)
+		private StringBuilder AppendInClassIInterface(StringBuilder builder, UhtClass classObj, List<UhtFunction> callbackFunctions, string api)
 		{
-			string InterfaceSourceName = "I" + Class.EngineName;
+			string interfaceSourceName = "I" + classObj.EngineName;
 
-			Builder.Append("protected: \\\r\n");
-			Builder.Append("\tvirtual ~").Append(InterfaceSourceName).Append("() {} \\\r\n");
-			Builder.Append("public: \\\r\n");
-			Builder.Append("\ttypedef ").Append(Class.SourceName).Append(" UClassType; \\\r\n");
-			Builder.Append("\ttypedef ").Append(InterfaceSourceName).Append(" ThisClass; \\\r\n");
+			builder.Append("protected: \\\r\n");
+			builder.Append("\tvirtual ~").Append(interfaceSourceName).Append("() {} \\\r\n");
+			builder.Append("public: \\\r\n");
+			builder.Append("\ttypedef ").Append(classObj.SourceName).Append(" UClassType; \\\r\n");
+			builder.Append("\ttypedef ").Append(interfaceSourceName).Append(" ThisClass; \\\r\n");
 
-			AppendInterfaceCallFunctions(Builder, CallbackFunctions);
+			AppendInterfaceCallFunctions(builder, callbackFunctions);
 
 			// we'll need a way to get to the UObject portion of a native interface, so that we can safely pass native interfaces
 			// to script VM functions
-			if (Class.SuperClass != null && Class.SuperClass.IsChildOf(this.Session.UInterface))
+			if (classObj.SuperClass != null && classObj.SuperClass.IsChildOf(this.Session.UInterface))
 			{
 				// Note: This used to be declared as a pure virtual function, but it was changed here in order to allow the Blueprint nativization process
 				// to detect C++ interface classes that explicitly declare pure virtual functions via type traits. This code will no longer trigger that check.
-				Builder.Append("\tvirtual UObject* _getUObject() const { return nullptr; } \\\r\n");
+				builder.Append("\tvirtual UObject* _getUObject() const { return nullptr; } \\\r\n");
 			}
 
-			if (Class.ClassExportFlags.HasAnyFlags(UhtClassExportFlags.SelfHasReplicatedProperties))
+			if (classObj.ClassExportFlags.HasAnyFlags(UhtClassExportFlags.SelfHasReplicatedProperties))
 			{
-				AppendReplicatedMacroData(Builder, Class, Api);
+				AppendReplicatedMacroData(builder, classObj, api);
 			}
 
-			return Builder;
+			return builder;
 		}
 
-		private static StringBuilder AppendReplicatedMacroData(StringBuilder Builder, UhtClass Class, string Api)
+		private static StringBuilder AppendReplicatedMacroData(StringBuilder builder, UhtClass classObj, string api)
 		{
-			if (!Class.ClassExportFlags.HasAnyFlags(UhtClassExportFlags.HasGetLifetimeReplicatedProps))
+			if (!classObj.ClassExportFlags.HasAnyFlags(UhtClassExportFlags.HasGetLifetimeReplicatedProps))
 			{
 				// Default version autogenerates declarations.
-				if (Class.GeneratedCodeVersion == EGeneratedCodeVersion.V1)
+				if (classObj.GeneratedCodeVersion == EGeneratedCodeVersion.V1)
 				{
-					Builder.Append("\tvoid GetLifetimeReplicatedProps(TArray<FLifetimeProperty>& OutLifetimeProps) const override; \\\r\n");
+					builder.Append("\tvoid GetLifetimeReplicatedProps(TArray<FLifetimeProperty>& OutLifetimeProps) const override; \\\r\n");
 				}
 			}
 
-			AppendNetData(Builder, Class, Api);
+			AppendNetData(builder, classObj, api);
 
 			// If this class has replicated properties and it owns the first one, that means
 			// it's the base most replicated class. In that case, go ahead and add our interface macro.
-			if (Class.ClassExportFlags.HasExactFlags(UhtClassExportFlags.HasReplciatedProperties, UhtClassExportFlags.SelfHasReplicatedProperties))
+			if (classObj.ClassExportFlags.HasExactFlags(UhtClassExportFlags.HasReplciatedProperties, UhtClassExportFlags.SelfHasReplicatedProperties))
 			{
-				Builder.Append("private: \\\r\n");
-				Builder.Append("\tREPLICATED_BASE_CLASS(").Append(Class.SourceName).Append(") \\\r\n");
-				Builder.Append("public: \\\r\n");
+				builder.Append("private: \\\r\n");
+				builder.Append("\tREPLICATED_BASE_CLASS(").Append(classObj.SourceName).Append(") \\\r\n");
+				builder.Append("public: \\\r\n");
 			}
-			return Builder;
+			return builder;
 		}
 
-		private static StringBuilder AppendNetData(StringBuilder Builder, UhtClass Class, string Api)
+		private static StringBuilder AppendNetData(StringBuilder builder, UhtClass classObj, string api)
 		{
-			bool bHasArray = false;
-			foreach (UhtProperty Property in Class.EnumerateReplicatedProperties(false))
+			bool hasArray = false;
+			foreach (UhtProperty property in classObj.EnumerateReplicatedProperties(false))
 			{
-				if (Property.bIsStaticArray)
+				if (property.IsStaticArray)
 				{
-					if (!bHasArray)
+					if (!hasArray)
 					{
-						bHasArray = true;
-						Builder.Append("\tenum class EArrayDims_Private : uint16 \\\r\n");
-						Builder.Append("\t{ \\\r\n");
+						hasArray = true;
+						builder.Append("\tenum class EArrayDims_Private : uint16 \\\r\n");
+						builder.Append("\t{ \\\r\n");
 					}
-					Builder.Append("\t\t").Append(Property.SourceName).Append('=').Append(Property.ArrayDimensions).Append(", \\\r\n");
+					builder.Append("\t\t").Append(property.SourceName).Append('=').Append(property.ArrayDimensions).Append(", \\\r\n");
 				}
 			}
 
-			if (bHasArray)
+			if (hasArray)
 			{
-				Builder.Append("\t}; \\\r\n");
+				builder.Append("\t}; \\\r\n");
 			}
 
-			Builder.Append("\tenum class ENetFields_Private : uint16 \\\r\n");
-			Builder.Append("\t{ \\\r\n");
-			Builder.Append("\t\tNETFIELD_REP_START=(uint16)((int32)Super::ENetFields_Private::NETFIELD_REP_END + (int32)1), \\\r\n");
+			builder.Append("\tenum class ENetFields_Private : uint16 \\\r\n");
+			builder.Append("\t{ \\\r\n");
+			builder.Append("\t\tNETFIELD_REP_START=(uint16)((int32)Super::ENetFields_Private::NETFIELD_REP_END + (int32)1), \\\r\n");
 
-			bool bIsFirst = true;
-			UhtProperty? LastProperty = null;
-			foreach (UhtProperty Property in Class.EnumerateReplicatedProperties(false))
+			bool isFirst = true;
+			UhtProperty? lastProperty = null;
+			foreach (UhtProperty property in classObj.EnumerateReplicatedProperties(false))
 			{
-				LastProperty = Property;
-				if (!Property.bIsStaticArray)
+				lastProperty = property;
+				if (!property.IsStaticArray)
 				{
-					if (bIsFirst)
+					if (isFirst)
 					{
-						Builder.Append("\t\t").Append(Property.SourceName).Append("=NETFIELD_REP_START, \\\r\n");
-						bIsFirst = false;
+						builder.Append("\t\t").Append(property.SourceName).Append("=NETFIELD_REP_START, \\\r\n");
+						isFirst = false;
 					}
 					else
 					{
-						Builder.Append("\t\t").Append(Property.SourceName).Append(", \\\r\n");
+						builder.Append("\t\t").Append(property.SourceName).Append(", \\\r\n");
 					}
 				}
 				else
 				{
-					if (bIsFirst)
+					if (isFirst)
 					{
-						Builder.Append("\t\t").Append(Property.SourceName).Append("_STATIC_ARRAY=NETFIELD_REP_START, \\\r\n");
-						bIsFirst = false;
+						builder.Append("\t\t").Append(property.SourceName).Append("_STATIC_ARRAY=NETFIELD_REP_START, \\\r\n");
+						isFirst = false;
 					}
 					else
 					{
-						Builder.Append("\t\t").Append(Property.SourceName).Append("_STATIC_ARRAY, \\\r\n");
+						builder.Append("\t\t").Append(property.SourceName).Append("_STATIC_ARRAY, \\\r\n");
 					}
 
-					Builder
+					builder
 						.Append("\t\t")
-						.Append(Property.SourceName)
+						.Append(property.SourceName)
 						.Append("_STATIC_ARRAY_END=((uint16)")
-						.Append(Property.SourceName)
+						.Append(property.SourceName)
 						.Append("_STATIC_ARRAY + (uint16)EArrayDims_Private::")
-						.Append(Property.SourceName)
+						.Append(property.SourceName)
 						.Append(" - (uint16)1), \\\r\n");
 				}
 			}
 
-			if (LastProperty != null)
+			if (lastProperty != null)
 			{
-				Builder.Append("\t\tNETFIELD_REP_END=").Append(LastProperty.SourceName);
-				if (LastProperty.bIsStaticArray)
+				builder.Append("\t\tNETFIELD_REP_END=").Append(lastProperty.SourceName);
+				if (lastProperty.IsStaticArray)
 				{
-					Builder.Append("_STATIC_ARRAY_END");
+					builder.Append("_STATIC_ARRAY_END");
 				}
 			}
-			Builder.Append("\t}; \\\r\n");
+			builder.Append("\t}; \\\r\n");
 
-			Builder.Append('\t').Append(Api).Append("virtual void ValidateGeneratedRepEnums(const TArray<struct FRepRecord>& ClassReps) const override; \\\r\n");
+			builder.Append('\t').Append(api).Append("virtual void ValidateGeneratedRepEnums(const TArray<struct FRepRecord>& ClassReps) const override; \\\r\n");
 
-			return Builder;
+			return builder;
 		}
 
-		private StringBuilder AppendInterfaceCallFunctions(StringBuilder Builder, List<UhtFunction> CallbackFunctions)
+		private StringBuilder AppendInterfaceCallFunctions(StringBuilder builder, List<UhtFunction> callbackFunctions)
 		{
 			const string ExtraArg = "UObject* O";
 			const string ConstExtraArg = "const UObject* O";
 
-			foreach (UhtFunction Function in CallbackFunctions)
+			foreach (UhtFunction function in callbackFunctions)
 			{
-				AppendNativeFunctionHeader(Builder, Function, UhtPropertyTextType.InterfaceFunctionArgOrRetVal, true, null,
-					Function.FunctionFlags.HasAnyFlags(EFunctionFlags.Const) ? ConstExtraArg : ExtraArg, UhtFunctionExportFlags.None,
+				AppendNativeFunctionHeader(builder, function, UhtPropertyTextType.InterfaceFunctionArgOrRetVal, true, null,
+					function.FunctionFlags.HasAnyFlags(EFunctionFlags.Const) ? ConstExtraArg : ExtraArg, UhtFunctionExportFlags.None,
 					"; \\\r\n");
 			}
-			return Builder;
+			return builder;
 		}
 
-		private StringBuilder AppendGeneratedBodyMacroBlock(StringBuilder Builder, UhtClass Class, UhtClass BodyClass, bool bIsLegacy, bool bHasEditorRpc, bool bHasCallbacks, string? DeprecatedMacroName)
+		private StringBuilder AppendGeneratedBodyMacroBlock(StringBuilder builder, UhtClass classObj, UhtClass bodyClassObj, bool isLegacy, bool hasEditorRpc, bool hasCallbacks, string? deprecatedMacroName)
 		{
-			bool bIsInterface = Class.ClassFlags.HasAnyFlags(EClassFlags.Interface);
-			using (UhtMacroCreator Macro = new UhtMacroCreator(Builder, this, BodyClass, bIsLegacy ? GeneratedBodyLegacyMacroSuffix : GeneratedBodyMacroSuffix))
+			bool isInterface = classObj.ClassFlags.HasAnyFlags(EClassFlags.Interface);
+			using (UhtMacroCreator macro = new UhtMacroCreator(builder, this, bodyClassObj, isLegacy ? GeneratedBodyLegacyMacroSuffix : GeneratedBodyMacroSuffix))
 			{
-				if (DeprecatedMacroName != null)
+				if (deprecatedMacroName != null)
 				{
-					AppendGeneratedMacroDeprecationWarning(Builder, DeprecatedMacroName);
+					AppendGeneratedMacroDeprecationWarning(builder, deprecatedMacroName);
 				}
-				Builder.Append(DisableDeprecationWarnings).Append(" \\\r\n");
-				Builder.Append("public: \\\r\n");
-				Builder.Append('\t').AppendMacroName(this, Class, SparseDataMacroSuffix).Append(" \\\r\n");
-				if (bIsLegacy)
+				builder.Append(DisableDeprecationWarnings).Append(" \\\r\n");
+				builder.Append("public: \\\r\n");
+				builder.Append('\t').AppendMacroName(this, classObj, SparseDataMacroSuffix).Append(" \\\r\n");
+				if (isLegacy)
 				{
-					Builder.Append('\t').AppendMacroName(this, Class, RpcWrappersMacroSuffix).Append(" \\\r\n");
+					builder.Append('\t').AppendMacroName(this, classObj, RpcWrappersMacroSuffix).Append(" \\\r\n");
 				}
 				else
 				{
-					Builder.Append('\t').AppendMacroName(this, Class, RpcWrappersNoPureDeclsMacroSuffix).Append(" \\\r\n");
+					builder.Append('\t').AppendMacroName(this, classObj, RpcWrappersNoPureDeclsMacroSuffix).Append(" \\\r\n");
 				}
-				if (bHasEditorRpc)
+				if (hasEditorRpc)
 				{
-					if (bIsLegacy)
+					if (isLegacy)
 					{
-						Builder.Append('\t').AppendMacroName(this, Class, EditorOnlyRpcWrappersMacroSuffix).Append(" \\\r\n");
+						builder.Append('\t').AppendMacroName(this, classObj, EditorOnlyRpcWrappersMacroSuffix).Append(" \\\r\n");
 					}
 					else
 					{
-						Builder.Append('\t').AppendMacroName(this, Class, EditorOnlyRpcWrappersNoPureDeclsMacroSuffix).Append(" \\\r\n");
+						builder.Append('\t').AppendMacroName(this, classObj, EditorOnlyRpcWrappersNoPureDeclsMacroSuffix).Append(" \\\r\n");
 					}
 				}
-				Builder.Append('\t').AppendMacroName(this, Class, AccessorsMacroSuffix).Append(" \\\r\n");
-				if (bHasCallbacks)
+				builder.Append('\t').AppendMacroName(this, classObj, AccessorsMacroSuffix).Append(" \\\r\n");
+				if (hasCallbacks)
 				{
-					Builder.Append('\t').AppendMacroName(this, Class, CallbackWrappersMacroSuffix).Append(" \\\r\n");
+					builder.Append('\t').AppendMacroName(this, classObj, CallbackWrappersMacroSuffix).Append(" \\\r\n");
 				}
-				if (bIsInterface)
+				if (isInterface)
 				{
-					if (bIsLegacy)
+					if (isLegacy)
 					{
-						Builder.Append('\t').AppendMacroName(this, Class, InClassIInterfaceMacroSuffix).Append(" \\\r\n");
+						builder.Append('\t').AppendMacroName(this, classObj, InClassIInterfaceMacroSuffix).Append(" \\\r\n");
 					}
 					else
 					{
-						Builder.Append('\t').AppendMacroName(this, Class, InClassIInterfaceNoPureDeclsMacroSuffix).Append(" \\\r\n");
+						builder.Append('\t').AppendMacroName(this, classObj, InClassIInterfaceNoPureDeclsMacroSuffix).Append(" \\\r\n");
 					}
 				}
 				else
 				{
-					if (bIsLegacy)
+					if (isLegacy)
 					{
-						Builder.Append('\t').AppendMacroName(this, Class, InClassMacroSuffix).Append(" \\\r\n");
+						builder.Append('\t').AppendMacroName(this, classObj, InClassMacroSuffix).Append(" \\\r\n");
 					}
 					else
 					{
-						Builder.Append('\t').AppendMacroName(this, Class, InClassNoPureDeclsMacroSuffix).Append(" \\\r\n");
+						builder.Append('\t').AppendMacroName(this, classObj, InClassNoPureDeclsMacroSuffix).Append(" \\\r\n");
 					}
 				}
-				if (!bIsInterface)
+				if (!isInterface)
 				{
-					if (bIsLegacy)
+					if (isLegacy)
 					{
-						Builder.Append('\t').AppendMacroName(this, Class, StandardConstructorsMacroSuffix).Append(" \\\r\n");
+						builder.Append('\t').AppendMacroName(this, classObj, StandardConstructorsMacroSuffix).Append(" \\\r\n");
 					}
 					else
 					{
-						Builder.Append('\t').AppendMacroName(this, Class, EnchancedConstructorsMacroSuffix).Append(" \\\r\n");
+						builder.Append('\t').AppendMacroName(this, classObj, EnchancedConstructorsMacroSuffix).Append(" \\\r\n");
 					}
-					if (NeedFieldNotifyCodeGen(Class))
+					if (NeedFieldNotifyCodeGen(classObj))
 					{
-						Builder.Append('\t').AppendMacroName(this, Class, FieldNotifyMacroSuffix).Append(" \\\r\n");
+						builder.Append('\t').AppendMacroName(this, classObj, FieldNotifyMacroSuffix).Append(" \\\r\n");
 					}
 				}
-				if (bIsLegacy)
+				if (isLegacy)
 				{
-					Builder.Append("public: \\\r\n");
+					builder.Append("public: \\\r\n");
 				}
 				else
 				{
-					AppendAccessSpecifier(Builder, Class);
-					Builder.Append(" \\\r\n");
+					AppendAccessSpecifier(builder, classObj);
+					builder.Append(" \\\r\n");
 				}
-				Builder.Append(EnableDeprecationWarnings).Append(" \\\r\n");
+				builder.Append(EnableDeprecationWarnings).Append(" \\\r\n");
 			}
-			return Builder;
+			return builder;
 		}
 
 		#region Enum helper methods
-		private static bool IsFullEnumName(string InEnumName)
+		private static bool IsFullEnumName(string inEnumName)
 		{
-			return InEnumName.Contains("::", StringComparison.Ordinal);
+			return inEnumName.Contains("::", StringComparison.Ordinal);
 		}
 
-		private static StringView GenerateEnumPrefix(UhtEnum Enum)
+		private static StringView GenerateEnumPrefix(UhtEnum enumObj)
 		{
-			StringView Prefix = new StringView();
-			if (Enum.EnumValues.Count > 0)
+			StringView prefix = new StringView();
+			if (enumObj.EnumValues.Count > 0)
 			{
-				Prefix = Enum.EnumValues[0].Name;
+				prefix = enumObj.EnumValues[0].Name;
 
 				// For each item in the enumeration, trim the prefix as much as necessary to keep it a prefix.
 				// This ensures that once all items have been processed, a common prefix will have been constructed.
 				// This will be the longest common prefix since as little as possible is trimmed at each step.
-				for (int NameIdx = 1; NameIdx < Enum.EnumValues.Count; ++NameIdx)
+				for (int nameIdx = 1; nameIdx < enumObj.EnumValues.Count; ++nameIdx)
 				{
-					StringView EnumItemName = Enum.EnumValues[NameIdx].Name;
+					StringView enumItemName = enumObj.EnumValues[nameIdx].Name;
 
 					// Find the length of the longest common prefix of Prefix and EnumItemName.
-					int PrefixIdx = 0;
-					while (PrefixIdx < Prefix.Length && PrefixIdx < EnumItemName.Length && Prefix.Span[PrefixIdx] == EnumItemName.Span[PrefixIdx])
+					int prefixIdx = 0;
+					while (prefixIdx < prefix.Length && prefixIdx < enumItemName.Length && prefix.Span[prefixIdx] == enumItemName.Span[prefixIdx])
 					{
-						PrefixIdx++;
+						prefixIdx++;
 					}
 
 					// Trim the prefix to the length of the common prefix.
-					Prefix = new StringView(Prefix, 0, PrefixIdx);
+					prefix = new StringView(prefix, 0, prefixIdx);
 				}
 
 				// Find the index of the rightmost underscore in the prefix.
-				int UnderscoreIdx = Prefix.Span.LastIndexOf('_');
+				int underscoreIdx = prefix.Span.LastIndexOf('_');
 
 				// If an underscore was found, trim the prefix so only the part before the rightmost underscore is included.
-				if (UnderscoreIdx > 0)
+				if (underscoreIdx > 0)
 				{
-					Prefix = new StringView(Prefix, 0, UnderscoreIdx);
+					prefix = new StringView(prefix, 0, underscoreIdx);
 				}
 				else
 				{
 					// no underscores in the common prefix - this probably indicates that the names
 					// for this enum are not using Epic's notation, so just empty the prefix so that
 					// the max item will use the full name of the enum
-					Prefix = new StringView();
+					prefix = new StringView();
 				}
 			}
 
 			// If no common prefix was found, or if the enum does not contain any entries,
 			// use the name of the enumeration instead.
-			if (Prefix.Length == 0)
+			if (prefix.Length == 0)
 			{
-				Prefix = Enum.EngineName;
+				prefix = enumObj.EngineName;
 			}
-			return Prefix;
+			return prefix;
 		}
 
-		private static string GenerateFullEnumName(UhtEnum Enum, string InEnumName)
+		private static string GenerateFullEnumName(UhtEnum enumObj, string inEnumName)
 		{
-			if (Enum.CppForm == UhtEnumCppForm.Regular || IsFullEnumName(InEnumName))
+			if (enumObj.CppForm == UhtEnumCppForm.Regular || IsFullEnumName(inEnumName))
 			{
-				return InEnumName;
+				return inEnumName;
 			}
-			return $"{Enum.EngineName}::{InEnumName}";
+			return $"{enumObj.EngineName}::{inEnumName}";
 		}
 
-		private static bool EnumHasExistingMax(UhtEnum Enum)
+		private static bool EnumHasExistingMax(UhtEnum enumObj)
 		{
-			if (Enum.GetIndexByName(GenerateFullEnumName(Enum, "MAX")) != -1)
+			if (enumObj.GetIndexByName(GenerateFullEnumName(enumObj, "MAX")) != -1)
 			{
 				return true;
 			}
 
-			string MaxEnumItem = GenerateFullEnumName(Enum, GenerateEnumPrefix(Enum).ToString() + "_MAX");
-			if (Enum.GetIndexByName(MaxEnumItem) != -1)
+			string maxEnumItem = GenerateFullEnumName(enumObj, GenerateEnumPrefix(enumObj).ToString() + "_MAX");
+			if (enumObj.GetIndexByName(maxEnumItem) != -1)
 			{
 				return true;
 			}
 			return false;
 		}
 
-		private static long GetMaxEnumValue(UhtEnum Enum)
+		private static long GetMaxEnumValue(UhtEnum enumObj)
 		{
-			if (Enum.EnumValues.Count == 0)
+			if (enumObj.EnumValues.Count == 0)
 			{
 				return 0;
 			}
 
-			long MaxValue = Enum.EnumValues[0].Value;
-			for (int i = 1; i < Enum.EnumValues.Count; ++i)
+			long maxValue = enumObj.EnumValues[0].Value;
+			for (int i = 1; i < enumObj.EnumValues.Count; ++i)
 			{
-				long CurrentValue = Enum.EnumValues[i].Value;
-				if (CurrentValue > MaxValue)
+				long currentValue = enumObj.EnumValues[i].Value;
+				if (currentValue > maxValue)
 				{
-					MaxValue = CurrentValue;
+					maxValue = currentValue;
 				}
 			}
 
-			return MaxValue;
+			return maxValue;
 		}
 		#endregion
 	}
 
 	static class UhtClassExtensions
 	{
-		public static IEnumerable<UhtProperty> EnumerateReplicatedProperties(this UhtClass Class, bool bIncludeSuper)
+		public static IEnumerable<UhtProperty> EnumerateReplicatedProperties(this UhtClass classObj, bool includeSuper)
 		{
-			if (bIncludeSuper && Class.SuperClass != null)
+			if (includeSuper && classObj.SuperClass != null)
 			{
-				foreach (UhtProperty Property in Class.SuperClass.EnumerateReplicatedProperties(true))
+				foreach (UhtProperty property in classObj.SuperClass.EnumerateReplicatedProperties(true))
 				{
-					yield return Property;
+					yield return property;
 				}
 			}
 
-			foreach (UhtProperty Property in Class.Properties)
+			foreach (UhtProperty property in classObj.Properties)
 			{
-				if (Property.PropertyFlags.HasAnyFlags(EPropertyFlags.Net))
+				if (property.PropertyFlags.HasAnyFlags(EPropertyFlags.Net))
 				{
-					yield return Property;
+					yield return property;
 				}
 			}
 		}
