More coding standard fixes

diff --git a/Engine/Source/Programs/Shared/EpicGames.UHT/Exporters/CodeGen/UhtMacroBlockEmitter.cs b/Engine/Source/Programs/Shared/EpicGames.UHT/Exporters/CodeGen/UhtMacroBlockEmitter.cs
--- a/Engine/Source/Programs/Shared/EpicGames.UHT/Exporters/CodeGen/UhtMacroBlockEmitter.cs
+++ b/Engine/Source/Programs/Shared/EpicGames.UHT/Exporters/CodeGen/UhtMacroBlockEmitter.cs
@@ -7,32 +7,32 @@ namespace EpicGames.UHT.Exporters.CodeGen
 {
 	internal struct UhtMacroBlockEmitter : IDisposable
 	{
-		private readonly StringBuilder Builder;
-		private readonly string MacroName;
-		private bool bEmitted;
+		private readonly StringBuilder _builder;
+		private readonly string _macroName;
+		private bool _bEmitted;
 
-		public UhtMacroBlockEmitter(StringBuilder Builder, string MacroName, bool bInitialState = false)
+		public UhtMacroBlockEmitter(StringBuilder builder, string macroName, bool initialState = false)
 		{
-			this.Builder = Builder;
-			this.MacroName = MacroName;
-			this.bEmitted = false;
-			Set(bInitialState);
+			this._builder = builder;
+			this._macroName = macroName;
+			this._bEmitted = false;
+			Set(initialState);
 		}
 
-		public void Set(bool bEmit)
+		public void Set(bool emit)
 		{
-			if (this.bEmitted == bEmit)
+			if (this._bEmitted == emit)
 			{
 				return;
 			}
-			this.bEmitted = bEmit;
-			if (this.bEmitted)
+			this._bEmitted = emit;
+			if (this._bEmitted)
 			{
-				Builder.Append("#if ").Append(this.MacroName).Append("\r\n");
+				_builder.Append("#if ").Append(this._macroName).Append("\r\n");
 			}
 			else
 			{
-				Builder.Append("#endif // ").Append(this.MacroName).Append("\r\n");
+				_builder.Append("#endif // ").Append(this._macroName).Append("\r\n");
 			}
 		}
 
diff --git a/Engine/Source/Programs/Shared/EpicGames.UHT/Exporters/CodeGen/UhtMacroCreator.cs b/Engine/Source/Programs/Shared/EpicGames.UHT/Exporters/CodeGen/UhtMacroCreator.cs
--- a/Engine/Source/Programs/Shared/EpicGames.UHT/Exporters/CodeGen/UhtMacroCreator.cs
+++ b/Engine/Source/Programs/Shared/EpicGames.UHT/Exporters/CodeGen/UhtMacroCreator.cs
@@ -1,55 +1,55 @@
 // Copyright Epic Games, Inc. All Rights Reserved.
 
-using EpicGames.UHT.Types;
 using System;
 using System.Text;
+using EpicGames.UHT.Types;
 
 namespace EpicGames.UHT.Exporters.CodeGen
 {
 	internal struct UhtMacroCreator : IDisposable
 	{
-		private readonly StringBuilder Builder;
-		private readonly int StartingLength;
+		private readonly StringBuilder _builder;
+		private readonly int _startingLength;
 
-		public UhtMacroCreator(StringBuilder Builder, UhtHeaderCodeGenerator Generator, int LineNumber, string MacroSuffix)
+		public UhtMacroCreator(StringBuilder builder, UhtHeaderCodeGenerator generator, int lineNumber, string macroSuffix)
 		{
-			Builder.Append("#define ").AppendMacroName(Generator, LineNumber, MacroSuffix).Append(" \\\r\n");
-			this.Builder = Builder;
-			this.StartingLength = Builder.Length;
+			builder.Append("#define ").AppendMacroName(generator, lineNumber, macroSuffix).Append(" \\\r\n");
+			this._builder = builder;
+			this._startingLength = builder.Length;
 		}
 
-		public UhtMacroCreator(StringBuilder Builder, UhtHeaderCodeGenerator Generator, UhtClass Class, string MacroSuffix)
+		public UhtMacroCreator(StringBuilder builder, UhtHeaderCodeGenerator generator, UhtClass classObj, string macroSuffix)
 		{
-			Builder.Append("#define ").AppendMacroName(Generator, Class, MacroSuffix).Append(" \\\r\n");
-			this.Builder = Builder;
-			this.StartingLength = Builder.Length;
+			builder.Append("#define ").AppendMacroName(generator, classObj, macroSuffix).Append(" \\\r\n");
+			this._builder = builder;
+			this._startingLength = builder.Length;
 		}
 
-		public UhtMacroCreator(StringBuilder Builder, UhtHeaderCodeGenerator Generator, UhtScriptStruct ScriptStruct, string MacroSuffix)
+		public UhtMacroCreator(StringBuilder builder, UhtHeaderCodeGenerator generator, UhtScriptStruct scriptStruct, string macroSuffix)
 		{
-			Builder.Append("#define ").AppendMacroName(Generator, ScriptStruct, MacroSuffix).Append(" \\\r\n");
-			this.Builder = Builder;
-			this.StartingLength = Builder.Length;
+			builder.Append("#define ").AppendMacroName(generator, scriptStruct, macroSuffix).Append(" \\\r\n");
+			this._builder = builder;
+			this._startingLength = builder.Length;
 		}
 
-		public UhtMacroCreator(StringBuilder Builder, UhtHeaderCodeGenerator Generator, UhtFunction Function, string MacroSuffix)
+		public UhtMacroCreator(StringBuilder builder, UhtHeaderCodeGenerator generator, UhtFunction function, string macroSuffix)
 		{
-			Builder.Append("#define ").AppendMacroName(Generator, Function, MacroSuffix).Append(" \\\r\n");
-			this.Builder = Builder;
-			this.StartingLength = Builder.Length;
+			builder.Append("#define ").AppendMacroName(generator, function, macroSuffix).Append(" \\\r\n");
+			this._builder = builder;
+			this._startingLength = builder.Length;
 		}
 
 		public void Dispose()
 		{
-			int FinalLength = Builder.Length;
-			Builder.Length = Builder.Length - 4;
-			if (FinalLength == StartingLength)
+			int finalLength = _builder.Length;
+			_builder.Length -= 4;
+			if (finalLength == _startingLength)
 			{
-				Builder.Append("\r\n");
+				_builder.Append("\r\n");
 			}
 			else
 			{
-				Builder.Append("\r\n\r\n\r\n");
+				_builder.Append("\r\n\r\n\r\n");
 			}
 		}
 	}
diff --git a/Engine/Source/Programs/Shared/EpicGames.UHT/Exporters/CodeGen/UhtPackageCodeGenerator.cs b/Engine/Source/Programs/Shared/EpicGames.UHT/Exporters/CodeGen/UhtPackageCodeGenerator.cs
--- a/Engine/Source/Programs/Shared/EpicGames.UHT/Exporters/CodeGen/UhtPackageCodeGenerator.cs
+++ b/Engine/Source/Programs/Shared/EpicGames.UHT/Exporters/CodeGen/UhtPackageCodeGenerator.cs
@@ -1,11 +1,11 @@
 // Copyright Epic Games, Inc. All Rights Reserved.
 
-using EpicGames.UHT.Types;
 using System.Text;
+using EpicGames.UHT.Types;
 
 namespace EpicGames.UHT.Exporters.CodeGen
 {
-	internal class UhtPackageCodeGenerator
+	class UhtPackageCodeGenerator
 	{
 		public static string HeaderCopyright =
 			"// Copyright Epic Games, Inc. All Rights Reserved.\r\n" +
@@ -28,16 +28,16 @@ namespace EpicGames.UHT.Exporters.CodeGen
 		public bool SaveExportedHeaders => this.Package.Module.SaveExportedHeaders;
 
 		public Utils.UhtSession Session => this.CodeGenerator.Session;
-		public UhtCodeGenerator.PackageInfo[] PackageInfos => this.CodeGenerator.PackageInfos;
-		public UhtCodeGenerator.HeaderInfo[] HeaderInfos => this.CodeGenerator.HeaderInfos;
-		public UhtCodeGenerator.ObjectInfo[] ObjectInfos => this.CodeGenerator.ObjectInfos;
-		public string PackageApi => this.PackageInfos[this.Package.PackageTypeIndex].Api;
-		public string PackageSingletonName => this.ObjectInfos[this.Package.ObjectTypeIndex].RegisteredSingletonName;
+		public UhtCodeGenerator.PackageInfo[] PackageInfos => this.CodeGenerator._packageInfos;
+		public UhtCodeGenerator.HeaderInfo[] HeaderInfos => this.CodeGenerator._headerInfos;
+		public UhtCodeGenerator.ObjectInfo[] ObjectInfos => this.CodeGenerator._objectInfos;
+		public string PackageApi => this.PackageInfos[this.Package.PackageTypeIndex]._api;
+		public string PackageSingletonName => this.ObjectInfos[this.Package.ObjectTypeIndex]._registeredSingletonName;
 
-		public UhtPackageCodeGenerator(UhtCodeGenerator CodeGenerator, UhtPackage Package)
+		public UhtPackageCodeGenerator(UhtCodeGenerator codeGenerator, UhtPackage package)
 		{
-			this.CodeGenerator = CodeGenerator;
-			this.Package = Package;
+			this.CodeGenerator = codeGenerator;
+			this.Package = package;
 		}
 
 		#region Utility functions
@@ -45,68 +45,68 @@ namespace EpicGames.UHT.Exporters.CodeGen
 		/// <summary>
 		/// Return the singleton name for an object
 		/// </summary>
-		/// <param name="Object">The object in question.</param>
-		/// <param name="bRegistered">If true, return the registered singleton name.  Otherwise return the unregistered.</param>
+		/// <param name="obj">The object in question.</param>
+		/// <param name="registered">If true, return the registered singleton name.  Otherwise return the unregistered.</param>
 		/// <returns>Singleton name of "nullptr" if Object is null</returns>
-		public string GetSingletonName(UhtObject? Object, bool bRegistered)
+		public string GetSingletonName(UhtObject? obj, bool registered)
 		{
-			return this.CodeGenerator.GetSingletonName(Object, bRegistered);
+			return this.CodeGenerator.GetSingletonName(obj, registered);
 		}
 
 		/// <summary>
 		/// Return the external declaration for an object
 		/// </summary>
-		/// <param name="Object">The object in question.</param>
-		/// <param name="bRegistered">If true, return the registered external declaration.  Otherwise return the unregistered.</param>
+		/// <param name="obj">The object in question.</param>
+		/// <param name="registered">If true, return the registered external declaration.  Otherwise return the unregistered.</param>
 		/// <returns>External declaration</returns>
-		public string GetExternalDecl(UhtObject Object, bool bRegistered)
+		public string GetExternalDecl(UhtObject obj, bool registered)
 		{
-			return this.CodeGenerator.GetExternalDecl(Object, bRegistered);
+			return this.CodeGenerator.GetExternalDecl(obj, registered);
 		}
 
 		/// <summary>
 		/// Return the external declaration for an object
 		/// </summary>
-		/// <param name="ObjectIndex">The object in question.</param>
-		/// <param name="bRegistered">If true, return the registered external declaration.  Otherwise return the unregistered.</param>
+		/// <param name="objectIndex">The object in question.</param>
+		/// <param name="registered">If true, return the registered external declaration.  Otherwise return the unregistered.</param>
 		/// <returns>External declaration</returns>
-		public string GetExternalDecl(int ObjectIndex, bool bRegistered)
+		public string GetExternalDecl(int objectIndex, bool registered)
 		{
-			return this.CodeGenerator.GetExternalDecl(ObjectIndex, bRegistered);
+			return this.CodeGenerator.GetExternalDecl(objectIndex, registered);
 		}
 
 		/// <summary>
 		/// Return the cross reference for an object
 		/// </summary>
-		/// <param name="Object">The object in question.</param>
-		/// <param name="bRegistered">If true, return the registered cross reference.  Otherwise return the unregistered.</param>
+		/// <param name="obj">The object in question.</param>
+		/// <param name="registered">If true, return the registered cross reference.  Otherwise return the unregistered.</param>
 		/// <returns>Cross reference</returns>
-		public string GetCrossReference(UhtObject Object, bool bRegistered)
+		public string GetCrossReference(UhtObject obj, bool registered)
 		{
-			return this.CodeGenerator.GetCrossReference(Object, bRegistered);
+			return this.CodeGenerator.GetCrossReference(obj, registered);
 		}
 
 		/// <summary>
 		/// Return the cross reference for an object
 		/// </summary>
-		/// <param name="ObjectIndex">The object in question.</param>
-		/// <param name="bRegistered">If true, return the registered cross reference.  Otherwise return the unregistered.</param>
+		/// <param name="objectIndex">The object in question.</param>
+		/// <param name="registered">If true, return the registered cross reference.  Otherwise return the unregistered.</param>
 		/// <returns>Cross reference</returns>
-		public string GetCrossReference(int ObjectIndex, bool bRegistered)
+		public string GetCrossReference(int objectIndex, bool registered)
 		{
-			return this.CodeGenerator.GetCrossReference(ObjectIndex, bRegistered);
+			return this.CodeGenerator.GetCrossReference(objectIndex, registered);
 		}
 
 		/// <summary>
 		/// Test to see if the given field is a delegate function
 		/// </summary>
-		/// <param name="Field">Field to be tested</param>
+		/// <param name="field">Field to be tested</param>
 		/// <returns>True if the field is a delegate function</returns>
-		public static bool IsDelegateFunction(UhtField Field)
+		public static bool IsDelegateFunction(UhtField field)
 		{
-			if (Field is UhtFunction Function)
+			if (field is UhtFunction function)
 			{
-				return Function.FunctionType.IsDelegate();
+				return function.FunctionType.IsDelegate();
 			}
 			return false;
 		}
@@ -123,12 +123,13 @@ namespace EpicGames.UHT.Exporters.CodeGen
 		/// <param name="A">Hash to merge</param>
 		/// <param name="C">Previously combined hash</param>
 		/// <returns>Resulting hash value</returns>
+		[System.Diagnostics.CodeAnalysis.SuppressMessage("Style", "IDE1006:Naming Styles", Justification = "<Pending>")]
+		[System.Diagnostics.CodeAnalysis.SuppressMessage("Style", "IDE2001:Embedded statements must be on their own line", Justification = "<Pending>")]
 		public static uint HashCombine(uint A, uint C)
 		{
 			uint B = 0x9e3779b9;
 			A += B;
 
-#pragma warning disable IDE2001 // Embedded statements must be on their own line
 			A -= B; A -= C; A ^= (C >> 13);
 			B -= C; B -= A; B ^= (A << 8);
 			C -= A; C -= B; C ^= (B >> 13);
@@ -138,7 +139,6 @@ namespace EpicGames.UHT.Exporters.CodeGen
 			A -= B; A -= C; A ^= (C >> 3);
 			B -= C; B -= A; B ^= (A << 10);
 			C -= A; C -= B; C ^= (B >> 15);
-#pragma warning restore IDE2001 // Embedded statements must be on their own line
 
 			return C;
 		}
@@ -147,22 +147,22 @@ namespace EpicGames.UHT.Exporters.CodeGen
 
 	internal static class UhtPackageCodeGeneratorStringBuilderExtensions
 	{
-		public static StringBuilder AppendBeginEditorOnlyGuard(this StringBuilder Builder, bool bEnable = true)
+		public static StringBuilder AppendBeginEditorOnlyGuard(this StringBuilder builder, bool enable = true)
 		{
-			if (bEnable)
+			if (enable)
 			{
-				Builder.Append(UhtPackageCodeGenerator.BeginEditorOnlyGuard);
+				builder.Append(UhtPackageCodeGenerator.BeginEditorOnlyGuard);
 			}
-			return Builder;
+			return builder;
 		}
 
-		public static StringBuilder AppendEndEditorOnlyGuard(this StringBuilder Builder, bool bEnable = true)
+		public static StringBuilder AppendEndEditorOnlyGuard(this StringBuilder builder, bool enable = true)
 		{
-			if (bEnable)
+			if (enable)
 			{
-				Builder.Append(UhtPackageCodeGenerator.EndEditorOnlyGuard);
+				builder.Append(UhtPackageCodeGenerator.EndEditorOnlyGuard);
 			}
-			return Builder;
+			return builder;
 		}
 	}
 }
diff --git a/Engine/Source/Programs/Shared/EpicGames.UHT/Exporters/CodeGen/UhtPackageCodeGeneratorCppFile.cs b/Engine/Source/Programs/Shared/EpicGames.UHT/Exporters/CodeGen/UhtPackageCodeGeneratorCppFile.cs
--- a/Engine/Source/Programs/Shared/EpicGames.UHT/Exporters/CodeGen/UhtPackageCodeGeneratorCppFile.cs
+++ b/Engine/Source/Programs/Shared/EpicGames.UHT/Exporters/CodeGen/UhtPackageCodeGeneratorCppFile.cs
@@ -1,11 +1,11 @@
 // Copyright Epic Games, Inc. All Rights Reserved.
 
-using EpicGames.Core;
-using EpicGames.UHT.Types;
-using EpicGames.UHT.Utils;
 using System;
 using System.Collections.Generic;
 using System.Text;
+using EpicGames.Core;
+using EpicGames.UHT.Types;
+using EpicGames.UHT.Utils;
 
 namespace EpicGames.UHT.Exporters.CodeGen
 {
@@ -14,137 +14,137 @@ namespace EpicGames.UHT.Exporters.CodeGen
 		/// <summary>
 		/// Construct an instance of this generator object
 		/// </summary>
-		/// <param name="CodeGenerator">The base code generator</param>
-		/// <param name="Package">Package being generated</param>
-		public UhtPackageCodeGeneratorCppFile(UhtCodeGenerator CodeGenerator, UhtPackage Package)
-			: base(CodeGenerator, Package)
+		/// <param name="codeGenerator">The base code generator</param>
+		/// <param name="package">Package being generated</param>
+		public UhtPackageCodeGeneratorCppFile(UhtCodeGenerator codeGenerator, UhtPackage package)
+			: base(codeGenerator, package)
 		{
 		}
 
 		/// <summary>
 		/// For a given UE header file, generated the generated cpp file
 		/// </summary>
-		/// <param name="Factory">Requesting factory</param>
-		/// <param name="PackageSortedHeaders">Sorted list of headers by name of all headers in the package</param>
-		public void Generate(IUhtExportFactory Factory, List<UhtHeaderFile> PackageSortedHeaders)
+		/// <param name="factory">Requesting factory</param>
+		/// <param name="packageSortedHeaders">Sorted list of headers by name of all headers in the package</param>
+		public void Generate(IUhtExportFactory factory, List<UhtHeaderFile> packageSortedHeaders)
 		{
-			using (BorrowStringBuilder Borrower = new BorrowStringBuilder(StringBuilderCache.Big))
+			using (BorrowStringBuilder borrower = new BorrowStringBuilder(StringBuilderCache.Big))
 			{
-				StringBuilder Builder = Borrower.StringBuilder;
 				const string MetaDataParamsName = "Package_MetaDataParams";
+				StringBuilder builder = borrower.StringBuilder;
 
 				// Collect information from all of the headers
-				List<UhtField> Singletons = new List<UhtField>();
-				StringBuilder Declarations = new StringBuilder();
-				uint BodiesHash = 0;
-				foreach (UhtHeaderFile HeaderFile in PackageSortedHeaders)
+				List<UhtField> singletons = new List<UhtField>();
+				StringBuilder declarations = new StringBuilder();
+				uint bodiesHash = 0;
+				foreach (UhtHeaderFile headerFile in packageSortedHeaders)
 				{
-					ref UhtCodeGenerator.HeaderInfo HeaderInfo = ref this.HeaderInfos[HeaderFile.HeaderFileTypeIndex];
-					ReadOnlyMemory<string> Sorted = HeaderFile.References.Declaration.GetSortedReferences(
-						(int ObjectIndex, bool bRegistered) => GetExternalDecl(ObjectIndex, bRegistered));
-					foreach (string Declaration in Sorted.Span)
+					ref UhtCodeGenerator.HeaderInfo headerInfo = ref this.HeaderInfos[headerFile.HeaderFileTypeIndex];
+					ReadOnlyMemory<string> sorted = headerFile.References.Declaration.GetSortedReferences(
+						(int objectIndex, bool registered) => GetExternalDecl(objectIndex, registered));
+					foreach (string declaration in sorted.Span)
 					{
-						Declarations.Append(Declaration);
+						declarations.Append(declaration);
 					}
 
-					Singletons.AddRange(HeaderFile.References.Singletons);
+					singletons.AddRange(headerFile.References.Singletons);
 
-					uint BodyHash = this.HeaderInfos[HeaderFile.HeaderFileTypeIndex].BodyHash;
-					if (BodiesHash == 0)
+					uint bodyHash = this.HeaderInfos[headerFile.HeaderFileTypeIndex]._bodyHash;
+					if (bodiesHash == 0)
 					{
 						// Don't combine in the first case because it keeps GUID backwards compatibility
-						BodiesHash = BodyHash;
+						bodiesHash = bodyHash;
 					}
 					else
 					{
-						BodiesHash = HashCombine(BodyHash, BodiesHash);
+						bodiesHash = HashCombine(bodyHash, bodiesHash);
 					}
 				}
 
 				// No need to generate output if we have no declarations
-				if (Declarations.Length == 0)
+				if (declarations.Length == 0)
 				{
 					return;
 				}
-				uint DeclarationsHash = UhtHash.GenenerateTextHash(Declarations.ToString());
+				uint declarationsHash = UhtHash.GenenerateTextHash(declarations.ToString());
 
-				string StrippedName = this.PackageInfos[this.Package.PackageTypeIndex].StrippedName;
-				string SingletonName = this.PackageSingletonName;
-				Singletons.Sort((UhtField Lhs, UhtField Rhs) =>
+				string strippedName = this.PackageInfos[this.Package.PackageTypeIndex]._strippedName;
+				string singletonName = this.PackageSingletonName;
+				singletons.Sort((UhtField lhs, UhtField rhs) =>
 				{
-					bool bLhsIsDel = IsDelegateFunction(Lhs);
-					bool bRhsIsDel = IsDelegateFunction(Rhs);
-					if (bLhsIsDel != bRhsIsDel)
+					bool lhsIsDel = IsDelegateFunction(lhs);
+					bool rhsIsDel = IsDelegateFunction(rhs);
+					if (lhsIsDel != rhsIsDel)
 					{
-						return !bLhsIsDel ? -1 : 1;
+						return !lhsIsDel ? -1 : 1;
 					}
 					return StringComparerUE.OrdinalIgnoreCase.Compare(
-						this.ObjectInfos[Lhs.ObjectTypeIndex].RegisteredSingletonName,
-						this.ObjectInfos[Rhs.ObjectTypeIndex].RegisteredSingletonName);
+						this.ObjectInfos[lhs.ObjectTypeIndex]._registeredSingletonName,
+						this.ObjectInfos[rhs.ObjectTypeIndex]._registeredSingletonName);
 				});
 
-				Builder.Append(HeaderCopyright);
-				Builder.Append(RequiredCPPIncludes);
-				Builder.Append(DisableDeprecationWarnings).Append("\r\n");
-				Builder.Append("void EmptyLinkFunctionForGeneratedCode").Append(this.Package.ShortName).Append("_init() {}\r\n");
+				builder.Append(HeaderCopyright);
+				builder.Append(RequiredCPPIncludes);
+				builder.Append(DisableDeprecationWarnings).Append("\r\n");
+				builder.Append("void EmptyLinkFunctionForGeneratedCode").Append(this.Package.ShortName).Append("_init() {}\r\n");
 
-				if (this.Session.bIncludeDebugOutput)
+				if (this.Session.IncludeDebugOutput)
 				{
-					Builder.Append("#if 0\r\n");
-					Builder.Append(Declarations);
-					Builder.Append("#endif\r\n");
+					builder.Append("#if 0\r\n");
+					builder.Append(declarations);
+					builder.Append("#endif\r\n");
 				}
 
-				foreach (UhtObject Object in Singletons)
+				foreach (UhtObject obj in singletons)
 				{
-					ref UhtCodeGenerator.ObjectInfo Info = ref this.ObjectInfos[Object.ObjectTypeIndex];
-					Builder.Append(Info.RegsiteredExternalDecl);
+					ref UhtCodeGenerator.ObjectInfo info = ref this.ObjectInfos[obj.ObjectTypeIndex];
+					builder.Append(info._regsiteredExternalDecl);
 				}
 
-				Builder.AppendMetaDataDef(this.Package, null, MetaDataParamsName, 3);
+				builder.AppendMetaDataDef(this.Package, null, MetaDataParamsName, 3);
 
-				Builder.Append("\tstatic FPackageRegistrationInfo Z_Registration_Info_UPackage_").Append(StrippedName).Append(";\r\n");
-				Builder.Append("\tFORCENOINLINE UPackage* ").Append(SingletonName).Append("()\r\n");
-				Builder.Append("\t{\r\n");
-				Builder.Append("\t\tif (!Z_Registration_Info_UPackage_").Append(StrippedName).Append(".OuterSingleton)\r\n");
-				Builder.Append("\t\t{\r\n");
+				builder.Append("\tstatic FPackageRegistrationInfo Z_Registration_Info_UPackage_").Append(strippedName).Append(";\r\n");
+				builder.Append("\tFORCENOINLINE UPackage* ").Append(singletonName).Append("()\r\n");
+				builder.Append("\t{\r\n");
+				builder.Append("\t\tif (!Z_Registration_Info_UPackage_").Append(strippedName).Append(".OuterSingleton)\r\n");
+				builder.Append("\t\t{\r\n");
 
-				if (Singletons.Count != 0)
+				if (singletons.Count != 0)
 				{
-					Builder.Append("\t\t\tstatic UObject* (*const SingletonFuncArray[])() = {\r\n");
-					foreach (UhtField Field in Singletons)
+					builder.Append("\t\t\tstatic UObject* (*const SingletonFuncArray[])() = {\r\n");
+					foreach (UhtField field in singletons)
 					{
-						Builder.Append("\t\t\t\t(UObject* (*)())").Append(this.ObjectInfos[Field.ObjectTypeIndex].RegisteredSingletonName).Append(",\r\n");
+						builder.Append("\t\t\t\t(UObject* (*)())").Append(this.ObjectInfos[field.ObjectTypeIndex]._registeredSingletonName).Append(",\r\n");
 					}
-					Builder.Append("\t\t\t};\r\n");
+					builder.Append("\t\t\t};\r\n");
 				}
 
-				EPackageFlags Flags = this.Package.PackageFlags & (EPackageFlags.ClientOptional | EPackageFlags.ServerSideOnly | EPackageFlags.EditorOnly | EPackageFlags.Developer | EPackageFlags.UncookedOnly);
-				Builder.Append("\t\t\tstatic const UECodeGen_Private::FPackageParams PackageParams = {\r\n");
-				Builder.Append("\t\t\t\t").AppendUTF8LiteralString(this.Package.SourceName).Append(",\r\n");
-				Builder.Append("\t\t\t\t").Append(Singletons.Count != 0 ? "SingletonFuncArray" : "nullptr").Append(",\r\n");
-				Builder.Append("\t\t\t\t").Append(Singletons.Count != 0 ? "UE_ARRAY_COUNT(SingletonFuncArray)" : "0").Append(",\r\n");
-				Builder.Append("\t\t\t\tPKG_CompiledIn | ").Append($"0x{(uint)Flags:X8}").Append(",\r\n");
-				Builder.Append("\t\t\t\t").Append($"0x{BodiesHash:X8}").Append(",\r\n");
-				Builder.Append("\t\t\t\t").Append($"0x{DeclarationsHash:X8}").Append(",\r\n");
-				Builder.Append("\t\t\t\t").AppendMetaDataParams(this.Package, null, MetaDataParamsName).Append("\r\n");
-				Builder.Append("\t\t\t};\r\n");
-				Builder.Append("\t\t\tUECodeGen_Private::ConstructUPackage(Z_Registration_Info_UPackage_").Append(StrippedName).Append(".OuterSingleton, PackageParams);\r\n");
-				Builder.Append("\t\t}\r\n");
-				Builder.Append("\t\treturn Z_Registration_Info_UPackage_").Append(StrippedName).Append(".OuterSingleton;\r\n");
-				Builder.Append("\t}\r\n");
+				EPackageFlags flags = this.Package.PackageFlags & (EPackageFlags.ClientOptional | EPackageFlags.ServerSideOnly | EPackageFlags.EditorOnly | EPackageFlags.Developer | EPackageFlags.UncookedOnly);
+				builder.Append("\t\t\tstatic const UECodeGen_Private::FPackageParams PackageParams = {\r\n");
+				builder.Append("\t\t\t\t").AppendUTF8LiteralString(this.Package.SourceName).Append(",\r\n");
+				builder.Append("\t\t\t\t").Append(singletons.Count != 0 ? "SingletonFuncArray" : "nullptr").Append(",\r\n");
+				builder.Append("\t\t\t\t").Append(singletons.Count != 0 ? "UE_ARRAY_COUNT(SingletonFuncArray)" : "0").Append(",\r\n");
+				builder.Append("\t\t\t\tPKG_CompiledIn | ").Append($"0x{(uint)flags:X8}").Append(",\r\n");
+				builder.Append("\t\t\t\t").Append($"0x{bodiesHash:X8}").Append(",\r\n");
+				builder.Append("\t\t\t\t").Append($"0x{declarationsHash:X8}").Append(",\r\n");
+				builder.Append("\t\t\t\t").AppendMetaDataParams(this.Package, null, MetaDataParamsName).Append("\r\n");
+				builder.Append("\t\t\t};\r\n");
+				builder.Append("\t\t\tUECodeGen_Private::ConstructUPackage(Z_Registration_Info_UPackage_").Append(strippedName).Append(".OuterSingleton, PackageParams);\r\n");
+				builder.Append("\t\t}\r\n");
+				builder.Append("\t\treturn Z_Registration_Info_UPackage_").Append(strippedName).Append(".OuterSingleton;\r\n");
+				builder.Append("\t}\r\n");
 
 				// Do not change the Z_CompiledInDeferPackage_UPackage_ without changing LC_SymbolPatterns
-				Builder.Append("\tstatic FRegisterCompiledInInfo Z_CompiledInDeferPackage_UPackage_").Append(StrippedName).Append('(').Append(SingletonName)
-					.Append(", TEXT(\"").Append(this.Package.SourceName).Append("\"), Z_Registration_Info_UPackage_").Append(StrippedName).Append(", CONSTRUCT_RELOAD_VERSION_INFO(FPackageReloadVersionInfo, ")
-					.Append($"0x{BodiesHash:X8}, 0x{DeclarationsHash:X8}").Append("));\r\n");
+				builder.Append("\tstatic FRegisterCompiledInInfo Z_CompiledInDeferPackage_UPackage_").Append(strippedName).Append('(').Append(singletonName)
+					.Append(", TEXT(\"").Append(this.Package.SourceName).Append("\"), Z_Registration_Info_UPackage_").Append(strippedName).Append(", CONSTRUCT_RELOAD_VERSION_INFO(FPackageReloadVersionInfo, ")
+					.Append($"0x{bodiesHash:X8}, 0x{declarationsHash:X8}").Append("));\r\n");
 
-				Builder.Append(EnableDeprecationWarnings).Append("\r\n");
+				builder.Append(EnableDeprecationWarnings).Append("\r\n");
 
 				if (this.SaveExportedHeaders)
 				{
-					string CppFilePath = Factory.MakePath(this.Package, ".init.gen.cpp");
-					Factory.CommitOutput(CppFilePath, Builder);
+					string cppFilePath = factory.MakePath(this.Package, ".init.gen.cpp");
+					factory.CommitOutput(cppFilePath, builder);
 				}
 			}
 		}
diff --git a/Engine/Source/Programs/Shared/EpicGames.UHT/Exporters/CodeGen/UhtPackageCodeGeneratorHFile.cs b/Engine/Source/Programs/Shared/EpicGames.UHT/Exporters/CodeGen/UhtPackageCodeGeneratorHFile.cs
--- a/Engine/Source/Programs/Shared/EpicGames.UHT/Exporters/CodeGen/UhtPackageCodeGeneratorHFile.cs
+++ b/Engine/Source/Programs/Shared/EpicGames.UHT/Exporters/CodeGen/UhtPackageCodeGeneratorHFile.cs
@@ -1,11 +1,11 @@
 // Copyright Epic Games, Inc. All Rights Reserved.
 
-using EpicGames.Core;
-using EpicGames.UHT.Types;
-using EpicGames.UHT.Utils;
 using System.Collections.Generic;
 using System.Linq;
 using System.Text;
+using EpicGames.Core;
+using EpicGames.UHT.Types;
+using EpicGames.UHT.Utils;
 
 namespace EpicGames.UHT.Exporters.CodeGen
 {
@@ -14,57 +14,57 @@ namespace EpicGames.UHT.Exporters.CodeGen
 		/// <summary>
 		/// Construct an instance of this generator object
 		/// </summary>
-		/// <param name="CodeGenerator">The base code generator</param>
-		/// <param name="Package">Package being generated</param>
-		public UhtPackageCodeGeneratorHFile(UhtCodeGenerator CodeGenerator, UhtPackage Package)
-			: base(CodeGenerator, Package)
+		/// <param name="codeGenerator">The base code generator</param>
+		/// <param name="package">Package being generated</param>
+		public UhtPackageCodeGeneratorHFile(UhtCodeGenerator codeGenerator, UhtPackage package)
+			: base(codeGenerator, package)
 		{
 		}
 
 		/// <summary>
 		/// For a given UE header file, generated the generated H file
 		/// </summary>
-		/// <param name="Factory">Requesting factory</param>
-		/// <param name="PackageSortedHeaders">Sorted list of headers by name of all headers in the package</param>
-		public void Generate(IUhtExportFactory Factory, List<UhtHeaderFile> PackageSortedHeaders)
+		/// <param name="factory">Requesting factory</param>
+		/// <param name="packageSortedHeaders">Sorted list of headers by name of all headers in the package</param>
+		public void Generate(IUhtExportFactory factory, List<UhtHeaderFile> packageSortedHeaders)
 		{
-			using (BorrowStringBuilder Borrower = new BorrowStringBuilder(StringBuilderCache.Big))
+			using (BorrowStringBuilder borrower = new BorrowStringBuilder(StringBuilderCache.Big))
 			{
-				StringBuilder Builder = Borrower.StringBuilder;
+				StringBuilder builder = borrower.StringBuilder;
 
-				Builder.Append(HeaderCopyright);
-				Builder.Append("#pragma once\r\n");
-				Builder.Append("\r\n");
-				Builder.Append("\r\n");
+				builder.Append(HeaderCopyright);
+				builder.Append("#pragma once\r\n");
+				builder.Append("\r\n");
+				builder.Append("\r\n");
 
-				List<UhtHeaderFile> HeaderFiles = new List<UhtHeaderFile>(this.Package.Children.Count * 2);
-				HeaderFiles.AddRange(PackageSortedHeaders);
+				List<UhtHeaderFile> headerFiles = new List<UhtHeaderFile>(this.Package.Children.Count * 2);
+				headerFiles.AddRange(packageSortedHeaders);
 
-				foreach (UhtHeaderFile HeaderFile in this.Package.Children)
+				foreach (UhtHeaderFile headerFile in this.Package.Children)
 				{
-					if (HeaderFile.HeaderFileType == UhtHeaderFileType.Classes)
+					if (headerFile.HeaderFileType == UhtHeaderFileType.Classes)
 					{
-						HeaderFiles.Add(HeaderFile);
+						headerFiles.Add(headerFile);
 					}
 				}
 
-				List<UhtHeaderFile> SortedHeaderFiles = new List<UhtHeaderFile>(HeaderFiles.Distinct());
-				SortedHeaderFiles.Sort((x, y) => StringComparerUE.OrdinalIgnoreCase.Compare(x.FilePath, y.FilePath));
+				List<UhtHeaderFile> sortedHeaderFiles = new List<UhtHeaderFile>(headerFiles.Distinct());
+				sortedHeaderFiles.Sort((x, y) => StringComparerUE.OrdinalIgnoreCase.Compare(x.FilePath, y.FilePath));
 
-				foreach (UhtHeaderFile HeaderFile in SortedHeaderFiles)
+				foreach (UhtHeaderFile headerFile in sortedHeaderFiles)
 				{
-					if (HeaderFile.HeaderFileType == UhtHeaderFileType.Classes)
+					if (headerFile.HeaderFileType == UhtHeaderFileType.Classes)
 					{
-						Builder.Append("#include \"").Append(this.HeaderInfos[HeaderFile.HeaderFileTypeIndex].IncludePath).Append("\"\r\n");
+						builder.Append("#include \"").Append(this.HeaderInfos[headerFile.HeaderFileTypeIndex]._includePath).Append("\"\r\n");
 					}
 				}
 
-				Builder.Append("\r\n");
+				builder.Append("\r\n");
 
 				if (this.SaveExportedHeaders)
 				{
-					string HeaderFilePath = Factory.MakePath(this.Package, "Classes.h");
-					Factory.CommitOutput(HeaderFilePath, Builder);
+					string headerFilePath = factory.MakePath(this.Package, "Classes.h");
+					factory.CommitOutput(headerFilePath, builder);
 				}
 			}
 		}
diff --git a/Engine/Source/Programs/Shared/EpicGames.UHT/Exporters/CodeGen/UhtRigVMStringBuilderExtensions.cs b/Engine/Source/Programs/Shared/EpicGames.UHT/Exporters/CodeGen/UhtRigVMStringBuilderExtensions.cs
--- a/Engine/Source/Programs/Shared/EpicGames.UHT/Exporters/CodeGen/UhtRigVMStringBuilderExtensions.cs
+++ b/Engine/Source/Programs/Shared/EpicGames.UHT/Exporters/CodeGen/UhtRigVMStringBuilderExtensions.cs
@@ -1,116 +1,116 @@
 // Copyright Epic Games, Inc. All Rights Reserved.
 
-using EpicGames.UHT.Types;
 using System;
 using System.Collections.Generic;
 using System.Text;
+using EpicGames.UHT.Types;
 
 namespace EpicGames.UHT.Exporters.CodeGen
 {
 	internal static class UhtRigVMStringBuilderExtensions
 	{
-		public static StringBuilder AppendArgumentsName(this StringBuilder Builder, UhtStruct Struct, UhtRigVMMethodInfo MethodInfo)
+		public static StringBuilder AppendArgumentsName(this StringBuilder builder, UhtStruct structObj, UhtRigVMMethodInfo methodInfo)
 		{
-			return Builder.Append("Arguments_").Append(Struct.SourceName).Append('_').Append(MethodInfo.Name);
+			return builder.Append("Arguments_").Append(structObj.SourceName).Append('_').Append(methodInfo.Name);
 		}
 
-		public static StringBuilder AppendParameterNames(this StringBuilder Builder, IEnumerable<UhtRigVMParameter> Parameters, 
-			bool bLeadingSeparator = false, string Separator = ", ", bool bCastName = false, bool bIncludeEditorOnly = false)
+		public static StringBuilder AppendParameterNames(this StringBuilder builder, IEnumerable<UhtRigVMParameter> parameters,
+			bool leadingSeparator = false, string separator = ", ", bool castName = false, bool includeEditorOnly = false)
 		{
-			bool bEmitSeparator = bLeadingSeparator;
-			foreach (UhtRigVMParameter Parameter in Parameters)
+			bool emitSeparator = leadingSeparator;
+			foreach (UhtRigVMParameter parameter in parameters)
 			{
-				if (bIncludeEditorOnly && !Parameter.bEditorOnly)
+				if (includeEditorOnly && !parameter.EditorOnly)
 				{
 					continue;
 				}
-				if (bEmitSeparator)
+				if (emitSeparator)
 				{
-					Builder.Append(Separator);
+					builder.Append(separator);
 				}
-				bEmitSeparator = true;
-				Builder.Append(Parameter.NameOriginal(bCastName));
+				emitSeparator = true;
+				builder.Append(parameter.NameOriginal(castName));
 			}
-			return Builder;
+			return builder;
 		}
 
-		public static StringBuilder AppendTypeConstRef(this StringBuilder Builder, UhtRigVMParameter Parameter, bool bCastType = false)
+		public static StringBuilder AppendTypeConstRef(this StringBuilder builder, UhtRigVMParameter parameter, bool castType = false)
 		{
-			Builder.Append("const ");
-			string String = Parameter.TypeOriginal(bCastType);
-			ReadOnlySpan<char> Span = String.AsSpan();
-			if (Span.EndsWith("&"))
+			builder.Append("const ");
+			string paramterType = parameter.TypeOriginal(castType);
+			ReadOnlySpan<char> parameterTypeSpan = paramterType.AsSpan();
+			if (parameterTypeSpan.EndsWith("&"))
 			{
-				Span = Span.Slice(0, Span.Length - 1);
+				parameterTypeSpan = parameterTypeSpan.Slice(0, parameterTypeSpan.Length - 1);
 			}
-			Builder.Append(Span);
-			if (Span.Length > 0 && (Span[0] == 'T' || Span[0] == 'F'))
+			builder.Append(parameterTypeSpan);
+			if (parameterTypeSpan.Length > 0 && (parameterTypeSpan[0] == 'T' || parameterTypeSpan[0] == 'F'))
 			{
-				Builder.Append('&');
+				builder.Append('&');
 			}
-			return Builder;
+			return builder;
 		}
 
-		public static StringBuilder AppendTypeRef(this StringBuilder Builder, UhtRigVMParameter Parameter, bool bCastType = false)
+		public static StringBuilder AppendTypeRef(this StringBuilder builder, UhtRigVMParameter parameter, bool castType = false)
 		{
-			string String = Parameter.TypeOriginal(bCastType);
-			ReadOnlySpan<char> Span = String.AsSpan();
-			if (Span.EndsWith("&"))
+			string paramterType = parameter.TypeOriginal(castType);
+			ReadOnlySpan<char> parameterTypeSpan = paramterType.AsSpan();
+			if (parameterTypeSpan.EndsWith("&"))
 			{
-				Builder.Append(Span);
+				builder.Append(parameterTypeSpan);
 			}
 			else
 			{
-				Builder.Append(Span).Append('&');
+				builder.Append(parameterTypeSpan).Append('&');
 			}
-			return Builder;
+			return builder;
 		}
 
-		public static StringBuilder AppendTypeNoRef(this StringBuilder Builder, UhtRigVMParameter Parameter, bool bCastType = false)
+		public static StringBuilder AppendTypeNoRef(this StringBuilder builder, UhtRigVMParameter parameter, bool castType = false)
 		{
-			string String = Parameter.TypeOriginal(bCastType);
-			ReadOnlySpan<char> Span = String.AsSpan();
-			if (Span.EndsWith("&"))
+			string paramterType = parameter.TypeOriginal(castType);
+			ReadOnlySpan<char> parameterTypeSpan = paramterType.AsSpan();
+			if (parameterTypeSpan.EndsWith("&"))
 			{
-				Span = Span.Slice(0, Span.Length - 1);
+				parameterTypeSpan = parameterTypeSpan.Slice(0, parameterTypeSpan.Length - 1);
 			}
-			return Builder.Append(Span);
+			return builder.Append(parameterTypeSpan);
 		}
 
-		public static StringBuilder AppendTypeVariableRef(this StringBuilder Builder, UhtRigVMParameter Parameter, bool bCastType = false)
+		public static StringBuilder AppendTypeVariableRef(this StringBuilder builder, UhtRigVMParameter parameter, bool castType = false)
 		{
-			if (Parameter.IsConst())
+			if (parameter.IsConst())
 			{
-				Builder.AppendTypeConstRef(Parameter, bCastType);
+				builder.AppendTypeConstRef(parameter, castType);
 			}
 			else
 			{
-				Builder.AppendTypeRef(Parameter, bCastType);
+				builder.AppendTypeRef(parameter, castType);
 			}
-			return Builder;
+			return builder;
 		}
 
-		public static StringBuilder AppendParameterDecls(this StringBuilder Builder, IEnumerable<UhtRigVMParameter> Parameters, 
-			bool bLeadingSeparator = false, string Separator = ", ", bool bCastType = false, bool bCastName = false, bool bIncludeEditorOnly = false)
+		public static StringBuilder AppendParameterDecls(this StringBuilder builder, IEnumerable<UhtRigVMParameter> parameters,
+			bool leadingSeparator = false, string separator = ", ", bool castType = false, bool castName = false, bool includeEditorOnly = false)
 		{
-			bool bEmitSeparator = bLeadingSeparator;
-			foreach (UhtRigVMParameter Parameter in Parameters)
+			bool emitSeparator = leadingSeparator;
+			foreach (UhtRigVMParameter parameter in parameters)
 			{
-				if (bIncludeEditorOnly && !Parameter.bEditorOnly)
+				if (includeEditorOnly && !parameter.EditorOnly)
 				{
 					continue;
 				}
-				if (bEmitSeparator)
+				if (emitSeparator)
 				{
-					Builder.Append(Separator);
+					builder.Append(separator);
 				}
-				bEmitSeparator = true;
-				Builder
-					.AppendTypeVariableRef(Parameter, bCastType)
+				emitSeparator = true;
+				builder
+					.AppendTypeVariableRef(parameter, castType)
 					.Append(' ')
-					.Append(bCastName ? Parameter.CastName : Parameter.Name);
+					.Append(castName ? parameter.CastName : parameter.Name);
 			}
-			return Builder;
+			return builder;
 		}
 	}
 }
diff --git a/Engine/Source/Programs/Shared/EpicGames.UHT/Exporters/Json/UhtJsonExporter.cs b/Engine/Source/Programs/Shared/EpicGames.UHT/Exporters/Json/UhtJsonExporter.cs
--- a/Engine/Source/Programs/Shared/EpicGames.UHT/Exporters/Json/UhtJsonExporter.cs
+++ b/Engine/Source/Programs/Shared/EpicGames.UHT/Exporters/Json/UhtJsonExporter.cs
@@ -1,60 +1,60 @@
 // Copyright Epic Games, Inc. All Rights Reserved.
 
-using EpicGames.Core;
-using EpicGames.UHT.Tables;
-using EpicGames.UHT.Types;
-using EpicGames.UHT.Utils;
 using System.Collections.Generic;
 using System.Text.Json;
 using System.Text.Json.Serialization;
 using System.Threading.Tasks;
+using EpicGames.Core;
+using EpicGames.UHT.Tables;
+using EpicGames.UHT.Types;
+using EpicGames.UHT.Utils;
 
 namespace EpicGames.UHT.Exporters.Json
 {
 	[UnrealHeaderTool]
-	internal class UhtJsonExporter
+	class UhtJsonExporter
 	{
 
 		[UhtExporter(Name = "Json", Description = "Json description of packages", Options = UhtExporterOptions.None)]
-		public static void JsonExporter(IUhtExportFactory Factory)
+		public static void JsonExporter(IUhtExportFactory factory)
 		{
-			new UhtJsonExporter(Factory).Export();
+			new UhtJsonExporter(factory).Export();
 		}
 
 		public readonly IUhtExportFactory Factory;
 		public UhtSession Session => Factory.Session;
 
-		private UhtJsonExporter(IUhtExportFactory Factory)
+		private UhtJsonExporter(IUhtExportFactory factory)
 		{
-			this.Factory = Factory;
+			this.Factory = factory;
 		}
 
 		private void Export()
 		{
 			// Generate the files for the packages
-			List<Task?> GeneratedPackages = new List<Task?>(this.Session.PackageTypeCount);
-			foreach (UhtPackage Package in this.Session.Packages)
+			List<Task?> generatedPackages = new List<Task?>(this.Session.PackageTypeCount);
+			foreach (UhtPackage package in this.Session.Packages)
 			{
-				UHTManifest.Module Module = Package.Module;
-				GeneratedPackages.Add(Factory.CreateTask(
-					(IUhtExportFactory Factory) =>
+				UHTManifest.Module module = package.Module;
+				generatedPackages.Add(Factory.CreateTask(
+					(IUhtExportFactory factory) =>
 					{
-						string JsonPath = Factory.MakePath(Package, ".json");
-						JsonSerializerOptions Options = new JsonSerializerOptions { WriteIndented = true, DefaultIgnoreCondition = JsonIgnoreCondition.WhenWritingNull };
-						Factory.CommitOutput(JsonPath, JsonSerializer.Serialize(Package, Options));
+						string jsonPath = factory.MakePath(package, ".json");
+						JsonSerializerOptions options = new JsonSerializerOptions { WriteIndented = true, DefaultIgnoreCondition = JsonIgnoreCondition.WhenWritingNull };
+						factory.CommitOutput(jsonPath, JsonSerializer.Serialize(package, options));
 					}));
 			}
 
 			// Wait for all the packages to complete
-			List<Task> PackageTasks = new List<Task>(this.Session.PackageTypeCount);
-			foreach (Task? Output in GeneratedPackages)
+			List<Task> packageTasks = new List<Task>(this.Session.PackageTypeCount);
+			foreach (Task? output in generatedPackages)
 			{
-				if (Output != null)
+				if (output != null)
 				{
-					PackageTasks.Add(Output);
+					packageTasks.Add(output);
 				}
 			}
-			Task.WaitAll(PackageTasks.ToArray());
+			Task.WaitAll(packageTasks.ToArray());
 		}
 	}
 }
diff --git a/Engine/Source/Programs/Shared/EpicGames.UHT/Exporters/Stats/UhtStatsExporter.cs b/Engine/Source/Programs/Shared/EpicGames.UHT/Exporters/Stats/UhtStatsExporter.cs
--- a/Engine/Source/Programs/Shared/EpicGames.UHT/Exporters/Stats/UhtStatsExporter.cs
+++ b/Engine/Source/Programs/Shared/EpicGames.UHT/Exporters/Stats/UhtStatsExporter.cs
@@ -1,11 +1,11 @@
 // Copyright Epic Games, Inc. All Rights Reserved.
 
+using System.Collections.Generic;
+using System.Diagnostics.CodeAnalysis;
 using EpicGames.Core;
 using EpicGames.UHT.Tables;
 using EpicGames.UHT.Types;
 using EpicGames.UHT.Utils;
-using System.Collections.Generic;
-using System.Diagnostics.CodeAnalysis;
 
 namespace EpicGames.UHT.Exporters.Stats
 {
@@ -14,37 +14,37 @@ namespace EpicGames.UHT.Exporters.Stats
 	{
 		[UhtExporter(Name = "Stats", Description = "Type Stats", Options = UhtExporterOptions.None)]
 		[SuppressMessage("CodeQuality", "IDE0051:Remove unused private members", Justification = "Attribute accessed method")]
-		private static void StatsExporter(IUhtExportFactory Factory)
+		private static void StatsExporter(IUhtExportFactory factory)
 		{
-			SortedDictionary<string, int> CountByType = new SortedDictionary<string, int>();
-			foreach (UhtType Type in Factory.Session.Packages)
+			SortedDictionary<string, int> countByType = new SortedDictionary<string, int>();
+			foreach (UhtType type in factory.Session.Packages)
 			{
-				Collect(CountByType, Type);
+				Collect(countByType, type);
 			}
 
 			Log.TraceInformation("Counts by type:");
 
-			foreach (KeyValuePair<string, int> Kvp in CountByType)
+			foreach (KeyValuePair<string, int> kvp in countByType)
 			{
-				Log.TraceInformation($"{Kvp.Key} {Kvp.Value}");
+				Log.TraceInformation($"{kvp.Key} {kvp.Value}");
 			}
 			Log.TraceInformation("");
 		}
 
-		private static void Collect(SortedDictionary<string, int> CountByType, UhtType Type)
+		private static void Collect(SortedDictionary<string, int> countByType, UhtType type)
 		{
-			if (CountByType.TryGetValue(Type.EngineClassName, out int Count))
+			if (countByType.TryGetValue(type.EngineClassName, out int count))
 			{
-				CountByType[Type.EngineClassName] = Count + 1;
+				countByType[type.EngineClassName] = count + 1;
 			}
 			else
 			{
-				CountByType[Type.EngineClassName] = 1;
+				countByType[type.EngineClassName] = 1;
 			}
 
-			foreach (UhtType Child in Type.Children)
+			foreach (UhtType child in type.Children)
 			{
-				Collect(CountByType, Child);
+				Collect(countByType, child);
 			}
 		}
 	}
diff --git a/Engine/Source/Programs/Shared/EpicGames.UHT/Parsers/UhtClassBaseParser.cs b/Engine/Source/Programs/Shared/EpicGames.UHT/Parsers/UhtClassBaseParser.cs
--- a/Engine/Source/Programs/Shared/EpicGames.UHT/Parsers/UhtClassBaseParser.cs
+++ b/Engine/Source/Programs/Shared/EpicGames.UHT/Parsers/UhtClassBaseParser.cs
@@ -1,10 +1,10 @@
 // Copyright Epic Games, Inc. All Rights Reserved.
 
+using System.Collections.Generic;
+using System.Text.Json.Serialization;
 using EpicGames.Core;
 using EpicGames.UHT.Tokenizer;
 using EpicGames.UHT.Types;
-using System.Collections.Generic;
-using System.Text.Json.Serialization;
 
 namespace EpicGames.UHT.Parsers
 {
@@ -31,17 +31,17 @@ namespace EpicGames.UHT.Parsers
 		/// <summary>
 		/// Construct a new base class parser
 		/// </summary>
-		/// <param name="Outer">Outer type</param>
-		/// <param name="LineNumber">Line number of class</param>
-		public UhtClassBaseParser(UhtType Outer, int LineNumber) : base(Outer, LineNumber)
+		/// <param name="outer">Outer type</param>
+		/// <param name="lineNumber">Line number of class</param>
+		public UhtClassBaseParser(UhtType outer, int lineNumber) : base(outer, lineNumber)
 		{
 		}
 
 		/// <inheritdoc/>
-		protected override void ResolveSuper(UhtResolvePhase ResolvePhase)
+		protected override void ResolveSuper(UhtResolvePhase resolvePhase)
 		{
-			base.ResolveSuper(ResolvePhase);
-			switch (ResolvePhase)
+			base.ResolveSuper(resolvePhase);
+			switch (resolvePhase)
 			{
 				case UhtResolvePhase.Bases:
 					BindAndResolveSuper(this.SuperIdentifier, UhtFindOptions.Class);
@@ -57,17 +57,17 @@ namespace EpicGames.UHT.Parsers
 		}
 
 		/// <inheritdoc/>
-		protected override bool ResolveSelf(UhtResolvePhase ResolvePhase)
+		protected override bool ResolveSelf(UhtResolvePhase resolvePhase)
 		{
-			bool bResult = base.ResolveSelf(ResolvePhase);
+			bool result = base.ResolveSelf(resolvePhase);
 
-			switch (ResolvePhase)
+			switch (resolvePhase)
 			{
 				case UhtResolvePhase.Properties:
 					UhtPropertyParser.ResolveChildren(this, UhtPropertyParseOptions.AddModuleRelativePath);
 					break;
 			}
-			return bResult;
+			return result;
 		}
 	}
 }
diff --git a/Engine/Source/Programs/Shared/EpicGames.UHT/Parsers/UhtClassParser.cs b/Engine/Source/Programs/Shared/EpicGames.UHT/Parsers/UhtClassParser.cs
--- a/Engine/Source/Programs/Shared/EpicGames.UHT/Parsers/UhtClassParser.cs
+++ b/Engine/Source/Programs/Shared/EpicGames.UHT/Parsers/UhtClassParser.cs
@@ -1,14 +1,14 @@
 // Copyright Epic Games, Inc. All Rights Reserved.
 
+using System;
+using System.Collections.Generic;
+using System.Diagnostics.CodeAnalysis;
+using System.Text;
 using EpicGames.Core;
 using EpicGames.UHT.Tables;
 using EpicGames.UHT.Tokenizer;
 using EpicGames.UHT.Types;
 using EpicGames.UHT.Utils;
-using System;
-using System.Collections.Generic;
-using System.Diagnostics.CodeAnalysis;
-using System.Text;
 
 namespace EpicGames.UHT.Parsers
 {
@@ -76,63 +76,63 @@ namespace EpicGames.UHT.Parsers
 		/// <summary>
 		/// Add the given class flags
 		/// </summary>
-		/// <param name="Flags">Flags to add</param>
-		public void AddClassFlags(EClassFlags Flags)
+		/// <param name="flags">Flags to add</param>
+		public void AddClassFlags(EClassFlags flags)
 		{
-			this.ClassFlags |= Flags;
-			this.RemovedClassFlags &= ~Flags;
+			this.ClassFlags |= flags;
+			this.RemovedClassFlags &= ~flags;
 		}
 
 		/// <summary>
 		/// Remove the given class flags
 		/// </summary>
-		/// <param name="Flags">Flags to remove</param>
-		public void RemoveClassFlags(EClassFlags Flags)
+		/// <param name="flags">Flags to remove</param>
+		public void RemoveClassFlags(EClassFlags flags)
 		{
-			this.RemovedClassFlags |= Flags;
-			this.ClassFlags &= ~Flags;
+			this.RemovedClassFlags |= flags;
+			this.ClassFlags &= ~flags;
 		}
 
 		/// <summary>
 		/// Construct a new class parser
 		/// </summary>
-		/// <param name="Outer">Outer type</param>
-		/// <param name="LineNumber">Line number</param>
-		public UhtClassParser(UhtType Outer, int LineNumber) : base(Outer, LineNumber)
+		/// <param name="outer">Outer type</param>
+		/// <param name="lineNumber">Line number</param>
+		public UhtClassParser(UhtType outer, int lineNumber) : base(outer, lineNumber)
 		{
 		}
 
 		/// <inheritdoc/>
-		protected override void ResolveSuper(UhtResolvePhase ResolvePhase)
+		protected override void ResolveSuper(UhtResolvePhase resolvePhase)
 		{
-			base.ResolveSuper(ResolvePhase);
+			base.ResolveSuper(resolvePhase);
 
 			// Make sure the class within is resolved.  We have to make sure we don't try to resolve ourselves.
 			if (this.ClassWithin != null && this.ClassWithin != this)
 			{
-				this.ClassWithin.Resolve(ResolvePhase);
+				this.ClassWithin.Resolve(resolvePhase);
 			}
 
-			switch (ResolvePhase)
+			switch (resolvePhase)
 			{
 				case UhtResolvePhase.Bases:
-					UhtClass? SuperClass = this.SuperClass;
+					UhtClass? superClass = this.SuperClass;
 
 					// Merge the super class flags
-					if (SuperClass != null)
+					if (superClass != null)
 					{
-						this.ClassFlags |= SuperClass.ClassFlags & EClassFlags.ScriptInherit & ~this.RemovedClassFlags;
+						this.ClassFlags |= superClass.ClassFlags & EClassFlags.ScriptInherit & ~this.RemovedClassFlags;
 					}
 
-					foreach (UhtStruct Base in this.Bases)
+					foreach (UhtStruct baseStruct in this.Bases)
 					{
-						if (Base is UhtClass BaseClass)
+						if (baseStruct is UhtClass baseClass)
 						{
-							if (!BaseClass.ClassFlags.HasAnyFlags(EClassFlags.Interface))
+							if (!baseClass.ClassFlags.HasAnyFlags(EClassFlags.Interface))
 							{
-								this.LogError($"Class '{BaseClass.SourceName}' is not an interface; Can only inherit from non-UObjects or UInterface derived interfaces");
+								this.LogError($"Class '{baseClass.SourceName}' is not an interface; Can only inherit from non-UObjects or UInterface derived interfaces");
 							}
-							this.ClassFlags |= BaseClass.ClassFlags & EClassFlags.ScriptInherit & ~this.RemovedClassFlags;
+							this.ClassFlags |= baseClass.ClassFlags & EClassFlags.ScriptInherit & ~this.RemovedClassFlags;
 						}
 					}
 
@@ -149,7 +149,7 @@ namespace EpicGames.UHT.Parsers
 					// Merge with categories inherited from the parent.
 					MergeCategories();
 
-					SetAndValidateWithinClass(ResolvePhase);
+					SetAndValidateWithinClass(resolvePhase);
 					SetAndValidateConfigName();
 
 					// Copy all of the lists back to the meta data
@@ -173,65 +173,65 @@ namespace EpicGames.UHT.Parsers
 
 			// Merge ShowFunctions and HideFunctions
 			AppendStringListMetaData(this.SuperClass, UhtNames.HideFunctions, this.HideFunctions);
-			foreach (string Value in this.ShowFunctions)
+			foreach (string value in this.ShowFunctions)
 			{
-				this.HideFunctions.RemoveSwap(Value);
+				this.HideFunctions.RemoveSwap(value);
 			}
 			this.ShowFunctions.Clear();
 
 			// Merge AutoExpandCategories and AutoCollapseCategories (we still want to keep AutoExpandCategories though!)
-			List<string> ParentAutoExpandCategories = GetStringListMetaData(this.SuperClass, UhtNames.AutoExpandCategories);
-			List<string> ParentAutoCollapseCategories = GetStringListMetaData(this.SuperClass, UhtNames.AutoCollapseCategories);
+			List<string> parentAutoExpandCategories = GetStringListMetaData(this.SuperClass, UhtNames.AutoExpandCategories);
+			List<string> parentAutoCollapseCategories = GetStringListMetaData(this.SuperClass, UhtNames.AutoCollapseCategories);
 
-			foreach (string Value in this.AutoExpandCategories)
+			foreach (string value in this.AutoExpandCategories)
 			{
-				this.AutoCollapseCategories.RemoveSwap(Value);
-				ParentAutoCollapseCategories.RemoveSwap(Value);
+				this.AutoCollapseCategories.RemoveSwap(value);
+				parentAutoCollapseCategories.RemoveSwap(value);
 			}
 
 			// Do the same as above but the other way around
-			foreach (string Value in this.AutoCollapseCategories)
+			foreach (string value in this.AutoCollapseCategories)
 			{
-				this.AutoExpandCategories.RemoveSwap(Value);
-				ParentAutoExpandCategories.RemoveSwap(Value);
+				this.AutoExpandCategories.RemoveSwap(value);
+				parentAutoExpandCategories.RemoveSwap(value);
 			}
 
 			// Once AutoExpandCategories and AutoCollapseCategories for THIS class have been parsed, add the parent inherited categories
-			this.AutoCollapseCategories.AddRange(ParentAutoCollapseCategories);
-			this.AutoExpandCategories.AddRange(ParentAutoExpandCategories);
+			this.AutoCollapseCategories.AddRange(parentAutoCollapseCategories);
+			this.AutoExpandCategories.AddRange(parentAutoExpandCategories);
 		}
 
 		private void MergeShowCategories()
 		{
 
 			// Add the super class hide categories and prime the output show categories with the parent
-			List<string> OutShowCategories = GetStringListMetaData(this.SuperClass, UhtNames.ShowCategories);
+			List<string> outShowCategories = GetStringListMetaData(this.SuperClass, UhtNames.ShowCategories);
 			AppendStringListMetaData(this.SuperClass, UhtNames.HideCategories, this.HideCategories);
 
 			// If this class has new show categories, then merge them
 			if (this.ShowCategories.Count != 0)
 			{
-				StringBuilder SubCategoryPath = new StringBuilder();
-				foreach (string Value in this.ShowCategories)
+				StringBuilder subCategoryPath = new StringBuilder();
+				foreach (string value in this.ShowCategories)
 				{
 
 					// if we didn't find this specific category path in the HideCategories metadata
-					if (!this.HideCategories.RemoveSwap(Value))
+					if (!this.HideCategories.RemoveSwap(value))
 					{
-						string[] SubCategories = Value.ToString().Split('|', StringSplitOptions.RemoveEmptyEntries);
+						string[] subCategories = value.ToString().Split('|', StringSplitOptions.RemoveEmptyEntries);
 
-						SubCategoryPath.Clear();
+						subCategoryPath.Clear();
 						// look to see if any of the parent paths are excluded in the HideCategories list
-						for (int CategoryPathIndex = 0; CategoryPathIndex < SubCategories.Length - 1; ++CategoryPathIndex)
+						for (int categoryPathIndex = 0; categoryPathIndex < subCategories.Length - 1; ++categoryPathIndex)
 						{
-							SubCategoryPath.Append(SubCategories[CategoryPathIndex]);
+							subCategoryPath.Append(subCategories[categoryPathIndex]);
 							// if we're hiding a parent category, then we need to flag this sub category for show
-							if (HideCategories.Contains(SubCategoryPath.ToString()))
+							if (HideCategories.Contains(subCategoryPath.ToString()))
 							{
-								OutShowCategories.AddUnique(Value);
+								outShowCategories.AddUnique(value);
 								break;
 							}
-							SubCategoryPath.Append('|');
+							subCategoryPath.Append('|');
 						}
 					}
 				}
@@ -239,19 +239,19 @@ namespace EpicGames.UHT.Parsers
 
 			// Replace the show categories
 			this.ShowCategories.Clear();
-			this.ShowCategories.AddRange(OutShowCategories);
+			this.ShowCategories.AddRange(outShowCategories);
 		}
 
-		private void SetAndValidateWithinClass(UhtResolvePhase ResolvePhase)
+		private void SetAndValidateWithinClass(UhtResolvePhase resolvePhase)
 		{
 			// The class within must be a child of any super class within
-			UhtClass ExpectedClassWithin = this.SuperClass != null ? this.SuperClass.ClassWithin : this.Session.UObject;
+			UhtClass expectedClassWithin = this.SuperClass != null ? this.SuperClass.ClassWithin : this.Session.UObject;
 
 			// Process all of the class specifiers
-			if (!string.IsNullOrEmpty(this.ClassWithinIdentifier))
+			if (!String.IsNullOrEmpty(this.ClassWithinIdentifier))
 			{
-				UhtClass? SpecifiedClassWithin = (UhtClass?)this.Session.FindType(null, UhtFindOptions.EngineName | UhtFindOptions.Class, this.ClassWithinIdentifier);
-				if (SpecifiedClassWithin == null)
+				UhtClass? specifiedClassWithin = (UhtClass?)this.Session.FindType(null, UhtFindOptions.EngineName | UhtFindOptions.Class, this.ClassWithinIdentifier);
+				if (specifiedClassWithin == null)
 				{
 					this.LogError($"Within class '{this.ClassWithinIdentifier}' not found");
 				}
@@ -260,22 +260,22 @@ namespace EpicGames.UHT.Parsers
 
 					// Make sure the class within has been resolved to this phase.  We don't need to worry about the super since we know 
 					// that it has already been resolved.
-					if (SpecifiedClassWithin != this)
+					if (specifiedClassWithin != this)
 					{
-						SpecifiedClassWithin.Resolve(ResolvePhase);
+						specifiedClassWithin.Resolve(resolvePhase);
 					}
 
-					if (SpecifiedClassWithin.IsChildOf(this.Session.UInterface))
+					if (specifiedClassWithin.IsChildOf(this.Session.UInterface))
 					{
 						this.LogError("Classes cannot be 'within' interfaces");
 					}
-					else if (!SpecifiedClassWithin.IsChildOf(ExpectedClassWithin))
+					else if (!specifiedClassWithin.IsChildOf(expectedClassWithin))
 					{
-						this.LogError($"Cannot override within class with '{SpecifiedClassWithin.SourceName}' since it isn't a child of parent class's expected within '{ExpectedClassWithin.SourceName}'");
+						this.LogError($"Cannot override within class with '{specifiedClassWithin.SourceName}' since it isn't a child of parent class's expected within '{expectedClassWithin.SourceName}'");
 					}
 					else
 					{
-						this.ClassWithin = SpecifiedClassWithin;
+						this.ClassWithin = specifiedClassWithin;
 					}
 				}
 			}
@@ -283,7 +283,7 @@ namespace EpicGames.UHT.Parsers
 			// If we don't have a class within set, then just use the expected within
 			else
 			{
-				this.ClassWithin = ExpectedClassWithin;
+				this.ClassWithin = expectedClassWithin;
 			}
 		}
 
@@ -299,9 +299,9 @@ namespace EpicGames.UHT.Parsers
 			// Set the class config flag if any properties have config
 			if (!this.ClassFlags.HasAnyFlags(EClassFlags.Config))
 			{
-				foreach (UhtProperty Property in this.Properties)
+				foreach (UhtProperty property in this.Properties)
 				{
-					if (Property.PropertyFlags.HasAnyFlags(EPropertyFlags.Config))
+					if (property.PropertyFlags.HasAnyFlags(EPropertyFlags.Config))
 					{
 						this.ClassFlags |= EClassFlags.Config;
 						break;
@@ -341,25 +341,25 @@ namespace EpicGames.UHT.Parsers
 			}
 		}
 
-		private static List<string> GetStringListMetaData(UhtType? Type, string Key)
+		private static List<string> GetStringListMetaData(UhtType? type, string key)
 		{
-			List<string> Out = new List<string>();
-			AppendStringListMetaData(Type, Key, Out);
-			return Out;
+			List<string> outStrings = new List<string>();
+			AppendStringListMetaData(type, key, outStrings);
+			return outStrings;
 		}
 
-		private static void AppendStringListMetaData(UhtType? Type, string Key, List<string> StringList)
+		private static void AppendStringListMetaData(UhtType? type, string key, List<string> stringList)
 		{
-			if (Type != null)
+			if (type != null)
 			{
-				string[]? Values = Type.MetaData.GetStringArray(Key);
-				if (Values != null)
+				string[]? values = type.MetaData.GetStringArray(key);
+				if (values != null)
 				{
-					foreach (string Value in Values)
+					foreach (string value in values)
 					{
 						//COMPATIBILITY-TODO - TEST - This preserves duplicates that are in old UHT
 						//StringList.AddUnique(Value);
-						StringList.Add(Value);
+						stringList.Add(value);
 					}
 				}
 			}
@@ -369,110 +369,110 @@ namespace EpicGames.UHT.Parsers
 		[UhtKeyword(Extends = UhtTableNames.Global)]
 		[SuppressMessage("CodeQuality", "IDE0051:Remove unused private members", Justification = "Attribute accessed method")]
 		[SuppressMessage("Style", "IDE0060:Remove unused parameter", Justification = "Attribute accessed method")]
-		private static UhtParseResult UCLASSKeyword(UhtParsingScope TopScope, UhtParsingScope ActionScope, ref UhtToken Token)
+		private static UhtParseResult UCLASSKeyword(UhtParsingScope topScope, UhtParsingScope actionScope, ref UhtToken token)
 		{
-			return ParseUClass(TopScope, ref Token);
+			return ParseUClass(topScope, ref token);
 		}
 
 		[UhtKeyword(Extends = UhtTableNames.Class)]
 		[UhtKeyword(Extends = UhtTableNames.Class, Keyword = "GENERATED_BODY")]
 		[SuppressMessage("CodeQuality", "IDE0051:Remove unused private members", Justification = "Attribute accessed method")]
 		[SuppressMessage("Style", "IDE0060:Remove unused parameter", Justification = "Attribute accessed method")]
-		private static UhtParseResult GENERATED_UCLASS_BODYKeyword(UhtParsingScope TopScope, UhtParsingScope ActionScope, ref UhtToken Token)
+		private static UhtParseResult GENERATED_UCLASS_BODYKeyword(UhtParsingScope topScope, UhtParsingScope actionScope, ref UhtToken token)
 		{
-			UhtClass Class = (UhtClass)TopScope.ScopeType;
+			UhtClass classObj = (UhtClass)topScope.ScopeType;
 
-			if (Token.IsValue("GENERATED_BODY"))
+			if (token.IsValue("GENERATED_BODY"))
 			{
-				Class.bHasGeneratedBody = true;
-				Class.GeneratedBodyAccessSpecifier = TopScope.AccessSpecifier;
+				classObj.HasGeneratedBody = true;
+				classObj.GeneratedBodyAccessSpecifier = topScope.AccessSpecifier;
 			}
 			else
 			{
-				TopScope.AccessSpecifier = UhtAccessSpecifier.Public;
+				topScope.AccessSpecifier = UhtAccessSpecifier.Public;
 			}
 
-			UhtParserHelpers.ParseCompileVersionDeclaration(TopScope.TokenReader, TopScope.Session.Config!, Class);
+			UhtParserHelpers.ParseCompileVersionDeclaration(topScope.TokenReader, topScope.Session.Config!, classObj);
 
-			Class.GeneratedBodyLineNumber = TopScope.TokenReader.InputLine;
+			classObj.GeneratedBodyLineNumber = topScope.TokenReader.InputLine;
 
-			TopScope.TokenReader.Optional(';');
+			topScope.TokenReader.Optional(';');
 			return UhtParseResult.Handled;
 		}
 		#endregion
 
-		private static UhtParseResult ParseUClass(UhtParsingScope ParentScope, ref UhtToken Token)
+		private static UhtParseResult ParseUClass(UhtParsingScope parentScope, ref UhtToken token)
 		{
-			UhtClassParser Class = new UhtClassParser(ParentScope.ScopeType, Token.InputLine);
-			using (var TopScope = new UhtParsingScope(ParentScope, Class, ParentScope.Session.GetKeywordTable(UhtTableNames.Class), UhtAccessSpecifier.Private))
+			UhtClassParser classObj = new UhtClassParser(parentScope.ScopeType, token.InputLine);
+			using (UhtParsingScope topScope = new UhtParsingScope(parentScope, classObj, parentScope.Session.GetKeywordTable(UhtTableNames.Class), UhtAccessSpecifier.Private))
 			{
 				const string ScopeName = "class";
 
-				using (var TokenContext = new UhtMessageContext(ScopeName))
+				using (UhtMessageContext tokenContext = new UhtMessageContext(ScopeName))
 				{
 					// Parse the specifiers
-					UhtSpecifierContext SpecifierContext = new UhtSpecifierContext(TopScope, TopScope.TokenReader, Class.MetaData);
-					UhtSpecifierParser Specifiers = TopScope.HeaderParser.GetCachedSpecifierParser(SpecifierContext, ScopeName, ParentScope.Session.GetSpecifierTable(UhtTableNames.Class));
-					Specifiers.ParseSpecifiers();
-					Class.PrologLineNumber = TopScope.TokenReader.InputLine;
-					Class.ClassFlags |= EClassFlags.Native;
+					UhtSpecifierContext specifierContext = new UhtSpecifierContext(topScope, topScope.TokenReader, classObj.MetaData);
+					UhtSpecifierParser specifiers = topScope.HeaderParser.GetCachedSpecifierParser(specifierContext, ScopeName, parentScope.Session.GetSpecifierTable(UhtTableNames.Class));
+					specifiers.ParseSpecifiers();
+					classObj.PrologLineNumber = topScope.TokenReader.InputLine;
+					classObj.ClassFlags |= EClassFlags.Native;
 
-					TopScope.AddFormattedCommentsAsTooltipMetaData();
+					topScope.AddFormattedCommentsAsTooltipMetaData();
 
 					// Consume the class specifier
-					TopScope.TokenReader.Require("class");
+					topScope.TokenReader.Require("class");
 
-					TopScope.TokenReader.SkipAlignasAndDeprecatedMacroIfNecessary();
+					topScope.TokenReader.SkipAlignasAndDeprecatedMacroIfNecessary();
 
 					// Read the class name and possible API macro name
-					UhtToken APIMacroToken;
-					TopScope.TokenReader.TryOptionalAPIMacro(out APIMacroToken);
-					Class.SourceName = TopScope.TokenReader.GetIdentifier().Value.ToString();
+					UhtToken apiMacroToken;
+					topScope.TokenReader.TryOptionalAPIMacro(out apiMacroToken);
+					classObj.SourceName = topScope.TokenReader.GetIdentifier().Value.ToString();
 
 					// Update the context for better error messages
-					TokenContext.Reset($"class '{Class.SourceName}'");
+					tokenContext.Reset($"class '{classObj.SourceName}'");
 
 					// Split the source name into the different parts
-					UhtEngineNameParts NameParts = UhtUtilities.GetEngineNameParts(Class.SourceName);
-					Class.EngineName = NameParts.EngineName.ToString();
+					UhtEngineNameParts nameParts = UhtUtilities.GetEngineNameParts(classObj.SourceName);
+					classObj.EngineName = nameParts.EngineName.ToString();
 
 					// Check for an empty engine name
-					if (Class.EngineName.Length == 0)
+					if (classObj.EngineName.Length == 0)
 					{
-						TopScope.TokenReader.LogError($"When compiling class definition for '{Class.SourceName}', attempting to strip prefix results in an empty name. Did you leave off a prefix?");
+						topScope.TokenReader.LogError($"When compiling class definition for '{classObj.SourceName}', attempting to strip prefix results in an empty name. Did you leave off a prefix?");
 					}
 
 					// Skip optional final keyword
-					TopScope.TokenReader.Optional("final");
+					topScope.TokenReader.Optional("final");
 
 					// Parse the inheritance
-					UhtToken SuperIdentifier;
-					List<UhtToken[]>? BaseIdentifiers;
-					UhtParserHelpers.ParseInheritance(TopScope.TokenReader, TopScope.Session.Config!, out SuperIdentifier, out BaseIdentifiers);
-					Class.SuperIdentifier = SuperIdentifier;
-					Class.BaseIdentifiers = BaseIdentifiers;
+					UhtToken superIdentifier;
+					List<UhtToken[]>? baseIdentifiers;
+					UhtParserHelpers.ParseInheritance(topScope.TokenReader, topScope.Session.Config!, out superIdentifier, out baseIdentifiers);
+					classObj.SuperIdentifier = superIdentifier;
+					classObj.BaseIdentifiers = baseIdentifiers;
 
-					if (APIMacroToken)
+					if (apiMacroToken)
 					{
-						Class.ClassFlags |= EClassFlags.RequiredAPI;
+						classObj.ClassFlags |= EClassFlags.RequiredAPI;
 					}
 
-					Specifiers.ParseDeferred();
+					specifiers.ParseDeferred();
 
-					if (Class.Outer != null)
+					if (classObj.Outer != null)
 					{
-						Class.Outer.AddChild(Class);
+						classObj.Outer.AddChild(classObj);
 					}
 
-					TopScope.AddModuleRelativePathToMetaData();
+					topScope.AddModuleRelativePathToMetaData();
 
-					TopScope.HeaderParser.ParseStatements('{', '}', true);
+					topScope.HeaderParser.ParseStatements('{', '}', true);
 
-					TopScope.TokenReader.Require(';');
+					topScope.TokenReader.Require(';');
 
-					if (Class.GeneratedBodyLineNumber == -1)
+					if (classObj.GeneratedBodyLineNumber == -1)
 					{
-						TopScope.TokenReader.LogError("Expected a GENERATED_BODY() at the start of the class");
+						topScope.TokenReader.LogError("Expected a GENERATED_BODY() at the start of the class");
 					}
 				}
 			}
diff --git a/Engine/Source/Programs/Shared/EpicGames.UHT/Parsers/UhtEnumParser.cs b/Engine/Source/Programs/Shared/EpicGames.UHT/Parsers/UhtEnumParser.cs
--- a/Engine/Source/Programs/Shared/EpicGames.UHT/Parsers/UhtEnumParser.cs
+++ b/Engine/Source/Programs/Shared/EpicGames.UHT/Parsers/UhtEnumParser.cs
@@ -18,9 +18,9 @@ namespace EpicGames.UHT.Parsers
 		/// <summary>
 		/// Construct a new enum parser
 		/// </summary>
-		/// <param name="Outer">Outer type</param>
-		/// <param name="LineNumber">Line number</param>
-		public UhtEnumParser(UhtType Outer, int LineNumber) : base(Outer, LineNumber)
+		/// <param name="outer">Outer type</param>
+		/// <param name="lineNumber">Line number</param>
+		public UhtEnumParser(UhtType outer, int lineNumber) : base(outer, lineNumber)
 		{
 		}
 
@@ -28,138 +28,138 @@ namespace EpicGames.UHT.Parsers
 		[UhtKeyword(Extends = UhtTableNames.Global)]
 		[SuppressMessage("CodeQuality", "IDE0051:Remove unused private members", Justification = "Attribute accessed method")]
 		[SuppressMessage("Style", "IDE0060:Remove unused parameter", Justification = "Attribute accessed method")]
-		private static UhtParseResult UENUMKeyword(UhtParsingScope TopScope, UhtParsingScope ActionScope, ref UhtToken Token)
+		private static UhtParseResult UENUMKeyword(UhtParsingScope topScope, UhtParsingScope actionScope, ref UhtToken token)
 		{
-			return ParseUEnum(TopScope, Token);
+			return ParseUEnum(topScope, token);
 		}
 		#endregion
 
-		private static UhtParseResult ParseUEnum(UhtParsingScope ParentScope, UhtToken KeywordToken)
+		private static UhtParseResult ParseUEnum(UhtParsingScope parentScope, UhtToken keywordToken)
 		{
-			UhtEnumParser Enum = new UhtEnumParser(ParentScope.ScopeType, KeywordToken.InputLine);
-			using (var TopScope = new UhtParsingScope(ParentScope, Enum, ParentScope.Session.GetKeywordTable(UhtTableNames.Enum), UhtAccessSpecifier.Public))
+			UhtEnumParser enumObject = new UhtEnumParser(parentScope.ScopeType, keywordToken.InputLine);
+			using (UhtParsingScope topScope = new UhtParsingScope(parentScope, enumObject, parentScope.Session.GetKeywordTable(UhtTableNames.Enum), UhtAccessSpecifier.Public))
 			{
 				const string ScopeName = "UENUM";
 
-				using (var TokenContext = new UhtMessageContext(ScopeName))
+				using (UhtMessageContext tokenContext = new UhtMessageContext(ScopeName))
 				{
 
 					// Parse the specifiers
-					UhtSpecifierContext SpecifierContext = new UhtSpecifierContext(TopScope, TopScope.TokenReader, Enum.MetaData);
-					UhtSpecifierParser Specifiers = TopScope.HeaderParser.GetCachedSpecifierParser(SpecifierContext, ScopeName, ParentScope.Session.GetSpecifierTable(UhtTableNames.Enum));
-					Specifiers.ParseSpecifiers();
+					UhtSpecifierContext specifierContext = new UhtSpecifierContext(topScope, topScope.TokenReader, enumObject.MetaData);
+					UhtSpecifierParser specifiers = topScope.HeaderParser.GetCachedSpecifierParser(specifierContext, ScopeName, parentScope.Session.GetSpecifierTable(UhtTableNames.Enum));
+					specifiers.ParseSpecifiers();
 
 					// Read the name and the CPP type
-					switch (TopScope.TokenReader.TryOptional(new string[] { "namespace", "enum" }))
+					switch (topScope.TokenReader.TryOptional(new string[] { "namespace", "enum" }))
 					{
 						case 0: // namespace
-							Enum.CppForm = UhtEnumCppForm.Namespaced;
-							TopScope.TokenReader.SkipDeprecatedMacroIfNecessary();
+							enumObject.CppForm = UhtEnumCppForm.Namespaced;
+							topScope.TokenReader.SkipDeprecatedMacroIfNecessary();
 							break;
 						case 1: // enum
-							Enum.CppForm = TopScope.TokenReader.TryOptional(new string[] { "class", "struct" }) >= 0 ? UhtEnumCppForm.EnumClass : UhtEnumCppForm.Regular;
-							TopScope.TokenReader.SkipAlignasAndDeprecatedMacroIfNecessary();
+							enumObject.CppForm = topScope.TokenReader.TryOptional(new string[] { "class", "struct" }) >= 0 ? UhtEnumCppForm.EnumClass : UhtEnumCppForm.Regular;
+							topScope.TokenReader.SkipAlignasAndDeprecatedMacroIfNecessary();
 							break;
 						default:
-							throw new UhtTokenException(TopScope.TokenReader, TopScope.TokenReader.PeekToken(), null);
+							throw new UhtTokenException(topScope.TokenReader, topScope.TokenReader.PeekToken(), null);
 					}
 
-					UhtToken EnumToken = TopScope.TokenReader.GetIdentifier("enumeration name");
+					UhtToken enumToken = topScope.TokenReader.GetIdentifier("enumeration name");
 
-					Enum.SourceName = EnumToken.Value.ToString();
+					enumObject.SourceName = enumToken.Value.ToString();
 
-					Specifiers.ParseFieldMetaData();
-					Specifiers.ParseDeferred();
+					specifiers.ParseFieldMetaData();
+					specifiers.ParseDeferred();
 
-					if (Enum.Outer != null)
+					if (enumObject.Outer != null)
 					{
-						Enum.Outer.AddChild(Enum);
+						enumObject.Outer.AddChild(enumObject);
 					}
 
-					if ((TopScope.HeaderParser.GetCurrentCompositeCompilerDirective() & UhtCompilerDirective.WithEditorOnlyData) != 0)
+					if ((topScope.HeaderParser.GetCurrentCompositeCompilerDirective() & UhtCompilerDirective.WithEditorOnlyData) != 0)
 					{
-						Enum.bIsEditorOnly = true;
+						enumObject.IsEditorOnly = true;
 					}
 
 					// Read base for enum class
-					if (Enum.CppForm == UhtEnumCppForm.EnumClass)
+					if (enumObject.CppForm == UhtEnumCppForm.EnumClass)
 					{
-						if (TopScope.TokenReader.TryOptional(':'))
+						if (topScope.TokenReader.TryOptional(':'))
 						{
-							UhtToken EnumType = TopScope.TokenReader.GetIdentifier("enumeration base");
+							UhtToken enumType = topScope.TokenReader.GetIdentifier("enumeration base");
 
-							UhtEnumUnderlyingType UnderlyingType;
-							if (!System.Enum.TryParse<UhtEnumUnderlyingType>(EnumType.Value.ToString(), out UnderlyingType) || UnderlyingType == UhtEnumUnderlyingType.Unspecified)
+							UhtEnumUnderlyingType underlyingType;
+							if (!System.Enum.TryParse<UhtEnumUnderlyingType>(enumType.Value.ToString(), out underlyingType) || underlyingType == UhtEnumUnderlyingType.Unspecified)
 							{
-								TopScope.TokenReader.LogError(EnumType.InputLine, $"Unsupported enum class base type '{EnumType.Value}'");
+								topScope.TokenReader.LogError(enumType.InputLine, $"Unsupported enum class base type '{enumType.Value}'");
 							}
-							Enum.UnderlyingType = UnderlyingType;
+							enumObject.UnderlyingType = underlyingType;
 						}
 						else
 						{
-							Enum.UnderlyingType = UhtEnumUnderlyingType.Unspecified;
+							enumObject.UnderlyingType = UhtEnumUnderlyingType.Unspecified;
 						}
 					}
 					else
 					{
-						if ((Enum.EnumFlags & EEnumFlags.Flags) != 0)
+						if ((enumObject.EnumFlags & EEnumFlags.Flags) != 0)
 						{
-							TopScope.TokenReader.LogError("The 'Flags' specifier can only be used on enum classes");
+							topScope.TokenReader.LogError("The 'Flags' specifier can only be used on enum classes");
 						}
 					}
 
-					if (Enum.UnderlyingType != UhtEnumUnderlyingType.uint8 && Enum.MetaData.ContainsKey("BlueprintType"))
+					if (enumObject.UnderlyingType != UhtEnumUnderlyingType.uint8 && enumObject.MetaData.ContainsKey("BlueprintType"))
 					{
-						TopScope.TokenReader.LogError("Invalid BlueprintType enum base - currently only uint8 supported");
+						topScope.TokenReader.LogError("Invalid BlueprintType enum base - currently only uint8 supported");
 					}
 
 					//EnumDef.GetDefinitionRange().Start = &Input[InputPos];
 
 					// Get the opening brace
-					TopScope.TokenReader.Require('{');
+					topScope.TokenReader.Require('{');
 
-					switch (Enum.CppForm)
+					switch (enumObject.CppForm)
 					{
 						case UhtEnumCppForm.Namespaced:
 							// Now handle the inner true enum portion
-							TopScope.TokenReader.Require("enum");
-							TopScope.TokenReader.SkipAlignasAndDeprecatedMacroIfNecessary();
+							topScope.TokenReader.Require("enum");
+							topScope.TokenReader.SkipAlignasAndDeprecatedMacroIfNecessary();
 
-							UhtToken InnerEnumToken = TopScope.TokenReader.GetIdentifier("enumeration type name");
+							UhtToken innerEnumToken = topScope.TokenReader.GetIdentifier("enumeration type name");
 
-							TopScope.TokenReader.Require('{');
-							Enum.CppType = $"{Enum.SourceName}::{InnerEnumToken.Value}";
+							topScope.TokenReader.Require('{');
+							enumObject.CppType = $"{enumObject.SourceName}::{innerEnumToken.Value}";
 							break;
 
 						case UhtEnumCppForm.EnumClass:
 						case UhtEnumCppForm.Regular:
-							Enum.CppType = Enum.SourceName;
+							enumObject.CppType = enumObject.SourceName;
 							break;
 					}
 
-					TokenContext.Reset($"UENUM {Enum.SourceName}");
+					tokenContext.Reset($"UENUM {enumObject.SourceName}");
 
-					TopScope.AddModuleRelativePathToMetaData();
-					TopScope.AddFormattedCommentsAsTooltipMetaData();
+					topScope.AddModuleRelativePathToMetaData();
+					topScope.AddFormattedCommentsAsTooltipMetaData();
 
 					// Parse all enum tags
-					bool bHasUnparsedValue = false;
-					long CurrentEnumValue = 0;
+					bool hasUnparsedValue = false;
+					long currentEnumValue = 0;
 
-					TopScope.TokenReader.RequireList('}', ',', true, () =>
+					topScope.TokenReader.RequireList('}', ',', true, () =>
 					{
-						UhtToken TagToken = TopScope.TokenReader.GetIdentifier();
+						UhtToken tagToken = topScope.TokenReader.GetIdentifier();
 
-						StringView FullEnumName;
-						switch (Enum.CppForm)
+						StringView fullEnumName;
+						switch (enumObject.CppForm)
 						{
 							case UhtEnumCppForm.Namespaced:
 							case UhtEnumCppForm.EnumClass:
-								FullEnumName = new StringView($"{Enum.SourceName}::{TagToken.Value}");
+								fullEnumName = new StringView($"{enumObject.SourceName}::{tagToken.Value}");
 								break;
 
 							case UhtEnumCppForm.Regular:
-								FullEnumName = TagToken.Value;
+								fullEnumName = tagToken.Value;
 								break;
 
 							default:
@@ -167,69 +167,69 @@ namespace EpicGames.UHT.Parsers
 						}
 
 						// Save the new tag with a default value.  This will be replaced later
-						int EnumIndex = Enum.AddEnumValue(TagToken.Value.ToString(), 0);
-						TopScope.AddFormattedCommentsAsTooltipMetaData(EnumIndex);
+						int enumIndex = enumObject.AddEnumValue(tagToken.Value.ToString(), 0);
+						topScope.AddFormattedCommentsAsTooltipMetaData(enumIndex);
 
 						// Skip any deprecation
-						TopScope.TokenReader.SkipDeprecatedMacroIfNecessary();
+						topScope.TokenReader.SkipDeprecatedMacroIfNecessary();
 
 						// Try to read an optional explicit enum value specification
-						if (TopScope.TokenReader.TryOptional('='))
+						if (topScope.TokenReader.TryOptional('='))
 						{
-							long ScatchValue;
-							bool bParsedValue = TopScope.TokenReader.TryOptionalConstLong(out ScatchValue);
-							if (bParsedValue)
+							long scatchValue;
+							bool parsedValue = topScope.TokenReader.TryOptionalConstLong(out scatchValue);
+							if (parsedValue)
 							{
-								CurrentEnumValue = ScatchValue;
+								currentEnumValue = scatchValue;
 							}
 							else
 							{
 								// We didn't parse a literal, so set an invalid value
-								CurrentEnumValue = -1;
-								bHasUnparsedValue = true;
+								currentEnumValue = -1;
+								hasUnparsedValue = true;
 							}
 
 							// Skip tokens until we encounter a comma, a closing brace or a UMETA declaration
 							// There are tokens after the initializer so it's not a standalone literal,
 							// so set it to an invalid value.
-							int SkippedCount = TopScope.TokenReader.ConsumeUntil(new string[] { ",", "}", "UMETA" });
-							if (SkippedCount == 0 && !bParsedValue)
+							int skippedCount = topScope.TokenReader.ConsumeUntil(new string[] { ",", "}", "UMETA" });
+							if (skippedCount == 0 && !parsedValue)
 							{
-								throw new UhtTokenException(TopScope.TokenReader, TopScope.TokenReader.PeekToken(), "enumerator initializer");
+								throw new UhtTokenException(topScope.TokenReader, topScope.TokenReader.PeekToken(), "enumerator initializer");
 							}
-							if (SkippedCount > 0)
+							if (skippedCount > 0)
 							{
-								CurrentEnumValue = -1;
-								bHasUnparsedValue = true;
+								currentEnumValue = -1;
+								hasUnparsedValue = true;
 							}
 						}
 
 						// Save the value and auto increment
-						UhtEnumValue Value = Enum.EnumValues[EnumIndex];
-						Value.Value = CurrentEnumValue;
-						Enum.EnumValues[EnumIndex] = Value;
-						if (CurrentEnumValue != -1)
+						UhtEnumValue value = enumObject.EnumValues[enumIndex];
+						value.Value = currentEnumValue;
+						enumObject.EnumValues[enumIndex] = value;
+						if (currentEnumValue != -1)
 						{
-							++CurrentEnumValue;
+							++currentEnumValue;
 						}
 
-						Enum.MetaData.Add(UhtNames.Name, EnumIndex, Enum.EnumValues[EnumIndex].Name.ToString());
+						enumObject.MetaData.Add(UhtNames.Name, enumIndex, enumObject.EnumValues[enumIndex].Name.ToString());
 
 						// check for metadata on this enum value
-						SpecifierContext.MetaNameIndex = EnumIndex;
-						Specifiers.ParseFieldMetaData();
+						specifierContext.MetaNameIndex = enumIndex;
+						specifiers.ParseFieldMetaData();
 					});
 
 					// Trailing brace and semicolon for the enum
-					TopScope.TokenReader.Require(';');
-					if (Enum.CppForm == UhtEnumCppForm.Namespaced)
+					topScope.TokenReader.Require(';');
+					if (enumObject.CppForm == UhtEnumCppForm.Namespaced)
 					{
-						TopScope.TokenReader.Require('}');
+						topScope.TokenReader.Require('}');
 					}
 
-					if (!bHasUnparsedValue && !Enum.IsValidEnumValue(0) && Enum.MetaData.ContainsKey(UhtNames.BlueprintType))
+					if (!hasUnparsedValue && !enumObject.IsValidEnumValue(0) && enumObject.MetaData.ContainsKey(UhtNames.BlueprintType))
 					{
-						Enum.LogWarning($"'{Enum.SourceName}' does not have a 0 entry! (This is a problem when the enum is initialized by default)");
+						enumObject.LogWarning($"'{enumObject.SourceName}' does not have a 0 entry! (This is a problem when the enum is initialized by default)");
 					}
 				}
 
