More coding standard fixes

diff --git a/Engine/Source/Programs/Shared/EpicGames.UHT/Parsers/UhtFunctionParser.cs b/Engine/Source/Programs/Shared/EpicGames.UHT/Parsers/UhtFunctionParser.cs
--- a/Engine/Source/Programs/Shared/EpicGames.UHT/Parsers/UhtFunctionParser.cs
+++ b/Engine/Source/Programs/Shared/EpicGames.UHT/Parsers/UhtFunctionParser.cs
@@ -1,15 +1,15 @@
 // Copyright Epic Games, Inc. All Rights Reserved.
 
-using EpicGames.Core;
-using EpicGames.UHT.Tables;
-using EpicGames.UHT.Tokenizer;
-using EpicGames.UHT.Types;
-using EpicGames.UHT.Utils;
 using System;
 using System.Collections.Generic;
 using System.Diagnostics.CodeAnalysis;
 using System.Linq;
 using System.Text;
+using EpicGames.Core;
+using EpicGames.UHT.Tables;
+using EpicGames.UHT.Tokenizer;
+using EpicGames.UHT.Types;
+using EpicGames.UHT.Utils;
 
 namespace EpicGames.UHT.Parsers
 {
@@ -20,31 +20,31 @@ namespace EpicGames.UHT.Parsers
 	 */
 	struct UhtAdvancedDisplayParameterHandler
 	{
-		private readonly UhtMetaData MetaData;
-		private readonly string[]? ParameterNames;
-		private readonly int NumberLeaveUnmarked;
-		private readonly bool bUseNumber;
-		private int AlreadyLeft;
+		private readonly UhtMetaData _metaData;
+		private readonly string[]? _parameterNames;
+		private readonly int _numberLeaveUnmarked;
+		private readonly bool _bUseNumber;
+		private int _alreadyLeft;
 
-		public UhtAdvancedDisplayParameterHandler(UhtMetaData MetaData)
+		public UhtAdvancedDisplayParameterHandler(UhtMetaData metaData)
 		{
-			this.MetaData = MetaData;
-			this.ParameterNames = null;
-			this.NumberLeaveUnmarked = -1;
-			this.AlreadyLeft = 0;
-			this.bUseNumber = false;
-
-			string? FoundString;
-			if (this.MetaData.TryGetValue(UhtNames.AdvancedDisplay, out FoundString))
+			this._metaData = metaData;
+			this._parameterNames = null;
+			this._numberLeaveUnmarked = -1;
+			this._alreadyLeft = 0;
+			this._bUseNumber = false;
+
+			string? foundString;
+			if (this._metaData.TryGetValue(UhtNames.AdvancedDisplay, out foundString))
 			{
-				this.ParameterNames = FoundString.ToString().Split(',', StringSplitOptions.RemoveEmptyEntries);
-				for (int Index = 0, EndIndex = this.ParameterNames.Length; Index < EndIndex; ++Index)
+				this._parameterNames = foundString.ToString().Split(',', StringSplitOptions.RemoveEmptyEntries);
+				for (int index = 0, endIndex = this._parameterNames.Length; index < endIndex; ++index)
 				{
-					this.ParameterNames[Index] = this.ParameterNames[Index].Trim();
+					this._parameterNames[index] = this._parameterNames[index].Trim();
 				}
-				if (this.ParameterNames.Length == 1)
+				if (this._parameterNames.Length == 1)
 				{
-					bUseNumber = int.TryParse(this.ParameterNames[0], out NumberLeaveUnmarked);
+					_bUseNumber = Int32.TryParse(this._parameterNames[0], out _numberLeaveUnmarked);
 				}
 			}
 		}
@@ -53,30 +53,30 @@ namespace EpicGames.UHT.Parsers
 		 * return if given parameter should be marked as Advance View, 
 		 * the function should be called only once for any parameter
 		 */
-		public bool ShouldMarkParameter(StringView ParameterName)
+		public bool ShouldMarkParameter(StringView parameterName)
 		{
-			if (this.bUseNumber)
+			if (this._bUseNumber)
 			{
-				if (this.NumberLeaveUnmarked < 0)
+				if (this._numberLeaveUnmarked < 0)
 				{
 					return false;
 				}
-				if (this.AlreadyLeft < this.NumberLeaveUnmarked)
+				if (this._alreadyLeft < this._numberLeaveUnmarked)
 				{
-					++this.AlreadyLeft;
+					++this._alreadyLeft;
 					return false;
 				}
 				return true;
 			}
 
-			if (this.ParameterNames == null)
+			if (this._parameterNames == null)
 			{
 				return false;
 			}
 
-			foreach (string Element in this.ParameterNames)
+			foreach (string element in this._parameterNames)
 			{
-				if (ParameterName.Span.Equals(Element, StringComparison.OrdinalIgnoreCase))
+				if (parameterName.Span.Equals(element, StringComparison.OrdinalIgnoreCase))
 				{
 					return true;
 				}
@@ -87,7 +87,7 @@ namespace EpicGames.UHT.Parsers
 		/** return if more parameters can be marked */
 		public bool CanMarkMore()
 		{
-			return this.bUseNumber ? this.NumberLeaveUnmarked > 0 : (this.ParameterNames != null && this.ParameterNames.Length > 0);
+			return this._bUseNumber ? this._numberLeaveUnmarked > 0 : (this._parameterNames != null && this._parameterNames.Length > 0);
 		}
 	}
 
@@ -100,42 +100,42 @@ namespace EpicGames.UHT.Parsers
 		/// <summary>
 		/// Construct a new function parser
 		/// </summary>
-		/// <param name="Outer">Outer object</param>
-		/// <param name="LineNumber">Line number</param>
-		public UhtFunctionParser(UhtType Outer, int LineNumber) : base(Outer, LineNumber)
+		/// <param name="outer">Outer object</param>
+		/// <param name="lineNumber">Line number</param>
+		public UhtFunctionParser(UhtType outer, int lineNumber) : base(outer, lineNumber)
 		{
 		}
 
 		/// <summary>
 		/// True if the function specifier has a getter/setter specified
 		/// </summary>
-		public bool bSawPropertyAccessor { get; set; } = false;
+		public bool SawPropertyAccessor { get; set; } = false;
 
 		/// <inheritdoc/>
-		protected override bool ResolveSelf(UhtResolvePhase ResolvePhase)
+		protected override bool ResolveSelf(UhtResolvePhase resolvePhase)
 		{
-			bool bResult = base.ResolveSelf(ResolvePhase);
-			switch (ResolvePhase)
+			bool result = base.ResolveSelf(resolvePhase);
+			switch (resolvePhase)
 			{
 				case UhtResolvePhase.Properties:
 					UhtPropertyParser.ResolveChildren(this, GetPropertyParseOptions(false));
-					foreach (UhtProperty Property in this.Properties)
+					foreach (UhtProperty property in this.Properties)
 					{
-						if (Property.DefaultValueTokens != null)
+						if (property.DefaultValueTokens != null)
 						{
-							string Key = "CPP_Default_" + Property.EngineName;
-							if (!this.MetaData.ContainsKey(Key))
+							string key = "CPP_Default_" + property.EngineName;
+							if (!this.MetaData.ContainsKey(key))
 							{
-								bool bParsed = false;
+								bool parsed = false;
 								try
 								{
 									// All tokens MUST be consumed from the reader
-									StringBuilder Builder = new StringBuilder();
-									IUhtTokenReader DefaultValueReader = UhtTokenReplayReader.GetThreadInstance(Property, this.HeaderFile.Data.Memory, Property.DefaultValueTokens.ToArray(), UhtTokenType.EndOfDefault);
-									bParsed = Property.SanitizeDefaultValue(DefaultValueReader, Builder) && DefaultValueReader.bIsEOF;
-									if (bParsed)
+									StringBuilder builder = new StringBuilder();
+									IUhtTokenReader defaultValueReader = UhtTokenReplayReader.GetThreadInstance(property, this.HeaderFile.Data.Memory, property.DefaultValueTokens.ToArray(), UhtTokenType.EndOfDefault);
+									parsed = property.SanitizeDefaultValue(defaultValueReader, builder) && defaultValueReader.IsEOF;
+									if (parsed)
 									{
-										this.MetaData.Add(Key, Builder.ToString());
+										this.MetaData.Add(key, builder.ToString());
 									}
 								}
 								catch (Exception)
@@ -143,36 +143,36 @@ namespace EpicGames.UHT.Parsers
 									// Ignore the exception for now
 								}
 
-								if (!bParsed)
+								if (!parsed)
 								{
-									StringView DefaultValueText = new StringView(this.HeaderFile.Data, Property.DefaultValueTokens.First().InputStartPos,
-										Property.DefaultValueTokens.Last().InputEndPos - Property.DefaultValueTokens.First().InputStartPos);
-									Property.LogError($"C++ Default parameter not parsed: {Property.SourceName} '{DefaultValueText}'");
+									StringView defaultValueText = new StringView(this.HeaderFile.Data, property.DefaultValueTokens.First().InputStartPos,
+										property.DefaultValueTokens.Last().InputEndPos - property.DefaultValueTokens.First().InputStartPos);
+									property.LogError($"C++ Default parameter not parsed: {property.SourceName} '{defaultValueText}'");
 								}
 							}
 						}
 					}
 					break;
 			}
-			return bResult;
+			return result;
 		}
 
-		private UhtPropertyParseOptions GetPropertyParseOptions(bool bReturnValue)
+		private UhtPropertyParseOptions GetPropertyParseOptions(bool returnValue)
 		{
 			switch (this.FunctionType)
 			{
 				case UhtFunctionType.Delegate:
 				case UhtFunctionType.SparseDelegate:
-					return (bReturnValue ? UhtPropertyParseOptions.None : UhtPropertyParseOptions.CommaSeparatedName) | UhtPropertyParseOptions.DontAddReturn;
+					return (returnValue ? UhtPropertyParseOptions.None : UhtPropertyParseOptions.CommaSeparatedName) | UhtPropertyParseOptions.DontAddReturn;
 
 				case UhtFunctionType.Function:
-					UhtPropertyParseOptions Options = UhtPropertyParseOptions.DontAddReturn; // Fetch the function name
-					Options |= bReturnValue ? UhtPropertyParseOptions.FunctionNameIncluded : UhtPropertyParseOptions.NameIncluded;
+					UhtPropertyParseOptions options = UhtPropertyParseOptions.DontAddReturn; // Fetch the function name
+					options |= returnValue ? UhtPropertyParseOptions.FunctionNameIncluded : UhtPropertyParseOptions.NameIncluded;
 					if (this.FunctionFlags.HasAllFlags(EFunctionFlags.BlueprintEvent | EFunctionFlags.Native))
 					{
-						Options |= UhtPropertyParseOptions.NoAutoConst;
+						options |= UhtPropertyParseOptions.NoAutoConst;
 					}
-					return Options;
+					return options;
 
 				default:
 					throw new UhtIceException("Unknown enumeration value");
@@ -184,9 +184,9 @@ namespace EpicGames.UHT.Parsers
 		[UhtKeyword(Extends = UhtTableNames.Class)]
 		[SuppressMessage("CodeQuality", "IDE0051:Remove unused private members", Justification = "Attribute accessed method")]
 		[SuppressMessage("Style", "IDE0060:Remove unused parameter", Justification = "Attribute accessed method")]
-		private static UhtParseResult UDELEGATEKeyword(UhtParsingScope TopScope, UhtParsingScope ActionScope, ref UhtToken Token)
+		private static UhtParseResult UDELEGATEKeyword(UhtParsingScope topScope, UhtParsingScope actionScope, ref UhtToken token)
 		{
-			return ParseUDelegate(TopScope, Token, true);
+			return ParseUDelegate(topScope, token, true);
 		}
 
 		[UhtKeyword(Extends = UhtTableNames.Class)]
@@ -194,117 +194,116 @@ namespace EpicGames.UHT.Parsers
 		[UhtKeyword(Extends = UhtTableNames.NativeInterface)]
 		[SuppressMessage("CodeQuality", "IDE0051:Remove unused private members", Justification = "Attribute accessed method")]
 		[SuppressMessage("Style", "IDE0060:Remove unused parameter", Justification = "Attribute accessed method")]
-		private static UhtParseResult UFUNCTIONKeyword(UhtParsingScope TopScope, UhtParsingScope ActionScope, ref UhtToken Token)
+		private static UhtParseResult UFUNCTIONKeyword(UhtParsingScope topScope, UhtParsingScope actionScope, ref UhtToken token)
 		{
-			return ParseUFunction(TopScope, Token);
+			return ParseUFunction(topScope, token);
 		}
 
 		[UhtKeywordCatchAll(Extends = UhtTableNames.Global)]
 		[SuppressMessage("CodeQuality", "IDE0051:Remove unused private members", Justification = "Attribute accessed method")]
-		[SuppressMessage("Style", "IDE0060:Remove unused parameter", Justification = "Attribute accessed method")]
-		private static UhtParseResult ParseCatchAllKeyword(UhtParsingScope TopScope, ref UhtToken Token)
+		private static UhtParseResult ParseCatchAllKeyword(UhtParsingScope topScope, ref UhtToken token)
 		{
-			if (UhtFunctionParser.IsValidateDelegateDeclaration(Token))
+			if (UhtFunctionParser.IsValidateDelegateDeclaration(token))
 			{
-				return ParseUDelegate(TopScope, Token, false);
+				return ParseUDelegate(topScope, token, false);
 			}
 			return UhtParseResult.Unhandled;
 		}
 		#endregion
 
-		private static UhtParseResult ParseUDelegate(UhtParsingScope ParentScope, UhtToken Token, bool bHasSpecifiers)
+		private static UhtParseResult ParseUDelegate(UhtParsingScope parentScope, UhtToken token, bool hasSpecifiers)
 		{
-			UhtFunctionParser Function = new UhtFunctionParser(ParentScope.ScopeType, Token.InputLine);
+			UhtFunctionParser function = new UhtFunctionParser(parentScope.ScopeType, token.InputLine);
 
-			using (var TopScope = new UhtParsingScope(ParentScope, Function, ParentScope.Session.GetKeywordTable(UhtTableNames.Function), UhtAccessSpecifier.Public))
+			using (UhtParsingScope topScope = new UhtParsingScope(parentScope, function, parentScope.Session.GetKeywordTable(UhtTableNames.Function), UhtAccessSpecifier.Public))
 			{
 				const string ScopeName = "delegate declaration";
 
-				using (var TokenContext = new UhtMessageContext(ScopeName))
+				using (UhtMessageContext tokenContext = new UhtMessageContext(ScopeName))
 				{
-					TopScope.AddModuleRelativePathToMetaData();
+					topScope.AddModuleRelativePathToMetaData();
 
-					UhtSpecifierContext SpecifierContext = new UhtSpecifierContext(TopScope, TopScope.TokenReader, Function.MetaData);
-					UhtSpecifierParser Specifiers = TopScope.HeaderParser.GetCachedSpecifierParser(SpecifierContext, ScopeName, ParentScope.Session.GetSpecifierTable(UhtTableNames.Function));
+					UhtSpecifierContext specifierContext = new UhtSpecifierContext(topScope, topScope.TokenReader, function.MetaData);
+					UhtSpecifierParser specifiers = topScope.HeaderParser.GetCachedSpecifierParser(specifierContext, ScopeName, parentScope.Session.GetSpecifierTable(UhtTableNames.Function));
 
 					// If this is a UDELEGATE, parse the specifiers
-					StringView DelegateMacro = new StringView();
-					if (bHasSpecifiers)
+					StringView delegateMacro = new StringView();
+					if (hasSpecifiers)
 					{
-						Specifiers.ParseSpecifiers();
-						Specifiers.ParseDeferred();
-						FinalizeFunctionSpecifiers(Function);
+						specifiers.ParseSpecifiers();
+						specifiers.ParseDeferred();
+						FinalizeFunctionSpecifiers(function);
 
-						UhtToken MacroToken = TopScope.TokenReader.GetToken();
-						if (!IsValidateDelegateDeclaration(MacroToken))
+						UhtToken macroToken = topScope.TokenReader.GetToken();
+						if (!IsValidateDelegateDeclaration(macroToken))
 						{
-							throw new UhtTokenException(TopScope.TokenReader, MacroToken, "delegate macro");
+							throw new UhtTokenException(topScope.TokenReader, macroToken, "delegate macro");
 						}
-						DelegateMacro = MacroToken.Value;
+						delegateMacro = macroToken.Value;
 					}
 					else
 					{
-						DelegateMacro = Token.Value;
+						delegateMacro = token.Value;
 					}
 
 					// Break the delegate declaration macro down into parts
-					bool bHasReturnValue = DelegateMacro.Span.Contains("_RetVal".AsSpan(), StringComparison.Ordinal);
-					bool bDeclaredConst = DelegateMacro.Span.Contains("_Const".AsSpan(), StringComparison.Ordinal);
-					bool bIsMulticast = DelegateMacro.Span.Contains("_MULTICAST".AsSpan(), StringComparison.Ordinal);
-					bool bIsSparse = DelegateMacro.Span.Contains("_SPARSE".AsSpan(), StringComparison.Ordinal);
+					bool hasReturnValue = delegateMacro.Span.Contains("_RetVal".AsSpan(), StringComparison.Ordinal);
+					bool declaredConst = delegateMacro.Span.Contains("_Const".AsSpan(), StringComparison.Ordinal);
+					bool isMulticast = delegateMacro.Span.Contains("_MULTICAST".AsSpan(), StringComparison.Ordinal);
+					bool isSparse = delegateMacro.Span.Contains("_SPARSE".AsSpan(), StringComparison.Ordinal);
 
 					// Determine the parameter count
-					int FoundParamIndex = TopScope.Session.Config!.FindDelegateParameterCount(DelegateMacro);
+					int foundParamIndex = topScope.Session.Config!.FindDelegateParameterCount(delegateMacro);
 
 					// Try reconstructing the string to make sure it matches our expectations
-					string ExpectedOriginalString = string.Format("DECLARE_DYNAMIC{0}{1}_DELEGATE{2}{3}{4}",
-						bIsMulticast ? "_MULTICAST" : "",
-						bIsSparse ? "_SPARSE" : "",
-						bHasReturnValue ? "_RetVal" : "",
-						TopScope.Session.Config!.GetDelegateParameterCountString(FoundParamIndex),
-						bDeclaredConst ? "_Const" : "");
-					if (DelegateMacro != ExpectedOriginalString)
+					string expectedOriginalString = String.Format("DECLARE_DYNAMIC{0}{1}_DELEGATE{2}{3}{4}",
+						isMulticast ? "_MULTICAST" : "",
+						isSparse ? "_SPARSE" : "",
+						hasReturnValue ? "_RetVal" : "",
+						topScope.Session.Config!.GetDelegateParameterCountString(foundParamIndex),
+						declaredConst ? "_Const" : "");
+					if (delegateMacro != expectedOriginalString)
 					{
-						throw new UhtException(TopScope.TokenReader, $"Unable to parse delegate declaration; expected '{ExpectedOriginalString}' but found '{DelegateMacro}'.");
+						throw new UhtException(topScope.TokenReader, $"Unable to parse delegate declaration; expected '{expectedOriginalString}' but found '{delegateMacro}'.");
 					}
 
 					// Multi-cast delegate function signatures are not allowed to have a return value
-					if (bHasReturnValue && bIsMulticast)
+					if (hasReturnValue && isMulticast)
 					{
-						throw new UhtException(TopScope.TokenReader, "Multi-cast delegates function signatures must not return a value");
+						throw new UhtException(topScope.TokenReader, "Multi-cast delegates function signatures must not return a value");
 					}
 
 					// Delegate signature
-					Function.FunctionType = bIsSparse ? UhtFunctionType.SparseDelegate : UhtFunctionType.Delegate;
-					Function.FunctionFlags |= EFunctionFlags.Public | EFunctionFlags.Delegate;
-					if (bIsMulticast)
+					function.FunctionType = isSparse ? UhtFunctionType.SparseDelegate : UhtFunctionType.Delegate;
+					function.FunctionFlags |= EFunctionFlags.Public | EFunctionFlags.Delegate;
+					if (isMulticast)
 					{
-						Function.FunctionFlags |= EFunctionFlags.MulticastDelegate;
+						function.FunctionFlags |= EFunctionFlags.MulticastDelegate;
 					}
 
 					// Now parse the macro body
-					TopScope.TokenReader.Require('(');
+					topScope.TokenReader.Require('(');
 
 					// Parse the return type
-					UhtProperty? ReturnValueProperty = null;
-					if (bHasReturnValue)
+					UhtProperty? returnValueProperty = null;
+					if (hasReturnValue)
 					{
-						TopScope.HeaderParser.GetCachedPropertyParser(TopScope).Parse(EPropertyFlags.None,
-							Function.GetPropertyParseOptions(true), UhtParsePropertyDeclarationStyle.None, UhtPropertyCategory.Return,
-							(UhtParsingScope TopScope, UhtProperty Property, ref UhtToken NameToken, UhtLayoutMacroType LayoutMacroType) =>
+						topScope.HeaderParser.GetCachedPropertyParser(topScope).Parse(EPropertyFlags.None,
+							function.GetPropertyParseOptions(true), UhtParsePropertyDeclarationStyle.None, UhtPropertyCategory.Return,
+							(UhtParsingScope topScope, UhtProperty property, ref UhtToken nameToken, UhtLayoutMacroType layoutMacroType) =>
 							{
-								Property.PropertyFlags |= EPropertyFlags.Parm | EPropertyFlags.OutParm | EPropertyFlags.ReturnParm;
-								ReturnValueProperty = Property;
+								property.PropertyFlags |= EPropertyFlags.Parm | EPropertyFlags.OutParm | EPropertyFlags.ReturnParm;
+								returnValueProperty = property;
 							});
-						TopScope.TokenReader.Require(',');
+						topScope.TokenReader.Require(',');
 					}
 
 					// Skip white spaces to get InputPos exactly on beginning of function name.
-					TopScope.TokenReader.SkipWhitespaceAndComments();
+					topScope.TokenReader.SkipWhitespaceAndComments();
 
 					// Get the delegate name
-					UhtToken FuncNameToken = TopScope.TokenReader.GetIdentifier("name");
-					Function.SourceName = FuncNameToken.Value.ToString();
+					UhtToken funcNameToken = topScope.TokenReader.GetIdentifier("name");
+					function.SourceName = funcNameToken.Value.ToString();
 
 					// If this is a delegate function then go ahead and mangle the name so we don't collide with
 					// actual functions or properties
@@ -312,307 +311,307 @@ namespace EpicGames.UHT.Parsers
 						//@TODO: UCREMOVAL: Eventually this mangling shouldn't occur
 
 						// Remove the leading F
-						if (Function.SourceName[0] != 'F')
+						if (function.SourceName[0] != 'F')
 						{
-							TopScope.TokenReader.LogError("Delegate type declarations must start with F");
+							topScope.TokenReader.LogError("Delegate type declarations must start with F");
 						}
-						Function.StrippedFunctionName = Function.SourceName.Substring(1);
-						Function.EngineName = $"{Function.StrippedFunctionName}{UhtFunction.GeneratedDelegateSignatureSuffix}";
+						function.StrippedFunctionName = function.SourceName.Substring(1);
+						function.EngineName = $"{function.StrippedFunctionName}{UhtFunction.GeneratedDelegateSignatureSuffix}";
 					}
 
-					SetFunctionNames(Function);
-					AddFunction(Function);
+					SetFunctionNames(function);
+					AddFunction(function);
 
 					// determine whether this function should be 'const'
-					if (bDeclaredConst)
+					if (declaredConst)
 					{
-						Function.FunctionFlags |= EFunctionFlags.Const;
-						Function.FunctionExportFlags |= UhtFunctionExportFlags.DeclaredConst;
+						function.FunctionFlags |= EFunctionFlags.Const;
+						function.FunctionExportFlags |= UhtFunctionExportFlags.DeclaredConst;
 					}
 
-					if (bIsSparse)
+					if (isSparse)
 					{
-						TopScope.TokenReader.Require(',');
-						UhtToken Name = TopScope.TokenReader.GetIdentifier("OwningClass specifier");
-						UhtEngineNameParts Parts = UhtUtilities.GetEngineNameParts(Name.Value);
-						Function.SparseOwningClassName = Parts.EngineName.ToString();
-						TopScope.TokenReader.Require(',');
-						Function.SparseDelegateName = TopScope.TokenReader.GetIdentifier("delegate name").Value.ToString();
+						topScope.TokenReader.Require(',');
+						UhtToken name = topScope.TokenReader.GetIdentifier("OwningClass specifier");
+						UhtEngineNameParts parts = UhtUtilities.GetEngineNameParts(name.Value);
+						function.SparseOwningClassName = parts.EngineName.ToString();
+						topScope.TokenReader.Require(',');
+						function.SparseDelegateName = topScope.TokenReader.GetIdentifier("delegate name").Value.ToString();
 					}
 
 					// Get parameter list.
-					if (FoundParamIndex >= 0)
+					if (foundParamIndex >= 0)
 					{
-						TopScope.TokenReader.Require(',');
+						topScope.TokenReader.Require(',');
 
-						ParseParameterList(TopScope, Function.GetPropertyParseOptions(false));
+						ParseParameterList(topScope, function.GetPropertyParseOptions(false));
 					}
 					else
 					{
 						// Require the closing paren even with no parameter list
-						TopScope.TokenReader.Require(')');
+						topScope.TokenReader.Require(')');
 					}
 
 					// Add back in the return value
-					if (ReturnValueProperty != null)
+					if (returnValueProperty != null)
 					{
-						TopScope.ScopeType.AddChild(ReturnValueProperty);
+						topScope.ScopeType.AddChild(returnValueProperty);
 					}
 
 					// Verify the number of parameters (FoundParamIndex = -1 means zero parameters, 0 means one, ...)
-					int ExpectedProperties = FoundParamIndex + 1 + (bHasReturnValue ? 1 : 0);
-					int PropertiesCount = Function.Properties.Count();
-					if (PropertiesCount != ExpectedProperties)
+					int expectedProperties = foundParamIndex + 1 + (hasReturnValue ? 1 : 0);
+					int propertiesCount = function.Properties.Count();
+					if (propertiesCount != expectedProperties)
 					{
-						throw new UhtException(TopScope.TokenReader, $"Expected {ExpectedProperties} parameters but found {PropertiesCount} parameters");
+						throw new UhtException(topScope.TokenReader, $"Expected {expectedProperties} parameters but found {propertiesCount} parameters");
 					}
 
 					// The macro line must be set here
-					Function.MacroLineNumber = TopScope.TokenReader.InputLine;
+					function.MacroLineNumber = topScope.TokenReader.InputLine;
 
 					// Try parsing metadata for the function
-					Specifiers.ParseFieldMetaData();
+					specifiers.ParseFieldMetaData();
 
-					TopScope.AddFormattedCommentsAsTooltipMetaData();
+					topScope.AddFormattedCommentsAsTooltipMetaData();
 
 					// Optionally consume a semicolon, it's not required for the delegate macro since it contains one internally
-					TopScope.TokenReader.Optional(';');
+					topScope.TokenReader.Optional(';');
 				}
 				return UhtParseResult.Handled;
 			}
 		}
 
-		private static UhtParseResult ParseUFunction(UhtParsingScope ParentScope, UhtToken Token)
+		private static UhtParseResult ParseUFunction(UhtParsingScope parentScope, UhtToken token)
 		{
-			UhtFunctionParser Function = new UhtFunctionParser(ParentScope.ScopeType, Token.InputLine);
+			UhtFunctionParser function = new UhtFunctionParser(parentScope.ScopeType, token.InputLine);
 
-			using (var TopScope = new UhtParsingScope(ParentScope, Function, ParentScope.Session.GetKeywordTable(UhtTableNames.Function), UhtAccessSpecifier.Public))
+			using (UhtParsingScope topScope = new UhtParsingScope(parentScope, function, parentScope.Session.GetKeywordTable(UhtTableNames.Function), UhtAccessSpecifier.Public))
 			{
-				UhtParsingScope OuterClassScope = TopScope.CurrentClassScope;
-				UhtClass OuterClass = (UhtClass)OuterClassScope.ScopeType;
-				string ScopeName = "function";
+				UhtParsingScope outerClassScope = topScope.CurrentClassScope;
+				UhtClass outerClass = (UhtClass)outerClassScope.ScopeType;
+				string scopeName = "function";
 
-				using (var TokenContext = new UhtMessageContext(ScopeName))
+				using (UhtMessageContext tokenContext = new UhtMessageContext(scopeName))
 				{
-					TopScope.AddModuleRelativePathToMetaData();
+					topScope.AddModuleRelativePathToMetaData();
 
-					UhtSpecifierContext SpecifierContext = new UhtSpecifierContext(TopScope, TopScope.TokenReader, Function.MetaData);
-					UhtSpecifierParser SpecifierParser = TopScope.HeaderParser.GetCachedSpecifierParser(SpecifierContext, ScopeName, ParentScope.Session.GetSpecifierTable(UhtTableNames.Function));
-					SpecifierParser.ParseSpecifiers();
+					UhtSpecifierContext specifierContext = new UhtSpecifierContext(topScope, topScope.TokenReader, function.MetaData);
+					UhtSpecifierParser specifierParser = topScope.HeaderParser.GetCachedSpecifierParser(specifierContext, scopeName, parentScope.Session.GetSpecifierTable(UhtTableNames.Function));
+					specifierParser.ParseSpecifiers();
 
-					if (!OuterClass.ClassFlags.HasAnyFlags(EClassFlags.Native))
+					if (!outerClass.ClassFlags.HasAnyFlags(EClassFlags.Native))
 					{
-						throw new UhtException(Function, "Should only be here for native classes!");
+						throw new UhtException(function, "Should only be here for native classes!");
 					}
 
-					Function.MacroLineNumber = TopScope.TokenReader.InputLine;
-					Function.FunctionFlags |= EFunctionFlags.Native;
+					function.MacroLineNumber = topScope.TokenReader.InputLine;
+					function.FunctionFlags |= EFunctionFlags.Native;
 
-					bool bAutomaticallyFinal = true;
-					switch (OuterClassScope.AccessSpecifier)
+					bool automaticallyFinal = true;
+					switch (outerClassScope.AccessSpecifier)
 					{
 						case UhtAccessSpecifier.Public:
-							Function.FunctionFlags |= EFunctionFlags.Public;
+							function.FunctionFlags |= EFunctionFlags.Public;
 							break;
 
 						case UhtAccessSpecifier.Protected:
-							Function.FunctionFlags |= EFunctionFlags.Protected;
+							function.FunctionFlags |= EFunctionFlags.Protected;
 							break;
 
 						case UhtAccessSpecifier.Private:
-							Function.FunctionFlags |= EFunctionFlags.Private | EFunctionFlags.Final;
+							function.FunctionFlags |= EFunctionFlags.Private | EFunctionFlags.Final;
 
 							// This is automatically final as well, but in a different way and for a different reason
-							bAutomaticallyFinal = false;
+							automaticallyFinal = false;
 							break;
 					}
 
-					if (TopScope.TokenReader.TryOptional("static"))
+					if (topScope.TokenReader.TryOptional("static"))
 					{
-						Function.FunctionFlags |= EFunctionFlags.Static;
-						Function.FunctionExportFlags |= UhtFunctionExportFlags.CppStatic;
+						function.FunctionFlags |= EFunctionFlags.Static;
+						function.FunctionExportFlags |= UhtFunctionExportFlags.CppStatic;
 					}
 
-					if (Function.MetaData.ContainsKey(UhtNames.CppFromBpEvent))
+					if (function.MetaData.ContainsKey(UhtNames.CppFromBpEvent))
 					{
-						Function.FunctionFlags |= EFunctionFlags.Event;
+						function.FunctionFlags |= EFunctionFlags.Event;
 					}
 
-					if ((TopScope.HeaderParser.GetCurrentCompositeCompilerDirective() & UhtCompilerDirective.WithEditor) != 0)
+					if ((topScope.HeaderParser.GetCurrentCompositeCompilerDirective() & UhtCompilerDirective.WithEditor) != 0)
 					{
-						Function.FunctionFlags |= EFunctionFlags.EditorOnly;
+						function.FunctionFlags |= EFunctionFlags.EditorOnly;
 					}
 
-					SpecifierParser.ParseDeferred();
-					FinalizeFunctionSpecifiers(Function);
+					specifierParser.ParseDeferred();
+					FinalizeFunctionSpecifiers(function);
 
-					if (Function.FunctionExportFlags.HasAnyFlags(UhtFunctionExportFlags.CustomThunk) && !Function.MetaData.ContainsKey(UhtNames.CustomThunk))
+					if (function.FunctionExportFlags.HasAnyFlags(UhtFunctionExportFlags.CustomThunk) && !function.MetaData.ContainsKey(UhtNames.CustomThunk))
 					{
-						Function.MetaData.Add(UhtNames.CustomThunk, true);
+						function.MetaData.Add(UhtNames.CustomThunk, true);
 					}
 
-					if (Function.FunctionFlags.HasAnyFlags(EFunctionFlags.Net))
+					if (function.FunctionFlags.HasAnyFlags(EFunctionFlags.Net))
 					{
 						// Network replicated functions are always events, and are only final if sealed
-						ScopeName = "event";
-						TokenContext.Reset(ScopeName);
-						bAutomaticallyFinal = false;
+						scopeName = "event";
+						tokenContext.Reset(scopeName);
+						automaticallyFinal = false;
 					}
 
-					if (Function.FunctionFlags.HasAnyFlags(EFunctionFlags.BlueprintEvent))
+					if (function.FunctionFlags.HasAnyFlags(EFunctionFlags.BlueprintEvent))
 					{
-						ScopeName = Function.FunctionFlags.HasAnyFlags(EFunctionFlags.Native) ? "BlueprintNativeEvent" : "BlueprintImplementableEvent";
-						TokenContext.Reset(ScopeName);
-						bAutomaticallyFinal = false;
+						scopeName = function.FunctionFlags.HasAnyFlags(EFunctionFlags.Native) ? "BlueprintNativeEvent" : "BlueprintImplementableEvent";
+						tokenContext.Reset(scopeName);
+						automaticallyFinal = false;
 					}
 
 					// Record the tokens so we can detect this function as a declaration later (i.e. RPC)
-					using (UhtTokenRecorder TokenRecorder = new UhtTokenRecorder(ParentScope, Function))
+					using (UhtTokenRecorder tokenRecorder = new UhtTokenRecorder(parentScope, function))
 					{
 
-						if (TopScope.TokenReader.TryOptional("virtual"))
+						if (topScope.TokenReader.TryOptional("virtual"))
 						{
-							Function.FunctionExportFlags |= UhtFunctionExportFlags.Virtual;
+							function.FunctionExportFlags |= UhtFunctionExportFlags.Virtual;
 						}
 
-						bool bInternalOnly = Function.MetaData.GetBoolean(UhtNames.BlueprintInternalUseOnly);
+						bool internalOnly = function.MetaData.GetBoolean(UhtNames.BlueprintInternalUseOnly);
 
 						// Peek ahead to look for a CORE_API style DLL import/export token if present
-						UhtToken APIMacroToken;
-						if (TopScope.TokenReader.TryOptionalAPIMacro(out APIMacroToken))
+						UhtToken apiMacroToken;
+						if (topScope.TokenReader.TryOptionalAPIMacro(out apiMacroToken))
 						{
 							//@TODO: Validate the module name for RequiredAPIMacroIfPresent
-							Function.FunctionFlags |= EFunctionFlags.RequiredAPI;
-							Function.FunctionExportFlags |= UhtFunctionExportFlags.RequiredAPI;
+							function.FunctionFlags |= EFunctionFlags.RequiredAPI;
+							function.FunctionExportFlags |= UhtFunctionExportFlags.RequiredAPI;
 						}
 
 						// Look for static again, in case there was an ENGINE_API token first
-						if (APIMacroToken && TopScope.TokenReader.TryOptional("static"))
+						if (apiMacroToken && topScope.TokenReader.TryOptional("static"))
 						{
-							TopScope.TokenReader.LogError($"Unexpected API macro '{APIMacroToken.Value}'. Did you mean to put '{APIMacroToken.Value}' after the static keyword?");
+							topScope.TokenReader.LogError($"Unexpected API macro '{apiMacroToken.Value}'. Did you mean to put '{apiMacroToken.Value}' after the static keyword?");
 						}
 
 						// Look for virtual again, in case there was an ENGINE_API token first
-						if (TopScope.TokenReader.TryOptional("virtual"))
+						if (topScope.TokenReader.TryOptional("virtual"))
 						{
-							Function.FunctionExportFlags |= UhtFunctionExportFlags.Virtual;
+							function.FunctionExportFlags |= UhtFunctionExportFlags.Virtual;
 						}
 
 						// If virtual, remove the implicit final, the user can still specifying an explicit final at the end of the declaration
-						if (Function.FunctionExportFlags.HasAnyFlags(UhtFunctionExportFlags.Virtual))
+						if (function.FunctionExportFlags.HasAnyFlags(UhtFunctionExportFlags.Virtual))
 						{
-							bAutomaticallyFinal = false;
+							automaticallyFinal = false;
 						}
 
 						// Handle the initial implicit/explicit final
 						// A user can still specify an explicit final after the parameter list as well.
-						if (bAutomaticallyFinal || Function.FunctionExportFlags.HasAnyFlags(UhtFunctionExportFlags.SealedEvent))
+						if (automaticallyFinal || function.FunctionExportFlags.HasAnyFlags(UhtFunctionExportFlags.SealedEvent))
 						{
-							Function.FunctionFlags |= EFunctionFlags.Final;
-							Function.FunctionExportFlags |= UhtFunctionExportFlags.Final | UhtFunctionExportFlags.AutoFinal;
+							function.FunctionFlags |= EFunctionFlags.Final;
+							function.FunctionExportFlags |= UhtFunctionExportFlags.Final | UhtFunctionExportFlags.AutoFinal;
 						}
 
 						// Get return type.  C++ style functions always have a return value type, even if it's void
-						UhtToken FuncNameToken = new UhtToken();
-						UhtProperty? ReturnValueProperty = null;
-						TopScope.HeaderParser.GetCachedPropertyParser(TopScope).Parse(EPropertyFlags.None,
-							Function.GetPropertyParseOptions(true), UhtParsePropertyDeclarationStyle.None, UhtPropertyCategory.Return,
-							(UhtParsingScope TopScope, UhtProperty Property, ref UhtToken NameToken, UhtLayoutMacroType LayoutMacroType) =>
+						UhtToken funcNameToken = new UhtToken();
+						UhtProperty? returnValueProperty = null;
+						topScope.HeaderParser.GetCachedPropertyParser(topScope).Parse(EPropertyFlags.None,
+							function.GetPropertyParseOptions(true), UhtParsePropertyDeclarationStyle.None, UhtPropertyCategory.Return,
+							(UhtParsingScope topScope, UhtProperty property, ref UhtToken nameToken, UhtLayoutMacroType layoutMacroType) =>
 							{
-								Property.PropertyFlags |= EPropertyFlags.Parm | EPropertyFlags.OutParm | EPropertyFlags.ReturnParm;
-								FuncNameToken = NameToken;
-								if (!(Property is UhtVoidProperty))
+								property.PropertyFlags |= EPropertyFlags.Parm | EPropertyFlags.OutParm | EPropertyFlags.ReturnParm;
+								funcNameToken = nameToken;
+								if (property is not UhtVoidProperty)
 								{
-									ReturnValueProperty = Property;
+									returnValueProperty = property;
 								}
 							});
 
-						if (FuncNameToken.Value.Length == 0)
+						if (funcNameToken.Value.Length == 0)
 						{
-							throw new UhtException(TopScope.TokenReader, "expected return value and function name");
+							throw new UhtException(topScope.TokenReader, "expected return value and function name");
 						}
 
 						// Get function or operator name.
-						Function.SourceName = FuncNameToken.Value.ToString();
+						function.SourceName = funcNameToken.Value.ToString();
 
-						ScopeName = $"{ScopeName} '{Function.SourceName}'";
-						TokenContext.Reset(ScopeName);
+						scopeName = $"{scopeName} '{function.SourceName}'";
+						tokenContext.Reset(scopeName);
 
-						TopScope.TokenReader.Require('(');
+						topScope.TokenReader.Require('(');
 
-						SetFunctionNames(Function);
-						AddFunction(Function);
+						SetFunctionNames(function);
+						AddFunction(function);
 
 						// Get parameter list.
-						ParseParameterList(TopScope, Function.GetPropertyParseOptions(false));
+						ParseParameterList(topScope, function.GetPropertyParseOptions(false));
 
 						// Add back in the return value
-						if (ReturnValueProperty != null)
+						if (returnValueProperty != null)
 						{
-							TopScope.ScopeType.AddChild(ReturnValueProperty);
+							topScope.ScopeType.AddChild(returnValueProperty);
 						}
 
 						// determine whether this function should be 'const'
-						if (TopScope.TokenReader.TryOptional("const"))
+						if (topScope.TokenReader.TryOptional("const"))
 						{
-							if (Function.FunctionFlags.HasAnyFlags(EFunctionFlags.Native))
+							if (function.FunctionFlags.HasAnyFlags(EFunctionFlags.Native))
 							{
 								// @TODO: UCREMOVAL Reconsider?
 								//Throwf(TEXT("'const' may only be used for native functions"));
 							}
 
-							Function.FunctionFlags |= EFunctionFlags.Const;
-							Function.FunctionExportFlags |= UhtFunctionExportFlags.DeclaredConst;
+							function.FunctionFlags |= EFunctionFlags.Const;
+							function.FunctionExportFlags |= UhtFunctionExportFlags.DeclaredConst;
 						}
 
 						// Try parsing metadata for the function
-						SpecifierParser.ParseFieldMetaData();
+						specifierParser.ParseFieldMetaData();
 
 						// COMPATIBILITY-TODO - Try to pull any comment following the declaration
-						TopScope.TokenReader.PeekToken();
-						TopScope.TokenReader.CommitPendingComments();
+						topScope.TokenReader.PeekToken();
+						topScope.TokenReader.CommitPendingComments();
 
-						TopScope.AddFormattedCommentsAsTooltipMetaData();
+						topScope.AddFormattedCommentsAsTooltipMetaData();
 
 						// 'final' and 'override' can appear in any order before an optional '= 0' pure virtual specifier
-						bool bFoundFinal = TopScope.TokenReader.TryOptional("final");
-						bool bFoundOverride = TopScope.TokenReader.TryOptional("override");
-						if (!bFoundFinal && bFoundOverride)
+						bool foundFinal = topScope.TokenReader.TryOptional("final");
+						bool foundOverride = topScope.TokenReader.TryOptional("override");
+						if (!foundFinal && foundOverride)
 						{
-							bFoundFinal = TopScope.TokenReader.TryOptional("final");
+							foundFinal = topScope.TokenReader.TryOptional("final");
 						}
 
 						// Handle C++ style functions being declared as abstract
-						if (TopScope.TokenReader.TryOptional('='))
+						if (topScope.TokenReader.TryOptional('='))
 						{
-							int ZeroValue = 1;
-							bool bGotZero = TopScope.TokenReader.TryOptionalConstInt(out ZeroValue);
-							bGotZero = bGotZero && (ZeroValue == 0);
-							if (!bGotZero || ZeroValue != 0)
+							int zeroValue = 1;
+							bool gotZero = topScope.TokenReader.TryOptionalConstInt(out zeroValue);
+							gotZero = gotZero && (zeroValue == 0);
+							if (!gotZero || zeroValue != 0)
 							{
-								throw new UhtException(TopScope.TokenReader, "Expected 0 to indicate function is abstract");
+								throw new UhtException(topScope.TokenReader, "Expected 0 to indicate function is abstract");
 							}
 						}
 
 						// Look for the final keyword to indicate this function is sealed
-						if (bFoundFinal)
+						if (foundFinal)
 						{
 							// This is a final (prebinding, non-overridable) function
-							Function.FunctionFlags |= EFunctionFlags.Final;
-							Function.FunctionExportFlags |= UhtFunctionExportFlags.Final;
+							function.FunctionFlags |= EFunctionFlags.Final;
+							function.FunctionExportFlags |= UhtFunctionExportFlags.Final;
 						}
 
 						// Optionally consume a semicolon
 						// This is optional to allow inline function definitions
-						if (TopScope.TokenReader.TryOptional(';'))
+						if (topScope.TokenReader.TryOptional(';'))
 						{
 							// Do nothing (consume it)
 						}
-						else if (TopScope.TokenReader.TryPeekOptional('{'))
+						else if (topScope.TokenReader.TryPeekOptional('{'))
 						{
 							// Skip inline function bodies
-							UhtToken TokenCopy = new UhtToken();
-							TopScope.TokenReader.SkipDeclaration(ref TokenCopy);
+							UhtToken tokenCopy = new UhtToken();
+							topScope.TokenReader.SkipDeclaration(ref tokenCopy);
 						}
 					}
 				}
@@ -620,162 +619,162 @@ namespace EpicGames.UHT.Parsers
 			}
 		}
 
-		private static void FinalizeFunctionSpecifiers(UhtFunction Function)
+		private static void FinalizeFunctionSpecifiers(UhtFunction function)
 		{
-			if (Function.FunctionFlags.HasAnyFlags(EFunctionFlags.Net))
+			if (function.FunctionFlags.HasAnyFlags(EFunctionFlags.Net))
 			{
 				// Network replicated functions are always events
-				Function.FunctionFlags |= EFunctionFlags.Event;
+				function.FunctionFlags |= EFunctionFlags.Event;
 			}
 		}
 
-		internal static bool IsValidateDelegateDeclaration(UhtToken Token)
+		internal static bool IsValidateDelegateDeclaration(UhtToken token)
 		{
-			return (Token.IsIdentifier() && Token.Value.Span.StartsWith("DECLARE_DYNAMIC_"));
+			return (token.IsIdentifier() && token.Value.Span.StartsWith("DECLARE_DYNAMIC_"));
 		}
 
-		private static void ParseParameterList(UhtParsingScope TopScope, UhtPropertyParseOptions Options)
+		private static void ParseParameterList(UhtParsingScope topScope, UhtPropertyParseOptions options)
 		{
-			UhtFunction Function = (UhtFunction)TopScope.ScopeType;
+			UhtFunction function = (UhtFunction)topScope.ScopeType;
 
-			bool bIsNetFunc = Function.FunctionFlags.HasAnyFlags(EFunctionFlags.Net);
-			UhtPropertyCategory PropertyCategory = bIsNetFunc ? UhtPropertyCategory.ReplicatedParameter : UhtPropertyCategory.RegularParameter;
-			EPropertyFlags DisallowFlags = ~(EPropertyFlags.ParmFlags | EPropertyFlags.AutoWeak | EPropertyFlags.RepSkip | EPropertyFlags.UObjectWrapper | EPropertyFlags.NativeAccessSpecifiers);
+			bool isNetFunc = function.FunctionFlags.HasAnyFlags(EFunctionFlags.Net);
+			UhtPropertyCategory propertyCategory = isNetFunc ? UhtPropertyCategory.ReplicatedParameter : UhtPropertyCategory.RegularParameter;
+			EPropertyFlags disallowFlags = ~(EPropertyFlags.ParmFlags | EPropertyFlags.AutoWeak | EPropertyFlags.RepSkip | EPropertyFlags.UObjectWrapper | EPropertyFlags.NativeAccessSpecifiers);
 
-			UhtAdvancedDisplayParameterHandler AdvancedDisplay = new UhtAdvancedDisplayParameterHandler(TopScope.ScopeType.MetaData);
+			UhtAdvancedDisplayParameterHandler advancedDisplay = new UhtAdvancedDisplayParameterHandler(topScope.ScopeType.MetaData);
 
-			TopScope.TokenReader.RequireList(')', ',', false, () =>
+			topScope.TokenReader.RequireList(')', ',', false, () =>
 			{
-				TopScope.HeaderParser.GetCachedPropertyParser(TopScope).Parse(DisallowFlags, Options, UhtParsePropertyDeclarationStyle.None, PropertyCategory,
-					(UhtParsingScope TopScope, UhtProperty Property, ref UhtToken NameToken, UhtLayoutMacroType LayoutMacroType) =>
+				topScope.HeaderParser.GetCachedPropertyParser(topScope).Parse(disallowFlags, options, UhtParsePropertyDeclarationStyle.None, propertyCategory,
+					(UhtParsingScope topScope, UhtProperty property, ref UhtToken nameToken, UhtLayoutMacroType layoutMacroType) =>
 					{
-						Property.PropertyFlags |= EPropertyFlags.Parm;
-						if (AdvancedDisplay.CanMarkMore() && AdvancedDisplay.ShouldMarkParameter(Property.EngineName))
+						property.PropertyFlags |= EPropertyFlags.Parm;
+						if (advancedDisplay.CanMarkMore() && advancedDisplay.ShouldMarkParameter(property.EngineName))
 						{
-							Property.PropertyFlags |= EPropertyFlags.AdvancedDisplay;
+							property.PropertyFlags |= EPropertyFlags.AdvancedDisplay;
 						}
 
 						// Default value.
-						if (TopScope.TokenReader.TryOptional('='))
+						if (topScope.TokenReader.TryOptional('='))
 						{
-							List<UhtToken> DefaultValueTokens = new List<UhtToken>();
-							int ParenthesisNestCount = 0;
-							while (!TopScope.TokenReader.bIsEOF)
+							List<UhtToken> defaultValueTokens = new List<UhtToken>();
+							int parenthesisNestCount = 0;
+							while (!topScope.TokenReader.IsEOF)
 							{
-								UhtToken Token = TopScope.TokenReader.PeekToken();
-								if (Token.IsSymbol(','))
+								UhtToken token = topScope.TokenReader.PeekToken();
+								if (token.IsSymbol(','))
 								{
-									if (ParenthesisNestCount == 0)
+									if (parenthesisNestCount == 0)
 									{
 										break;
 									}
-									DefaultValueTokens.Add(Token);
-									TopScope.TokenReader.ConsumeToken();
+									defaultValueTokens.Add(token);
+									topScope.TokenReader.ConsumeToken();
 								}
-								else if (Token.IsSymbol(')'))
+								else if (token.IsSymbol(')'))
 								{
-									if (ParenthesisNestCount == 0)
+									if (parenthesisNestCount == 0)
 									{
 										break;
 									}
-									DefaultValueTokens.Add(Token);
-									TopScope.TokenReader.ConsumeToken();
-									--ParenthesisNestCount;
+									defaultValueTokens.Add(token);
+									topScope.TokenReader.ConsumeToken();
+									--parenthesisNestCount;
 								}
-								else if (Token.IsSymbol('('))
+								else if (token.IsSymbol('('))
 								{
-									++ParenthesisNestCount;
-									DefaultValueTokens.Add(Token);
-									TopScope.TokenReader.ConsumeToken();
+									++parenthesisNestCount;
+									defaultValueTokens.Add(token);
+									topScope.TokenReader.ConsumeToken();
 								}
 								else
 								{
-									DefaultValueTokens.Add(Token);
-									TopScope.TokenReader.ConsumeToken();
+									defaultValueTokens.Add(token);
+									topScope.TokenReader.ConsumeToken();
 								}
 							}
 
 							// allow exec functions to be added to the metaData, this is so we can have default params for them.
-							bool bStoreCppDefaultValueInMetaData = Function.FunctionFlags.HasAnyFlags(EFunctionFlags.BlueprintCallable | EFunctionFlags.Exec);
-							if (DefaultValueTokens.Count > 0 && bStoreCppDefaultValueInMetaData)
+							bool storeCppDefaultValueInMetaData = function.FunctionFlags.HasAnyFlags(EFunctionFlags.BlueprintCallable | EFunctionFlags.Exec);
+							if (defaultValueTokens.Count > 0 && storeCppDefaultValueInMetaData)
 							{
-								Property.DefaultValueTokens = DefaultValueTokens;
+								property.DefaultValueTokens = defaultValueTokens;
 							}
 						}
 					});
 			});
 		}
 
-		private static void AddFunction(UhtFunction Function)
+		private static void AddFunction(UhtFunction function)
 		{
-			if (Function.Outer != null)
+			if (function.Outer != null)
 			{
-				Function.Outer.AddChild(Function);
+				function.Outer.AddChild(function);
 			}
 		}
 
-		private static void SetFunctionNames(UhtFunction Function)
+		private static void SetFunctionNames(UhtFunction function)
 		{
 			// The source name won't have the suffix applied to delegate names, however, the engine name will
 			// We use the engine name because we need to detect the suffix for delegates
-			string FunctionName = Function.EngineName;
-			if (FunctionName.EndsWith(UhtFunction.GeneratedDelegateSignatureSuffix, StringComparison.Ordinal))
+			string functionName = function.EngineName;
+			if (functionName.EndsWith(UhtFunction.GeneratedDelegateSignatureSuffix, StringComparison.Ordinal))
 			{
-				FunctionName = FunctionName.Substring(0, FunctionName.Length - UhtFunction.GeneratedDelegateSignatureSuffix.Length);
+				functionName = functionName.Substring(0, functionName.Length - UhtFunction.GeneratedDelegateSignatureSuffix.Length);
 			}
 
-			Function.UnMarshalAndCallName = "exec" + FunctionName;
+			function.UnMarshalAndCallName = "exec" + functionName;
 
-			if (Function.FunctionFlags.HasAnyFlags(EFunctionFlags.BlueprintEvent))
+			if (function.FunctionFlags.HasAnyFlags(EFunctionFlags.BlueprintEvent))
 			{
-				Function.MarshalAndCallName = FunctionName;
-				if (Function.FunctionFlags.HasAllFlags(EFunctionFlags.BlueprintEvent | EFunctionFlags.Native))
+				function.MarshalAndCallName = functionName;
+				if (function.FunctionFlags.HasAllFlags(EFunctionFlags.BlueprintEvent | EFunctionFlags.Native))
 				{
-					Function.CppImplName = Function.EngineName + "_Implementation";
+					function.CppImplName = function.EngineName + "_Implementation";
 				}
 			}
-			else if (Function.FunctionFlags.HasAllFlags(EFunctionFlags.Native | EFunctionFlags.Net))
+			else if (function.FunctionFlags.HasAllFlags(EFunctionFlags.Native | EFunctionFlags.Net))
 			{
-				Function.MarshalAndCallName = FunctionName;
-				if (Function.FunctionFlags.HasAnyFlags(EFunctionFlags.NetResponse))
+				function.MarshalAndCallName = functionName;
+				if (function.FunctionFlags.HasAnyFlags(EFunctionFlags.NetResponse))
 				{
 					// Response function implemented by programmer and called directly from thunk
-					Function.CppImplName = Function.EngineName;
+					function.CppImplName = function.EngineName;
 				}
 				else
 				{
-					if (Function.CppImplName.Length == 0)
+					if (function.CppImplName.Length == 0)
 					{
-						Function.CppImplName = Function.EngineName + "_Implementation";
+						function.CppImplName = function.EngineName + "_Implementation";
 					}
-					else if (Function.CppImplName == FunctionName)
+					else if (function.CppImplName == functionName)
 					{
-						Function.LogError("Native implementation function must be different than original function name.");
+						function.LogError("Native implementation function must be different than original function name.");
 					}
 
-					if (Function.CppValidationImplName.Length == 0 && Function.FunctionFlags.HasAnyFlags(EFunctionFlags.NetValidate))
+					if (function.CppValidationImplName.Length == 0 && function.FunctionFlags.HasAnyFlags(EFunctionFlags.NetValidate))
 					{
-						Function.CppValidationImplName = Function.EngineName + "_Validate";
+						function.CppValidationImplName = function.EngineName + "_Validate";
 					}
-					else if (Function.CppValidationImplName == FunctionName)
+					else if (function.CppValidationImplName == functionName)
 					{
-						Function.LogError("Validation function must be different than original function name.");
+						function.LogError("Validation function must be different than original function name.");
 					}
 				}
 			}
-			else if (Function.FunctionFlags.HasAnyFlags(EFunctionFlags.Delegate))
+			else if (function.FunctionFlags.HasAnyFlags(EFunctionFlags.Delegate))
 			{
-				Function.MarshalAndCallName = "delegate" + FunctionName;
+				function.MarshalAndCallName = "delegate" + functionName;
 			}
 
-			if (Function.CppImplName.Length == 0)
+			if (function.CppImplName.Length == 0)
 			{
-				Function.CppImplName = FunctionName;
+				function.CppImplName = functionName;
 			}
 
-			if (Function.MarshalAndCallName.Length == 0)
+			if (function.MarshalAndCallName.Length == 0)
 			{
-				Function.MarshalAndCallName = "event" + FunctionName;
+				function.MarshalAndCallName = "event" + functionName;
 			}
 		}
 	}
