More coding standard fixes

diff --git a/Engine/Source/Programs/Shared/EpicGames.UHT/Parsers/UhtHeaderFileParser.cs b/Engine/Source/Programs/Shared/EpicGames.UHT/Parsers/UhtHeaderFileParser.cs
--- a/Engine/Source/Programs/Shared/EpicGames.UHT/Parsers/UhtHeaderFileParser.cs
+++ b/Engine/Source/Programs/Shared/EpicGames.UHT/Parsers/UhtHeaderFileParser.cs
@@ -1,14 +1,14 @@
 // Copyright Epic Games, Inc. All Rights Reserved.
 
+using System;
+using System.Collections.Generic;
+using System.Diagnostics.CodeAnalysis;
+using System.IO;
 using EpicGames.Core;
 using EpicGames.UHT.Tables;
 using EpicGames.UHT.Tokenizer;
 using EpicGames.UHT.Types;
 using EpicGames.UHT.Utils;
-using System;
-using System.Collections.Generic;
-using System.Diagnostics.CodeAnalysis;
-using System.IO;
 
 namespace EpicGames.UHT.Parsers
 {
@@ -97,35 +97,35 @@ namespace EpicGames.UHT.Parsers
 		/// <summary>
 		/// Test to see if any of the specified flags are set
 		/// </summary>
-		/// <param name="InFlags">Current flags</param>
-		/// <param name="TestFlags">Flags to test for</param>
+		/// <param name="inFlags">Current flags</param>
+		/// <param name="testFlags">Flags to test for</param>
 		/// <returns>True if any of the flags are set</returns>
-		public static bool HasAnyFlags(this UhtCompilerDirective InFlags, UhtCompilerDirective TestFlags)
+		public static bool HasAnyFlags(this UhtCompilerDirective inFlags, UhtCompilerDirective testFlags)
 		{
-			return (InFlags & TestFlags) != 0;
+			return (inFlags & testFlags) != 0;
 		}
 
 		/// <summary>
 		/// Test to see if all of the specified flags are set
 		/// </summary>
-		/// <param name="InFlags">Current flags</param>
-		/// <param name="TestFlags">Flags to test for</param>
+		/// <param name="inFlags">Current flags</param>
+		/// <param name="testFlags">Flags to test for</param>
 		/// <returns>True if all the flags are set</returns>
-		public static bool HasAllFlags(this UhtCompilerDirective InFlags, UhtCompilerDirective TestFlags)
+		public static bool HasAllFlags(this UhtCompilerDirective inFlags, UhtCompilerDirective testFlags)
 		{
-			return (InFlags & TestFlags) == TestFlags;
+			return (inFlags & testFlags) == testFlags;
 		}
 
 		/// <summary>
 		/// Test to see if a specific set of flags have a specific value.
 		/// </summary>
-		/// <param name="InFlags">Current flags</param>
-		/// <param name="TestFlags">Flags to test for</param>
-		/// <param name="MatchFlags">Expected value of the tested flags</param>
+		/// <param name="inFlags">Current flags</param>
+		/// <param name="testFlags">Flags to test for</param>
+		/// <param name="matchFlags">Expected value of the tested flags</param>
 		/// <returns>True if the given flags have a specific value.</returns>
-		public static bool HasExactFlags(this UhtCompilerDirective InFlags, UhtCompilerDirective TestFlags, UhtCompilerDirective MatchFlags)
+		public static bool HasExactFlags(this UhtCompilerDirective inFlags, UhtCompilerDirective testFlags, UhtCompilerDirective matchFlags)
 		{
-			return (InFlags & TestFlags) == MatchFlags;
+			return (inFlags & testFlags) == matchFlags;
 		}
 	}
 
@@ -140,34 +140,34 @@ namespace EpicGames.UHT.Parsers
 		[UhtKeyword(Extends = UhtTableNames.ScriptStruct, Keyword = "public")]
 		[SuppressMessage("CodeQuality", "IDE0051:Remove unused private members", Justification = "Attribute accessed method")]
 		[SuppressMessage("Style", "IDE0060:Remove unused parameter", Justification = "Attribute accessed method")]
-		private static UhtParseResult PublicKeyword(UhtParsingScope TopScope, UhtParsingScope ActionScope, ref UhtToken Token)
+		private static UhtParseResult PublicKeyword(UhtParsingScope topScope, UhtParsingScope actionScope, ref UhtToken token)
 		{
-			return SetAccessSpecifier(TopScope, UhtAccessSpecifier.Public);
+			return SetAccessSpecifier(topScope, UhtAccessSpecifier.Public);
 		}
 
 		[UhtKeyword(Extends = UhtTableNames.ClassBase, Keyword = "protected")]
 		[UhtKeyword(Extends = UhtTableNames.ScriptStruct, Keyword = "protected")]
 		[SuppressMessage("CodeQuality", "IDE0051:Remove unused private members", Justification = "Attribute accessed method")]
 		[SuppressMessage("Style", "IDE0060:Remove unused parameter", Justification = "Attribute accessed method")]
-		private static UhtParseResult ProtectedKeyword(UhtParsingScope TopScope, UhtParsingScope ActionScope, ref UhtToken Token)
+		private static UhtParseResult ProtectedKeyword(UhtParsingScope topScope, UhtParsingScope actionScope, ref UhtToken token)
 		{
-			return SetAccessSpecifier(TopScope, UhtAccessSpecifier.Protected);
+			return SetAccessSpecifier(topScope, UhtAccessSpecifier.Protected);
 		}
 
 		[UhtKeyword(Extends = UhtTableNames.ClassBase, Keyword = "private")]
 		[UhtKeyword(Extends = UhtTableNames.ScriptStruct, Keyword = "private")]
 		[SuppressMessage("CodeQuality", "IDE0051:Remove unused private members", Justification = "Attribute accessed method")]
 		[SuppressMessage("Style", "IDE0060:Remove unused parameter", Justification = "Attribute accessed method")]
-		private static UhtParseResult PrivateKeyword(UhtParsingScope TopScope, UhtParsingScope ActionScope, ref UhtToken Token)
+		private static UhtParseResult PrivateKeyword(UhtParsingScope topScope, UhtParsingScope actionScope, ref UhtToken token)
 		{
-			return SetAccessSpecifier(TopScope, UhtAccessSpecifier.Private);
+			return SetAccessSpecifier(topScope, UhtAccessSpecifier.Private);
 		}
 		#endregion
 
-		private static UhtParseResult SetAccessSpecifier(UhtParsingScope TopScope, UhtAccessSpecifier AccessSpecifier)
+		private static UhtParseResult SetAccessSpecifier(UhtParsingScope topScope, UhtAccessSpecifier accessSpecifier)
 		{
-			TopScope.AccessSpecifier = AccessSpecifier;
-			TopScope.TokenReader.Require(':');
+			topScope.AccessSpecifier = accessSpecifier;
+			topScope.TokenReader.Require(':');
 			return UhtParseResult.Handled;
 		}
 	}
@@ -179,8 +179,8 @@ namespace EpicGames.UHT.Parsers
 	{
 		private struct CompilerDirective
 		{
-			public UhtCompilerDirective Element;
-			public UhtCompilerDirective Composite;
+			public UhtCompilerDirective _element;
+			public UhtCompilerDirective _composite;
 		}
 
 		/// <summary>
@@ -196,144 +196,144 @@ namespace EpicGames.UHT.Parsers
 		/// <summary>
 		/// If true, this header file belongs to the engine
 		/// </summary>
-		public bool bIsPartOfEngine => this.HeaderFile.Package.bIsPartOfEngine;
+		public bool IsPartOfEngine => this.HeaderFile.Package.IsPartOfEngine;
 
 		/// <summary>
 		/// If true, the inclusion of the generated header file was seen
 		/// </summary>
-		public bool bSpottedAutogeneratedHeaderInclude { get; set; } = false;
+		public bool SpottedAutogeneratedHeaderInclude { get; set; } = false;
 
 		/// <summary>
-		/// For a given header file, we share a common specifier parser to reduce the number of allocations.
-		/// Before the parser can be reused, the ParseDeferred method must be called to dispatch that list.
+		/// For a given header file, we share a common property parser to reduce the number of allocations.
 		/// </summary>
-		private UhtSpecifierParser? SpecifierParser { get; set; } = null;
+		public UhtPropertyParser? PropertyParser { get; set; } = null;
 
 		/// <summary>
-		/// For a given header file, we share a common property parser to reduce the number of allocations.
+		/// For a given header file, we share a common specifier parser to reduce the number of allocations.
+		/// Before the parser can be reused, the ParseDeferred method must be called to dispatch that list.
 		/// </summary>
-		public UhtPropertyParser? PropertyParser { get; set; } = null;
+		private UhtSpecifierParser? _specifierParser = null;
 
 		/// <summary>
 		/// Stack of current #if states
 		/// </summary>
-		private readonly List<CompilerDirective> CompilerDirectives = new List<CompilerDirective>();
+		private readonly List<CompilerDirective> _compilerDirectives = new List<CompilerDirective>();
 
 		/// <summary>
 		/// Stack of current #if states saved as part of the preprocessor state
 		/// </summary>
-		private readonly List<CompilerDirective> SavedCompilerDirectives = new List<CompilerDirective>();
+		private readonly List<CompilerDirective> _savedCompilerDirectives = new List<CompilerDirective>();
 
 		/// <summary>
 		/// Current top of the parsing scopes.  Classes, structures and functions all allocate scopes.
 		/// </summary>
-		private UhtParsingScope? TopScope = null;
+		private UhtParsingScope? _topScope = null;
 
 		/// <summary>
 		/// Parse the given header file
 		/// </summary>
-		/// <param name="HeaderFile">Header file to parse</param>
+		/// <param name="headerFile">Header file to parse</param>
 		/// <returns>Parser</returns>
-		public static UhtHeaderFileParser Parse(UhtHeaderFile HeaderFile)
+		public static UhtHeaderFileParser Parse(UhtHeaderFile headerFile)
 		{
-			UhtHeaderFileParser HeaderParser = new UhtHeaderFileParser(HeaderFile);
-			using (UhtParsingScope TopScope = new UhtParsingScope(HeaderParser, HeaderParser.HeaderFile, HeaderFile.Session.GetKeywordTable(UhtTableNames.Global)))
+			UhtHeaderFileParser headerParser = new UhtHeaderFileParser(headerFile);
+			using (UhtParsingScope topScope = new UhtParsingScope(headerParser, headerParser.HeaderFile, headerFile.Session.GetKeywordTable(UhtTableNames.Global)))
 			{
-				HeaderParser.ParseStatements();
+				headerParser.ParseStatements();
 
-				if (!HeaderParser.bSpottedAutogeneratedHeaderInclude && HeaderParser.HeaderFile.Data.Length > 0)
+				if (!headerParser.SpottedAutogeneratedHeaderInclude && headerParser.HeaderFile.Data.Length > 0)
 				{
-					bool bNoExportClassesOnly = true;
-					foreach (UhtType Type in HeaderParser.HeaderFile.Children)
+					bool noExportClassesOnly = true;
+					foreach (UhtType type in headerParser.HeaderFile.Children)
 					{
-						if (Type is UhtClass Class)
+						if (type is UhtClass classObj)
 						{
-							if (Class.ClassType != UhtClassType.NativeInterface && !Class.ClassExportFlags.HasAnyFlags(UhtClassExportFlags.NoExport))
+							if (classObj.ClassType != UhtClassType.NativeInterface && !classObj.ClassExportFlags.HasAnyFlags(UhtClassExportFlags.NoExport))
 							{
-								bNoExportClassesOnly = false;
+								noExportClassesOnly = false;
 								break;
 							}
 						}
 					}
 
-					if (!bNoExportClassesOnly)
+					if (!noExportClassesOnly)
 					{
-						HeaderParser.HeaderFile.LogError($"Expected an include at the top of the header the follows all other includes: '#include \"{HeaderParser.HeaderFile.GeneratedHeaderFileName}\"'");
+						headerParser.HeaderFile.LogError($"Expected an include at the top of the header the follows all other includes: '#include \"{headerParser.HeaderFile.GeneratedHeaderFileName}\"'");
 					}
 				}
 			}
-			return HeaderParser;
+			return headerParser;
 		}
 
-		private UhtHeaderFileParser(UhtHeaderFile HeaderFile)
+		private UhtHeaderFileParser(UhtHeaderFile headerFile)
 		{
-			this.TokenReader = new UhtTokenBufferReader(HeaderFile, HeaderFile.Data.Memory);
-			this.HeaderFile = HeaderFile;
+			this.TokenReader = new UhtTokenBufferReader(headerFile, headerFile.Data.Memory);
+			this.HeaderFile = headerFile;
 			this.TokenReader.TokenPreprocessor = this;
 		}
 
 		/// <summary>
 		/// Push a new scope
 		/// </summary>
-		/// <param name="Scope">Scope to push</param>
+		/// <param name="scope">Scope to push</param>
 		/// <exception cref="UhtIceException">Throw if the new scope isn't parented by the current scope</exception>
-		public void PushScope(UhtParsingScope Scope)
+		public void PushScope(UhtParsingScope scope)
 		{
-			if (Scope.ParentScope != this.TopScope)
+			if (scope.ParentScope != this._topScope)
 			{
 				throw new UhtIceException("Pushing a new scope whose parent isn't the current top scope.");
 			}
-			this.TopScope = Scope;
+			this._topScope = scope;
 		}
 
 		/// <summary>
 		/// Pop the given scope
 		/// </summary>
-		/// <param name="Scope">Scope to be popped</param>
+		/// <param name="scope">Scope to be popped</param>
 		/// <exception cref="UhtIceException">Thrown if the given scope isn't the top scope</exception>
-		public void PopScope(UhtParsingScope Scope)
+		public void PopScope(UhtParsingScope scope)
 		{
-			if (Scope != this.TopScope)
+			if (scope != this._topScope)
 			{
 				throw new UhtIceException("Attempt to pop a scope that isn't the top scope");
 			}
-			this.TopScope = Scope.ParentScope;
+			this._topScope = scope.ParentScope;
 		}
 
 		/// <summary>
 		/// Get the cached specifier parser
 		/// </summary>
-		/// <param name="SpecifierContext">Specifier context</param>
-		/// <param name="Context">User facing context</param>
-		/// <param name="Table">Specifier table</param>
+		/// <param name="specifierContext">Specifier context</param>
+		/// <param name="context">User facing context</param>
+		/// <param name="table">Specifier table</param>
 		/// <returns>Specifier parser</returns>
-		public UhtSpecifierParser GetCachedSpecifierParser(UhtSpecifierContext SpecifierContext, StringView Context, UhtSpecifierTable Table)
+		public UhtSpecifierParser GetCachedSpecifierParser(UhtSpecifierContext specifierContext, StringView context, UhtSpecifierTable table)
 		{
-			if (this.SpecifierParser == null)
+			if (this._specifierParser == null)
 			{
-				this.SpecifierParser = new UhtSpecifierParser(SpecifierContext, Context, Table);
+				this._specifierParser = new UhtSpecifierParser(specifierContext, context, table);
 			}
 			else
 			{
-				this.SpecifierParser.Reset(SpecifierContext, Context, Table);
+				this._specifierParser.Reset(specifierContext, context, table);
 			}
-			return this.SpecifierParser;
+			return this._specifierParser;
 		}
 
 		/// <summary>
 		/// Get the cached property parser
 		/// </summary>
-		/// <param name="TopScope">Top scope requesting parser</param>
+		/// <param name="topScope">Top scope requesting parser</param>
 		/// <returns>Property parser</returns>
-		public UhtPropertyParser GetCachedPropertyParser(UhtParsingScope TopScope)
+		public UhtPropertyParser GetCachedPropertyParser(UhtParsingScope topScope)
 		{
 			if (this.PropertyParser == null)
 			{
-				this.PropertyParser = new UhtPropertyParser(TopScope, TopScope.TokenReader);
+				this.PropertyParser = new UhtPropertyParser(topScope, topScope.TokenReader);
 			}
 			else
 			{
-				this.PropertyParser.Reset(TopScope, TopScope.TokenReader);
+				this.PropertyParser.Reset(topScope, topScope.TokenReader);
 			}
 			return this.PropertyParser;
 		}
@@ -344,7 +344,7 @@ namespace EpicGames.UHT.Parsers
 		/// <returns>Enumeration flags for all active compiler directives</returns>
 		public UhtCompilerDirective GetCurrentCompositeCompilerDirective()
 		{
-			return this.CompilerDirectives.Count > 0 ? this.CompilerDirectives[this.CompilerDirectives.Count - 1].Composite : UhtCompilerDirective.None;
+			return this._compilerDirectives.Count > 0 ? this._compilerDirectives[^1]._composite : UhtCompilerDirective.None;
 		}
 
 		/// <summary>
@@ -353,34 +353,34 @@ namespace EpicGames.UHT.Parsers
 		/// <returns>Current compiler directive</returns>
 		public UhtCompilerDirective GetCurrentNonCompositeCompilerDirective()
 		{
-			return this.CompilerDirectives.Count > 0 ? this.CompilerDirectives[this.CompilerDirectives.Count - 1].Element : UhtCompilerDirective.None;
+			return this._compilerDirectives.Count > 0 ? this._compilerDirectives[^1]._element : UhtCompilerDirective.None;
 		}
 
 		#region ITokenPreprocessor implementation
 		/// <inheritdoc/>
-		public bool ParsePreprocessorDirective(ref UhtToken Token, bool bIsBeingIncluded, out bool bClearComments, out bool bIllegalContentsCheck)
+		public bool ParsePreprocessorDirective(ref UhtToken token, bool isBeingIncluded, out bool clearComments, out bool illegalContentsCheck)
 		{
-			bClearComments = true;
-			if (ParseDirectiveInternal(bIsBeingIncluded))
+			clearComments = true;
+			if (ParseDirectiveInternal(isBeingIncluded))
 			{
-				bClearComments = ClearCommentsCompilerDirective();
+				clearComments = ClearCommentsCompilerDirective();
 			}
-			bIllegalContentsCheck = !GetCurrentNonCompositeCompilerDirective().HasAnyFlags(UhtCompilerDirective.ZeroBlock | UhtCompilerDirective.WithEditorOnlyData);
+			illegalContentsCheck = !GetCurrentNonCompositeCompilerDirective().HasAnyFlags(UhtCompilerDirective.ZeroBlock | UhtCompilerDirective.WithEditorOnlyData);
 			return IncludeCurrentCompilerDirective();
 		}
 
 		/// <inheritdoc/>
 		public void SaveState()
 		{
-			this.SavedCompilerDirectives.Clear();
-			this.SavedCompilerDirectives.AddRange(this.CompilerDirectives);
+			this._savedCompilerDirectives.Clear();
+			this._savedCompilerDirectives.AddRange(this._compilerDirectives);
 		}
 
 		/// <inheritdoc/>
 		public void RestoreState()
 		{
-			this.CompilerDirectives.Clear();
-			this.CompilerDirectives.AddRange(this.SavedCompilerDirectives);
+			this._compilerDirectives.Clear();
+			this._compilerDirectives.AddRange(this._savedCompilerDirectives);
 		}
 		#endregion
 
@@ -397,39 +397,39 @@ namespace EpicGames.UHT.Parsers
 		/// <summary>
 		/// Parse the statements between the given symbols
 		/// </summary>
-		/// <param name="Initiator">Starting symbol</param>
-		/// <param name="Terminator">Ending symbol</param>
-		/// <param name="bLogUnhandledKeywords">If true, log any unhandled keywords</param>
-		public void ParseStatements(char Initiator, char Terminator, bool bLogUnhandledKeywords)
+		/// <param name="initiator">Starting symbol</param>
+		/// <param name="terminator">Ending symbol</param>
+		/// <param name="logUnhandledKeywords">If true, log any unhandled keywords</param>
+		public void ParseStatements(char initiator, char terminator, bool logUnhandledKeywords)
 		{
-			if (this.TopScope == null)
+			if (this._topScope == null)
 			{
 				return;
 			}
 
-			if (Initiator != 0)
+			if (initiator != 0)
 			{
-				this.TokenReader.Require(Initiator);
+				this.TokenReader.Require(initiator);
 			}
 			while (true)
 			{
-				UhtToken Token = this.TokenReader.GetToken();
-				if (Token.TokenType.IsEndType())
+				UhtToken token = this.TokenReader.GetToken();
+				if (token.TokenType.IsEndType())
 				{
-					if (this.TopScope != null && this.TopScope.ParentScope == null)
+					if (this._topScope != null && this._topScope.ParentScope == null)
 					{
-						CheckEof(ref Token);
+						CheckEof(ref token);
 					}
 					return;
 				}
-				else if (Terminator != 0 && Token.IsSymbol(Terminator))
+				else if (terminator != 0 && token.IsSymbol(terminator))
 				{
 					return;
 				}
 
-				if (this.TopScope != null)
+				if (this._topScope != null)
 				{
-					ParseStatement(this.TopScope, ref Token, bLogUnhandledKeywords);
+					ParseStatement(this._topScope, ref token, logUnhandledKeywords);
 					this.TokenReader.ClearComments();
 				}
 			}
@@ -438,63 +438,63 @@ namespace EpicGames.UHT.Parsers
 		/// <summary>
 		/// Parse a statement
 		/// </summary>
-		/// <param name="TopScope">Current top scope</param>
-		/// <param name="Token">Token starting the statement</param>
-		/// <param name="bLogUnhandledKeywords">If true, log unhandled keywords</param>
+		/// <param name="topScope">Current top scope</param>
+		/// <param name="token">Token starting the statement</param>
+		/// <param name="logUnhandledKeywords">If true, log unhandled keywords</param>
 		/// <returns>Always returns true ATM</returns>
-		public static bool ParseStatement(UhtParsingScope TopScope, ref UhtToken Token, bool bLogUnhandledKeywords)
+		public static bool ParseStatement(UhtParsingScope topScope, ref UhtToken token, bool logUnhandledKeywords)
 		{
-			UhtParseResult ParseResult = UhtParseResult.Unhandled;
+			UhtParseResult parseResult = UhtParseResult.Unhandled;
 
-			switch (Token.TokenType)
+			switch (token.TokenType)
 			{
 				case UhtTokenType.Identifier:
-					ParseResult = DispatchKeyword(TopScope, ref Token);
+					parseResult = DispatchKeyword(topScope, ref token);
 					break;
 
 				case UhtTokenType.Symbol:
 					// Ignore any extra semicolons
-					if (Token.IsSymbol(';'))
+					if (token.IsSymbol(';'))
 					{
 						return true;
 					}
-					ParseResult = DispatchKeyword(TopScope, ref Token);
+					parseResult = DispatchKeyword(topScope, ref token);
 					break;
 			}
 
-			if (ParseResult == UhtParseResult.Unhandled)
+			if (parseResult == UhtParseResult.Unhandled)
 			{
-				ParseResult = DispatchCatchAll(TopScope, ref Token);
+				parseResult = DispatchCatchAll(topScope, ref token);
 			}
 
-			if (ParseResult == UhtParseResult.Unhandled)
+			if (parseResult == UhtParseResult.Unhandled)
 			{
-				ParseResult = ProbablyAnUnknownObjectLikeMacro(TopScope.TokenReader, ref Token);
+				parseResult = ProbablyAnUnknownObjectLikeMacro(topScope.TokenReader, ref token);
 			}
 
-			if (ParseResult == UhtParseResult.Unhandled && bLogUnhandledKeywords)
+			if (parseResult == UhtParseResult.Unhandled && logUnhandledKeywords)
 			{
-				TopScope.Session.LogUnhandledKeywordError(TopScope.TokenReader, Token);
+				topScope.Session.LogUnhandledKeywordError(topScope.TokenReader, token);
 			}
 
-			if (ParseResult == UhtParseResult.Unhandled || ParseResult == UhtParseResult.Invalid)
+			if (parseResult == UhtParseResult.Unhandled || parseResult == UhtParseResult.Invalid)
 			{
-				if (TopScope.ScopeType is UhtClass Class)
+				if (topScope.ScopeType is UhtClass classObj)
 				{
-					using (UhtTokenRecorder Recorder = new UhtTokenRecorder(TopScope, ref Token))
+					using (UhtTokenRecorder recorder = new UhtTokenRecorder(topScope, ref token))
 					{
-						TopScope.TokenReader.SkipDeclaration(ref Token);
-						if (Recorder.Stop())
+						topScope.TokenReader.SkipDeclaration(ref token);
+						if (recorder.Stop())
 						{
-							if (Class.Declarations != null)
+							if (classObj.Declarations != null)
 							{
-								UhtDeclaration Declaration = Class.Declarations[Class.Declarations.Count - 1];
-								if (TopScope.HeaderParser.CheckForConstructor(Class, Declaration))
+								UhtDeclaration declaration = classObj.Declarations[classObj.Declarations.Count - 1];
+								if (topScope.HeaderParser.CheckForConstructor(classObj, declaration))
 								{
 								}
-								else if (Class.ClassType == UhtClassType.Class)
+								else if (classObj.ClassType == UhtClassType.Class)
 								{
-									if (TopScope.HeaderParser.CheckForSerialize(Class, Declaration))
+									if (topScope.HeaderParser.CheckForSerialize(classObj, declaration))
 									{
 									}
 								}
@@ -504,39 +504,39 @@ namespace EpicGames.UHT.Parsers
 				}
 				else
 				{
-					TopScope.TokenReader.SkipDeclaration(ref Token);
+					topScope.TokenReader.SkipDeclaration(ref token);
 				}
 			}
 			return true;
 		}
 
-		private static UhtParseResult DispatchKeyword(UhtParsingScope TopScope, ref UhtToken Token)
+		private static UhtParseResult DispatchKeyword(UhtParsingScope topScope, ref UhtToken token)
 		{
-			UhtParseResult ParseResult = UhtParseResult.Unhandled;
-			for (UhtParsingScope? CurrentScope = TopScope; CurrentScope != null && ParseResult == UhtParseResult.Unhandled; CurrentScope = CurrentScope.ParentScope)
+			UhtParseResult parseResult = UhtParseResult.Unhandled;
+			for (UhtParsingScope? currentScope = topScope; currentScope != null && parseResult == UhtParseResult.Unhandled; currentScope = currentScope.ParentScope)
 			{
-				UhtKeyword KeywordInfo;
-				if (CurrentScope.ScopeKeywordTable.TryGetValue(Token.Value, out KeywordInfo))
+				UhtKeyword keywordInfo;
+				if (currentScope.ScopeKeywordTable.TryGetValue(token.Value, out keywordInfo))
 				{
-					if (KeywordInfo.bAllScopes || TopScope == CurrentScope)
+					if (keywordInfo.AllScopes || topScope == currentScope)
 					{
-						ParseResult = KeywordInfo.Delegate(TopScope, CurrentScope, ref Token);
+						parseResult = keywordInfo.Delegate(topScope, currentScope, ref token);
 					}
 				}
 			}
-			return ParseResult;
+			return parseResult;
 		}
 
-		private static UhtParseResult DispatchCatchAll(UhtParsingScope TopScope, ref UhtToken Token)
+		private static UhtParseResult DispatchCatchAll(UhtParsingScope topScope, ref UhtToken token)
 		{
-			for (UhtParsingScope? CurrentScope = TopScope; CurrentScope != null; CurrentScope = CurrentScope.ParentScope)
+			for (UhtParsingScope? currentScope = topScope; currentScope != null; currentScope = currentScope.ParentScope)
 			{
-				foreach (var Delegate in CurrentScope.ScopeKeywordTable.CatchAlls)
+				foreach (UhtKeywordCatchAllDelegate catchAll in currentScope.ScopeKeywordTable.CatchAlls)
 				{
-					UhtParseResult ParseResult = Delegate(TopScope, ref Token);
-					if (ParseResult != UhtParseResult.Unhandled)
+					UhtParseResult parseResult = catchAll(topScope, ref token);
+					if (parseResult != UhtParseResult.Unhandled)
 					{
-						return ParseResult;
+						return parseResult;
 					}
 				}
 			}
@@ -546,33 +546,33 @@ namespace EpicGames.UHT.Parsers
 		/// <summary>
 		/// Tests if an identifier looks like a macro which doesn't have a following open parenthesis.
 		/// </summary>
-		/// <param name="TokenReader">Token reader</param>
-		/// <param name="Token">The current token that initiated the process</param>
+		/// <param name="tokenReader">Token reader</param>
+		/// <param name="token">The current token that initiated the process</param>
 		/// <returns>Result if matching the token</returns>
-		private static UhtParseResult ProbablyAnUnknownObjectLikeMacro(IUhtTokenReader TokenReader, ref UhtToken Token)
+		private static UhtParseResult ProbablyAnUnknownObjectLikeMacro(IUhtTokenReader tokenReader, ref UhtToken token)
 		{
 			// Non-identifiers are not macros
-			if (!Token.IsIdentifier())
+			if (!token.IsIdentifier())
 			{
 				return UhtParseResult.Unhandled;
 			}
 
 			// Macros must start with a capitalized alphanumeric character or underscore
-			char FirstChar = Token.Value.Span[0];
-			if (FirstChar != '_' && (FirstChar < 'A' || FirstChar > 'Z'))
+			char firstChar = token.Value.Span[0];
+			if (firstChar != '_' && (firstChar < 'A' || firstChar > 'Z'))
 			{
 				return UhtParseResult.Unhandled;
 			}
 
 			// We'll guess about it being a macro based on it being fully-capitalized with at least one underscore.
-			int UnderscoreCount = 0;
-			foreach (char Ch in Token.Value.Span.Slice(1))
+			int underscoreCount = 0;
+			foreach (char ch in token.Value.Span.Slice(1))
 			{
-				if (Ch == '_')
+				if (ch == '_')
 				{
-					++UnderscoreCount;
+					++underscoreCount;
 				}
-				else if ((Ch < 'A' || Ch > 'Z') && (Ch < '0' || Ch > '9'))
+				else if ((ch < 'A' || ch > 'Z') && (ch < '0' || ch > '9'))
 				{
 					return UhtParseResult.Unhandled;
 				}
@@ -580,34 +580,34 @@ namespace EpicGames.UHT.Parsers
 
 			// We look for at least one underscore as a convenient way of allowing many known macros
 			// like FORCEINLINE and CONSTEXPR, and non-macros like FPOV and TCHAR.
-			if (UnderscoreCount == 0)
+			if (underscoreCount == 0)
 			{
 				return UhtParseResult.Unhandled;
 			}
 
 			// Identifiers which end in _API are known
-			if (Token.Value.Span.Length > 4 && Token.Value.Span.EndsWith("_API"))
+			if (token.Value.Span.Length > 4 && token.Value.Span.EndsWith("_API"))
 			{
 				return UhtParseResult.Unhandled;
 			}
 
 			// Ignore certain known macros or identifiers that look like macros.
-			if (Token.IsValue("FORCEINLINE_DEBUGGABLE") ||
-				Token.IsValue("FORCEINLINE_STATS") ||
-				Token.IsValue("SIZE_T"))
+			if (token.IsValue("FORCEINLINE_DEBUGGABLE") ||
+				token.IsValue("FORCEINLINE_STATS") ||
+				token.IsValue("SIZE_T"))
 			{
 				return UhtParseResult.Unhandled;
 			}
 
 			// Check if there's an open parenthesis following the token.
-			return TokenReader.PeekToken().IsSymbol('(') ? UhtParseResult.Unhandled : UhtParseResult.Handled;
+			return tokenReader.PeekToken().IsSymbol('(') ? UhtParseResult.Unhandled : UhtParseResult.Handled;
 		}
 
-		private void CheckEof(ref UhtToken Token)
+		private void CheckEof(ref UhtToken token)
 		{
-			if (this.CompilerDirectives.Count > 0)
+			if (this._compilerDirectives.Count > 0)
 			{
-				throw new UhtException(this.TokenReader, Token.InputLine, "Missing #endif");
+				throw new UhtException(this.TokenReader, token.InputLine, "Missing #endif");
 			}
 		}
 		#endregion
@@ -616,125 +616,125 @@ namespace EpicGames.UHT.Parsers
 		/// <summary>
 		/// Parse a preprocessor directive.
 		/// </summary>
-		/// <param name="bIsBeingIncluded">If true, then this directive is in an active block</param>
+		/// <param name="isBeingIncluded">If true, then this directive is in an active block</param>
 		/// <returns>True if we should check to see if tokenizer should clear comments</returns>
-		private bool ParseDirectiveInternal(bool bIsBeingIncluded)
+		private bool ParseDirectiveInternal(bool isBeingIncluded)
 		{
-			bool bCheckClearComments = false;
+			bool checkClearComments = false;
 
 			// Collect all the lines of the preprocessor statement including any continuations.
 			// We assume that the vast majority of lines will not be continuations.  So avoid using the
 			// string builder as much as possible.
-			int StartingLine = this.TokenReader.InputLine;
-			StringViewBuilder SVB = new StringViewBuilder();
+			int startingLine = this.TokenReader.InputLine;
+			StringViewBuilder builder = new StringViewBuilder();
 			while (true)
 			{
-				UhtToken LineToken = this.TokenReader.GetLine();
-				if (LineToken.TokenType != UhtTokenType.Line)
+				UhtToken lineToken = this.TokenReader.GetLine();
+				if (lineToken.TokenType != UhtTokenType.Line)
 				{
 					break;
 				}
-				if (LineToken.Value.Span.Length > 0 && LineToken.Value.Span[LineToken.Value.Span.Length - 1] == '\\')
+				if (lineToken.Value.Span.Length > 0 && lineToken.Value.Span[^1] == '\\')
 				{
-					SVB.Append(new StringView(LineToken.Value, 0, LineToken.Value.Span.Length - 1));
+					builder.Append(new StringView(lineToken.Value, 0, lineToken.Value.Span.Length - 1));
 				}
 				else
 				{
-					SVB.Append(LineToken.Value);
+					builder.Append(lineToken.Value);
 					break;
 				}
 			}
-			StringView Line = SVB.ToStringView();
+			StringView line = builder.ToStringView();
 
 			// Create a token reader we will use to decode
-			UhtTokenBufferReader LineTokenReader = new UhtTokenBufferReader(this.HeaderFile, Line.Memory);
-			LineTokenReader.InputLine = StartingLine;
+			UhtTokenBufferReader lineTokenReader = new UhtTokenBufferReader(this.HeaderFile, line.Memory);
+			lineTokenReader.InputLine = startingLine;
 
-			UhtToken Directive;
-			if (!LineTokenReader.TryOptionalIdentifier(out Directive))
+			UhtToken directive;
+			if (!lineTokenReader.TryOptionalIdentifier(out directive))
 			{
-				if (bIsBeingIncluded)
+				if (isBeingIncluded)
 				{
-					throw new UhtException(this.TokenReader, Directive.InputLine, "Missing compiler directive after '#'");
+					throw new UhtException(this.TokenReader, directive.InputLine, "Missing compiler directive after '#'");
 				}
-				return bCheckClearComments;
+				return checkClearComments;
 			}
 
-			if (Directive.IsValue("error"))
+			if (directive.IsValue("error"))
 			{
-				if (bIsBeingIncluded)
+				if (isBeingIncluded)
 				{
-					throw new UhtException(this.TokenReader, Directive.InputLine, "#error directive encountered");
+					throw new UhtException(this.TokenReader, directive.InputLine, "#error directive encountered");
 				}
 			}
-			else if (Directive.IsValue("pragma"))
+			else if (directive.IsValue("pragma"))
 			{
 				// Ignore all pragmas
 			}
-			else if (Directive.IsValue("linenumber"))
+			else if (directive.IsValue("linenumber"))
 			{
-				int NewInputLine;
-				if (!LineTokenReader.TryOptionalConstInt(out NewInputLine))
+				int newInputLine;
+				if (!lineTokenReader.TryOptionalConstInt(out newInputLine))
 				{
-					throw new UhtException(this.TokenReader, Directive.InputLine, "Missing line number in line number directive");
+					throw new UhtException(this.TokenReader, directive.InputLine, "Missing line number in line number directive");
 				}
-				this.TokenReader.InputLine = NewInputLine;
+				this.TokenReader.InputLine = newInputLine;
 			}
-			else if (Directive.IsValue("include"))
+			else if (directive.IsValue("include"))
 			{
-				if (bIsBeingIncluded)
+				if (isBeingIncluded)
 				{
-					UhtToken IncludeName = LineTokenReader.GetToken();
-					if (IncludeName.IsConstString())
+					UhtToken includeName = lineTokenReader.GetToken();
+					if (includeName.IsConstString())
 					{
-						StringView IncludeNameString = IncludeName.GetUnescapedString(this.HeaderFile);
-						if (this.bSpottedAutogeneratedHeaderInclude)
+						StringView includeNameString = includeName.GetUnescapedString(this.HeaderFile);
+						if (this.SpottedAutogeneratedHeaderInclude)
 						{
 							this.HeaderFile.LogError("#include found after .generated.h file - the .generated.h file should always be the last #include in a header");
 						}
-						if (this.HeaderFile.GeneratedHeaderFileName.AsSpan().Equals(IncludeNameString.Span, StringComparison.OrdinalIgnoreCase))
+						if (this.HeaderFile.GeneratedHeaderFileName.AsSpan().Equals(includeNameString.Span, StringComparison.OrdinalIgnoreCase))
 						{
-							this.bSpottedAutogeneratedHeaderInclude = true;
+							this.SpottedAutogeneratedHeaderInclude = true;
 						}
-						if (IncludeNameString.Span.Contains(".generated.h", StringComparison.Ordinal))
+						if (includeNameString.Span.Contains(".generated.h", StringComparison.Ordinal))
 						{
-							string RefName = Path.GetFileName(IncludeNameString.ToString()).Replace(".generated.h", ".h", StringComparison.Ordinal);
-							this.HeaderFile.AddReferencedHeader(RefName, true);
+							string refName = Path.GetFileName(includeNameString.ToString()).Replace(".generated.h", ".h", StringComparison.Ordinal);
+							this.HeaderFile.AddReferencedHeader(refName, true);
 						}
 						else
 						{
-							this.HeaderFile.AddReferencedHeader(IncludeNameString.ToString(), true);
+							this.HeaderFile.AddReferencedHeader(includeNameString.ToString(), true);
 						}
 					}
 				}
 			}
-			else if (Directive.IsValue("if"))
+			else if (directive.IsValue("if"))
 			{
-				bCheckClearComments = true;
-				PushCompilerDirective(ParserConditional(LineTokenReader));
+				checkClearComments = true;
+				PushCompilerDirective(ParserConditional(lineTokenReader));
 			}
-			else if (Directive.IsValue("ifdef") || Directive.IsValue("ifndef"))
+			else if (directive.IsValue("ifdef") || directive.IsValue("ifndef"))
 			{
-				bCheckClearComments = true;
+				checkClearComments = true;
 				PushCompilerDirective(UhtCompilerDirective.Unrecognized);
 			}
-			else if (Directive.IsValue("elif"))
+			else if (directive.IsValue("elif"))
 			{
-				bCheckClearComments = true;
-				UhtCompilerDirective OldCompilerDirective = PopCompilerDirective(Directive);
-				UhtCompilerDirective NewCompilerDirective = ParserConditional(LineTokenReader);
-				if (SupportsElif(OldCompilerDirective) != SupportsElif(NewCompilerDirective))
+				checkClearComments = true;
+				UhtCompilerDirective oldCompilerDirective = PopCompilerDirective(directive);
+				UhtCompilerDirective newCompilerDirective = ParserConditional(lineTokenReader);
+				if (SupportsElif(oldCompilerDirective) != SupportsElif(newCompilerDirective))
 				{
-					throw new UhtException(this.TokenReader, Directive.InputLine, 
-						$"Mixing {GetCompilerDirectiveText(OldCompilerDirective)} with {GetCompilerDirectiveText(NewCompilerDirective)} in an #elif preprocessor block is not supported");
+					throw new UhtException(this.TokenReader, directive.InputLine,
+						$"Mixing {GetCompilerDirectiveText(oldCompilerDirective)} with {GetCompilerDirectiveText(newCompilerDirective)} in an #elif preprocessor block is not supported");
 				}
-				PushCompilerDirective(NewCompilerDirective);
+				PushCompilerDirective(newCompilerDirective);
 			}
-			else if (Directive.IsValue("else"))
+			else if (directive.IsValue("else"))
 			{
-				bCheckClearComments = true;
-				UhtCompilerDirective OldCompilerDirective = PopCompilerDirective(Directive);
-				switch (OldCompilerDirective)
+				checkClearComments = true;
+				UhtCompilerDirective oldCompilerDirective = PopCompilerDirective(directive);
+				switch (oldCompilerDirective)
 				{
 					case UhtCompilerDirective.ZeroBlock:
 						PushCompilerDirective(UhtCompilerDirective.OneBlock);
@@ -749,40 +749,40 @@ namespace EpicGames.UHT.Parsers
 						PushCompilerDirective(UhtCompilerDirective.NotCPPBlock);
 						break;
 					case UhtCompilerDirective.WithHotReload:
-						throw new UhtException(this.TokenReader, Directive.InputLine, "Can not use WITH_HOT_RELOAD with an #else clause");
+						throw new UhtException(this.TokenReader, directive.InputLine, "Can not use WITH_HOT_RELOAD with an #else clause");
 					default:
-						PushCompilerDirective(OldCompilerDirective);
+						PushCompilerDirective(oldCompilerDirective);
 						break;
 				}
 			}
-			else if (Directive.IsValue("endif"))
+			else if (directive.IsValue("endif"))
 			{
-				PopCompilerDirective(Directive);
+				PopCompilerDirective(directive);
 			}
-			else if (Directive.IsValue("define"))
+			else if (directive.IsValue("define"))
 			{
 			}
-			else if (Directive.IsValue("undef"))
+			else if (directive.IsValue("undef"))
 			{
 			}
 			else
 			{
-				if (bIsBeingIncluded)
+				if (isBeingIncluded)
 				{
-					throw new UhtException(this.TokenReader, Directive.InputLine, $"Unrecognized compiler directive {Directive.Value}");
+					throw new UhtException(this.TokenReader, directive.InputLine, $"Unrecognized compiler directive {directive.Value}");
 				}
 			}
-			return bCheckClearComments;
+			return checkClearComments;
 		}
 
-		private UhtCompilerDirective ParserConditional(UhtTokenBufferReader LineTokenReader)
+		private UhtCompilerDirective ParserConditional(UhtTokenBufferReader lineTokenReader)
 		{
 			// Get any possible ! and the identifier
-			UhtToken Define = LineTokenReader.GetToken();
-			bool bNotPresent = Define.IsSymbol('!');
-			if (bNotPresent)
+			UhtToken define = lineTokenReader.GetToken();
+			bool notPresent = define.IsSymbol('!');
+			if (notPresent)
 			{
-				Define = LineTokenReader.GetToken();
+				define = lineTokenReader.GetToken();
 			}
 
 			//COMPATIBILITY-TODO
@@ -790,42 +790,42 @@ namespace EpicGames.UHT.Parsers
 			// Checking for this being the only token causes that to fail
 #if COMPATIBILITY_DISABLE
 			// Make sure there is nothing left
-			UhtToken End = LineTokenReader.GetToken();
-			if (!End.TokenType.IsEndType())
+			UhtToken end = LineTokenReader.GetToken();
+			if (!end.TokenType.IsEndType())
 			{
 				return UhtCompilerDirective.Unrecognized;
 			}
 #endif
 
-			switch (Define.TokenType)
+			switch (define.TokenType)
 			{
 				case UhtTokenType.DecimalConst:
-					if (Define.IsValue("0"))
+					if (define.IsValue("0"))
 					{
 						return UhtCompilerDirective.ZeroBlock;
 					}
-					else if (Define.IsValue("1"))
+					else if (define.IsValue("1"))
 					{
 						return UhtCompilerDirective.OneBlock;
 					}
 					break;
 
 				case UhtTokenType.Identifier:
-					if (Define.IsValue("WITH_EDITORONLY_DATA"))
+					if (define.IsValue("WITH_EDITORONLY_DATA"))
 					{
-						return bNotPresent ? UhtCompilerDirective.Unrecognized : UhtCompilerDirective.WithEditorOnlyData;
+						return notPresent ? UhtCompilerDirective.Unrecognized : UhtCompilerDirective.WithEditorOnlyData;
 					}
-					else if (Define.IsValue("WITH_EDITOR"))
+					else if (define.IsValue("WITH_EDITOR"))
 					{
-						return bNotPresent ? UhtCompilerDirective.Unrecognized : UhtCompilerDirective.WithEditor;
+						return notPresent ? UhtCompilerDirective.Unrecognized : UhtCompilerDirective.WithEditor;
 					}
-					else if (Define.IsValue("WITH_HOT_RELOAD"))
+					else if (define.IsValue("WITH_HOT_RELOAD"))
 					{
 						return UhtCompilerDirective.WithHotReload;
 					}
-					else if (Define.IsValue("CPP"))
+					else if (define.IsValue("CPP"))
 					{
-						return bNotPresent ? UhtCompilerDirective.NotCPPBlock : UhtCompilerDirective.CPPBlock;
+						return notPresent ? UhtCompilerDirective.NotCPPBlock : UhtCompilerDirective.CPPBlock;
 					}
 					break;
 
@@ -833,7 +833,7 @@ namespace EpicGames.UHT.Parsers
 				case UhtTokenType.EndOfDefault:
 				case UhtTokenType.EndOfType:
 				case UhtTokenType.EndOfDeclaration:
-					throw new UhtException(this.TokenReader, Define.InputLine, "#if with no expression");
+					throw new UhtException(this.TokenReader, define.InputLine, "#if with no expression");
 			}
 
 			return UhtCompilerDirective.Unrecognized;
@@ -842,32 +842,32 @@ namespace EpicGames.UHT.Parsers
 		/// <summary>
 		/// Add a new compiler directive to the stack
 		/// </summary>
-		/// <param name="CompilerDirective">Directive to be added</param>
-		private void PushCompilerDirective(UhtCompilerDirective CompilerDirective)
+		/// <param name="compilerDirective">Directive to be added</param>
+		private void PushCompilerDirective(UhtCompilerDirective compilerDirective)
 		{
-			CompilerDirective New = new CompilerDirective();
-			New.Element = CompilerDirective;
-			New.Composite = GetCurrentCompositeCompilerDirective() | CompilerDirective;
-			this.CompilerDirectives.Add(New);
+			CompilerDirective newCompileDirective = new CompilerDirective();
+			newCompileDirective._element = compilerDirective;
+			newCompileDirective._composite = GetCurrentCompositeCompilerDirective() | compilerDirective;
+			this._compilerDirectives.Add(newCompileDirective);
 		}
 
 		/// <summary>
 		/// Remove the top level compiler directive from the stack
 		/// </summary>
-		private UhtCompilerDirective PopCompilerDirective(UhtToken Token)
+		private UhtCompilerDirective PopCompilerDirective(UhtToken token)
 		{
-			if (this.CompilerDirectives.Count == 0)
+			if (this._compilerDirectives.Count == 0)
 			{
-				throw new UhtException(this.TokenReader, Token.InputLine, $"Unmatched '#{Token.Value}'");
+				throw new UhtException(this.TokenReader, token.InputLine, $"Unmatched '#{token.Value}'");
 			}
-			UhtCompilerDirective CompilerDirective = this.CompilerDirectives[this.CompilerDirectives.Count - 1].Element;
-			this.CompilerDirectives.RemoveAt(this.CompilerDirectives.Count - 1);
-			return CompilerDirective;
+			UhtCompilerDirective compilerDirective = this._compilerDirectives[^1]._element;
+			this._compilerDirectives.RemoveAt(this._compilerDirectives.Count - 1);
+			return compilerDirective;
 		}
 
 		private bool IncludeCurrentCompilerDirective()
 		{
-			if (this.CompilerDirectives.Count == 0)
+			if (this._compilerDirectives.Count == 0)
 			{
 				return true;
 			}
@@ -884,12 +884,12 @@ namespace EpicGames.UHT.Parsers
 		/// <exception cref="UhtIceException"></exception>
 		private bool ClearCommentsCompilerDirective()
 		{
-			if (this.CompilerDirectives.Count == 0)
+			if (this._compilerDirectives.Count == 0)
 			{
 				return true;
 			}
-			UhtCompilerDirective CompilerDirective = this.CompilerDirectives[this.CompilerDirectives.Count - 1].Element;
-			switch (CompilerDirective)
+			UhtCompilerDirective compilerDirective = this._compilerDirectives[^1]._element;
+			switch (compilerDirective)
 			{
 				case UhtCompilerDirective.CPPBlock:
 				case UhtCompilerDirective.NotCPPBlock:
@@ -908,17 +908,17 @@ namespace EpicGames.UHT.Parsers
 			}
 		}
 
-		private static bool SupportsElif(UhtCompilerDirective CompilerDirective)
+		private static bool SupportsElif(UhtCompilerDirective compilerDirective)
 		{
 			return
-				CompilerDirective == UhtCompilerDirective.WithEditor ||
-				CompilerDirective == UhtCompilerDirective.WithEditorOnlyData ||
-				CompilerDirective == UhtCompilerDirective.WithHotReload;
+				compilerDirective == UhtCompilerDirective.WithEditor ||
+				compilerDirective == UhtCompilerDirective.WithEditorOnlyData ||
+				compilerDirective == UhtCompilerDirective.WithHotReload;
 		}
 
-		private static string GetCompilerDirectiveText(UhtCompilerDirective CompilerDirective)
+		private static string GetCompilerDirectiveText(UhtCompilerDirective compilerDirective)
 		{
-			switch (CompilerDirective)
+			switch (compilerDirective)
 			{
 				case UhtCompilerDirective.CPPBlock: return "CPP";
 				case UhtCompilerDirective.NotCPPBlock: return "!CPP";
@@ -931,203 +931,203 @@ namespace EpicGames.UHT.Parsers
 			}
 		}
 
-		private static void SkipVirtualAndAPI(IUhtTokenReader ReplayReader)
+		private static void SkipVirtualAndAPI(IUhtTokenReader replayReader)
 		{
 			while (true)
 			{
-				UhtToken PeekToken = ReplayReader.PeekToken();
-				if (!PeekToken.IsValue("virtual") && !PeekToken.Value.Span.EndsWith("_API"))
+				UhtToken peekToken = replayReader.PeekToken();
+				if (!peekToken.IsValue("virtual") && !peekToken.Value.Span.EndsWith("_API"))
 				{
 					break;
 				}
-				ReplayReader.ConsumeToken();
+				replayReader.ConsumeToken();
 			}
 		}
 
-		private bool CheckForConstructor(UhtClass Class, UhtDeclaration Declaration)
+		private bool CheckForConstructor(UhtClass classObj, UhtDeclaration declaration)
 		{
-			IUhtTokenReader ReplayReader = UhtTokenReplayReader.GetThreadInstance(this.HeaderFile, this.HeaderFile.Data.Memory, new ReadOnlyMemory<UhtToken>(Declaration.Tokens), UhtTokenType.EndOfDeclaration);
+			IUhtTokenReader replayReader = UhtTokenReplayReader.GetThreadInstance(this.HeaderFile, this.HeaderFile.Data.Memory, new ReadOnlyMemory<UhtToken>(declaration.Tokens), UhtTokenType.EndOfDeclaration);
 
 			// Allow explicit constructors
 			{
-				bool bFoundExplicit = ReplayReader.TryOptional("explicit");
-				if (ReplayReader.PeekToken().Value.Span.EndsWith("_API"))
+				bool foundExplicit = replayReader.TryOptional("explicit");
+				if (replayReader.PeekToken().Value.Span.EndsWith("_API"))
 				{
-					ReplayReader.ConsumeToken();
-					if (!bFoundExplicit)
+					replayReader.ConsumeToken();
+					if (!foundExplicit)
 					{
-						ReplayReader.TryOptional("explicit");
+						replayReader.TryOptional("explicit");
 					}
 				}
 			}
 
-			if (!ReplayReader.TryOptional(Class.SourceName) ||
-				!ReplayReader.TryOptional('('))
+			if (!replayReader.TryOptional(classObj.SourceName) ||
+				!replayReader.TryOptional('('))
 			{
 				return false;
 			}
 
-			bool bOICtor = false;
-			bool bVTCtor = false;
+			bool oiCtor = false;
+			bool vtCtor = false;
 
-			if (!Class.ClassExportFlags.HasAnyFlags(UhtClassExportFlags.HasDefaultConstructor) && ReplayReader.TryOptional(')'))
+			if (!classObj.ClassExportFlags.HasAnyFlags(UhtClassExportFlags.HasDefaultConstructor) && replayReader.TryOptional(')'))
 			{
-				Class.ClassExportFlags |= UhtClassExportFlags.HasDefaultConstructor;
+				classObj.ClassExportFlags |= UhtClassExportFlags.HasDefaultConstructor;
 			}
-			else if (!Class.ClassExportFlags.HasAllFlags(UhtClassExportFlags.HasObjectInitializerConstructor | UhtClassExportFlags.HasCustomVTableHelperConstructor))
+			else if (!classObj.ClassExportFlags.HasAllFlags(UhtClassExportFlags.HasObjectInitializerConstructor | UhtClassExportFlags.HasCustomVTableHelperConstructor))
 			{
-				bool bIsConst = false;
-				bool bIsRef = false;
-				int ParenthesesNestingLevel = 1;
+				bool isConst = false;
+				bool isRef = false;
+				int parenthesesNestingLevel = 1;
 
-				while (ParenthesesNestingLevel != 0)
+				while (parenthesesNestingLevel != 0)
 				{
-					UhtToken Token = ReplayReader.GetToken();
-					if (!Token)
+					UhtToken token = replayReader.GetToken();
+					if (!token)
 					{
 						break;
 					}
 
 					// Template instantiation or additional parameter excludes ObjectInitializer constructor.
-					if (Token.IsValue(',') || Token.IsValue('<'))
+					if (token.IsValue(',') || token.IsValue('<'))
 					{
-						bOICtor = false;
-						bVTCtor = false;
+						oiCtor = false;
+						vtCtor = false;
 						break;
 					}
 
-					if (Token.IsValue('('))
+					if (token.IsValue('('))
 					{
-						ParenthesesNestingLevel++;
+						parenthesesNestingLevel++;
 						continue;
 					}
 
-					if (Token.IsValue(')'))
+					if (token.IsValue(')'))
 					{
-						ParenthesesNestingLevel--;
+						parenthesesNestingLevel--;
 						continue;
 					}
 
-					if (Token.IsValue("const"))
+					if (token.IsValue("const"))
 					{
-						bIsConst = true;
+						isConst = true;
 						continue;
 					}
 
-					if (Token.IsValue('&'))
+					if (token.IsValue('&'))
 					{
-						bIsRef = true;
+						isRef = true;
 						continue;
 					}
 
 					// FPostConstructInitializeProperties is deprecated, but left here, so it won't break legacy code.
-					if (Token.IsValue("FObjectInitializer") || Token.IsValue("FPostConstructInitializeProperties"))
+					if (token.IsValue("FObjectInitializer") || token.IsValue("FPostConstructInitializeProperties"))
 					{
-						bOICtor = true;
+						oiCtor = true;
 					}
 
-					if (Token.IsValue("FVTableHelper"))
+					if (token.IsValue("FVTableHelper"))
 					{
-						bVTCtor = true;
+						vtCtor = true;
 					}
 				}
 
 				// Parse until finish.
-				if (ParenthesesNestingLevel != 0)
+				if (parenthesesNestingLevel != 0)
 				{
-					ReplayReader.SkipBrackets('(', ')', ParenthesesNestingLevel);
+					replayReader.SkipBrackets('(', ')', parenthesesNestingLevel);
 				}
 
-				if (bOICtor && bIsRef && bIsConst)
+				if (oiCtor && isRef && isConst)
 				{
-					Class.ClassExportFlags |= UhtClassExportFlags.HasObjectInitializerConstructor;
-					Class.MetaData.Add(UhtNames.ObjectInitializerConstructorDeclared, "");
+					classObj.ClassExportFlags |= UhtClassExportFlags.HasObjectInitializerConstructor;
+					classObj.MetaData.Add(UhtNames.ObjectInitializerConstructorDeclared, "");
 				}
-				if (bVTCtor && bIsRef)
+				if (vtCtor && isRef)
 				{
-					Class.ClassExportFlags |= UhtClassExportFlags.HasCustomVTableHelperConstructor;
+					classObj.ClassExportFlags |= UhtClassExportFlags.HasCustomVTableHelperConstructor;
 				}
 			}
 
-			if (!bVTCtor)
+			if (!vtCtor)
 			{
-				Class.ClassExportFlags |= UhtClassExportFlags.HasConstructor;
+				classObj.ClassExportFlags |= UhtClassExportFlags.HasConstructor;
 			}
 
 			return false;
 		}
 
-		bool CheckForSerialize(UhtClass Class, UhtDeclaration Declaration)
+		bool CheckForSerialize(UhtClass classObj, UhtDeclaration declaration)
 		{
-			IUhtTokenReader ReplayReader = UhtTokenReplayReader.GetThreadInstance(this.HeaderFile, this.HeaderFile.Data.Memory, new ReadOnlyMemory<UhtToken>(Declaration.Tokens), UhtTokenType.EndOfDeclaration);
+			IUhtTokenReader replayReader = UhtTokenReplayReader.GetThreadInstance(this.HeaderFile, this.HeaderFile.Data.Memory, new ReadOnlyMemory<UhtToken>(declaration.Tokens), UhtTokenType.EndOfDeclaration);
 
-			SkipVirtualAndAPI(ReplayReader);
+			SkipVirtualAndAPI(replayReader);
 
-			if (!ReplayReader.TryOptional("void") ||
-				!ReplayReader.TryOptional("Serialize") ||
-				!ReplayReader.TryOptional('('))
+			if (!replayReader.TryOptional("void") ||
+				!replayReader.TryOptional("Serialize") ||
+				!replayReader.TryOptional('('))
 			{
 				return false;
 			}
 
-			UhtToken Token = ReplayReader.GetToken();
+			UhtToken token = replayReader.GetToken();
 
-			UhtSerializerArchiveType ArchiveType = UhtSerializerArchiveType.None;
-			if (Token.IsValue("FArchive"))
+			UhtSerializerArchiveType archiveType = UhtSerializerArchiveType.None;
+			if (token.IsValue("FArchive"))
 			{
-				if (ReplayReader.TryOptional('&'))
+				if (replayReader.TryOptional('&'))
 				{
 					// Allow the declaration to not define a name for the archive parameter
-					if (!ReplayReader.PeekToken().IsValue(')'))
+					if (!replayReader.PeekToken().IsValue(')'))
 					{
-						ReplayReader.SkipOne();
+						replayReader.SkipOne();
 					}
-					if (ReplayReader.TryOptional(')'))
+					if (replayReader.TryOptional(')'))
 					{
-						ArchiveType = UhtSerializerArchiveType.Archive;
+						archiveType = UhtSerializerArchiveType.Archive;
 					}
 				}
 			}
-			else if (Token.IsValue("FStructuredArchive"))
+			else if (token.IsValue("FStructuredArchive"))
 			{
-				if (ReplayReader.TryOptional("::") &&
-					ReplayReader.TryOptional("FRecord"))
+				if (replayReader.TryOptional("::") &&
+					replayReader.TryOptional("FRecord"))
 				{
 					// Allow the declaration to not define a name for the archive parameter
-					if (!ReplayReader.PeekToken().IsValue(')'))
+					if (!replayReader.PeekToken().IsValue(')'))
 					{
-						ReplayReader.SkipOne();
+						replayReader.SkipOne();
 					}
-					if (ReplayReader.TryOptional(')'))
+					if (replayReader.TryOptional(')'))
 					{
-						ArchiveType = UhtSerializerArchiveType.StructuredArchiveRecord;
+						archiveType = UhtSerializerArchiveType.StructuredArchiveRecord;
 					}
 				}
 			}
-			else if (Token.IsValue("FStructuredArchiveRecord"))
+			else if (token.IsValue("FStructuredArchiveRecord"))
 			{
 				// Allow the declaration to not define a name for the archive parameter
-				if (!ReplayReader.PeekToken().IsValue(')'))
+				if (!replayReader.PeekToken().IsValue(')'))
 				{
-					ReplayReader.SkipOne();
+					replayReader.SkipOne();
 				}
-				if (ReplayReader.TryOptional(')'))
+				if (replayReader.TryOptional(')'))
 				{
-					ArchiveType = UhtSerializerArchiveType.StructuredArchiveRecord;
+					archiveType = UhtSerializerArchiveType.StructuredArchiveRecord;
 				}
 			}
 
-			if (ArchiveType != UhtSerializerArchiveType.None)
+			if (archiveType != UhtSerializerArchiveType.None)
 			{
 				// Found what we want!
-				if (Declaration.CompilerDirectives == UhtCompilerDirective.None || Declaration.CompilerDirectives == UhtCompilerDirective.WithEditorOnlyData)
+				if (declaration.CompilerDirectives == UhtCompilerDirective.None || declaration.CompilerDirectives == UhtCompilerDirective.WithEditorOnlyData)
 				{
-					Class.SerializerArchiveType |= ArchiveType;
-					Class.EnclosingDefine = Declaration.CompilerDirectives == UhtCompilerDirective.None ? "" : "WITH_EDITORONLY_DATA";
+					classObj.SerializerArchiveType |= archiveType;
+					classObj.EnclosingDefine = declaration.CompilerDirectives == UhtCompilerDirective.None ? "" : "WITH_EDITORONLY_DATA";
 				}
 				else
 				{
-					Class.LogError("Serialize functions must not be inside preprocessor blocks, except for WITH_EDITORONLY_DATA");
+					classObj.LogError("Serialize functions must not be inside preprocessor blocks, except for WITH_EDITORONLY_DATA");
 				}
 				return true;
 			}
diff --git a/Engine/Source/Programs/Shared/EpicGames.UHT/Parsers/UhtInterfaceClassParser.cs b/Engine/Source/Programs/Shared/EpicGames.UHT/Parsers/UhtInterfaceClassParser.cs
--- a/Engine/Source/Programs/Shared/EpicGames.UHT/Parsers/UhtInterfaceClassParser.cs
+++ b/Engine/Source/Programs/Shared/EpicGames.UHT/Parsers/UhtInterfaceClassParser.cs
@@ -19,28 +19,28 @@ namespace EpicGames.UHT.Parsers
 		/// <summary>
 		/// Construct a new interface class parser
 		/// </summary>
-		/// <param name="Outer">Outer object</param>
-		/// <param name="LineNumber">Line number</param>
-		public UhtInterfaceClassParser(UhtType Outer, int LineNumber) : base(Outer, LineNumber)
+		/// <param name="outer">Outer object</param>
+		/// <param name="lineNumber">Line number</param>
+		public UhtInterfaceClassParser(UhtType outer, int lineNumber) : base(outer, lineNumber)
 		{
 		}
 
 		/// <inheritdoc/>
-		protected override void ResolveSuper(UhtResolvePhase ResolvePhase)
+		protected override void ResolveSuper(UhtResolvePhase resolvePhase)
 		{
-			base.ResolveSuper( ResolvePhase);
-			switch (ResolvePhase)
+			base.ResolveSuper(resolvePhase);
+			switch (resolvePhase)
 			{
 				case UhtResolvePhase.Bases:
-					UhtClass? SuperClass = this.SuperClass;
-					if (SuperClass != null)
+					UhtClass? superClass = this.SuperClass;
+					if (superClass != null)
 					{
-						this.ClassFlags |= SuperClass.ClassFlags & EClassFlags.ScriptInherit;
-						if (!SuperClass.ClassFlags.HasAnyFlags(EClassFlags.Native))
+						this.ClassFlags |= superClass.ClassFlags & EClassFlags.ScriptInherit;
+						if (!superClass.ClassFlags.HasAnyFlags(EClassFlags.Native))
 						{
 							throw new UhtException(this, $"Native classes cannot extend non-native classes");
 						}
-						this.ClassWithin = SuperClass.ClassWithin;
+						this.ClassWithin = superClass.ClassWithin;
 					}
 					else
 					{
@@ -51,133 +51,133 @@ namespace EpicGames.UHT.Parsers
 		}
 
 		/// <inheritdoc/>
-		protected override bool ResolveSelf(UhtResolvePhase ResolvePhase)
+		protected override bool ResolveSelf(UhtResolvePhase resolvePhase)
 		{
-			bool bResult = base.ResolveSelf(ResolvePhase);
+			bool result = base.ResolveSelf(resolvePhase);
 
-			switch (ResolvePhase)
+			switch (resolvePhase)
 			{
 				case UhtResolvePhase.InvalidCheck:
 					{
-						string NativeInterfaceName = "I" + this.EngineName;
-						UhtType? NativeInterface = this.Session.FindType(null, UhtFindOptions.SourceName | UhtFindOptions.Class, NativeInterfaceName);
-						if (NativeInterface == null)
+						string nativeInterfaceName = "I" + this.EngineName;
+						UhtType? nativeInterface = this.Session.FindType(null, UhtFindOptions.SourceName | UhtFindOptions.Class, nativeInterfaceName);
+						if (nativeInterface == null)
 						{
-							this.LogError($"UInterface '{this.SourceName}' parsed without a corresponding '{NativeInterfaceName}'");
+							this.LogError($"UInterface '{this.SourceName}' parsed without a corresponding '{nativeInterfaceName}'");
 						}
 						else
 						{
 							// Copy the children
-							this.AddChildren(NativeInterface.DetchChildren());
+							this.AddChildren(nativeInterface.DetachChildren());
 						}
 					}
 					break;
 			}
 
-			return bResult;
+			return result;
 		}
 		#region Keywords
 		[UhtKeyword(Extends = UhtTableNames.Global)]
 		[SuppressMessage("CodeQuality", "IDE0051:Remove unused private members", Justification = "Attribute accessed method")]
 		[SuppressMessage("Style", "IDE0060:Remove unused parameter", Justification = "Attribute accessed method")]
-		private static UhtParseResult UINTERFACEKeyword(UhtParsingScope TopScope, UhtParsingScope ActionScope, ref UhtToken Token)
+		private static UhtParseResult UINTERFACEKeyword(UhtParsingScope topScope, UhtParsingScope actionScope, ref UhtToken token)
 		{
-			return ParseUInterface(TopScope, ref Token);
+			return ParseUInterface(topScope, ref token);
 		}
 
 		[UhtKeyword(Extends = UhtTableNames.Interface)]
 		[UhtKeyword(Extends = UhtTableNames.Interface, Keyword = "GENERATED_BODY")]
 		[SuppressMessage("CodeQuality", "IDE0051:Remove unused private members", Justification = "Attribute accessed method")]
 		[SuppressMessage("Style", "IDE0060:Remove unused parameter", Justification = "Attribute accessed method")]
-		private static UhtParseResult GENERATED_UINTERFACE_BODYKeyword(UhtParsingScope TopScope, UhtParsingScope ActionScope, ref UhtToken Token)
+		private static UhtParseResult GENERATED_UINTERFACE_BODYKeyword(UhtParsingScope topScope, UhtParsingScope actionScope, ref UhtToken token)
 		{
-			UhtClass Class = (UhtClass)TopScope.ScopeType;
+			UhtClass classObj = (UhtClass)topScope.ScopeType;
 
-			UhtParserHelpers.ParseCompileVersionDeclaration(TopScope.TokenReader, TopScope.Session.Config!, Class);
+			UhtParserHelpers.ParseCompileVersionDeclaration(topScope.TokenReader, topScope.Session.Config!, classObj);
 
-			Class.GeneratedBodyAccessSpecifier = TopScope.AccessSpecifier;
-			Class.GeneratedBodyLineNumber = TopScope.TokenReader.InputLine;
+			classObj.GeneratedBodyAccessSpecifier = topScope.AccessSpecifier;
+			classObj.GeneratedBodyLineNumber = topScope.TokenReader.InputLine;
 
-			if (Token.IsValue("GENERATED_UINTERFACE_BODY"))
+			if (token.IsValue("GENERATED_UINTERFACE_BODY"))
 			{
-				TopScope.AccessSpecifier = UhtAccessSpecifier.Public;
+				topScope.AccessSpecifier = UhtAccessSpecifier.Public;
 			}
 			return UhtParseResult.Handled;
 		}
 		#endregion
 
-		private static UhtParseResult ParseUInterface(UhtParsingScope ParentScope, ref UhtToken Token)
+		private static UhtParseResult ParseUInterface(UhtParsingScope parentScope, ref UhtToken token)
 		{
-			UhtInterfaceClassParser Class = new UhtInterfaceClassParser(ParentScope.ScopeType, Token.InputLine);
-			Class.ClassType = UhtClassType.Interface;
+			UhtInterfaceClassParser classObj = new UhtInterfaceClassParser(parentScope.ScopeType, token.InputLine);
+			classObj.ClassType = UhtClassType.Interface;
 
-			using (var TopScope = new UhtParsingScope(ParentScope, Class, ParentScope.Session.GetKeywordTable(UhtTableNames.Interface), UhtAccessSpecifier.Private))
+			using (UhtParsingScope topScope = new UhtParsingScope(parentScope, classObj, parentScope.Session.GetKeywordTable(UhtTableNames.Interface), UhtAccessSpecifier.Private))
 			{
 				const string ScopeName = "interface";
 
-				using (var TokenContext = new UhtMessageContext(ScopeName))
+				using (UhtMessageContext tokenContext = new UhtMessageContext(ScopeName))
 				{
 					// Parse the specifiers
-					UhtSpecifierContext SpecifierContext = new UhtSpecifierContext(TopScope, TopScope.TokenReader, Class.MetaData);
-					UhtSpecifierParser Specifiers = TopScope.HeaderParser.GetCachedSpecifierParser(SpecifierContext, ScopeName, ParentScope.Session.GetSpecifierTable(UhtTableNames.Interface));
-					Specifiers.ParseSpecifiers();
-					Class.PrologLineNumber = TopScope.TokenReader.InputLine;
-					Class.ClassFlags |= EClassFlags.Native | EClassFlags.Interface | EClassFlags.Abstract;
-					Class.InternalObjectFlags |= EInternalObjectFlags.Native;
+					UhtSpecifierContext specifierContext = new UhtSpecifierContext(topScope, topScope.TokenReader, classObj.MetaData);
+					UhtSpecifierParser specifiers = topScope.HeaderParser.GetCachedSpecifierParser(specifierContext, ScopeName, parentScope.Session.GetSpecifierTable(UhtTableNames.Interface));
+					specifiers.ParseSpecifiers();
+					classObj.PrologLineNumber = topScope.TokenReader.InputLine;
+					classObj.ClassFlags |= EClassFlags.Native | EClassFlags.Interface | EClassFlags.Abstract;
+					classObj.InternalObjectFlags |= EInternalObjectFlags.Native;
 
 					// Consume the class specifier
-					TopScope.TokenReader.Require("class");
+					topScope.TokenReader.Require("class");
 
 					// Get the optional API macro
-					UhtToken APIMacroToken;
-					TopScope.TokenReader.TryOptionalAPIMacro(out APIMacroToken);
+					UhtToken apiMacroToken;
+					topScope.TokenReader.TryOptionalAPIMacro(out apiMacroToken);
 
 					// Get the name of the interface
-					Class.SourceName = TopScope.TokenReader.GetIdentifier().Value.ToString();
+					classObj.SourceName = topScope.TokenReader.GetIdentifier().Value.ToString();
 
 					// Parse the inheritance
-					TopScope.TokenReader.OptionalInheritance((ref UhtToken Identifier) =>
+					topScope.TokenReader.OptionalInheritance((ref UhtToken identifier) =>
 					{
-						Class.SuperIdentifier = Identifier;
+						classObj.SuperIdentifier = identifier;
 					});
 
 					// Split the source name into the different parts
-					UhtEngineNameParts NameParts = UhtUtilities.GetEngineNameParts(Class.SourceName);
-					Class.EngineName = NameParts.EngineName.ToString();
+					UhtEngineNameParts nameParts = UhtUtilities.GetEngineNameParts(classObj.SourceName);
+					classObj.EngineName = nameParts.EngineName.ToString();
 
 					// Interfaces must start with a valid prefix
-					if (NameParts.Prefix != "U")
+					if (nameParts.Prefix != "U")
 					{
-						TopScope.TokenReader.LogError($"Interface name '{Class.SourceName}' is invalid, the first class should be identified as 'U{NameParts.EngineName}'");
+						topScope.TokenReader.LogError($"Interface name '{classObj.SourceName}' is invalid, the first class should be identified as 'U{nameParts.EngineName}'");
 					}
 
 					// Check for an empty engine name
-					if (Class.EngineName.Length == 0)
+					if (classObj.EngineName.Length == 0)
 					{
-						TopScope.TokenReader.LogError($"When compiling class definition for '{Class.SourceName}', attempting to strip prefix results in an empty name. Did you leave off a prefix?");
+						topScope.TokenReader.LogError($"When compiling class definition for '{classObj.SourceName}', attempting to strip prefix results in an empty name. Did you leave off a prefix?");
 					}
 
-					if (APIMacroToken)
+					if (apiMacroToken)
 					{
-						Class.ClassFlags |= EClassFlags.RequiredAPI;
+						classObj.ClassFlags |= EClassFlags.RequiredAPI;
 					}
 
-					Specifiers.ParseDeferred();
+					specifiers.ParseDeferred();
 
-					if (Class.Outer != null)
+					if (classObj.Outer != null)
 					{
-						Class.Outer.AddChild(Class);
+						classObj.Outer.AddChild(classObj);
 					}
 
-					TopScope.AddModuleRelativePathToMetaData();
+					topScope.AddModuleRelativePathToMetaData();
 
-					TopScope.HeaderParser.ParseStatements('{', '}', true);
+					topScope.HeaderParser.ParseStatements('{', '}', true);
 
-					TopScope.TokenReader.Require(';');
+					topScope.TokenReader.Require(';');
 
-					if (Class.GeneratedBodyLineNumber == -1)
+					if (classObj.GeneratedBodyLineNumber == -1)
 					{
-						TopScope.TokenReader.LogError("Expected a GENERATED_BODY() at the start of the interface");
+						topScope.TokenReader.LogError("Expected a GENERATED_BODY() at the start of the interface");
 					}
 				}
 			}
diff --git a/Engine/Source/Programs/Shared/EpicGames.UHT/Parsers/UhtNativeInterfaceClassParser.cs b/Engine/Source/Programs/Shared/EpicGames.UHT/Parsers/UhtNativeInterfaceClassParser.cs
--- a/Engine/Source/Programs/Shared/EpicGames.UHT/Parsers/UhtNativeInterfaceClassParser.cs
+++ b/Engine/Source/Programs/Shared/EpicGames.UHT/Parsers/UhtNativeInterfaceClassParser.cs
@@ -19,40 +19,40 @@ namespace EpicGames.UHT.Parsers
 		/// <summary>
 		/// Construct a new native interface
 		/// </summary>
-		/// <param name="Outer">Outer object</param>
-		/// <param name="LineNumber">Line number of declaration</param>
-		public UhtParserNativeInterfaceClass(UhtType Outer, int LineNumber) : base(Outer, LineNumber)
+		/// <param name="outer">Outer object</param>
+		/// <param name="lineNumber">Line number of declaration</param>
+		public UhtParserNativeInterfaceClass(UhtType outer, int lineNumber) : base(outer, lineNumber)
 		{
 		}
 
 		/// <inheritdoc/>
-		protected override bool ResolveSelf(UhtResolvePhase ResolvePhase)
+		protected override bool ResolveSelf(UhtResolvePhase resolvePhase)
 		{
-			bool bResult = base.ResolveSelf(ResolvePhase);
+			bool result = base.ResolveSelf(resolvePhase);
 
-			switch (ResolvePhase)
+			switch (resolvePhase)
 			{
 				case UhtResolvePhase.InvalidCheck:
-					string InterfaceName = "U" + this.EngineName;
-					UhtClass? Interface = (UhtClass?) this.Session.FindType(null, UhtFindOptions.SourceName | UhtFindOptions.Class, InterfaceName);
-					if (Interface == null)
+					string interfaceName = "U" + this.EngineName;
+					UhtClass? interfaceObj = (UhtClass?)this.Session.FindType(null, UhtFindOptions.SourceName | UhtFindOptions.Class, interfaceName);
+					if (interfaceObj == null)
 					{
-						if (this.bHasGeneratedBody || this.Children.Count != 0)
+						if (this.HasGeneratedBody || this.Children.Count != 0)
 						{
-							this.LogError($"Native interface '{this.SourceName}' parsed without a corresponding '{InterfaceName}'");
+							this.LogError($"Native interface '{this.SourceName}' parsed without a corresponding '{interfaceName}'");
 						}
 						else
 						{
-							this.bVisibleType = false;
-							bResult = false;
+							this.VisibleType = false;
+							result = false;
 						}
 					}
 					else
 					{
-						this.AlternateObject = Interface;
-						Interface.NativeInterface = this;
+						this.AlternateObject = interfaceObj;
+						interfaceObj.NativeInterface = this;
 						//COMPATIBILITY-TODO - Use the native interface access specifier
-						Interface.GeneratedBodyAccessSpecifier = this.GeneratedBodyAccessSpecifier;
+						interfaceObj.GeneratedBodyAccessSpecifier = this.GeneratedBodyAccessSpecifier;
 
 						if (this.GeneratedBodyLineNumber == -1)
 						{
@@ -62,126 +62,126 @@ namespace EpicGames.UHT.Parsers
 					break;
 			}
 
-			return bResult;
+			return result;
 		}
 
 		#region Keywords
 		[UhtKeyword(Extends = UhtTableNames.Global, Keyword = "class", DisableUsageError = true)]
 		[SuppressMessage("CodeQuality", "IDE0051:Remove unused private members", Justification = "Attribute accessed method")]
 		[SuppressMessage("Style", "IDE0060:Remove unused parameter", Justification = "Attribute accessed method")]
-		private static UhtParseResult ClassKeyword(UhtParsingScope TopScope, UhtParsingScope ActionScope, ref UhtToken Token)
+		private static UhtParseResult ClassKeyword(UhtParsingScope topScope, UhtParsingScope actionScope, ref UhtToken token)
 		{
-			using (var SaveState = new UhtTokenSaveState(TopScope.TokenReader))
+			using (UhtTokenSaveState saveState = new UhtTokenSaveState(topScope.TokenReader))
 			{
-				UhtToken SourceName = new UhtToken();
-				UhtToken SuperName = new UhtToken();
-				if (TryParseIInterface(TopScope, out SourceName, out SuperName))
+				UhtToken sourceName = new UhtToken();
+				UhtToken superName = new UhtToken();
+				if (TryParseIInterface(topScope, out sourceName, out superName))
 				{
-					SaveState.AbandonState();
-					ParseIInterface(TopScope, ref Token, SourceName, SuperName);
+					saveState.AbandonState();
+					ParseIInterface(topScope, ref token, sourceName, superName);
 					return UhtParseResult.Handled;
 				}
 			}
-			return TopScope.TokenReader.SkipDeclaration(ref Token) ? UhtParseResult.Handled : UhtParseResult.Invalid;
+			return topScope.TokenReader.SkipDeclaration(ref token) ? UhtParseResult.Handled : UhtParseResult.Invalid;
 		}
 
 		[UhtKeyword(Extends = UhtTableNames.NativeInterface)]
 		[UhtKeyword(Extends = UhtTableNames.NativeInterface, Keyword = "GENERATED_BODY")]
 		[SuppressMessage("CodeQuality", "IDE0051:Remove unused private members", Justification = "Attribute accessed method")]
 		[SuppressMessage("Style", "IDE0060:Remove unused parameter", Justification = "Attribute accessed method")]
-		private static UhtParseResult GENERATED_IINTERFACE_BODYKeyword(UhtParsingScope TopScope, UhtParsingScope ActionScope, ref UhtToken Token)
+		private static UhtParseResult GENERATED_IINTERFACE_BODYKeyword(UhtParsingScope topScope, UhtParsingScope actionScope, ref UhtToken token)
 		{
-			UhtClass Class = (UhtClass)TopScope.ScopeType;
+			UhtClass classObj = (UhtClass)topScope.ScopeType;
 
-			UhtParserHelpers.ParseCompileVersionDeclaration(TopScope.TokenReader, TopScope.Session.Config!, Class);
+			UhtParserHelpers.ParseCompileVersionDeclaration(topScope.TokenReader, topScope.Session.Config!, classObj);
 
-			Class.GeneratedBodyAccessSpecifier = TopScope.AccessSpecifier;
-			Class.GeneratedBodyLineNumber = TopScope.TokenReader.InputLine;
-			Class.bHasGeneratedBody = true;
+			classObj.GeneratedBodyAccessSpecifier = topScope.AccessSpecifier;
+			classObj.GeneratedBodyLineNumber = topScope.TokenReader.InputLine;
+			classObj.HasGeneratedBody = true;
 
-			if (Token.IsValue("GENERATED_IINTERFACE_BODY"))
+			if (token.IsValue("GENERATED_IINTERFACE_BODY"))
 			{
-				TopScope.AccessSpecifier = UhtAccessSpecifier.Public;
+				topScope.AccessSpecifier = UhtAccessSpecifier.Public;
 			}
 			return UhtParseResult.Handled;
 		}
 		#endregion
 
-		private static bool TryParseIInterface(UhtParsingScope ParentScope, out UhtToken SourceName, out UhtToken SuperName)
+		private static bool TryParseIInterface(UhtParsingScope parentScope, out UhtToken sourceName, out UhtToken superName)
 		{
-			IUhtTokenReader TokenReader = ParentScope.TokenReader;
+			IUhtTokenReader tokenReader = parentScope.TokenReader;
 
 			// Get the optional API macro
-			TokenReader.TryOptionalAPIMacro(out UhtToken _);
+			tokenReader.TryOptionalAPIMacro(out UhtToken _);
 
 			// Get the name of the interface
-			SourceName = TokenReader.GetIdentifier();
+			sourceName = tokenReader.GetIdentifier();
 
 			// Old UHT would still parse the inheritance 
-			SuperName = new UhtToken();
-			if (TokenReader.TryOptional(':'))
+			superName = new UhtToken();
+			if (tokenReader.TryOptional(':'))
 			{
-				if (!TokenReader.TryOptional("public"))
+				if (!tokenReader.TryOptional("public"))
 				{
 					return false;
 				}
-				SuperName = TokenReader.GetIdentifier();
+				superName = tokenReader.GetIdentifier();
 			}
 
 			// Only process classes starting with 'I'
-			if (SourceName.Value.Span[0] != 'I')
+			if (sourceName.Value.Span[0] != 'I')
 			{
 				return false;
 			}
 
 			// If we end with a ';', then this is a forward declaration
-			if (TokenReader.TryOptional(';'))
+			if (tokenReader.TryOptional(';'))
 			{
 				return false;
 			}
 
 			// If we don't have a '{', then this is something else
-			if (!TokenReader.TryPeekOptional('{'))
+			if (!tokenReader.TryPeekOptional('{'))
 			{
 				return false;
 			}
 			return true;
 		}
 
-		private static void ParseIInterface(UhtParsingScope ParentScope, ref UhtToken Token, UhtToken SourceName, UhtToken SuperName)
+		private static void ParseIInterface(UhtParsingScope parentScope, ref UhtToken token, UhtToken sourceName, UhtToken superName)
 		{
-			IUhtTokenReader TokenReader = ParentScope.TokenReader;
+			IUhtTokenReader tokenReader = parentScope.TokenReader;
 
-			UhtParserNativeInterfaceClass Class = new UhtParserNativeInterfaceClass(ParentScope.ScopeType, Token.InputLine);
-			Class.ClassType = UhtClassType.NativeInterface;
-			Class.SourceName = SourceName.Value.ToString();
-			Class.ClassFlags |= EClassFlags.Native | EClassFlags.Interface;
+			UhtParserNativeInterfaceClass classObj = new UhtParserNativeInterfaceClass(parentScope.ScopeType, token.InputLine);
+			classObj.ClassType = UhtClassType.NativeInterface;
+			classObj.SourceName = sourceName.Value.ToString();
+			classObj.ClassFlags |= EClassFlags.Native | EClassFlags.Interface;
 
 			// Split the source name into the different parts
-			UhtEngineNameParts NameParts = UhtUtilities.GetEngineNameParts(Class.SourceName);
-			Class.EngineName = NameParts.EngineName.ToString();
+			UhtEngineNameParts nameParts = UhtUtilities.GetEngineNameParts(classObj.SourceName);
+			classObj.EngineName = nameParts.EngineName.ToString();
 
 			// Check for an empty engine name
-			if (Class.EngineName.Length == 0)
+			if (classObj.EngineName.Length == 0)
 			{
-				TokenReader.LogError($"When compiling class definition for '{Class.SourceName}', attempting to strip prefix results in an empty name. Did you leave off a prefix?");
+				tokenReader.LogError($"When compiling class definition for '{classObj.SourceName}', attempting to strip prefix results in an empty name. Did you leave off a prefix?");
 			}
 
-			if (Class.Outer != null)
+			if (classObj.Outer != null)
 			{
-				Class.Outer.AddChild(Class);
+				classObj.Outer.AddChild(classObj);
 			}
 
-			Class.SuperIdentifier = SuperName;
+			classObj.SuperIdentifier = superName;
 
-			using (var TopScope = new UhtParsingScope(ParentScope, Class, ParentScope.Session.GetKeywordTable(UhtTableNames.NativeInterface), UhtAccessSpecifier.Private))
+			using (UhtParsingScope topScope = new UhtParsingScope(parentScope, classObj, parentScope.Session.GetKeywordTable(UhtTableNames.NativeInterface), UhtAccessSpecifier.Private))
 			{
 				const string ScopeName = "native interface";
 
-				using (var TokenContext = new UhtMessageContext(ScopeName))
+				using (UhtMessageContext tokenContext = new UhtMessageContext(ScopeName))
 				{
-					TopScope.HeaderParser.ParseStatements('{', '}', false);
-					TokenReader.Require(';');
+					topScope.HeaderParser.ParseStatements('{', '}', false);
+					tokenReader.Require(';');
 				}
 			}
 			return;
diff --git a/Engine/Source/Programs/Shared/EpicGames.UHT/Parsers/UhtParserHelpers.cs b/Engine/Source/Programs/Shared/EpicGames.UHT/Parsers/UhtParserHelpers.cs
--- a/Engine/Source/Programs/Shared/EpicGames.UHT/Parsers/UhtParserHelpers.cs
+++ b/Engine/Source/Programs/Shared/EpicGames.UHT/Parsers/UhtParserHelpers.cs
@@ -1,11 +1,11 @@
 // Copyright Epic Games, Inc. All Rights Reserved.
 
+using System;
+using System.Collections.Generic;
 using EpicGames.Core;
 using EpicGames.UHT.Tokenizer;
 using EpicGames.UHT.Types;
 using EpicGames.UHT.Utils;
-using System;
-using System.Collections.Generic;
 
 namespace EpicGames.UHT.Parsers
 {
@@ -18,62 +18,62 @@ namespace EpicGames.UHT.Parsers
 		/// <summary>
 		/// Parse the inheritance 
 		/// </summary>
-		/// <param name="TokenReader">Token reader</param>
-		/// <param name="Config">Configuration</param>
-		/// <param name="SuperIdentifier">Output super identifier</param>
-		/// <param name="BaseIdentifiers">Output base identifiers</param>
-		public static void ParseInheritance(IUhtTokenReader TokenReader, IUhtConfig Config, out UhtToken SuperIdentifier, out List<UhtToken[]>? BaseIdentifiers)
+		/// <param name="tokenReader">Token reader</param>
+		/// <param name="config">Configuration</param>
+		/// <param name="superIdentifier">Output super identifier</param>
+		/// <param name="baseIdentifiers">Output base identifiers</param>
+		public static void ParseInheritance(IUhtTokenReader tokenReader, IUhtConfig config, out UhtToken superIdentifier, out List<UhtToken[]>? baseIdentifiers)
 		{
-			UhtToken SuperIdentifierTemp = new UhtToken();
-			List<UhtToken[]>? BaseIdentifiersTemp = null;
-			TokenReader.OptionalInheritance(
-				(ref UhtToken Identifier) =>
+			UhtToken superIdentifierTemp = new UhtToken();
+			List<UhtToken[]>? baseIdentifiersTemp = null;
+			tokenReader.OptionalInheritance(
+				(ref UhtToken identifier) =>
 				{
-					Config.RedirectTypeIdentifier(ref Identifier);
-					SuperIdentifierTemp = Identifier;
+					config.RedirectTypeIdentifier(ref identifier);
+					superIdentifierTemp = identifier;
 				},
-				(UhtTokenList Identifier) =>
+				(UhtTokenList identifier) =>
 				{
-					if (BaseIdentifiersTemp == null)
+					if (baseIdentifiersTemp == null)
 					{
-						BaseIdentifiersTemp = new List<UhtToken[]>();
+						baseIdentifiersTemp = new List<UhtToken[]>();
 					}
-					BaseIdentifiersTemp.Add(Identifier.ToArray());
+					baseIdentifiersTemp.Add(identifier.ToArray());
 				});
-			SuperIdentifier = SuperIdentifierTemp;
-			BaseIdentifiers = BaseIdentifiersTemp;
+			superIdentifier = superIdentifierTemp;
+			baseIdentifiers = baseIdentifiersTemp;
 		}
 
 		/// <summary>
 		/// Parse compiler version declaration
 		/// </summary>
-		/// <param name="TokenReader">Token reader</param>
-		/// <param name="Config">Configuration</param>
-		/// <param name="Struct">Struct being parsed</param>
-		public static void ParseCompileVersionDeclaration(IUhtTokenReader TokenReader, IUhtConfig Config, UhtStruct Struct)
+		/// <param name="tokenReader">Token reader</param>
+		/// <param name="config">Configuration</param>
+		/// <param name="structObj">Struct being parsed</param>
+		public static void ParseCompileVersionDeclaration(IUhtTokenReader tokenReader, IUhtConfig config, UhtStruct structObj)
 		{
 
 			// Fetch the default generation code version. If supplied, then package code version overrides the default.
-			EGeneratedCodeVersion Version = Struct.Package.Module.GeneratedCodeVersion;
-			if (Version == EGeneratedCodeVersion.None)
+			EGeneratedCodeVersion version = structObj.Package.Module.GeneratedCodeVersion;
+			if (version == EGeneratedCodeVersion.None)
 			{
-				Version = Config.DefaultGeneratedCodeVersion;
+				version = config.DefaultGeneratedCodeVersion;
 			}
 
 			// Fetch the code version from header file
-			TokenReader
+			tokenReader
 				.Require('(')
-				.OptionalIdentifier((ref UhtToken Identifier) =>
+				.OptionalIdentifier((ref UhtToken identifier) =>
 				{
-					if (!Enum.TryParse(Identifier.Value.ToString(), true, out Version))
+					if (!Enum.TryParse(identifier.Value.ToString(), true, out version))
 					{
-						Version = EGeneratedCodeVersion.None;
+						version = EGeneratedCodeVersion.None;
 					}
 				})
 				.Require(')');
 
 			// Save the results
-			Struct.GeneratedCodeVersion = Version;
+			structObj.GeneratedCodeVersion = version;
 		}
 	}
 }
